Merging UE4-Streaming to UE4

- Linkers are no longer UObjects. Renamed ULinker, ULinkerLoad and ULinkerSave to FLinker, FLinkerLoad, FLinkerSave respectively
- Linkers are now associated with their UPackages
- Linker version is now stored in UPackages
- Async loading is now performed on a separate thread (if platform supports it and only in cooked builds), with the exception of PostLoad which is still done on the game thread
- Added UObject::IsPostLoadThreadSafe() function to determine if PostLoad is thread safe and can be executed on the async loading thread (defaults to false)
- UObject creation is now thread safe and can be performed on any thread
- Move many of the linker/UObject globals into FUObjectThreadContext (TLS)
- GetAsyncLoadPercentage() now takes PostLoad into account
- More async loading stats
- Added AtomicallySetFlags/ClearFlags to UObject
- Made FModuleManager thread safe.
- Added FGCScopeGuard as means of preventing GC from executing from non-game thread
- It's possible to disable async loading thread through ini settings.
- Cancelling async loading will now also trigger GC
- Implemented a basic version of async streaming priorities.

Change 2410813 by [REDACTED_PERSON]:
	Change Sleep in while loop to ConditionalSleep in FMultiReaderSingleWriterGT
Change 2410734 by [REDACTED_PERSON]:
	Make FModuleManager thread safe.
Change 2399879 by [REDACTED_PERSON]:
	Basic version of async streaming priorities.
Change 2410707 by [REDACTED_PERSON]:
	Implement conditional and no stat versions of sleep.
Change 2371939 by [REDACTED_PERSON]:
	Async Loading Improvements: adding more stats (accumulators)
Change 2372403 by [REDACTED_PERSON]:
	Fixing compile errors when STATs are not enabled
Change 2371526 by [REDACTED_PERSON]:
	AsyncLoading Improvements (WIP)
Change 2407198 by [REDACTED_PERSON]:
	Re-implementing delegate fixes for Async Loading
Change 2407425 by [REDACTED_PERSON]:
	Re-implementing cancelling async loading in the async loading branch.
Change 2484362 by [REDACTED_PERSON]:
	Making it possible to disable async loading thread through ini settings.
Change 2484744 by [REDACTED_PERSON]:
	Minimizing locks in GC and other threads when handling UObjects
Change 2480190 by [REDACTED_PERSON]:
	Fixing infinite stall after canceling async loading in non-cooked builds
Change 2484268 by [REDACTED_PERSON]:
	Fixing crash when allocating permanent object pool.
Change 2489761 by [REDACTED_PERSON]:
	Fixing BulkData using linker archive on the main thread even if the linker was created on the async loading thread.
Change 2493624 by [REDACTED_PERSON]:
	Cancelling async loading will now also trigger GC
Change 2487881 by [REDACTED_PERSON]:
	Making ShaderIdMap operations thread safe.
Change 2488067 by [REDACTED_PERSON]:
	Fixing GetAsyncLoadPercentage. It will now also respect PostLoad.
Change 2458640 by [REDACTED_PERSON]:
	Fixing crash in PIE
Change 2458825 by [REDACTED_PERSON]:
	Fixing a few crashes when streaming and the package is missing.
Change 2476935 by [REDACTED_PERSON]:
	Fixing crash while async loading ANavigationData
Change 2477361 by [REDACTED_PERSON]:
	Fixing crashes in cooked game
Change 2480095 by [REDACTED_PERSON]:
	Making FUObjectArray more thread safe
Change 2475443 by [REDACTED_PERSON]:
	Re-enabling single-threaded async loading path for the editor and platforms that don't support multithreading.
Change 2475458 by [REDACTED_PERSON]:
	Making sure bulk data is only loaded on a separate thread if it's not being loaded on the async loading thread.
Change 2476661 by [REDACTED_PERSON]:
	Fixing FlushAsyncLoading not flushing everything
Change 2401089 by [REDACTED_PERSON]:
	Core - Added AtomicallySetFlags/ClearFlags to UObject, added a comment to ThisThreadAtomicallyClearedRFUnreachable

diff --git a/Engine/Config/BaseEngine.ini b/Engine/Config/BaseEngine.ini
--- a/Engine/Config/BaseEngine.ini
+++ b/Engine/Config/BaseEngine.ini
@@ -1303,6 +1303,7 @@ ZeroEngineVersionWarning=True
 UseStrictEngineVersioning=True
 FlushStreamingOnGC=False
 MinBulkDataSizeForAsyncLoading=131072
+AsyncLoadingThreadEnabled=True
 
 [Internationalization]
 +LocalizationPaths=../../../Engine/Content/Localization/Engine
diff --git a/Engine/Plugins/Experimental/LiveEditor/Source/LiveEditor/Private/LiveEditor.cpp b/Engine/Plugins/Experimental/LiveEditor/Source/LiveEditor/Private/LiveEditor.cpp
--- a/Engine/Plugins/Experimental/LiveEditor/Source/LiveEditor/Private/LiveEditor.cpp
+++ b/Engine/Plugins/Experimental/LiveEditor/Source/LiveEditor/Private/LiveEditor.cpp
@@ -206,10 +206,10 @@ void FLiveEditor::InstallHooks()
 	FEdGraphUtilities::RegisterVisualPinFactory(PinFactory);
 
 	ObjectCreationListener = new FLiveEditorObjectCreateListener();
-	GUObjectArray.AddUObjectCreateListener( ObjectCreationListener );
+	GetUObjectArray().AddUObjectCreateListener(ObjectCreationListener);
 
 	ObjectDeletionListener = new FLiveEditorObjectDeleteListener();
-	GUObjectArray.AddUObjectDeleteListener( ObjectDeletionListener );
+	GetUObjectArray().AddUObjectDeleteListener(ObjectDeletionListener);
 
 	FGlobalTabmanager::Get()->RegisterNomadTabSpawner(LiveEditorModule::LiveEditorApp, FOnSpawnTab::CreateStatic(&SpawnLiveEditorTab))
 		.SetDisplayName(NSLOCTEXT("LiveEditorPlugin", "TabTitle", "Live Editor"))
@@ -223,8 +223,8 @@ void FLiveEditor::RemoveHooks()
 
 	FEdGraphUtilities::UnregisterVisualPinFactory(PinFactory);
 
-	GUObjectArray.RemoveUObjectCreateListener( ObjectCreationListener );
-	GUObjectArray.RemoveUObjectDeleteListener( ObjectDeletionListener );
+	GetUObjectArray().RemoveUObjectCreateListener(ObjectCreationListener);
+	GetUObjectArray().RemoveUObjectDeleteListener(ObjectDeletionListener);
 
 	FGlobalTabmanager::Get()->UnregisterNomadTabSpawner(LiveEditorModule::LiveEditorApp);
 }
diff --git a/Engine/Plugins/Runtime/LiveEditorListenServer/Source/LiveEditorListenServer/Private/LiveEditorListenServer.cpp b/Engine/Plugins/Runtime/LiveEditorListenServer/Source/LiveEditorListenServer/Private/LiveEditorListenServer.cpp
--- a/Engine/Plugins/Runtime/LiveEditorListenServer/Source/LiveEditorListenServer/Private/LiveEditorListenServer.cpp
+++ b/Engine/Plugins/Runtime/LiveEditorListenServer/Source/LiveEditorListenServer/Private/LiveEditorListenServer.cpp
@@ -221,10 +221,10 @@ void FLiveEditorListenServer::ShutdownModule()
 void FLiveEditorListenServer::InstallHooks()
 {
 	ObjectCreateListener = new nLiveEditorListenServer::FCreateListener(this);
-	GUObjectArray.AddUObjectCreateListener( ObjectCreateListener );
+	GetUObjectArray().AddUObjectCreateListener(ObjectCreateListener);
 
 	ObjectDeleteListener = new nLiveEditorListenServer::FDeleteListener(this);
-	GUObjectArray.AddUObjectDeleteListener( ObjectDeleteListener );
+	GetUObjectArray().AddUObjectDeleteListener(ObjectDeleteListener);
 
 	TickObject = new nLiveEditorListenServer::FTickObject(this);
 
diff --git a/Engine/Source/Developer/AssetTools/Private/AssetFixUpRedirectors.cpp b/Engine/Source/Developer/AssetTools/Private/AssetFixUpRedirectors.cpp
--- a/Engine/Source/Developer/AssetTools/Private/AssetFixUpRedirectors.cpp
+++ b/Engine/Source/Developer/AssetTools/Private/AssetFixUpRedirectors.cpp
@@ -336,11 +336,6 @@ void FAssetFixUpRedirectors::DeleteRedirectors(TArray<FRedirectorRefs>& Redirect
 			if ( !bContainsAtLeastOneOtherAsset )
 			{
 				RedirectorPackage->RemoveFromRoot();
-				ULinkerLoad* Linker = ULinkerLoad::FindExistingLinkerForPackage(RedirectorPackage);
-				if ( Linker )
-				{
-					Linker->RemoveFromRoot();
-				}
 
 				// @todo we shouldnt be worrying about metadata objects here, ObjectTools::CleanUpAfterSuccessfulDelete should
 				if ( PackageMetaData )
diff --git a/Engine/Source/Developer/ModuleUI/Private/SModuleUI.cpp b/Engine/Source/Developer/ModuleUI/Private/SModuleUI.cpp
--- a/Engine/Source/Developer/ModuleUI/Private/SModuleUI.cpp
+++ b/Engine/Source/Developer/ModuleUI/Private/SModuleUI.cpp
@@ -277,7 +277,7 @@ FReply SModuleUI::FModuleListItem::OnRecompileClicked()
 
 EVisibility SModuleUI::FModuleListItem::GetVisibilityBasedOnLoadedAndShutdownableState() const
 {
-	if ( GIsSavingPackage || GIsGarbageCollecting )
+	if ( GIsSavingPackage || IsGarbageCollecting() )
 	{
 		return EVisibility::Hidden;
 	}
@@ -299,7 +299,7 @@ EVisibility SModuleUI::FModuleListItem::GetVisibilityBasedOnReloadableState() co
 
 EVisibility SModuleUI::FModuleListItem::GetVisibilityBasedOnRecompilableState() const
 {
-	if ( GIsSavingPackage || GIsGarbageCollecting )
+	if ( GIsSavingPackage || IsGarbageCollecting() )
 	{
 		return EVisibility::Hidden;
 	}
diff --git a/Engine/Source/Developer/PackageDependencyInfo/Private/PackageDependencyInfo.cpp b/Engine/Source/Developer/PackageDependencyInfo/Private/PackageDependencyInfo.cpp
--- a/Engine/Source/Developer/PackageDependencyInfo/Private/PackageDependencyInfo.cpp
+++ b/Engine/Source/Developer/PackageDependencyInfo/Private/PackageDependencyInfo.cpp
@@ -370,7 +370,7 @@ void FPackageDependencyInfo::RecursiveDeterminePackageDependentTimeStamp(const T
 
 		// We have the package info, so process the actual package.
 		BeginLoad();
-		ULinkerLoad* Linker = GetPackageLinker(NULL, InPackageName, LOAD_NoVerify, NULL, NULL);
+		auto Linker = GetPackageLinker(NULL, InPackageName, LOAD_NoVerify, NULL, NULL);
 		EndLoad();
 		if (Linker != NULL)
 		{
diff --git a/Engine/Source/Editor/ContentBrowser/Private/AssetViewWidgets.cpp b/Engine/Source/Editor/ContentBrowser/Private/AssetViewWidgets.cpp
--- a/Engine/Source/Editor/ContentBrowser/Private/AssetViewWidgets.cpp
+++ b/Engine/Source/Editor/ContentBrowser/Private/AssetViewWidgets.cpp
@@ -654,7 +654,7 @@ EVisibility SAssetViewItem::GetCheckedOutByOtherTextVisibility() const
 
 FText SAssetViewItem::GetCheckedOutByOtherText() const
 {
-	if ( AssetItem.IsValid() && AssetItem->GetType() != EAssetItemType::Folder && !GIsSavingPackage && !GIsGarbageCollecting )
+	if ( AssetItem.IsValid() && AssetItem->GetType() != EAssetItemType::Folder && !GIsSavingPackage && !IsGarbageCollecting() )
 	{
 		const FAssetData& AssetData = StaticCastSharedPtr<FAssetViewAsset>(AssetItem)->Data;
 		ISourceControlProvider& SourceControlProvider = ISourceControlModule::Get().GetProvider();
diff --git a/Engine/Source/Editor/IntroTutorials/Private/STutorialOverlay.cpp b/Engine/Source/Editor/IntroTutorials/Private/STutorialOverlay.cpp
--- a/Engine/Source/Editor/IntroTutorials/Private/STutorialOverlay.cpp
+++ b/Engine/Source/Editor/IntroTutorials/Private/STutorialOverlay.cpp
@@ -287,7 +287,7 @@ void STutorialOverlay::FocusOnAnyBlueprintNodes(const FTutorialWidgetContent &Wi
 		if (!FocusObject)
 		{
 			const FString ObjectName = FPackageName::ObjectPathToObjectName(ObjectPath.ToString());
-			const FName RedirectedObjectName = ULinkerLoad::FindNewNameForClass(*ObjectName, false);
+			const FName RedirectedObjectName = FLinkerLoad::FindNewNameForClass(*ObjectName, false);
 			if (!RedirectedObjectName.IsNone())
 			{
 				FocusObject = FindObject<UClass>(ANY_PACKAGE, *RedirectedObjectName.ToString());
diff --git a/Engine/Source/Editor/Kismet/Private/SBlueprintActionMenu.cpp b/Engine/Source/Editor/Kismet/Private/SBlueprintActionMenu.cpp
--- a/Engine/Source/Editor/Kismet/Private/SBlueprintActionMenu.cpp
+++ b/Engine/Source/Editor/Kismet/Private/SBlueprintActionMenu.cpp
@@ -437,7 +437,7 @@ void SBlueprintActionMenu::CollectAllActions(FGraphActionListBuilderBase& OutAll
 	
 	FBlueprintActionMenuBuilder MenuBuilder(EditorPtr);
 	// NOTE: cannot call GetGraphContextActions() during serialization and GC due to its use of FindObject()
-	if(!GIsSavingPackage && !GIsGarbageCollecting)
+	if(!GIsSavingPackage && !IsGarbageCollecting())
 	{
 		FBlueprintActionMenuUtils::MakeContextMenu(FilterContext, bIsContextSensitive, ContextTargetMask, MenuBuilder);
 	}
diff --git a/Engine/Source/Editor/Kismet/Private/STimelineEditor.cpp b/Engine/Source/Editor/Kismet/Private/STimelineEditor.cpp
--- a/Engine/Source/Editor/Kismet/Private/STimelineEditor.cpp
+++ b/Engine/Source/Editor/Kismet/Private/STimelineEditor.cpp
@@ -1364,7 +1364,7 @@ void STimelineEditor::OnTrackNameCommitted( const FText& StringName, ETextCommit
 bool STimelineEditor::IsCurveAssetSelected() const
 {
 	// Note: Cannot call GetContentBrowserSelectionClasses() during serialization and GC due to its use of FindObject()
-	if(!GIsSavingPackage && !GIsGarbageCollecting)
+	if(!GIsSavingPackage && !IsGarbageCollecting())
 	{
 		TArray<UClass*> SelectionList;
 		GEditor->GetContentBrowserSelectionClasses(SelectionList);
diff --git a/Engine/Source/Editor/KismetCompiler/Private/AnimBlueprintCompiler.cpp b/Engine/Source/Editor/KismetCompiler/Private/AnimBlueprintCompiler.cpp
--- a/Engine/Source/Editor/KismetCompiler/Private/AnimBlueprintCompiler.cpp
+++ b/Engine/Source/Editor/KismetCompiler/Private/AnimBlueprintCompiler.cpp
@@ -593,7 +593,7 @@ int32 FAnimBlueprintCompiler::ExpandGraphAndProcessNodes(UEdGraph* SourceGraph,
 	check(TargetRootNode);
 
 	// Move the cloned nodes into the consolidated event graph
-	const bool bIsLoading = Blueprint->bIsRegeneratingOnLoad || GIsAsyncLoading;
+	const bool bIsLoading = Blueprint->bIsRegeneratingOnLoad || IsAsyncLoading();
 	ClonedGraph->MoveNodesToAnotherGraph(ConsolidatedEventGraph, bIsLoading);
 
 	// Process any animation nodes
@@ -1108,7 +1108,7 @@ void FAnimBlueprintCompiler::MergeUbergraphPagesIn(UEdGraph* Ubergraph)
 			{
 				// Merge all the animation nodes, contents, etc... into the ubergraph
 				UEdGraph* ClonedGraph = FEdGraphUtilities::CloneGraph(SourceGraph, NULL, &MessageLog, true);
-				const bool bIsLoading = Blueprint->bIsRegeneratingOnLoad || GIsAsyncLoading;
+				const bool bIsLoading = Blueprint->bIsRegeneratingOnLoad || IsAsyncLoading();
 				ClonedGraph->MoveNodesToAnotherGraph(ConsolidatedEventGraph, bIsLoading);
 			}
 		}
diff --git a/Engine/Source/Editor/KismetCompiler/Private/KismetCompiler.cpp b/Engine/Source/Editor/KismetCompiler/Private/KismetCompiler.cpp
--- a/Engine/Source/Editor/KismetCompiler/Private/KismetCompiler.cpp
+++ b/Engine/Source/Editor/KismetCompiler/Private/KismetCompiler.cpp
@@ -177,7 +177,7 @@ void FKismetCompilerContext::CleanAndSanitizeClass(UBlueprintGeneratedClass* Cla
 		FString TransientCDOString = FString::Printf(TEXT("TRASH_%s"), *OldCDO->GetName());
 		FName TransientCDOName = MakeUniqueObjectName(GetTransientPackage(), TransientClass, FName(*TransientCDOString));
 		OldCDO->Rename(*TransientCDOName.ToString(), GetTransientPackage(), RenFlags);
-		ULinkerLoad::InvalidateExport(OldCDO);
+		FLinkerLoad::InvalidateExport(OldCDO);
 	}
 
 	// Purge all subobjects (properties, functions, params) of the class, as they will be regenerated
@@ -202,7 +202,7 @@ void FKismetCompilerContext::CleanAndSanitizeClass(UBlueprintGeneratedClass* Cla
 		}
 		else
 		{
-			ULinkerLoad::InvalidateExport(CurrSubObj);
+			FLinkerLoad::InvalidateExport(CurrSubObj);
 		}
 	}
 
@@ -2872,7 +2872,7 @@ void FKismetCompilerContext::ExpandTunnelsAndMacros(UEdGraph* SourceGraph)
 				}
 			}
 
-			ClonedGraph->MoveNodesToAnotherGraph(SourceGraph, GIsAsyncLoading || bIsLoading);
+			ClonedGraph->MoveNodesToAnotherGraph(SourceGraph, IsAsyncLoading() || bIsLoading);
 			FEdGraphUtilities::MergeChildrenGraphsIn(SourceGraph, ClonedGraph, /*bRequireSchemaMatch=*/ true);
 
 			// When emitting intermediate products; make an effort to make them readable by preventing overlaps and adding informative comments
@@ -3252,7 +3252,7 @@ void FKismetCompilerContext::Compile()
 
 	UObject* OldCDO = NULL;
 	int32 OldGenLinkerIdx = INDEX_NONE;
-	ULinkerLoad* OldLinker = Blueprint->GetLinker();
+	FLinkerLoad* OldLinker = Blueprint->GetLinker();
 
 	if (OldLinker)
 	{
diff --git a/Engine/Source/Editor/KismetCompiler/Private/KismetCompilerMisc.cpp b/Engine/Source/Editor/KismetCompiler/Private/KismetCompilerMisc.cpp
--- a/Engine/Source/Editor/KismetCompiler/Private/KismetCompilerMisc.cpp
+++ b/Engine/Source/Editor/KismetCompiler/Private/KismetCompilerMisc.cpp
@@ -328,7 +328,7 @@ void FKismetCompilerUtilities::ConsignToOblivion(UClass* OldClass, bool bForceNo
 		for( TFieldIterator<UFunction> ItFunc(OldClass,EFieldIteratorFlags::ExcludeSuper); ItFunc; ++ItFunc )
 		{
 			UFunction* CurrentFunc = *ItFunc;
-			ULinkerLoad::InvalidateExport(CurrentFunc);
+			FLinkerLoad::InvalidateExport(CurrentFunc);
 
 			for( TFieldIterator<UProperty> It(CurrentFunc,EFieldIteratorFlags::ExcludeSuper); It; ++It )
 			{
@@ -351,10 +351,10 @@ void FKismetCompilerUtilities::InvalidatePropertyExport(UProperty* PropertyToInv
 	UArrayProperty* ArrayProp = Cast<UArrayProperty>(PropertyToInvalidate);
  	if( ArrayProp && ArrayProp->Inner )
  	{
-		ULinkerLoad::InvalidateExport(ArrayProp->Inner);
+		FLinkerLoad::InvalidateExport(ArrayProp->Inner);
  	}
 
-	ULinkerLoad::InvalidateExport(PropertyToInvalidate);
+	FLinkerLoad::InvalidateExport(PropertyToInvalidate);
 }
 
 void FKismetCompilerUtilities::RemoveObjectRedirectorIfPresent(UObject* Package, const FString& NewName, UObject* ObjectBeingMovedIn)
diff --git a/Engine/Source/Editor/KismetCompiler/Private/UserDefinedStructureCompilerUtils.cpp b/Engine/Source/Editor/KismetCompiler/Private/UserDefinedStructureCompilerUtils.cpp
--- a/Engine/Source/Editor/KismetCompiler/Private/UserDefinedStructureCompilerUtils.cpp
+++ b/Engine/Source/Editor/KismetCompiler/Private/UserDefinedStructureCompilerUtils.cpp
@@ -92,7 +92,7 @@ struct FUserDefinedStructureCompilerInner
 			}
 			else
 			{
-				ULinkerLoad::InvalidateExport(CurrSubObj);
+				FLinkerLoad::InvalidateExport(CurrSubObj);
 			}
 		}
 
diff --git a/Engine/Source/Editor/MainFrame/Private/Frame/MainFrameHandler.h b/Engine/Source/Editor/MainFrame/Private/Frame/MainFrameHandler.h
--- a/Engine/Source/Editor/MainFrame/Private/Frame/MainFrameHandler.h
+++ b/Engine/Source/Editor/MainFrame/Private/Frame/MainFrameHandler.h
@@ -61,7 +61,7 @@ public:
 			// Cannot exit right now.
 			return false;
 		}
-		else if (GIsSavingPackage || GIsGarbageCollecting)
+		else if (GIsSavingPackage || IsGarbageCollecting())
 		{
 			// SAVING/GC PATH
 
diff --git a/Engine/Source/Editor/MainFrame/Private/MainFrameModule.cpp b/Engine/Source/Editor/MainFrame/Private/MainFrameModule.cpp
--- a/Engine/Source/Editor/MainFrame/Private/MainFrameModule.cpp
+++ b/Engine/Source/Editor/MainFrame/Private/MainFrameModule.cpp
@@ -273,7 +273,7 @@ TSharedRef<SWidget> FMainFrameModule::MakeDeveloperTools() const
 
 		static FText GetUObjectCountAsString() 
 		{
-			return FText::AsNumber(GUObjectArray.GetObjectArrayNumMinusAvailable());
+			return FText::AsNumber(GetUObjectArray().GetObjectArrayNumMinusAvailable());
 		}
 
 		static void OpenVideo( FString SourceFilePath )
diff --git a/Engine/Source/Editor/PropertyEditor/Private/UserInterface/PropertyEditor/SPropertyEditorAsset.cpp b/Engine/Source/Editor/PropertyEditor/Private/UserInterface/PropertyEditor/SPropertyEditorAsset.cpp
--- a/Engine/Source/Editor/PropertyEditor/Private/UserInterface/PropertyEditor/SPropertyEditorAsset.cpp
+++ b/Engine/Source/Editor/PropertyEditor/Private/UserInterface/PropertyEditor/SPropertyEditorAsset.cpp
@@ -479,7 +479,7 @@ FPropertyAccess::Result SPropertyEditorAsset::GetValue( FObjectOrAssetData& OutV
 {
 	// Potentially accessing the value while garbage collecting or saving the package could trigger a crash.
 	// so we fail to get the value when that is occuring.
-	if ( GIsSavingPackage || GIsGarbageCollecting )
+	if ( GIsSavingPackage || IsGarbageCollecting() )
 	{
 		return FPropertyAccess::Fail;
 	}
diff --git a/Engine/Source/Editor/UMGEditor/Private/UMGEditorModule.cpp b/Engine/Source/Editor/UMGEditor/Private/UMGEditorModule.cpp
--- a/Engine/Source/Editor/UMGEditor/Private/UMGEditorModule.cpp
+++ b/Engine/Source/Editor/UMGEditor/Private/UMGEditorModule.cpp
@@ -107,6 +107,11 @@ public:
 
 	void PreCompile(UBlueprint* Blueprint)
 	{
+		if (!CanCompile(Blueprint))
+		{
+			return;
+		}
+
 		if ( CompileCount == 0 )
 		{
 			ReRegister = new TComponentReregisterContext<UWidgetComponent>();
diff --git a/Engine/Source/Editor/UnrealEd/Classes/Commandlets/ResavePackagesCommandlet.h b/Engine/Source/Editor/UnrealEd/Classes/Commandlets/ResavePackagesCommandlet.h
--- a/Engine/Source/Editor/UnrealEd/Classes/Commandlets/ResavePackagesCommandlet.h
+++ b/Engine/Source/Editor/UnrealEd/Classes/Commandlets/ResavePackagesCommandlet.h
@@ -98,7 +98,7 @@ protected:
 	 * @param	bSavePackage	[in]	indicates whether the package is currently going to be saved
 	 *							[out]	set to true to resave the package
 	 */
-	virtual bool PerformPreloadOperations( ULinkerLoad* PackageLinker, bool& bSavePackage );
+	virtual bool PerformPreloadOperations( FLinkerLoad* PackageLinker, bool& bSavePackage );
 
 	/**
 	 * Allows the commandlet to perform any additional operations on the object before it is resaved.
diff --git a/Engine/Source/Editor/UnrealEd/Classes/Editor/EditorEngine.h b/Engine/Source/Editor/UnrealEd/Classes/Editor/EditorEngine.h
--- a/Engine/Source/Editor/UnrealEd/Classes/Editor/EditorEngine.h
+++ b/Engine/Source/Editor/UnrealEd/Classes/Editor/EditorEngine.h
@@ -2080,7 +2080,7 @@ public:
 
 	/** The editor wrapper for UPackage::SavePackage. Auto-adds files to source control when necessary */
 	bool SavePackage( UPackage* InOuter, UObject* Base, EObjectFlags TopLevelFlags, const TCHAR* Filename, 
-		FOutputDevice* Error=GError, ULinkerLoad* Conform=NULL, bool bForceByteSwapping=false, bool bWarnOfLongFilename=true, 
+		FOutputDevice* Error=GError, FLinkerLoad* Conform=NULL, bool bForceByteSwapping=false, bool bWarnOfLongFilename=true, 
 		uint32 SaveFlags=SAVE_None, const class ITargetPlatform* TargetPlatform = NULL, const FDateTime& FinalTimeStamp = FDateTime::MinValue(), bool bSlowTask = true );
 
 	/** Invoked before a UWorld is saved to update editor systems */
diff --git a/Engine/Source/Editor/UnrealEd/Private/AssetDeleteModel.cpp b/Engine/Source/Editor/UnrealEd/Private/AssetDeleteModel.cpp
--- a/Engine/Source/Editor/UnrealEd/Private/AssetDeleteModel.cpp
+++ b/Engine/Source/Editor/UnrealEd/Private/AssetDeleteModel.cpp
@@ -510,11 +510,6 @@ void FAssetDeleteModel::PrepareToDelete(UObject* InObject)
 		if ( !bContainsAtLeastOneOtherAsset )
 		{
 			RedirectorPackage->RemoveFromRoot();
-			ULinkerLoad* Linker = ULinkerLoad::FindExistingLinkerForPackage(RedirectorPackage);
-			if ( Linker )
-			{
-				Linker->RemoveFromRoot();
-			}
 
 			// @todo we shouldnt be worrying about metadata objects here, ObjectTools::CleanUpAfterSuccessfulDelete should
 			if ( PackageMetaData )
diff --git a/Engine/Source/Editor/UnrealEd/Private/Commandlets/ContentCommandlets.cpp b/Engine/Source/Editor/UnrealEd/Private/Commandlets/ContentCommandlets.cpp
--- a/Engine/Source/Editor/UnrealEd/Private/Commandlets/ContentCommandlets.cpp
+++ b/Engine/Source/Editor/UnrealEd/Private/Commandlets/ContentCommandlets.cpp
@@ -237,7 +237,7 @@ void UResavePackagesCommandlet::LoadAndSaveOnePackage(const FString& Filename)
 		VerboseMessage(TEXT("Pre GetPackageLinker"));
 
 		BeginLoad();
-		ULinkerLoad* Linker = GetPackageLinker(NULL,*Filename,LOAD_NoVerify,NULL,NULL);
+		auto Linker = GetPackageLinker(NULL,*Filename,LOAD_NoVerify,NULL,NULL);
 		EndLoad();
 	
 		// Bail early if we don't have a valid linker (package was out of date, etc)
@@ -595,7 +595,7 @@ FText UResavePackagesCommandlet::GetChangelistDescription() const
 }
 
 
-bool UResavePackagesCommandlet::PerformPreloadOperations( ULinkerLoad* PackageLinker, bool& bSavePackage )
+bool UResavePackagesCommandlet::PerformPreloadOperations( FLinkerLoad* PackageLinker, bool& bSavePackage )
 {
 	bool bResult = false;
 
@@ -1064,7 +1064,7 @@ int32 UWrangleContentCommandlet::Main( const FString& Params )
 					UPackage* Package = LoadPackage(NULL, *PackageFilename, LOAD_None);
 
 					BeginLoad();
-					ULinkerLoad* Linker = GetPackageLinker( NULL, *PackageFilename, LOAD_Quiet|LOAD_NoWarn|LOAD_NoVerify, NULL, NULL );
+					auto Linker = GetPackageLinker( NULL, *PackageFilename, LOAD_Quiet|LOAD_NoWarn|LOAD_NoVerify, NULL, NULL );
 					EndLoad();
 
 					// look for special package types
@@ -1219,7 +1219,7 @@ int32 UWrangleContentCommandlet::Main( const FString& Params )
 			}
 
 			BeginLoad();
-			ULinkerLoad* Linker = GetPackageLinker( NULL, *PackageFilename, LOAD_Quiet|LOAD_NoWarn|LOAD_NoVerify, NULL, NULL );
+			auto Linker = GetPackageLinker( NULL, *PackageFilename, LOAD_Quiet|LOAD_NoWarn|LOAD_NoVerify, NULL, NULL );
 			EndLoad();
 
 			// go through the exports in the package, looking for public objects
diff --git a/Engine/Source/Editor/UnrealEd/Private/Commandlets/DiffPackagesCommandlet.cpp b/Engine/Source/Editor/UnrealEd/Private/Commandlets/DiffPackagesCommandlet.cpp
--- a/Engine/Source/Editor/UnrealEd/Private/Commandlets/DiffPackagesCommandlet.cpp
+++ b/Engine/Source/Editor/UnrealEd/Private/Commandlets/DiffPackagesCommandlet.cpp
@@ -1298,7 +1298,7 @@ void UDiffPackagesCommandlet::LoadNativePropertyData( UObject* Object, TArray<ui
 	// first, validate our input parameters
 	check(Object);
 
-	ULinkerLoad* ObjectLinker = Object->GetLinker();
+	auto ObjectLinker = Object->GetLinker();
 	check(ObjectLinker);
 
 	int32 ObjectLinkerIndex = Object->GetLinkerIndex();
diff --git a/Engine/Source/Editor/UnrealEd/Private/Commandlets/PackageUtilities.cpp b/Engine/Source/Editor/UnrealEd/Private/Commandlets/PackageUtilities.cpp
--- a/Engine/Source/Editor/UnrealEd/Private/Commandlets/PackageUtilities.cpp
+++ b/Engine/Source/Editor/UnrealEd/Private/Commandlets/PackageUtilities.cpp
@@ -199,7 +199,7 @@ bool NormalizePackageNames( TArray<FString> PackageNames, TArray<FString>& Packa
 
 * @return true if successful
 */
-bool SavePackageHelper(UPackage* Package, FString Filename, EObjectFlags KeepObjectFlags, FOutputDevice* ErrorDevice, ULinkerLoad* LinkerToConformAgainst, ESaveFlags SaveFlags)
+bool SavePackageHelper(UPackage* Package, FString Filename, EObjectFlags KeepObjectFlags, FOutputDevice* ErrorDevice, FLinkerLoad* LinkerToConformAgainst, ESaveFlags SaveFlags)
 {
 	// look for a world object in the package (if there is one, there's a map)
 	UWorld* World = UWorld::FindWorldInPackage(Package);
@@ -692,7 +692,7 @@ int32 ULoadPackageCommandlet::Main( const FString& Params )
 		if (bCheckForLegacyPackages)
 		{
 			BeginLoad();
-			ULinkerLoad* Linker = GetPackageLinker(NULL,*Filename,LOAD_NoVerify,NULL,NULL);
+			auto Linker = GetPackageLinker(NULL,*Filename,LOAD_NoVerify,NULL,NULL);
 			EndLoad();
 			MinVersion = FMath::Min<int32>(MinVersion, Linker->Summary.GetFileVersionUE4());
 		}
@@ -729,7 +729,7 @@ struct FExportInfo
 	FString PathName;
 	FString OuterPathName;
 
-	FExportInfo( ULinkerLoad* Linker, int32 InIndex )
+	FExportInfo( FLinkerLoad* Linker, int32 InIndex )
 	: Export(Linker->ExportMap[InIndex]), ExportIndex(InIndex)
 	, OuterPathName(TEXT("NULL"))
 	{
@@ -737,7 +737,7 @@ struct FExportInfo
 		SetOuterPathName(Linker);
 	}
 
-	void SetOuterPathName( ULinkerLoad* Linker )
+	void SetOuterPathName( FLinkerLoad* Linker )
 	{
 		if ( !Export.OuterIndex.IsNull() )
 		{
@@ -816,7 +816,7 @@ namespace
  *
  * @param	InLinker	if specified, changes this reporter's Linker before generating the report.
  */
-void FPkgInfoReporter_Log::GeneratePackageReport( ULinkerLoad* InLinker/*=NULL*/ )
+void FPkgInfoReporter_Log::GeneratePackageReport( FLinkerLoad* InLinker/*=NULL*/ )
 {
 	check(InLinker);
 
@@ -1394,7 +1394,7 @@ int32 UPkgInfoCommandlet::Main( const FString& Params )
 		}
 
 		BeginLoad();
-		ULinkerLoad* Linker = GetPackageLinker( NULL, *Filename, LOAD_NoVerify, NULL, NULL );
+		auto Linker = GetPackageLinker( NULL, *Filename, LOAD_NoVerify, NULL, NULL );
 		EndLoad();
 
 		if( Linker )
diff --git a/Engine/Source/Editor/UnrealEd/Private/EdGraphUtilities.cpp b/Engine/Source/Editor/UnrealEd/Private/EdGraphUtilities.cpp
--- a/Engine/Source/Editor/UnrealEd/Private/EdGraphUtilities.cpp
+++ b/Engine/Source/Editor/UnrealEd/Private/EdGraphUtilities.cpp
@@ -203,7 +203,7 @@ void FEdGraphUtilities::CloneAndMergeGraphIn(UEdGraph* MergeTarget, UEdGraph* So
 	const bool bIsLoading = Blueprint ? Blueprint->bIsRegeneratingOnLoad : false;
 
 	// Move them all to the destination
-	ClonedGraph->MoveNodesToAnotherGraph(MergeTarget, GIsAsyncLoading || bIsLoading, bInIsCompiling);
+	ClonedGraph->MoveNodesToAnotherGraph(MergeTarget, IsAsyncLoading() || bIsLoading, bInIsCompiling);
 }
 
 // Moves the contents of all of the children graphs (recursively) into the target graph.  This does not clone, it's destructive to the source
@@ -225,7 +225,7 @@ void FEdGraphUtilities::MergeChildrenGraphsIn(UEdGraph* MergeTarget, UEdGraph* P
 			// Even if we don't require a match to recurse, we do to actually copy the nodes
 			if (bSchemaMatches)
 			{
-				ChildGraph->MoveNodesToAnotherGraph(MergeTarget, GIsAsyncLoading || bIsLoading, bInIsCompiling);
+				ChildGraph->MoveNodesToAnotherGraph(MergeTarget, IsAsyncLoading() || bIsLoading, bInIsCompiling);
 			}
 
 			MergeChildrenGraphsIn(MergeTarget, ChildGraph, bRequireSchemaMatch, bInIsCompiling);
diff --git a/Engine/Source/Editor/UnrealEd/Private/Editor.cpp b/Engine/Source/Editor/UnrealEd/Private/Editor.cpp
--- a/Engine/Source/Editor/UnrealEd/Private/Editor.cpp
+++ b/Engine/Source/Editor/UnrealEd/Private/Editor.cpp
@@ -4314,7 +4314,7 @@ void UEditorEngine::OnSourceControlDialogClosed(bool bEnabled)
 }
 
 bool UEditorEngine::SavePackage( UPackage* InOuter, UObject* InBase, EObjectFlags TopLevelFlags, const TCHAR* Filename, 
-				 FOutputDevice* Error, ULinkerLoad* Conform, bool bForceByteSwapping, bool bWarnOfLongFilename, 
+				 FOutputDevice* Error, FLinkerLoad* Conform, bool bForceByteSwapping, bool bWarnOfLongFilename, 
 				 uint32 SaveFlags, const class ITargetPlatform* TargetPlatform, const FDateTime& FinalTimeStamp, bool bSlowTask )
 {
 	FScopedSlowTask SlowTask(100, FText(), bSlowTask);
@@ -7512,7 +7512,7 @@ void UEditorEngine::OnAssetLoaded(UObject* Asset)
 	UWorld* WorldAsset = Cast<UWorld>(Asset);
 	if (WorldAsset != NULL)
 	{
-		// Init inactive worlds here instead of UWorld::PostLoad because it is illegal to call UpdateWorldComponents while GIsRoutingPostLoad
+		// Init inactive worlds here instead of UWorld::PostLoad because it is illegal to call UpdateWorldComponents while IsRoutingPostLoad
 		if (WorldAsset->WorldType == EWorldType::Inactive)
 		{
 			// Create the world without a physics scene because creating too many physics scenes causes deadlock issues in PhysX. The scene will be created when it is opened in the level editor.
diff --git a/Engine/Source/Editor/UnrealEd/Private/EditorServer.cpp b/Engine/Source/Editor/UnrealEd/Private/EditorServer.cpp
--- a/Engine/Source/Editor/UnrealEd/Private/EditorServer.cpp
+++ b/Engine/Source/Editor/UnrealEd/Private/EditorServer.cpp
@@ -1231,7 +1231,7 @@ void UEditorEngine::PostUndo (bool bSuccess)
 bool UEditorEngine::UndoTransaction()
 {
 	// make sure we're in a valid state to perform this
-	if (GIsSavingPackage || GIsGarbageCollecting)
+	if (GIsSavingPackage || IsGarbageCollecting())
 	{
 		return false;
 	}
@@ -1242,7 +1242,7 @@ bool UEditorEngine::UndoTransaction()
 bool UEditorEngine::RedoTransaction()
 {
 	// make sure we're in a valid state to perform this
-	if (GIsSavingPackage || GIsGarbageCollecting)
+	if (GIsSavingPackage || IsGarbageCollecting())
 	{
 		return false;
 	}
@@ -5068,13 +5068,13 @@ void ListMapPackageDependencies(const TCHAR* InStr)
 			UE_LOG(LogEditorServer, Warning, TEXT("\tResave packages and run again to ensure accurate results."));
 		}
 
-		ULinkerLoad* Linker = ProcessingPackage->GetLinker();
-		if (Linker == NULL)
+		auto Linker = ProcessingPackage->GetLinker();
+		if (!Linker)
 		{
 			// Create a new linker object which goes off and tries load the file.
 			Linker = GetPackageLinker(NULL, *(ProcessingPackage->GetName()), LOAD_None, NULL, NULL );
 		}
-		if (Linker != NULL)
+		if (Linker)
 		{
 			for (int32 ImportIdx = 0; ImportIdx < Linker->ImportMap.Num(); ImportIdx++)
 			{
@@ -5111,13 +5111,13 @@ void ListMapPackageDependencies(const TCHAR* InStr)
 			UPackage* RefdPackage = LoadPackage(NULL, *RefdPkgName, LOAD_None);
 			if (RefdPackage != NULL)
 			{
-				ULinkerLoad* Linker = RefdPackage->GetLinker();
-				if (Linker == NULL)
+				auto Linker = RefdPackage->GetLinker();
+				if (!Linker)
 				{
 					// Create a new linker object which goes off and tries load the file.
 					Linker = GetPackageLinker(NULL, *RefdPkgName,  LOAD_None, NULL, NULL );
 				}
-				if (Linker != NULL)
+				if (Linker)
 				{
 					for (int32 ExportIdx = 0; ExportIdx < Linker->ExportMap.Num(); ExportIdx++)
 					{
diff --git a/Engine/Source/Editor/UnrealEd/Private/Factories/PackFactory.cpp b/Engine/Source/Editor/UnrealEd/Private/Factories/PackFactory.cpp
--- a/Engine/Source/Editor/UnrealEd/Private/Factories/PackFactory.cpp
+++ b/Engine/Source/Editor/UnrealEd/Private/Factories/PackFactory.cpp
@@ -419,8 +419,8 @@ UObject* UPackFactory::FactoryCreateBinary
 						FString FinalIniFileName;
 						GConfig->LoadGlobalIniFile(FinalIniFileName, *RedirectsSection, NULL, NULL, true);
 
-						ULinkerLoad::AddGameNameRedirect(*LongOldGameName, *LongNewGameName);
-						ULinkerLoad::AddGameNameRedirect(*ConfigParameters.GameName, *LongNewGameName);
+						FLinkerLoad::AddGameNameRedirect(*LongOldGameName, *LongNewGameName);
+						FLinkerLoad::AddGameNameRedirect(*ConfigParameters.GameName, *LongNewGameName);
 					}
 					break;
 				}
diff --git a/Engine/Source/Editor/UnrealEd/Private/GroupActor.cpp b/Engine/Source/Editor/UnrealEd/Private/GroupActor.cpp
--- a/Engine/Source/Editor/UnrealEd/Private/GroupActor.cpp
+++ b/Engine/Source/Editor/UnrealEd/Private/GroupActor.cpp
@@ -775,7 +775,7 @@ void AGroupActor::PostRemove()
 			MarkPackageDirty();
 
 			// If not currently garbage collecting (changing maps, saving, etc), remove the group immediately
-			if(!GIsGarbageCollecting)
+			if(!IsGarbageCollecting())
 			{
 				// Refresh all editor browsers after removal
 				FScopedRefreshAllBrowsers LevelRefreshAllBrowsers;
diff --git a/Engine/Source/Editor/UnrealEd/Private/Kismet2/BlueprintEditorUtils.cpp b/Engine/Source/Editor/UnrealEd/Private/Kismet2/BlueprintEditorUtils.cpp
--- a/Engine/Source/Editor/UnrealEd/Private/Kismet2/BlueprintEditorUtils.cpp
+++ b/Engine/Source/Editor/UnrealEd/Private/Kismet2/BlueprintEditorUtils.cpp
@@ -576,7 +576,7 @@ void FBlueprintEditorUtils::PreloadConstructionScript(UBlueprint* Blueprint)
 {
 	if ( Blueprint )
 	{
-		ULinkerLoad* TargetLinker = Blueprint->SimpleConstructionScript ? Blueprint->SimpleConstructionScript->GetLinker() : NULL;
+		FLinkerLoad* TargetLinker = Blueprint->SimpleConstructionScript ? Blueprint->SimpleConstructionScript->GetLinker() : NULL;
 		if (TargetLinker)
 		{
 			TargetLinker->Preload(Blueprint->SimpleConstructionScript);
@@ -607,7 +607,7 @@ void FBlueprintEditorUtils::PreloadConstructionScript(UBlueprint* Blueprint)
 	}
 }
 
-void FBlueprintEditorUtils::PatchNewCDOIntoLinker(UObject* CDO, ULinkerLoad* Linker, int32 ExportIndex, TArray<UObject*>& ObjLoaded)
+void FBlueprintEditorUtils::PatchNewCDOIntoLinker(UObject* CDO, FLinkerLoad* Linker, int32 ExportIndex, TArray<UObject*>& ObjLoaded)
 {
 	if( (CDO != NULL) && (Linker != NULL) && (ExportIndex != INDEX_NONE) )
 	{
@@ -907,7 +907,7 @@ struct FRegenerationHelper
 	/**
 	 * A helper function that loads (and regenerates) interface dependencies.
 	 * Accounts for circular dependencies by following how we handle parent 
-	 * classes in ULinkerLoad::RegenerateBlueprintClass() (that is, to complete 
+	 * classes in FLinkerLoad::RegenerateBlueprintClass() (that is, to complete 
 	 * the interface's compilation/regeneration before we utilize it for the
 	 * specified blueprint).
 	 * 
@@ -1290,7 +1290,7 @@ UClass* FBlueprintEditorUtils::RegenerateBlueprintClass(UBlueprint* Blueprint, U
 		Blueprint->GetBlueprintCDONames(GeneratedName, SkeletonName);
 		int32 OldSkelLinkerIdx = INDEX_NONE;
 		int32 OldGenLinkerIdx = INDEX_NONE;
-		ULinkerLoad* OldLinker = Blueprint->GetLinker();
+		auto OldLinker = Blueprint->GetLinker();
 		for( int32 i = 0; i < OldLinker->ExportMap.Num(); i++ )
 		{
 			FObjectExport& ThisExport = OldLinker->ExportMap[i];
@@ -1646,7 +1646,7 @@ void FBlueprintEditorUtils::PatchCDOSubobjectsIntoExport(UObject* PreviousCDO, U
 			UObject** NewComponentPtr = NewComponentsMap.Find(OldComponentName);
 			if (NewComponentPtr && *NewComponentPtr)
 			{
-				ULinkerLoad::PRIVATE_PatchNewObjectIntoExport(OldComponent, *NewComponentPtr);
+				FLinkerLoad::PRIVATE_PatchNewObjectIntoExport(OldComponent, *NewComponentPtr);
 			}
 		}
 		NewCDO->CheckDefaultSubobjects();
diff --git a/Engine/Source/Editor/UnrealEd/Private/Tests/EditorAutomationTests.cpp b/Engine/Source/Editor/UnrealEd/Private/Tests/EditorAutomationTests.cpp
--- a/Engine/Source/Editor/UnrealEd/Private/Tests/EditorAutomationTests.cpp
+++ b/Engine/Source/Editor/UnrealEd/Private/Tests/EditorAutomationTests.cpp
@@ -2094,7 +2094,7 @@ namespace StaticMeshUVTest
 			}
 			else
 			{
-				LoadPackageAsync(PackageName,FLoadPackageAsyncDelegate::CreateSP(this, &FUVTestHelper::PackageLoadCallback));
+				LoadPackageAsync(*PackageName, FLoadPackageAsyncDelegate::CreateSP(this, &FUVTestHelper::PackageLoadCallback));
 				CurrentState = EStaticMeshUVTestState::WaitingForPackage;
 			}
 		}
diff --git a/Engine/Source/Editor/UnrealEd/Private/UnrealEdEngine.cpp b/Engine/Source/Editor/UnrealEd/Private/UnrealEdEngine.cpp
--- a/Engine/Source/Editor/UnrealEd/Private/UnrealEdEngine.cpp
+++ b/Engine/Source/Editor/UnrealEd/Private/UnrealEdEngine.cpp
@@ -31,6 +31,7 @@
 #include "AutoReimport/AutoReimportManager.h"
 #include "NotificationManager.h"
 #include "SNotificationList.h"
+#include "UObject/UObjectThreadContext.h"
 
 DEFINE_LOG_CATEGORY_STATIC(LogUnrealEdEngine, Log, All);
 
@@ -429,8 +430,8 @@ void UUnrealEdEngine::OnPackageDirtyStateUpdated( UPackage* Pkg)
 	const FString PackageName = Package->GetName();
 
 	// Alert the user if they have modified a package that won't be able to be saved because
-	// it's already been saved with an engine version that is newer than the current one.
-	if ( !GIsRoutingPostLoad && Package->IsDirty() && !PackagesCheckedForEngineVersion.Contains( PackageName ) )
+	// it's already been saved with an engine version that is newer than the current one.	
+	if (!FUObjectThreadContext::Get().IsRoutingPostLoad && Package->IsDirty() && !PackagesCheckedForEngineVersion.Contains(PackageName))
 	{
 		EWriteDisallowedWarningState WarningStateToSet = WDWS_WarningUnnecessary;
 				
@@ -457,7 +458,7 @@ void UUnrealEdEngine::OnPackageDirtyStateUpdated( UPackage* Pkg)
 
 	// Alert the user if they have modified a package that they do not have sufficient permission to write to disk.
 	// This can be due to the content being in the "Program Files" folder and the user does not have admin privileges.
-	if ( !GIsRoutingPostLoad && Package->IsDirty() && !PackagesCheckedForWritePermission.Contains( PackageName ) )
+	if (!FUObjectThreadContext::Get().IsRoutingPostLoad && Package->IsDirty() && !PackagesCheckedForWritePermission.Contains(PackageName))
 	{
 		EWriteDisallowedWarningState WarningStateToSet = GetWarningStateForWritePermission(PackageName);
 
diff --git a/Engine/Source/Editor/UnrealEd/Public/Kismet2/BlueprintEditorUtils.h b/Engine/Source/Editor/UnrealEd/Public/Kismet2/BlueprintEditorUtils.h
--- a/Engine/Source/Editor/UnrealEd/Public/Kismet2/BlueprintEditorUtils.h
+++ b/Engine/Source/Editor/UnrealEd/Public/Kismet2/BlueprintEditorUtils.h
@@ -117,7 +117,7 @@ public:
 	/** 
 	 * Helper function to patch the new CDO into the linker where the old one existed 
 	 */
-	static void PatchNewCDOIntoLinker(UObject* CDO, ULinkerLoad* Linker, int32 ExportIndex, TArray<UObject*>& ObjLoaded);
+	static void PatchNewCDOIntoLinker(UObject* CDO, FLinkerLoad* Linker, int32 ExportIndex, TArray<UObject*>& ObjLoaded);
 
 	/**
 	 * Regenerates the class at class load time, and refreshes the blueprint
diff --git a/Engine/Source/Editor/UnrealEd/Public/PackageHelperFunctions.h b/Engine/Source/Editor/UnrealEd/Public/PackageHelperFunctions.h
--- a/Engine/Source/Editor/UnrealEd/Public/PackageHelperFunctions.h
+++ b/Engine/Source/Editor/UnrealEd/Public/PackageHelperFunctions.h
@@ -61,7 +61,7 @@ bool NormalizePackageNames( TArray<FString> PackageNames, TArray<FString>& Packa
  *
  * @return true if successful
  */
-bool SavePackageHelper(UPackage* Package, FString Filename,  EObjectFlags KeepObjectFlags = RF_Standalone, FOutputDevice* ErrorDevice=GWarn, ULinkerLoad* LinkerToConformAgainst=NULL, ESaveFlags SaveFlags = SAVE_None);
+bool SavePackageHelper(UPackage* Package, FString Filename,  EObjectFlags KeepObjectFlags = RF_Standalone, FOutputDevice* ErrorDevice=GWarn, FLinkerLoad* LinkerToConformAgainst=NULL, ESaveFlags SaveFlags = SAVE_None);
 
 
 /**
diff --git a/Engine/Source/Editor/UnrealEd/Public/PackageUtilityWorkers.h b/Engine/Source/Editor/UnrealEd/Public/PackageUtilityWorkers.h
--- a/Engine/Source/Editor/UnrealEd/Public/PackageUtilityWorkers.h
+++ b/Engine/Source/Editor/UnrealEd/Public/PackageUtilityWorkers.h
@@ -36,7 +36,7 @@ struct FPkgInfoReporter
 	FPkgInfoReporter() 
 	: InfoFlags(PKGINFO_None), bHideOffsets(false), Linker(NULL), PackageCount(0)
 	{}
-	FPkgInfoReporter( uint32 InInfoFlags, bool bInHideOffsets, ULinkerLoad* InLinker=NULL )
+	FPkgInfoReporter( uint32 InInfoFlags, bool bInHideOffsets, FLinkerLoad* InLinker=NULL )
 	: InfoFlags(InInfoFlags), bHideOffsets(bInHideOffsets), Linker(InLinker), PackageCount(0)
 	{}
 	FPkgInfoReporter( const FPkgInfoReporter& Other )
@@ -51,12 +51,12 @@ struct FPkgInfoReporter
 	 *
 	 * @param	InLinker	if specified, changes this reporter's Linker before generating the report.
 	 */
-	virtual void GeneratePackageReport( class ULinkerLoad* InLinker=NULL )=0;
+	virtual void GeneratePackageReport( class FLinkerLoad* InLinker=NULL )=0;
 
 	/**
 	 * Changes the target linker for this reporter.  Useful when generating reports for multiple packages.
 	 */
-	void SetLinker( class ULinkerLoad* NewLinker )
+	void SetLinker( class FLinkerLoad* NewLinker )
 	{
 		Linker = NewLinker;
 	}
@@ -76,7 +76,7 @@ protected:
 	/**
 	 * The linker of the package to generate the report for
 	 */
-	class ULinkerLoad* Linker;
+	class FLinkerLoad* Linker;
 
 	/**
 	 * The number of packages evaluated by this reporter so far.  Must be incremented by child classes.
@@ -86,7 +86,7 @@ protected:
 
 struct FPkgInfoReporter_Log : public FPkgInfoReporter
 {
-	FPkgInfoReporter_Log( uint32 InInfoFlags, bool bInHideOffsets, ULinkerLoad* InLinker=NULL )
+	FPkgInfoReporter_Log( uint32 InInfoFlags, bool bInHideOffsets, FLinkerLoad* InLinker=NULL )
 	: FPkgInfoReporter(InInfoFlags, bInHideOffsets, InLinker)
 	{}
 	FPkgInfoReporter_Log( const FPkgInfoReporter_Log& Other )
@@ -97,7 +97,7 @@ struct FPkgInfoReporter_Log : public FPkgInfoReporter
 	 *
 	 * @param	InLinker	if specified, changes this reporter's Linker before generating the report.
 	 */
-	virtual void GeneratePackageReport( class ULinkerLoad* InLinker=NULL );
+	virtual void GeneratePackageReport( class FLinkerLoad* InLinker=NULL );
 };
 
 #endif
diff --git a/Engine/Source/Runtime/AssetRegistry/Public/AssetData.h b/Engine/Source/Runtime/AssetRegistry/Public/AssetData.h
--- a/Engine/Source/Runtime/AssetRegistry/Public/AssetData.h
+++ b/Engine/Source/Runtime/AssetRegistry/Public/AssetData.h
@@ -147,7 +147,7 @@ public:
 		if (!FoundClass)
 		{
 			// Look for class redirectors
-			FName NewPath = ULinkerLoad::FindNewNameForClass(AssetClass, false);
+			FName NewPath = FLinkerLoad::FindNewNameForClass(AssetClass, false);
 
 			if (NewPath != NAME_None)
 			{
diff --git a/Engine/Source/Runtime/Core/Private/HAL/ThreadingBase.cpp b/Engine/Source/Runtime/Core/Private/HAL/ThreadingBase.cpp
--- a/Engine/Source/Runtime/Core/Private/HAL/ThreadingBase.cpp
+++ b/Engine/Source/Runtime/Core/Private/HAL/ThreadingBase.cpp
@@ -50,6 +50,13 @@ CORE_API bool IsInRHIThread()
 CORE_API FRunnableThread* GRHIThread = nullptr;
 // Fake threads
 
+// Core version of IsInAsyncLoadingThread
+static bool IsInAsyncLoadingThreadCoreInternal()
+{
+	// No async loading in Core
+	return false;
+}
+bool(*IsInAsyncLoadingThread)() = &IsInAsyncLoadingThreadCoreInternal;
 
 /**
  * Fake thread created when multi-threading is disabled.
@@ -651,3 +658,48 @@ void FTlsAutoCleanup::Register()
 		RunnableThread->TlsInstances.Add( this );
 	}
 }
+
+
+
+FMultiReaderSingleWriterGT::FMultiReaderSingleWriterGT()
+{
+	CanRead = TFunction<bool()>([=]() { return FPlatformAtomics::InterlockedCompareExchange(&CriticalSection.Action, ReadingAction, NoAction) == ReadingAction; });
+	CanWrite = TFunction<bool()>([=]() { return FPlatformAtomics::InterlockedCompareExchange(&CriticalSection.Action, WritingAction, NoAction) == WritingAction; });
+}
+
+void FMultiReaderSingleWriterGT::LockRead()
+{
+	if (!IsInGameThread())
+	{
+		FPlatformProcess::ConditionalSleep(CanRead);
+	}
+	CriticalSection.ReadCounter.Increment();
+}
+
+void FMultiReaderSingleWriterGT::UnlockRead()
+{
+	if (CriticalSection.ReadCounter.Decrement() != 0)
+	{
+		return;
+	}
+
+	if (!IsInGameThread())
+	{
+		FPlatformAtomics::InterlockedExchange(&CriticalSection.Action, NoAction);
+	}
+}
+
+void FMultiReaderSingleWriterGT::LockWrite()
+{
+	check(IsInGameThread());
+	FPlatformProcess::ConditionalSleep(CanWrite);
+	CriticalSection.WriteCounter.Increment();
+}
+
+void FMultiReaderSingleWriterGT::UnlockWrite()
+{
+	if (CriticalSection.WriteCounter.Decrement() == 0)
+	{
+		FPlatformAtomics::InterlockedExchange(&CriticalSection.Action, NoAction);
+	}
+}
diff --git a/Engine/Source/Runtime/Core/Private/Misc/Core.cpp b/Engine/Source/Runtime/Core/Private/Misc/Core.cpp
--- a/Engine/Source/Runtime/Core/Private/Misc/Core.cpp
+++ b/Engine/Source/Runtime/Core/Private/Misc/Core.cpp
@@ -109,7 +109,6 @@ bool					GIsServer						= false;					/* Whether engine was launched as a server,
 bool					GIsCriticalError				= false;					/* An appError() has occured */
 bool					GIsGuarded						= false;					/* Whether execution is happening within main()/WinMain()'s try/catch handler */
 bool					GIsRunning						= false;					/* Whether execution is happening within MainLoop() */
-bool					GIsGarbageCollecting			= false;					/* Whether we are inside garbage collection */
 bool					GIsDuplicatingClassForReinstancing = false;					/* Whether we are currently using SDO on a UClass or CDO for live reinstancing */
 /** This specifies whether the engine was launched as a build machine process								*/
 bool					GIsBuildMachine					= false;
@@ -165,7 +164,12 @@ IMPLEMENT_FOREIGN_ENGINE_DIR()
 /** Exec handler for game debugging tool, allowing commands like "editactor", ...							*/
 FExec*					GDebugToolExec					= NULL;
 /** Whether we're currently in the async loading codepath or not											*/
-bool					GIsAsyncLoading					= false;
+static bool IsAsyncLoadingCoreInternal()
+{
+	// No Async loading in Core
+	return false;
+}
+bool(*IsAsyncLoading)() = &IsAsyncLoadingCoreInternal;
 /** Whether the editor is currently loading a package or not												*/
 bool					GIsEditorLoadingPackage				= false;
 /** Whether GWorld points to the play in editor world														*/
@@ -182,8 +186,6 @@ double					GStartTime						= FPlatformTime::InitTiming();
 FString					GSystemStartTime;
 /** Whether we are still in the initial loading proces.														*/
 bool					GIsInitialLoad					= true;
-/** true when we are routing ConditionalPostLoad/PostLoad to objects										*/
-bool					GIsRoutingPostLoad				= false;
 /** Steadily increasing frame counter.																		*/
 uint64					GFrameCounter					= 0;
 /** Incremented once per frame before the scene is being rendered. In split screen mode this is incremented once for all views (not for each view). */
diff --git a/Engine/Source/Runtime/Core/Private/Modules/ModuleManager.cpp b/Engine/Source/Runtime/Core/Private/Modules/ModuleManager.cpp
--- a/Engine/Source/Runtime/Core/Private/Modules/ModuleManager.cpp
+++ b/Engine/Source/Runtime/Core/Private/Modules/ModuleManager.cpp
@@ -17,12 +17,20 @@ DEFINE_LOG_CATEGORY_STATIC(LogModuleManager, Log, All);
 
 int32 FModuleManager::FModuleInfo::CurrentLoadOrder = 1;
 
+TSharedRef<FModuleManager::FModuleInfo>* FModuleManager::FindModule(FName InModuleName)
+{
+	FReadScopeLock Lock(&ModulesCriticalSection);
+	return Modules.Find(InModuleName);
+}
 
-FModuleManager& FModuleManager::Get()
+TSharedRef<FModuleManager::FModuleInfo> FModuleManager::FindModuleChecked(FName InModuleName)
 {
-	// FModuleManager is not thread-safe
-	ensure(IsInGameThread());
+	FReadScopeLock Lock(&ModulesCriticalSection);
+	return Modules.FindChecked(InModuleName);
+}
 
+FModuleManager& FModuleManager::Get()
+{
 	// NOTE: The reason we initialize to NULL here is due to an issue with static initialization of variables with
 	// constructors/destructors across DLL boundaries, where a function called from a statically constructed object
 	// calls a function in another module (such as this function) that creates a static variable.  A crash can occur
@@ -32,6 +40,9 @@ FModuleManager& FModuleManager::Get()
 
 	if( ModuleManager == NULL )
 	{
+		// FModuleManager is not thread-safe
+		ensure(IsInGameThread());
+
 		ModuleManager = new FModuleManager();
 
 		//temp workaround for IPlatformFile being used for FPaths::DirectoryExists before main() sets up the commandline.
@@ -90,10 +101,10 @@ void FModuleManager::FindModules(const TCHAR* WildcardWithoutExtension, TArray<F
 }
 
 
-bool FModuleManager::IsModuleLoaded( const FName InModuleName ) const
+bool FModuleManager::IsModuleLoaded( const FName InModuleName )
 {
 	// Do we even know about this module?
-	const TSharedRef< FModuleInfo >* ModuleInfoPtr = Modules.Find( InModuleName );
+	auto ModuleInfoPtr = FindModule(InModuleName);
 	if( ModuleInfoPtr != NULL )
 	{
 		const TSharedRef< FModuleInfo >& ModuleInfo( *ModuleInfoPtr );
@@ -124,131 +135,159 @@ bool FModuleManager::IsModuleUpToDate( const FName InModuleName ) const
 	return CheckModuleCompatibility(*TMap<FName, FString>::TConstIterator(ModulePathMap).Value());
 }
 
+bool FindNewestModuleFile(TArray<FString>& FilesToSearch, const FDateTime& NewerThan, const FString& ModuleFileSearchDirectory, const FString& Prefix, const FString& Suffix, FString& OutFilename)
+{
+	// Figure out what the newest module file is
+	bool bFound = false;
+	FDateTime NewestFoundFileTime = NewerThan;
+
+	for (const auto& FoundFile : FilesToSearch)
+	{
+		// FoundFiles contains file names with no directory information, but we need the full path up
+		// to the file, so we'll prefix it back on if we have a path.
+		const FString FoundFilePath = ModuleFileSearchDirectory.IsEmpty() ? FoundFile : (ModuleFileSearchDirectory / FoundFile);
+
+		// need to reject some files here that are not numbered...release executables, do have a suffix, so we need to make sure we don't find the debug version
+		check(FoundFilePath.Len() > Prefix.Len() + Suffix.Len());
+		FString Center = FoundFilePath.Mid(Prefix.Len(), FoundFilePath.Len() - Prefix.Len() - Suffix.Len());
+		check(Center.StartsWith(TEXT("-"))); // a minus sign is still considered numeric, so we can leave it.
+		if (!Center.IsNumeric())
+		{
+			// this is a debug DLL or something, it is not a numbered hot DLL
+			continue;
+		}
+
+		// Check the time stamp for this file
+		const FDateTime FoundFileTime = IFileManager::Get().GetTimeStamp(*FoundFilePath);
+		if (ensure(FoundFileTime != -1.0))
+		{
+			// Was this file modified more recently than our others?
+			if (FoundFileTime > NewestFoundFileTime)
+			{
+				bFound = true;
+				NewestFoundFileTime = FoundFileTime;
+			}
+		}
+		else
+		{
+			// File wasn't found, should never happen as we searched for these files just now
+		}
+	}
+
+	return bFound;
+}
 
-void FModuleManager::AddModule( const FName InModuleName )
+void FModuleManager::AddModule(const FName InModuleName)
 {
 	// Do we already know about this module?  If not, we'll create information for this module now.
-	if( ensureMsg( InModuleName != NAME_None, TEXT( "FModuleManager::AddModule() was called with an invalid module name (empty string or 'None'.)  This is not allowed." ) ) &&
-		!Modules.Contains( InModuleName ) )
+	if (!((ensureMsg(InModuleName != NAME_None, TEXT("FModuleManager::AddModule() was called with an invalid module name (empty string or 'None'.)  This is not allowed.")) &&
+		!Modules.Contains(InModuleName))))
+	{
+		return;
+	}
+
+	TSharedRef<FModuleInfo> ModuleInfo(new FModuleInfo());
+	
+	// Make sure module info is added to known modules and proper delegates are fired on exit.
+	struct FAtExit
 	{
-		// Add this module to the set of modules that we know about
-		TSharedRef< FModuleInfo > ModuleInfo( new FModuleInfo() );
+		FAtExit(const FName InModuleName, const TSharedRef<FModuleInfo>& InModuleInfo)
+			: ModuleName(InModuleName)
+			, ModuleInfo(InModuleInfo)
+		{ }
+
+		~FAtExit()
+		{
+			// Update hash table
+			FModuleManager::Get().Modules.Add(ModuleName, ModuleInfo);
+
+			// List of known modules has changed.  Fire callbacks.
+			FModuleManager::Get().ModulesChangedEvent.Broadcast(ModuleName, EModuleChangeReason::PluginDirectoryChanged);
+		}
+
+		FName ModuleName;
+		TSharedRef<FModuleInfo> ModuleInfo;
+	} AtExit(InModuleName, ModuleInfo);
 
 #if !IS_MONOLITHIC
-		FString ModuleNameString = InModuleName.ToString();
+	FString ModuleNameString = InModuleName.ToString();
 
-		TMap<FName, FString> ModulePathMap;
-		FindModulePaths(*ModuleNameString, ModulePathMap);
+	TMap<FName, FString> ModulePathMap;
+	FindModulePaths(*ModuleNameString, ModulePathMap);
 
-		if(ModulePathMap.Num() == 1)
-		{
-			// Add this module to the set of modules that we know about
-			ModuleInfo->OriginalFilename = TMap<FName, FString>::TConstIterator(ModulePathMap).Value();
-			ModuleInfo->Filename = ModuleInfo->OriginalFilename;
-
-			// When iterating on code during development, it's possible there are multiple rolling versions of this
-			// module's DLL file.  This can happen if the programmer is recompiling DLLs while the game is loaded.  In
-			// this case, we want to load the newest iteration of the DLL file, so that behavior is the same after
-			// restarting the application.
-			{
-				// NOTE: We leave this enabled in UE_BUILD_SHIPPING editor builds so module authors can iterate on custom modules
+	if (ModulePathMap.Num() != 1)
+	{
+		return;
+	}
+
+	// Add this module to the set of modules that we know about
+	ModuleInfo->OriginalFilename = TMap<FName, FString>::TConstIterator(ModulePathMap).Value();
+	ModuleInfo->Filename = ModuleInfo->OriginalFilename;
+
+	// When iterating on code during development, it's possible there are multiple rolling versions of this
+	// module's DLL file.  This can happen if the programmer is recompiling DLLs while the game is loaded.  In
+	// this case, we want to load the newest iteration of the DLL file, so that behavior is the same after
+	// restarting the application.
+
+	// NOTE: We leave this enabled in UE_BUILD_SHIPPING editor builds so module authors can iterate on custom modules
 #if !(UE_BUILD_SHIPPING || UE_BUILD_TEST) || (UE_BUILD_SHIPPING && WITH_EDITOR)
-				// In some cases, sadly, modules may be loaded before appInit() is called.  We can't cleanly support rolling files for those modules.
-				{
-					// First, check to see if the module we added already exists on disk
-					const FDateTime OriginalModuleFileTime = IFileManager::Get().GetTimeStamp( *ModuleInfo->OriginalFilename );
-					if( OriginalModuleFileTime != FDateTime::MinValue() )
-					{
-						const FString ModuleName = *InModuleName.ToString();
-						const int32 MatchPos = ModuleInfo->OriginalFilename.Find( ModuleName, ESearchCase::IgnoreCase, ESearchDir::FromEnd);
-						if( ensureMsgf( MatchPos != INDEX_NONE, TEXT("Could not find module name '%s' in module filename '%s'"), *ModuleName, *ModuleInfo->OriginalFilename ) )
-						{
-							const int32 SuffixPos = MatchPos + ModuleName.Len();
+	// In some cases, sadly, modules may be loaded before appInit() is called.  We can't cleanly support rolling files for those modules.
 
-							const FString Prefix = ModuleInfo->OriginalFilename.Left( SuffixPos );
-							const FString Suffix = ModuleInfo->OriginalFilename.Right( ModuleInfo->OriginalFilename.Len() - SuffixPos );
+	// First, check to see if the module we added already exists on disk
+	const FDateTime OriginalModuleFileTime = IFileManager::Get().GetTimeStamp(*ModuleInfo->OriginalFilename);
+	if (OriginalModuleFileTime == FDateTime::MinValue())
+	{
+		return;
+	}
 
-							const FString ModuleFileSearchString = FString::Printf( TEXT( "%s-*%s" ), *Prefix, *Suffix);
-							const FString ModuleFileSearchDirectory = FPaths::GetPath(ModuleFileSearchString);
+	const FString ModuleName = *InModuleName.ToString();
+	const int32 MatchPos = ModuleInfo->OriginalFilename.Find(ModuleName, ESearchCase::IgnoreCase, ESearchDir::FromEnd);
+	if (!ensureMsgf(MatchPos != INDEX_NONE, TEXT("Could not find module name '%s' in module filename '%s'"), *ModuleName, *ModuleInfo->OriginalFilename))
+	{
+		return;
+	}
 
-							// Search for module files
-							TArray<FString> FoundFiles;
-							IFileManager::Get().FindFiles( FoundFiles, *ModuleFileSearchString, true, false );
+	const int32 SuffixPos = MatchPos + ModuleName.Len();
 
-							// Figure out what the newest module file is
-							int32 NewestFoundFileIndex = INDEX_NONE;
-							FDateTime NewestFoundFileTime = OriginalModuleFileTime;
-							if( FoundFiles.Num() > 0 )
-							{
-								for( int32 CurFileIndex = 0; CurFileIndex < FoundFiles.Num(); ++CurFileIndex )
-								{
-									// FoundFiles contains file names with no directory information, but we need the full path up
-									// to the file, so we'll prefix it back on if we have a path.
-									const FString& FoundFile = FoundFiles[ CurFileIndex ];
-									const FString FoundFilePath = ModuleFileSearchDirectory.IsEmpty() ? FoundFile : ( ModuleFileSearchDirectory / FoundFile );
-
-									// need to reject some files here that are not numbered...release executables, do have a suffix, so we need to make sure we don't find the debug version
-									check(FoundFilePath.Len() > Prefix.Len() + Suffix.Len());
-									FString Center = FoundFilePath.Mid(Prefix.Len(), FoundFilePath.Len() - Prefix.Len() - Suffix.Len());
-									check(Center.StartsWith(TEXT("-"), ESearchCase::CaseSensitive)); // a minus sign is still considered numeric, so we can leave it.
-									if (!Center.IsNumeric())
-									{
-										// this is a debug DLL or something, it is not a numbered hot DLL
-										continue;
-									}
-
-
-									// Check the time stamp for this file
-									const FDateTime FoundFileTime = IFileManager::Get().GetTimeStamp( *FoundFilePath );
-									if( ensure( FoundFileTime != -1.0 ) )
-									{
-										// Was this file modified more recently than our others?
-										if( FoundFileTime > NewestFoundFileTime )
-										{
-											NewestFoundFileIndex = CurFileIndex;
-											NewestFoundFileTime = FoundFileTime;
-										}
-									}
-									else
-									{
-										// File wasn't found, should never happen as we searched for these files just now
-									}
-								}
-
-
-								// Did we find a variant of the module file that is newer than our original file?
-								if( NewestFoundFileIndex != INDEX_NONE )
-								{
-									// Update the module working file name to the most recently-modified copy of that module
-									const FString& NewestModuleFilename = FoundFiles[ NewestFoundFileIndex ];
-									const FString NewestModuleFilePath = ModuleFileSearchDirectory.IsEmpty() ? NewestModuleFilename : ( ModuleFileSearchDirectory / NewestModuleFilename );
-									ModuleInfo->Filename = NewestModuleFilePath;
-								}
-								else
-								{
-									// No variants were found that were newer than the original module file name, so
-									// we'll continue to use that!
-								}
-							}
-						}
-					}
-				}
-#endif	// !(UE_BUILD_SHIPPING || UE_BUILD_TEST)
-			}
-		}
-#endif	// !IS_MONOLITHIC
+	const FString Prefix = ModuleInfo->OriginalFilename.Left(SuffixPos);
+	const FString Suffix = ModuleInfo->OriginalFilename.Right(ModuleInfo->OriginalFilename.Len() - SuffixPos);
+
+	const FString ModuleFileSearchString = FString::Printf(TEXT("%s-*%s"), *Prefix, *Suffix);
+
+	// Search for module files
+	TArray<FString> FoundFiles;
+	IFileManager::Get().FindFiles(FoundFiles, *ModuleFileSearchString, true, false);
+	if (FoundFiles.Num() != 0)
+	{
+		return;
+	}
+
+	const FString ModuleFileSearchDirectory = FPaths::GetPath(ModuleFileSearchString);
 
-		// Update hash table
-		Modules.Add( InModuleName, ModuleInfo );
+	FString NewestModuleFilename;
+	bool bFoundNewestFile = FindNewestModuleFile(FoundFiles, OriginalModuleFileTime, ModuleFileSearchDirectory, Prefix, Suffix, NewestModuleFilename);
 
-		// List of known modules has changed.  Fire callbacks.
-		ModulesChangedEvent.Broadcast( InModuleName, EModuleChangeReason::PluginDirectoryChanged );
+	// Did we find a variant of the module file that is newer than our original file?
+	if (!bFoundNewestFile)
+	{
+		// No variants were found that were newer than the original module file name, so
+		// we'll continue to use that!
+		return;
 	}
+
+	// Update the module working file name to the most recently-modified copy of that module
+	const FString NewestModuleFilePath = ModuleFileSearchDirectory.IsEmpty() ? NewestModuleFilename : (ModuleFileSearchDirectory / NewestModuleFilename);
+	ModuleInfo->Filename = NewestModuleFilePath;
+#endif	// !(UE_BUILD_SHIPPING || UE_BUILD_TEST)
+#endif	// !IS_MONOLITHIC
 }
 
 
 TSharedPtr<IModuleInterface> FModuleManager::LoadModule( const FName InModuleName, bool bWasReloaded )
 {
+	// FModuleManager is not thread-safe
+	ensure(IsInGameThread());
+
 	EModuleLoadResult FailureReason;
 	return LoadModuleWithFailureReason(InModuleName, FailureReason, bWasReloaded );
 }
@@ -263,7 +302,7 @@ TSharedPtr<IModuleInterface> FModuleManager::LoadModuleChecked( const FName InMo
 }
 
 
-TSharedPtr<IModuleInterface> FModuleManager::LoadModuleWithFailureReason( const FName InModuleName, EModuleLoadResult& OutFailureReason, bool bWasReloaded /*=false*/ )
+TSharedPtr<IModuleInterface> FModuleManager::LoadModuleWithFailureReason(const FName InModuleName, EModuleLoadResult& OutFailureReason, bool bWasReloaded /*=false*/)
 {
 	DECLARE_SCOPE_CYCLE_COUNTER(TEXT("Module Load"), STAT_ModuleLoad, STATGROUP_LoadTime);
 
@@ -276,139 +315,139 @@ TSharedPtr<IModuleInterface> FModuleManager::LoadModuleWithFailureReason( const
 	// Grab the module info.  This has the file name of the module, as well as other info.
 	TSharedRef< FModuleInfo > ModuleInfo = Modules.FindChecked( InModuleName );
 
-	if( ModuleInfo->Module.IsValid() )
+	if (ModuleInfo->Module.IsValid())
 	{
 		// Assign the already loaded module into the return value, otherwise the return value gives the impression the module failed load!
 		LoadedModule = ModuleInfo->Module;
 	}
 	else
 	{
-		// Make sure this isn't a module that we had previously loaded, and then unloaded at shutdown time.
-		//
-		// If this assert goes off, your trying to load a module during the shutdown phase that was already
-		// cleaned up.  The easiest way to fix this is to change your code to query for an already-loaded
-		// module instead of trying to load it directly.
-		checkf( ( !ModuleInfo->bWasUnloadedAtShutdown ), TEXT( "Attempted to load module '%s' that was already unloaded at shutdown.  FModuleManager::LoadModule() was called to load a module that was previously loaded, and was unloaded at shutdown time.  If this assert goes off, your trying to load a module during the shutdown phase that was already cleaned up.  The easiest way to fix this is to change your code to query for an already-loaded module instead of trying to load it directly." ), *InModuleName.ToString() );
-
-		// Check if we're statically linked with the module.  Those modules register with the module manager using a static variable,
-		// so hopefully we already know about the name of the module and how to initialize it.
-		const FInitializeStaticallyLinkedModule* ModuleInitializerPtr = StaticallyLinkedModuleInitializers.Find( InModuleName );
-		if( ModuleInitializerPtr != NULL )
-		{
-			const FInitializeStaticallyLinkedModule& ModuleInitializer( *ModuleInitializerPtr );
+	// Make sure this isn't a module that we had previously loaded, and then unloaded at shutdown time.
+	//
+	// If this assert goes off, your trying to load a module during the shutdown phase that was already
+	// cleaned up.  The easiest way to fix this is to change your code to query for an already-loaded
+	// module instead of trying to load it directly.
+	checkf((!ModuleInfo->bWasUnloadedAtShutdown), TEXT("Attempted to load module '%s' that was already unloaded at shutdown.  FModuleManager::LoadModule() was called to load a module that was previously loaded, and was unloaded at shutdown time.  If this assert goes off, your trying to load a module during the shutdown phase that was already cleaned up.  The easiest way to fix this is to change your code to query for an already-loaded module instead of trying to load it directly."), *InModuleName.ToString());
+
+	// Check if we're statically linked with the module.  Those modules register with the module manager using a static variable,
+	// so hopefully we already know about the name of the module and how to initialize it.
+	const FInitializeStaticallyLinkedModule* ModuleInitializerPtr = StaticallyLinkedModuleInitializers.Find(InModuleName);
+		if (ModuleInitializerPtr != NULL)
+	{
+		const FInitializeStaticallyLinkedModule& ModuleInitializer(*ModuleInitializerPtr);
 
-			// Initialize the module!
-			ModuleInfo->Module = MakeShareable( ModuleInitializer.Execute() );
+		// Initialize the module!
+		ModuleInfo->Module = MakeShareable(ModuleInitializer.Execute());
 
-			if( ModuleInfo->Module.IsValid() )
-			{
-				// Startup the module
-				ModuleInfo->Module->StartupModule();
+		if (ModuleInfo->Module.IsValid())
+		{
+			// Startup the module
+			ModuleInfo->Module->StartupModule();
 
-				// Module was started successfully!  Fire callbacks.
-				ModulesChangedEvent.Broadcast( InModuleName, EModuleChangeReason::ModuleLoaded );
+			// Module was started successfully!  Fire callbacks.
+			ModulesChangedEvent.Broadcast(InModuleName, EModuleChangeReason::ModuleLoaded);
 
-				// Set the return parameter
-				LoadedModule = ModuleInfo->Module;
-			}
-			else
-			{
-				UE_LOG(LogModuleManager, Warning, TEXT( "ModuleManager: Unable to load module '%s' because InitializeModule function failed (returned NULL pointer.)" ), *InModuleName.ToString() );
-				OutFailureReason = EModuleLoadResult::FailedToInitialize;
-			}
+			// Set the return parameter
+			LoadedModule = ModuleInfo->Module;
+		}
+		else
+		{
+			UE_LOG(LogModuleManager, Warning, TEXT("ModuleManager: Unable to load module '%s' because InitializeModule function failed (returned NULL pointer.)"), *InModuleName.ToString());
+			OutFailureReason = EModuleLoadResult::FailedToInitialize;
 		}
+	}
 #if IS_MONOLITHIC
 		else
 		{
-			// Monolithic builds that do not have the initializer were *not found* during the build step, so return FileNotFound
-			// (FileNotFound is an acceptable error in some case - ie loading a content only project)
-			UE_LOG(LogModuleManager, Warning, TEXT( "ModuleManager: Module '%s' not found - its StaticallyLinkedModuleInitializers function is null." ), *InModuleName.ToString() );
-			OutFailureReason = EModuleLoadResult::FileNotFound;
+	// Monolithic builds that do not have the initializer were *not found* during the build step, so return FileNotFound
+	// (FileNotFound is an acceptable error in some case - ie loading a content only project)
+	UE_LOG(LogModuleManager, Warning, TEXT("ModuleManager: Module '%s' not found - its StaticallyLinkedModuleInitializers function is null."), *InModuleName.ToString());
+	OutFailureReason = EModuleLoadResult::FileNotFound;
 		}
 #endif
 #if !IS_MONOLITHIC
 		else
 		{
-			// Make sure that any UObjects that need to be registered were already processed before we go and
-			// load another module.  We just do this so that we can easily tell whether UObjects are present
-			// in the module being loaded.
-			if( bCanProcessNewlyLoadedObjects )
-			{
-				ProcessLoadedObjectsCallback.Broadcast();
-			}
+	// Make sure that any UObjects that need to be registered were already processed before we go and
+	// load another module.  We just do this so that we can easily tell whether UObjects are present
+	// in the module being loaded.
+	if (bCanProcessNewlyLoadedObjects)
+	{
+		ProcessLoadedObjectsCallback.Broadcast();
+	}
 
-			// Try to dynamically load the DLL
+	// Try to dynamically load the DLL
 
-			UE_LOG(LogModuleManager, Verbose, TEXT( "ModuleManager: Load Module '%s' DLL '%s'" ), *InModuleName.ToString(), *ModuleInfo->Filename);
+	UE_LOG(LogModuleManager, Verbose, TEXT("ModuleManager: Load Module '%s' DLL '%s'"), *InModuleName.ToString(), *ModuleInfo->Filename);
 
-			// Determine which file to load for this module.
-			const FString ModuleFileToLoad = FPaths::ConvertRelativePathToFull(ModuleInfo->Filename);
+	// Determine which file to load for this module.
+	const FString ModuleFileToLoad = FPaths::ConvertRelativePathToFull(ModuleInfo->Filename);
 
-			// Clear the handle and set it again below if the module is successfully loaded
-			ModuleInfo->Handle = NULL;
+	// Clear the handle and set it again below if the module is successfully loaded
+	ModuleInfo->Handle = NULL;
 
-			// Skip this check if file manager has not yet been initialized
-			if ( FPaths::FileExists(ModuleFileToLoad) )
-			{
-				if ( CheckModuleCompatibility(*ModuleFileToLoad) )
-				{
-					ModuleInfo->Handle = FPlatformProcess::GetDllHandle(*ModuleFileToLoad);
-					if( ModuleInfo->Handle != NULL )
-					{
-						// First things first.  If the loaded DLL has UObjects in it, then their generated code's
-						// static initialization will have run during the DLL loading phase, and we'll need to
-						// go in and make sure those new UObject classes are properly registered.
-						{
-							// Sometimes modules are loaded before even the UObject systems are ready.  We need to assume
-							// these modules aren't using UObjects.
-							if( bCanProcessNewlyLoadedObjects)
-							{
-								// OK, we've verified that loading the module caused new UObject classes to be
-								// registered, so we'll treat this module as a module with UObjects in it.
-								ProcessLoadedObjectsCallback.Broadcast();
-							}
-						}
+	// Skip this check if file manager has not yet been initialized
+			if (FPaths::FileExists(ModuleFileToLoad))
+	{
+				if (CheckModuleCompatibility(*ModuleFileToLoad))
+	{
+	ModuleInfo->Handle = FPlatformProcess::GetDllHandle(*ModuleFileToLoad);
+					if (ModuleInfo->Handle != NULL)
+	{
+	// First things first.  If the loaded DLL has UObjects in it, then their generated code's
+	// static initialization will have run during the DLL loading phase, and we'll need to
+	// go in and make sure those new UObject classes are properly registered.
+	{
+		// Sometimes modules are loaded before even the UObject systems are ready.  We need to assume
+		// these modules aren't using UObjects.
+		if (bCanProcessNewlyLoadedObjects)
+		{
+			// OK, we've verified that loading the module caused new UObject classes to be
+			// registered, so we'll treat this module as a module with UObjects in it.
+			ProcessLoadedObjectsCallback.Broadcast();
+		}
+	}
 
-						// Find our "InitializeModule" global function, which must exist for all module DLLs
+	// Find our "InitializeModule" global function, which must exist for all module DLLs
 						FInitializeModuleFunctionPtr InitializeModuleFunctionPtr =
-							( FInitializeModuleFunctionPtr )FPlatformProcess::GetDllExport( ModuleInfo->Handle, TEXT( "InitializeModule" ) );
-						if( InitializeModuleFunctionPtr != NULL )
+							(FInitializeModuleFunctionPtr)FPlatformProcess::GetDllExport(ModuleInfo->Handle, TEXT("InitializeModule"));
+						if (InitializeModuleFunctionPtr != NULL)
 						{
 							// Initialize the module!
-							ModuleInfo->Module = MakeShareable( InitializeModuleFunctionPtr() );
+							ModuleInfo->Module = MakeShareable(InitializeModuleFunctionPtr());
 
-							if( ModuleInfo->Module.IsValid() )
+							if (ModuleInfo->Module.IsValid())
 							{
 								// Startup the module
 								ModuleInfo->Module->StartupModule();
 
 								// Module was started successfully!  Fire callbacks.
-								ModulesChangedEvent.Broadcast( InModuleName, EModuleChangeReason::ModuleLoaded );
+								ModulesChangedEvent.Broadcast(InModuleName, EModuleChangeReason::ModuleLoaded);
 
 								// Set the return parameter
 								LoadedModule = ModuleInfo->Module;
 							}
 							else
-							{
-								UE_LOG(LogModuleManager, Warning, TEXT( "ModuleManager: Unable to load module '%s' because InitializeModule function failed (returned NULL pointer.)" ), *ModuleFileToLoad );
+	{
+								UE_LOG(LogModuleManager, Warning, TEXT("ModuleManager: Unable to load module '%s' because InitializeModule function failed (returned NULL pointer.)"), *ModuleFileToLoad);
 
-								FPlatformProcess::FreeDllHandle( ModuleInfo->Handle );
-								ModuleInfo->Handle = NULL;
-								OutFailureReason = EModuleLoadResult::FailedToInitialize;
+		FPlatformProcess::FreeDllHandle(ModuleInfo->Handle);
+		ModuleInfo->Handle = NULL;
+		OutFailureReason = EModuleLoadResult::FailedToInitialize;
 							}
-						}
+	}
 						else
-						{
-							UE_LOG(LogModuleManager, Warning, TEXT( "ModuleManager: Unable to load module '%s' because InitializeModule function was not found." ), *ModuleFileToLoad );
+	{
+							UE_LOG(LogModuleManager, Warning, TEXT("ModuleManager: Unable to load module '%s' because InitializeModule function was not found."), *ModuleFileToLoad);
 
-							FPlatformProcess::FreeDllHandle( ModuleInfo->Handle );
-							ModuleInfo->Handle = NULL;
-							OutFailureReason = EModuleLoadResult::FailedToInitialize;
+		FPlatformProcess::FreeDllHandle(ModuleInfo->Handle);
+		ModuleInfo->Handle = NULL;
+		OutFailureReason = EModuleLoadResult::FailedToInitialize;
 						}
 					}
 					else
 					{
-						UE_LOG(LogModuleManager, Warning, TEXT( "ModuleManager: Unable to load module '%s' because the file couldn't be loaded by the OS." ), *ModuleFileToLoad );
+						UE_LOG(LogModuleManager, Warning, TEXT("ModuleManager: Unable to load module '%s' because the file couldn't be loaded by the OS."), *ModuleFileToLoad);
 						OutFailureReason = EModuleLoadResult::CouldNotBeLoadedByOS;
 					}
 				}
@@ -420,10 +459,10 @@ TSharedPtr<IModuleInterface> FModuleManager::LoadModuleWithFailureReason( const
 			}
 			else
 			{
-				UE_LOG(LogModuleManager, Warning, TEXT( "ModuleManager: Unable to load module '%s' because the file '%s' was not found." ), *InModuleName.ToString(), *ModuleFileToLoad );
+				UE_LOG(LogModuleManager, Warning, TEXT("ModuleManager: Unable to load module '%s' because the file '%s' was not found."), *InModuleName.ToString(), *ModuleFileToLoad);
 				OutFailureReason = EModuleLoadResult::FileNotFound;
 			}
-		}
+	}
 #endif
 	}
 
@@ -434,7 +473,7 @@ TSharedPtr<IModuleInterface> FModuleManager::LoadModuleWithFailureReason( const
 bool FModuleManager::UnloadModule( const FName InModuleName, bool bIsShutdown )
 {
 	// Do we even know about this module?
-	TSharedRef< FModuleInfo >* ModuleInfoPtr = Modules.Find( InModuleName );
+	auto ModuleInfoPtr = FindModule(InModuleName);
 	if( ModuleInfoPtr != NULL )
 	{
 		TSharedRef< FModuleInfo >& ModuleInfo( *ModuleInfoPtr );
@@ -496,7 +535,7 @@ bool FModuleManager::UnloadModule( const FName InModuleName, bool bIsShutdown )
 void FModuleManager::AbandonModule( const FName InModuleName )
 {
 	// Do we even know about this module?
-	TSharedRef< FModuleInfo >* ModuleInfoPtr = Modules.Find( InModuleName );
+	auto ModuleInfoPtr = FindModule(InModuleName);
 	if( ModuleInfoPtr != NULL )
 	{
 		TSharedRef< FModuleInfo >& ModuleInfo( *ModuleInfoPtr );
@@ -519,6 +558,8 @@ void FModuleManager::AbandonModule( const FName InModuleName )
 
 void FModuleManager::UnloadModulesAtShutdown()
 {
+	ensure(IsInGameThread());
+
 	struct FModulePair
 	{
 		FName ModuleName;
@@ -535,9 +576,9 @@ void FModuleManager::UnloadModulesAtShutdown()
 		}
 	};
 	TArray<FModulePair> ModulesToUnload;
-	for( FModuleMap::TConstIterator ModuleIt( Modules ); ModuleIt; ++ModuleIt )
+	for (const auto ModuleIt : Modules)
 	{
-		TSharedRef< FModuleInfo > ModuleInfo( ModuleIt.Value() );
+		TSharedRef< FModuleInfo > ModuleInfo( ModuleIt.Value );
 
 		// Only if already loaded
 		if( ModuleInfo->Module.IsValid() )
@@ -545,7 +586,7 @@ void FModuleManager::UnloadModulesAtShutdown()
 			// Only if the module supports shutting down in this phase
 			if( ModuleInfo->Module->SupportsAutomaticShutdown() )
 			{
-				new (ModulesToUnload) FModulePair(ModuleIt.Key(), ModuleIt.Value()->LoadOrder);
+				new (ModulesToUnload) FModulePair(ModuleIt.Key, ModuleIt.Value->LoadOrder);
 			}
 		}
 	}
@@ -564,7 +605,7 @@ void FModuleManager::UnloadModulesAtShutdown()
 TSharedPtr<IModuleInterface> FModuleManager::GetModule( const FName InModuleName )
 {
 	// Do we even know about this module?
-	const TSharedRef<FModuleInfo>* ModuleInfo = Modules.Find(InModuleName);
+	auto ModuleInfo = FindModule(InModuleName);
 
 	if (ModuleInfo == nullptr)
 	{
@@ -706,62 +747,62 @@ bool FModuleManager::Exec( UWorld* Inworld, const TCHAR* Cmd, FOutputDevice& Ar
 bool FModuleManager::QueryModule( const FName InModuleName, FModuleStatus& OutModuleStatus )
 {
 	// Do we even know about this module?
-	const TSharedRef< FModuleInfo >* ModuleInfoPtr = Modules.Find( InModuleName );
-	if( ModuleInfoPtr != NULL )
+	auto ModuleInfoPtr = FindModule(InModuleName);
+
+	if (!ModuleInfoPtr)
 	{
-		const TSharedRef< FModuleInfo >& ModuleInfo( *ModuleInfoPtr );
+		// Not known to us
+		return false;
+	}
 
-		OutModuleStatus.Name = InModuleName.ToString();
-		OutModuleStatus.FilePath = FPaths::ConvertRelativePathToFull(ModuleInfo->Filename);
-		OutModuleStatus.bIsLoaded = ( ModuleInfo->Module.IsValid() != false );
+	const TSharedRef< FModuleInfo >& ModuleInfo( *ModuleInfoPtr );
 
-		if( OutModuleStatus.bIsLoaded )
-		{
-			OutModuleStatus.bIsGameModule = ModuleInfo->Module->IsGameModule();
-		}
+	OutModuleStatus.Name = InModuleName.ToString();
+	OutModuleStatus.FilePath = FPaths::ConvertRelativePathToFull(ModuleInfo->Filename);
+	OutModuleStatus.bIsLoaded = ( ModuleInfo->Module.IsValid() != false );
 
-		return true;
+	if( OutModuleStatus.bIsLoaded )
+	{
+		OutModuleStatus.bIsGameModule = ModuleInfo->Module->IsGameModule();
 	}
 
-	// Not known to us
-	return false;
+	return true;
 }
 
 
 void FModuleManager::QueryModules( TArray< FModuleStatus >& OutModuleStatuses )
 {
 	OutModuleStatuses.Reset();
-
-	for( FModuleMap::TConstIterator ModuleIt( Modules ); ModuleIt; ++ModuleIt )
+	FReadScopeLock Lock(&ModulesCriticalSection);
+	for (const auto ModuleIt : Modules)
 	{
-		const FName CurModuleFName = ModuleIt.Key();
-		const TSharedRef< FModuleInfo > CurModule = ModuleIt.Value();
+		const auto CurModuleFName = ModuleIt.Key;
+		const auto CurModule = ModuleIt.Value;
 
 		FModuleStatus ModuleStatus;
-		{
-			ModuleStatus.Name = CurModuleFName.ToString();
-			ModuleStatus.FilePath = FPaths::ConvertRelativePathToFull(CurModule->Filename);
-			ModuleStatus.bIsLoaded = CurModule->Module.IsValid();
+		ModuleStatus.Name = CurModuleFName.ToString();
+		ModuleStatus.FilePath = FPaths::ConvertRelativePathToFull(CurModule->Filename);
+		ModuleStatus.bIsLoaded = CurModule->Module.IsValid();
 
-			if( ModuleStatus.bIsLoaded  )
-			{
-				ModuleStatus.bIsGameModule = CurModule->Module->IsGameModule();
-			}
+		if( ModuleStatus.bIsLoaded  )
+		{
+			ModuleStatus.bIsGameModule = CurModule->Module->IsGameModule();
 		}
+
 		OutModuleStatuses.Add( ModuleStatus );
 	}
 }
 
 
-FString FModuleManager::GetModuleFilename(FName ModuleName) const
+FString FModuleManager::GetModuleFilename(FName ModuleName)
 {
-	return Modules.FindChecked(ModuleName)->Filename;
+	return FindModuleChecked(ModuleName)->Filename;
 }
 
 
 void FModuleManager::SetModuleFilename(FName ModuleName, const FString& Filename)
 {
-	auto& Module = Modules.FindChecked(ModuleName);
+	auto Module = FindModuleChecked(ModuleName);
 	Module->Filename = Filename;
 	// If it's a new module then also update its OriginalFilename
 	if (Module->OriginalFilename.IsEmpty())
@@ -886,12 +927,12 @@ void FModuleManager::FindModulePathsInDirectory(const FString& InDirectoryName,
 
 void FModuleManager::UnloadOrAbandonModuleWithCallback( const FName InModuleName, FOutputDevice &Ar )
 {
-	TSharedRef< FModuleInfo > Module = Modules.FindChecked( InModuleName );
+	auto Module = FindModuleChecked(InModuleName);
 	
 	Module->Module->PreUnloadCallback();
 
 	const bool bIsHotReloadable = DoesLoadedModuleHaveUObjects( InModuleName );
-	if( !bIsHotReloadable && Module->Module->SupportsDynamicReloading() )
+	if( bIsHotReloadable && Module->Module->SupportsDynamicReloading() )
 	{
 		if( !UnloadModule( InModuleName ))
 		{
@@ -926,7 +967,7 @@ bool FModuleManager::LoadModuleWithCallback( const FName InModuleName, FOutputDe
 
 void FModuleManager::MakeUniqueModuleFilename( const FName InModuleName, FString& UniqueSuffix, FString& UniqueModuleFileName )
 {
-	TSharedRef<FModuleInfo> Module = Modules.FindChecked(InModuleName);
+	auto Module = FindModuleChecked(InModuleName);
 
 	do
 	{
@@ -1034,3 +1075,136 @@ bool FModuleManager::DoesLoadedModuleHaveUObjects( const FName ModuleName )
 
 	return false;
 }
+
+TSharedRef<FModuleManager::FModuleInfo> FModuleManager::GetOrCreateModule(FName InModuleName)
+{
+	check(IsInGameThread());
+	ensureMsg(InModuleName != NAME_None, TEXT("FModuleManager::AddModule() was called with an invalid module name (empty string or 'None'.)  This is not allowed."));
+
+	if (Modules.Contains(InModuleName))
+	{
+		return FindModuleChecked(InModuleName);
+	}
+
+	// Add this module to the set of modules that we know about
+	TSharedRef< FModuleInfo > ModuleInfo(new FModuleInfo());
+
+#if !IS_MONOLITHIC
+	FString ModuleNameString = InModuleName.ToString();
+
+	TMap<FName, FString> ModulePathMap;
+	FindModulePaths(*ModuleNameString, ModulePathMap);
+
+	if (ModulePathMap.Num() != 1)
+	{
+		return ModuleInfo;
+	}
+
+	// Add this module to the set of modules that we know about
+	ModuleInfo->OriginalFilename = TMap<FName, FString>::TConstIterator(ModulePathMap).Value();
+	ModuleInfo->Filename = ModuleInfo->OriginalFilename;
+
+	// When iterating on code during development, it's possible there are multiple rolling versions of this
+	// module's DLL file.  This can happen if the programmer is recompiling DLLs while the game is loaded.  In
+	// this case, we want to load the newest iteration of the DLL file, so that behavior is the same after
+	// restarting the application.
+
+	// NOTE: We leave this enabled in UE_BUILD_SHIPPING editor builds so module authors can iterate on custom modules
+#if !(UE_BUILD_SHIPPING || UE_BUILD_TEST) || (UE_BUILD_SHIPPING && WITH_EDITOR)
+	// In some cases, sadly, modules may be loaded before appInit() is called.  We can't cleanly support rolling files for those modules.
+
+	// First, check to see if the module we added already exists on disk
+	const FDateTime OriginalModuleFileTime = IFileManager::Get().GetTimeStamp(*ModuleInfo->OriginalFilename);
+	if (OriginalModuleFileTime == FDateTime::MinValue())
+	{
+		return ModuleInfo;
+	}
+
+	const FString ModuleName = *InModuleName.ToString();
+	const int32 MatchPos = ModuleInfo->OriginalFilename.Find(ModuleName, ESearchCase::IgnoreCase, ESearchDir::FromEnd);
+	if (!ensureMsgf(MatchPos != INDEX_NONE, TEXT("Could not find module name '%s' in module filename '%s'"), *ModuleName, *ModuleInfo->OriginalFilename))
+	{
+		return ModuleInfo;
+	}
+
+	const int32 SuffixPos = MatchPos + ModuleName.Len();
+
+	const FString Prefix = ModuleInfo->OriginalFilename.Left(SuffixPos);
+	const FString Suffix = ModuleInfo->OriginalFilename.Right(ModuleInfo->OriginalFilename.Len() - SuffixPos);
+
+	const FString ModuleFileSearchString = FString::Printf(TEXT("%s-*%s"), *Prefix, *Suffix);
+	const FString ModuleFileSearchDirectory = FPaths::GetPath(ModuleFileSearchString);
+
+	// Search for module files
+	TArray<FString> FoundFiles;
+	IFileManager::Get().FindFiles(FoundFiles, *ModuleFileSearchString, true, false);
+
+	if (FoundFiles.Num() == 0)
+	{
+		return ModuleInfo;
+	}
+
+	// Figure out what the newest module file is
+	int32 NewestFoundFileIndex = INDEX_NONE;
+	FDateTime NewestFoundFileTime = OriginalModuleFileTime;
+
+	for (int32 CurFileIndex = 0; CurFileIndex < FoundFiles.Num(); ++CurFileIndex)
+	{
+		// FoundFiles contains file names with no directory information, but we need the full path up
+		// to the file, so we'll prefix it back on if we have a path.
+		const FString& FoundFile = FoundFiles[CurFileIndex];
+		const FString FoundFilePath = ModuleFileSearchDirectory.IsEmpty() ? FoundFile : (ModuleFileSearchDirectory / FoundFile);
+
+		// need to reject some files here that are not numbered...release executables, do have a suffix, so we need to make sure we don't find the debug version
+		check(FoundFilePath.Len() > Prefix.Len() + Suffix.Len());
+		FString Center = FoundFilePath.Mid(Prefix.Len(), FoundFilePath.Len() - Prefix.Len() - Suffix.Len());
+		check(Center.StartsWith(TEXT("-"), ESearchCase::CaseSensitive)); // a minus sign is still considered numeric, so we can leave it.
+		if (!Center.IsNumeric())
+		{
+			// this is a debug DLL or something, it is not a numbered hot DLL
+			continue;
+		}
+
+
+		// Check the time stamp for this file
+		const FDateTime FoundFileTime = IFileManager::Get().GetTimeStamp(*FoundFilePath);
+		if (ensure(FoundFileTime != -1.0))
+		{
+			// Was this file modified more recently than our others?
+			if (FoundFileTime > NewestFoundFileTime)
+			{
+				NewestFoundFileIndex = CurFileIndex;
+				NewestFoundFileTime = FoundFileTime;
+			}
+		}
+		else
+		{
+			// File wasn't found, should never happen as we searched for these files just now
+		}
+	}
+
+
+	// Did we find a variant of the module file that is newer than our original file?
+	if (NewestFoundFileIndex != INDEX_NONE)
+	{
+		// Update the module working file name to the most recently-modified copy of that module
+		const FString& NewestModuleFilename = FoundFiles[NewestFoundFileIndex];
+		const FString NewestModuleFilePath = ModuleFileSearchDirectory.IsEmpty() ? NewestModuleFilename : (ModuleFileSearchDirectory / NewestModuleFilename);
+		ModuleInfo->Filename = NewestModuleFilePath;
+	}
+	else
+	{
+		// No variants were found that were newer than the original module file name, so
+		// we'll continue to use that!
+	}
+#endif	// !(UE_BUILD_SHIPPING || UE_BUILD_TEST)
+#endif	// !IS_MONOLITHIC
+
+	return ModuleInfo;
+}
+
+int32 FModuleManager::GetModuleCount()
+{
+	FReadScopeLock Lock(&ModulesCriticalSection);
+	return Modules.Num();
+}
diff --git a/Engine/Source/Runtime/Core/Private/UObject/UnrealNames.cpp b/Engine/Source/Runtime/Core/Private/UObject/UnrealNames.cpp
--- a/Engine/Source/Runtime/Core/Private/UObject/UnrealNames.cpp
+++ b/Engine/Source/Runtime/Core/Private/UObject/UnrealNames.cpp
@@ -393,7 +393,7 @@ FName::FName( const TCHAR* Name, int32 InNumber, EFindName FindType )
 
 
 /**
- * Constructor used by ULinkerLoad when loading its name table; Creates an FName with an instance
+ * Constructor used by FLinkerLoad when loading its name table; Creates an FName with an instance
  * number of 0 that does not attempt to split the FName into string and number portions.
  */
 FName::FName( ELinkerNameTableConstructor, const WIDECHAR* Name )
@@ -402,7 +402,7 @@ FName::FName( ELinkerNameTableConstructor, const WIDECHAR* Name )
 }
 
 /**
- * Constructor used by ULinkerLoad when loading its name table; Creates an FName with an instance
+ * Constructor used by FLinkerLoad when loading its name table; Creates an FName with an instance
  * number of 0 that does not attempt to split the FName into string and number portions.
  */
 FName::FName( ELinkerNameTableConstructor, const ANSICHAR* Name )
diff --git a/Engine/Source/Runtime/Core/Private/Windows/WindowsRunnableThread.h b/Engine/Source/Runtime/Core/Private/Windows/WindowsRunnableThread.h
--- a/Engine/Source/Runtime/Core/Private/Windows/WindowsRunnableThread.h
+++ b/Engine/Source/Runtime/Core/Private/Windows/WindowsRunnableThread.h
@@ -104,7 +104,7 @@ public:
 		}
 	}
 
-	virtual void Suspend( bool bShouldPause = 1 ) override
+	virtual void Suspend( bool bShouldPause = true ) override
 	{
 		check(Thread);
 		if (bShouldPause == true)
diff --git a/Engine/Source/Runtime/Core/Public/Containers/ChunkedArray.h b/Engine/Source/Runtime/Core/Public/Containers/ChunkedArray.h
--- a/Engine/Source/Runtime/Core/Public/Containers/ChunkedArray.h
+++ b/Engine/Source/Runtime/Core/Public/Containers/ChunkedArray.h
@@ -115,6 +115,19 @@ public:
 		return Chunks.GetAllocatedSize();
 	}
 
+	/**
+	* Tests if index is valid, i.e. greater than zero and less than number of
+	* elements in array.
+	*
+	* @param Index Index to test.
+	*
+	* @returns True if index is valid. False otherwise.
+	*/
+	FORCEINLINE bool IsValidIndex(int32 Index) const
+	{
+		return Index >= 0 && Index < NumElements;
+	}
+
 	/**
 	 * Adds a new item to the end of the chunked array.
 	 *
diff --git a/Engine/Source/Runtime/Core/Public/CoreGlobals.h b/Engine/Source/Runtime/Core/Public/CoreGlobals.h
--- a/Engine/Source/Runtime/Core/Public/CoreGlobals.h
+++ b/Engine/Source/Runtime/Core/Public/CoreGlobals.h
@@ -105,7 +105,6 @@ extern CORE_API bool GIsClient;
 extern CORE_API bool GIsServer;
 extern CORE_API bool GIsCriticalError;
 extern CORE_API bool GIsRunning;
-extern CORE_API bool GIsGarbageCollecting;
 extern CORE_API bool GIsDuplicatingClassForReinstancing;
 
 /**
@@ -168,7 +167,7 @@ extern CORE_API const TCHAR* GForeignEngineDir;
 extern CORE_API FExec* GDebugToolExec;
 
 /** Whether we're currently in the async loading code path or not */
-extern CORE_API bool GIsAsyncLoading;
+extern CORE_API bool(*IsAsyncLoading)();
 
 /** Whether the editor is currently loading a package or not */
 extern CORE_API bool GIsEditorLoadingPackage;
@@ -196,9 +195,6 @@ extern CORE_API FString GSystemStartTime;
 /** Whether we are still in the initial loading process. */
 extern CORE_API bool GIsInitialLoad;
 
-/** true when we are routing ConditionalPostLoad/PostLoad to objects */
-extern CORE_API bool GIsRoutingPostLoad;
-
 #if WITH_HOT_RELOAD_CTORS
 /** true when we are retrieving VTablePtr from UClass */
 extern CORE_API bool GIsRetrievingVTablePtr;
diff --git a/Engine/Source/Runtime/Core/Public/HAL/ThreadingBase.h b/Engine/Source/Runtime/Core/Public/HAL/ThreadingBase.h
--- a/Engine/Source/Runtime/Core/Public/HAL/ThreadingBase.h
+++ b/Engine/Source/Runtime/Core/Public/HAL/ThreadingBase.h
@@ -835,6 +835,135 @@ public:
 	}
 };
 
+/**
+ * Object synchronizing read access (from any thread) and write access (only from game thread).
+ * Multiple simultaneous reads from different threads are allowed. Allows multiple writes from
+ * reentrant calls on main thread. Locking write doesn't prevent reads from main thread.
+ **/
+class FMultiReaderSingleWriterGT
+{
+public:
+	FMultiReaderSingleWriterGT();
+
+	/** Protect data from modification while reading. If issued on game thread, doesn't wait for write to finish. */
+	void LockRead();
+
+	/** Ends protecting data from modification while reading. */
+	void UnlockRead();
+
+	/** Protect data from reading while modifying. Can be called only on game thread. Reentrant. */
+	void LockWrite();
+
+	/** Ends protecting data from reading while modifying. */
+	void UnlockWrite();
+
+private:
+	struct FMRSWCriticalSection
+	{
+		FMRSWCriticalSection()
+			: Action(0)
+			, ReadCounter(0)
+			, WriteCounter(0)
+		{ }
+
+		int32 Action;
+		FThreadSafeCounter ReadCounter;
+		FThreadSafeCounter WriteCounter;
+	};
+
+	FMRSWCriticalSection CriticalSection;
+
+	static const int32 WritingAction = -1;
+	static const int32 NoAction = 0;
+	static const int32 ReadingAction = 1;
+
+	TFunction<bool()> CanRead;
+	TFunction<bool()> CanWrite;
+};
+
+class FReadScopeLock
+{
+public:
+
+	/**
+	* Constructor that performs a lock on the synchronization object
+	*
+	* @param InSynchObject The synchronization object to manage
+	*/
+	FReadScopeLock(FMultiReaderSingleWriterGT* InSynchObject)
+		: SynchObject(InSynchObject)
+	{
+		check(SynchObject);
+		SynchObject->LockRead();
+	}
+
+	/** Destructor that performs a release on the synchronization object. */
+	~FReadScopeLock()
+	{
+		check(SynchObject);
+		SynchObject->UnlockRead();
+	}
+private:
+
+	/** Default constructor (hidden on purpose). */
+	FReadScopeLock();
+
+	/** Copy constructor( hidden on purpose). */
+	FReadScopeLock(FReadScopeLock* InScopeLock);
+
+	/** Assignment operator (hidden on purpose). */
+	FReadScopeLock& operator=(FReadScopeLock& InScopeLock)
+	{
+		return *this;
+	}
+
+private:
+
+	// Holds the synchronization object to aggregate and scope manage.
+	FMultiReaderSingleWriterGT* SynchObject;
+};
+
+class FWriteScopeLock
+{
+public:
+
+	/**
+	* Constructor that performs a lock on the synchronization object
+	*
+	* @param InSynchObject The synchronization object to manage
+	*/
+	FWriteScopeLock(FMultiReaderSingleWriterGT* InSynchObject)
+		: SynchObject(InSynchObject)
+	{
+		check(SynchObject);
+		SynchObject->LockWrite();
+	}
+
+	/** Destructor that performs a release on the synchronization object. */
+	~FWriteScopeLock()
+	{
+		check(SynchObject);
+		SynchObject->UnlockWrite();
+	}
+private:
+
+	/** Default constructor (hidden on purpose). */
+	FWriteScopeLock();
+
+	/** Copy constructor( hidden on purpose). */
+	FWriteScopeLock(FWriteScopeLock* InScopeLock);
+
+	/** Assignment operator (hidden on purpose). */
+	FWriteScopeLock& operator=(FWriteScopeLock& InScopeLock)
+	{
+		return *this;
+	}
+
+private:
+
+	// Holds the synchronization object to aggregate and scope manage.
+	FMultiReaderSingleWriterGT* SynchObject;
+};
 
 /**
  * Implements a scope lock.
@@ -912,6 +1041,9 @@ extern CORE_API bool IsInParallelRenderingThread();
 // Unlike IsInRenderingThread, this will always return false if we are running single threaded. It only returns true if this is actually a separate rendering thread. Mostly useful for checks
 extern CORE_API bool IsInActualRenderingThread();
 
+/** @return True if called from the async loading thread if it's enabled, otherwise if called from game thread while is async loading code. */
+extern CORE_API bool (*IsInAsyncLoadingThread)();
+
 /** Thread used for rendering */
 extern CORE_API FRunnableThread* GRenderingThread;
 
diff --git a/Engine/Source/Runtime/Core/Public/Misc/CallbackDevice.h b/Engine/Source/Runtime/Core/Public/Misc/CallbackDevice.h
--- a/Engine/Source/Runtime/Core/Public/Misc/CallbackDevice.h
+++ b/Engine/Source/Runtime/Core/Public/Misc/CallbackDevice.h
@@ -61,7 +61,7 @@ public:
 	DECLARE_DELEGATE_RetVal_ThreeParams(EAppReturnType::Type, FOnModalMessageBox, EAppMsgType::Type, const FText&, const FText&);
 
 	// Callback for PER_MODULE_BOILERPLATE macro's GObjectArrayForDebugVisualizers
-	DECLARE_DELEGATE_RetVal(TArray<class UObjectBase*>*, FObjectArrayForDebugVisualizersDelegate);
+	DECLARE_DELEGATE_RetVal(class UObjectBase***, FObjectArrayForDebugVisualizersDelegate);
 
 	// Called in PER_MODULE_BOILERPLATE macro.
 	static FObjectArrayForDebugVisualizersDelegate& GetObjectArrayForDebugVisualizersDelegate();
diff --git a/Engine/Source/Runtime/Core/Public/Modules/Boilerplate/ModuleBoilerplate.h b/Engine/Source/Runtime/Core/Public/Modules/Boilerplate/ModuleBoilerplate.h
--- a/Engine/Source/Runtime/Core/Public/Modules/Boilerplate/ModuleBoilerplate.h
+++ b/Engine/Source/Runtime/Core/Public/Modules/Boilerplate/ModuleBoilerplate.h
@@ -34,6 +34,6 @@
 	TArray<FNameEntry const*>* GFNameTableForDebuggerVisualizers = FName::GetNameTableForDebuggerVisualizers_ST(); \
 	FNameEntry*** GFNameTableForDebuggerVisualizers_MT = FName::GetNameTableForDebuggerVisualizers_MT(); \
 	int32*** GSerialNumberBlocksForDebugVisualizers = FCoreDelegates::GetSerialNumberBlocksForDebugVisualizersDelegate().IsBound() ? FCoreDelegates::GetSerialNumberBlocksForDebugVisualizersDelegate().Execute() : NULL; \
-	TArray<UObjectBase*>* GObjectArrayForDebugVisualizers = FCoreDelegates::GetObjectArrayForDebugVisualizersDelegate().IsBound() ? FCoreDelegates::GetObjectArrayForDebugVisualizersDelegate().Execute() : NULL; \
+	UObjectBase*** GObjectArrayForDebugVisualizers = FCoreDelegates::GetObjectArrayForDebugVisualizersDelegate().IsBound() ? FCoreDelegates::GetObjectArrayForDebugVisualizersDelegate().Execute() : NULL; \
 	bool GFNameDebuggerVisualizersIsUE3=false; \
 	REPLACEMENT_OPERATOR_NEW_AND_DELETE
diff --git a/Engine/Source/Runtime/Core/Public/Modules/ModuleManager.h b/Engine/Source/Runtime/Core/Public/Modules/ModuleManager.h
--- a/Engine/Source/Runtime/Core/Public/Modules/ModuleManager.h
+++ b/Engine/Source/Runtime/Core/Public/Modules/ModuleManager.h
@@ -113,10 +113,7 @@ public:
 	/**
 	 * Adds a module to our list of modules, unless it's already known.
 	 *
-	 * This method is used by the plug-in manager to register a plug-in module.
-	 *
 	 * @param InModuleName The base name of the module file.  Should not include path, extension or platform/configuration info.  This is just the "name" part of the module file name.  Names should be globally unique.
-	 * @param InBinariesDirectory The directory where to find this file, or an empty string to search in the default locations.  This parameter is used by the plugin system to locate plugin binaries.
 	 */
 	void AddModule( const FName InModuleName );
 
@@ -138,7 +135,7 @@ public:
 	 * @return true if module is currently loaded, false otherwise.
 	 * @see AbandonModule, LoadModule, LoadModuleWithFailureReason, UnloadModule
 	 */
-	bool IsModuleLoaded( const FName InModuleName ) const;
+	bool IsModuleLoaded( const FName InModuleName );
 
 	/**
 	 * Loads the specified module.
@@ -316,10 +313,7 @@ public:
 	 *
 	 * @return The number of modules.
 	 */
-	int32 GetModuleCount( ) const
-	{
-		return Modules.Num();
-	}
+	int32 GetModuleCount( );
 
 	/**
 	 * Unloads modules during the shutdown process.
@@ -389,7 +383,7 @@ public:
 	static const TCHAR *GetUBTConfiguration( );
 
 	/** Gets the filename for a module. The return value is a full path of a module known to the module manager. */
-	FString GetModuleFilename(FName ModuleName) const;
+	FString GetModuleFilename(FName ModuleName);
 
 	/** Sets the filename for a module. The module is not reloaded immediately, but the new name will be used for subsequent unload/load events. */
 	void SetModuleFilename(FName ModuleName, const FString& Filename);
@@ -503,6 +497,10 @@ public:
 	void MakeUniqueModuleFilename( const FName InModuleName, FString& UniqueSuffix, FString& UniqueModuleFileName );
 
 private:
+	/** Thread safe module finding routine. */
+	TSharedRef<FModuleInfo>* FindModule(FName InModuleName);
+	TSharedRef<FModuleInfo> FindModuleChecked(FName InModuleName);
+
 	/** Compares file versions between the current executing engine version and the specified dll */
 	static bool CheckModuleCompatibility(const TCHAR *Filename);
 
@@ -516,6 +514,8 @@ private:
 	void FindModulePathsInDirectory(const FString &DirectoryName, bool bIsGameDirectory, const TCHAR *NamePattern, TMap<FName, FString> &OutModulePaths) const;
 
 private:
+	/** Gets module with given name from Modules or creates a new one. Doesn't modify Modules. */
+	TSharedRef<FModuleInfo> GetOrCreateModule(FName InModuleName);
 
 	/** Map of all modules.  Maps the case-insensitive module name to information about that module, loaded or not. */
 	FModuleMap Modules;
@@ -543,6 +543,9 @@ private:
 
 	/** Array of game binaries directories. */
 	TArray<FString> GameBinariesDirectories;
+
+	/** Critical section object controlling R/W access to Modules. */
+	FMultiReaderSingleWriterGT ModulesCriticalSection;
 };
 
 /**
diff --git a/Engine/Source/Runtime/Core/Public/Serialization/ArchiveBase.h b/Engine/Source/Runtime/Core/Public/Serialization/ArchiveBase.h
--- a/Engine/Source/Runtime/Core/Public/Serialization/ArchiveBase.h
+++ b/Engine/Source/Runtime/Core/Public/Serialization/ArchiveBase.h
@@ -348,11 +348,11 @@ public:
 	virtual FString GetArchiveName() const;
 
 	/**
-	 * If this archive is a ULinkerLoad or ULinkerSave, returns a pointer to the ULinker portion.
+	 * If this archive is a FLinkerLoad or FLinkerSave, returns a pointer to the ULinker portion.
 	 *
 	 * @return The linker, or nullptr if the archive is not a linker.
 	 */
-	virtual class ULinker* GetLinker()
+	virtual class FLinker* GetLinker()
 	{
 		return nullptr;
 	}
@@ -711,7 +711,7 @@ public:
 	}
 
 	/**
-	 * Sets the archive version number. Used by the code that makes sure that ULinkerLoad's 
+	 * Sets the archive version number. Used by the code that makes sure that FLinkerLoad's 
 	 * internal archive versions match the file reader it creates.
 	 *
 	 * @param UE4Ver	new version number
@@ -722,7 +722,7 @@ public:
 	}
 
 	/**
-	 * Sets the archive licensee version number. Used by the code that makes sure that ULinkerLoad's 
+	 * Sets the archive licensee version number. Used by the code that makes sure that FLinkerLoad's 
 	 * internal archive versions match the file reader it creates.
 	 *
 	 * @param Ver	new version number
@@ -733,7 +733,7 @@ public:
 	}
 
 	/**
-	* Sets the archive engine version. Used by the code that makes sure that ULinkerLoad's
+	* Sets the archive engine version. Used by the code that makes sure that FLinkerLoad's
 	* internal archive versions match the file reader it creates.
 	*
 	* @param InVer	new version number
@@ -1111,7 +1111,7 @@ public:
 
 	CORE_API virtual FString GetArchiveName() const override;
 
-	virtual class ULinker* GetLinker() override
+	virtual class FLinker* GetLinker() override
 	{
 		return InnerArchive.GetLinker();
 	}
diff --git a/Engine/Source/Runtime/Core/Public/Templates/RefCounting.h b/Engine/Source/Runtime/Core/Public/Templates/RefCounting.h
--- a/Engine/Source/Runtime/Core/Public/Templates/RefCounting.h
+++ b/Engine/Source/Runtime/Core/Public/Templates/RefCounting.h
@@ -130,6 +130,11 @@ public:
 		return InReference.Reference != nullptr;
 	}
 
+	FORCEINLINE bool IsValid() const
+	{
+		return Reference != nullptr;
+	}
+
 	void SafeRelease()
 	{
 		*this = nullptr;
