Core - FileProfiler misc fixes

diff --git a/Engine/Source/Runtime/Core/Public/HAL/IPlatformFileProfilerWrapper.h b/Engine/Source/Runtime/Core/Public/HAL/IPlatformFileProfilerWrapper.h
--- a/Engine/Source/Runtime/Core/Public/HAL/IPlatformFileProfilerWrapper.h
+++ b/Engine/Source/Runtime/Core/Public/HAL/IPlatformFileProfilerWrapper.h
@@ -34,9 +34,10 @@ struct FProfiledFileStatsBase
 	}
 };
 
+
 struct FProfiledFileStatsOp : public FProfiledFileStatsBase
 {
-	enum OpType
+	enum class EOpType : uint8
 	{
 		Unknown = 0,
 		Tell = 1,
@@ -44,7 +45,7 @@ struct FProfiledFileStatsOp : public FProfiledFileStatsBase
 		Read,
 		Write,
 		Size,
-		OpenRead,		
+		OpenRead,
 		OpenWrite,
 		Exists,
 		Delete,
@@ -62,7 +63,7 @@ struct FProfiledFileStatsOp : public FProfiledFileStatsBase
 	};
 
 	/** Operation type */
-	uint8 Type;
+	FProfiledFileStatsOp::EOpType Type;
 
 	/** Number of bytes processed */
 	int64 Bytes;
@@ -70,7 +71,7 @@ struct FProfiledFileStatsOp : public FProfiledFileStatsBase
 	/** The last time this operation was executed */
 	double LastOpTime;
 
-	FProfiledFileStatsOp( uint8 InType )
+	FProfiledFileStatsOp( FProfiledFileStatsOp::EOpType InType )
 		: Type( InType )
 		, Bytes( 0 )
 		, LastOpTime( 0.0 )
@@ -96,7 +97,7 @@ struct FProfiledFileStatsFileDetailed : public FProfiledFileStatsFileBase
 	: FProfiledFileStatsFileBase( Filename )
 	{}
 
-	FORCEINLINE FProfiledFileStatsOp* CreateOpStat( uint8 Type )
+	FORCEINLINE FProfiledFileStatsOp* CreateOpStat( FProfiledFileStatsOp::EOpType Type )
 	{
 		FScopeLock ScopeLock(&SynchronizationObject);
 		TSharedPtr< FProfiledFileStatsOp > Stat( new FProfiledFileStatsOp( Type ) );
@@ -112,16 +113,16 @@ struct FProfiledFileStatsFileSimple : public FProfiledFileStatsFileBase
 	FProfiledFileStatsFileSimple( const TCHAR* Filename )
 	: FProfiledFileStatsFileBase( Filename )
 	{
-		for( uint8 TypeIndex = 0; TypeIndex < FProfiledFileStatsOp::Count; TypeIndex++ )
+		for( uint8 TypeIndex = 0; TypeIndex < (uint8)FProfiledFileStatsOp::EOpType::Count; TypeIndex++ )
 		{
-			TSharedPtr< FProfiledFileStatsOp > Stat( new FProfiledFileStatsOp( TypeIndex ) );
+			TSharedPtr< FProfiledFileStatsOp > Stat( new FProfiledFileStatsOp( FProfiledFileStatsOp::EOpType( TypeIndex ) ) );
 			Children.Add( Stat );
 		}
 	}
 
-	FORCEINLINE FProfiledFileStatsOp* CreateOpStat( uint8 Type )
+	FORCEINLINE FProfiledFileStatsOp* CreateOpStat( FProfiledFileStatsOp::EOpType Type )
 	{
-		TSharedPtr< FProfiledFileStatsOp > Stat = Children[ Type ];
+		TSharedPtr< FProfiledFileStatsOp > Stat = Children[ (uint8)Type ];
 		Stat->LastOpTime = FPlatformTime::Seconds() * 1000.0;
 		if( Stat->StartTime == 0.0 )
 		{
@@ -149,28 +150,28 @@ public:
 
 	virtual int64		Tell() override
 	{
-		FProfiledFileStatsOp* Stat( FileStats->CreateOpStat( FProfiledFileStatsOp::Tell ) );
+		FProfiledFileStatsOp* Stat( FileStats->CreateOpStat( FProfiledFileStatsOp::EOpType::Tell ) );
 		int64 Result = FileHandle->Tell();
 		Stat->Duration += FPlatformTime::Seconds() * 1000.0 - Stat->LastOpTime;
 		return Result;
 	}
 	virtual bool		Seek(int64 NewPosition) override
 	{
-		FProfiledFileStatsOp* Stat( FileStats->CreateOpStat( FProfiledFileStatsOp::Seek ) );
+		FProfiledFileStatsOp* Stat( FileStats->CreateOpStat( FProfiledFileStatsOp::EOpType::Seek ) );
 		bool Result = FileHandle->Seek(NewPosition);
 		Stat->Duration += FPlatformTime::Seconds() * 1000.0 - Stat->LastOpTime;
 		return Result;
 	}
 	virtual bool		SeekFromEnd(int64 NewPositionRelativeToEnd) override
 	{
-		FProfiledFileStatsOp* Stat( FileStats->CreateOpStat( FProfiledFileStatsOp::Seek ) );
+		FProfiledFileStatsOp* Stat( FileStats->CreateOpStat( FProfiledFileStatsOp::EOpType::Seek ) );
 		bool Result = FileHandle->SeekFromEnd(NewPositionRelativeToEnd);
 		Stat->Duration += FPlatformTime::Seconds() * 1000.0 - Stat->LastOpTime;
 		return Result;
 	}
 	virtual bool		Read(uint8* Destination, int64 BytesToRead) override
 	{
-		FProfiledFileStatsOp* Stat( FileStats->CreateOpStat( FProfiledFileStatsOp::Read ) );
+		FProfiledFileStatsOp* Stat( FileStats->CreateOpStat( FProfiledFileStatsOp::EOpType::Read ) );
 		bool Result = FileHandle->Read(Destination, BytesToRead);
 		Stat->Duration += FPlatformTime::Seconds() * 1000.0 - Stat->LastOpTime;
 		Stat->Bytes += BytesToRead;
@@ -178,7 +179,7 @@ public:
 	}
 	virtual bool		Write(const uint8* Source, int64 BytesToWrite) override
 	{
-		FProfiledFileStatsOp* Stat( FileStats->CreateOpStat( FProfiledFileStatsOp::Write ) );
+		FProfiledFileStatsOp* Stat( FileStats->CreateOpStat( FProfiledFileStatsOp::EOpType::Write ) );
 		bool Result = FileHandle->Write(Source, BytesToWrite);
 		Stat->Duration += FPlatformTime::Seconds() * 1000.0 - Stat->LastOpTime;
 		Stat->Bytes += BytesToWrite;
@@ -186,7 +187,7 @@ public:
 	}
 	virtual int64		Size() override
 	{
-		FProfiledFileStatsOp* Stat( FileStats->CreateOpStat( FProfiledFileStatsOp::Size ) );
+		FProfiledFileStatsOp* Stat( FileStats->CreateOpStat( FProfiledFileStatsOp::EOpType::Size ) );
 		int64 Result = FileHandle->Size();
 		Stat->Duration += FPlatformTime::Seconds() * 1000.0 - Stat->LastOpTime;
 		return Result;
@@ -216,7 +217,7 @@ public:
 
 	virtual bool ShouldBeUsed(IPlatformFile* Inner, const TCHAR* CmdLine) const override
 	{
-		return false;
+		return FParse::Param( CmdLine, GetName() );
 	}
 
 	virtual bool Initialize(IPlatformFile* Inner, const TCHAR* CommandLineParam) override
@@ -283,7 +284,7 @@ public:
 	virtual bool		FileExists(const TCHAR* Filename) override
 	{
 		StatsType* FileStat = CreateStat( Filename );
-		FProfiledFileStatsOp* OpStat = FileStat->CreateOpStat( FProfiledFileStatsOp::Exists );
+		FProfiledFileStatsOp* OpStat = FileStat->CreateOpStat( FProfiledFileStatsOp::EOpType::Exists );
 		bool Result = LowerLevel->FileExists(Filename);
 		OpStat->Duration += FPlatformTime::Seconds() * 1000.0 - OpStat->LastOpTime;
 		return Result;
@@ -291,7 +292,7 @@ public:
 	virtual int64		FileSize(const TCHAR* Filename) override
 	{
 		StatsType* FileStat = CreateStat( Filename );
-		FProfiledFileStatsOp* OpStat = FileStat->CreateOpStat( FProfiledFileStatsOp::Size );
+		FProfiledFileStatsOp* OpStat = FileStat->CreateOpStat( FProfiledFileStatsOp::EOpType::Size );
 		int64 Result = LowerLevel->FileSize(Filename);
 		OpStat->Duration += FPlatformTime::Seconds() * 1000.0 - OpStat->LastOpTime;
 		return Result;
@@ -299,7 +300,7 @@ public:
 	virtual bool		DeleteFile(const TCHAR* Filename) override
 	{
 		StatsType* FileStat = CreateStat( Filename );
-		FProfiledFileStatsOp* OpStat = FileStat->CreateOpStat( FProfiledFileStatsOp::Delete );
+		FProfiledFileStatsOp* OpStat = FileStat->CreateOpStat( FProfiledFileStatsOp::EOpType::Delete );
 		bool Result = LowerLevel->DeleteFile(Filename);
 		OpStat->Duration += FPlatformTime::Seconds() * 1000.0 - OpStat->LastOpTime;
 		return Result;
@@ -307,7 +308,7 @@ public:
 	virtual bool		IsReadOnly(const TCHAR* Filename) override
 	{
 		StatsType* FileStat = CreateStat( Filename );
-		FProfiledFileStatsOp* OpStat = FileStat->CreateOpStat( FProfiledFileStatsOp::IsReadOnly );
+		FProfiledFileStatsOp* OpStat = FileStat->CreateOpStat( FProfiledFileStatsOp::EOpType::IsReadOnly );
 		bool Result = LowerLevel->IsReadOnly(Filename);
 		OpStat->Duration += FPlatformTime::Seconds() * 1000.0 - OpStat->LastOpTime;
 		return Result;
@@ -315,7 +316,7 @@ public:
 	virtual bool		MoveFile(const TCHAR* To, const TCHAR* From) override
 	{
 		StatsType* FileStat = CreateStat( From );
-		FProfiledFileStatsOp* OpStat = FileStat->CreateOpStat( FProfiledFileStatsOp::Move );
+		FProfiledFileStatsOp* OpStat = FileStat->CreateOpStat( FProfiledFileStatsOp::EOpType::Move );
 		bool Result = LowerLevel->MoveFile(To, From);
 		OpStat->Duration += FPlatformTime::Seconds() * 1000.0 - OpStat->LastOpTime;
 		return Result;
@@ -323,7 +324,7 @@ public:
 	virtual bool		SetReadOnly(const TCHAR* Filename, bool bNewReadOnlyValue) override
 	{
 		StatsType* FileStat = CreateStat( Filename );
-		FProfiledFileStatsOp* OpStat = FileStat->CreateOpStat( FProfiledFileStatsOp::SetReadOnly );
+		FProfiledFileStatsOp* OpStat = FileStat->CreateOpStat( FProfiledFileStatsOp::EOpType::SetReadOnly );
 		bool Result = LowerLevel->SetReadOnly(Filename, bNewReadOnlyValue);
 		OpStat->Duration += FPlatformTime::Seconds() * 1000.0 - OpStat->LastOpTime;
 		return Result;
@@ -331,7 +332,7 @@ public:
 	virtual FDateTime	GetTimeStamp(const TCHAR* Filename) override
 	{
 		StatsType* FileStat = CreateStat( Filename );
-		FProfiledFileStatsOp* OpStat = FileStat->CreateOpStat( FProfiledFileStatsOp::GetTimeStamp );
+		FProfiledFileStatsOp* OpStat = FileStat->CreateOpStat( FProfiledFileStatsOp::EOpType::GetTimeStamp );
 		double StartTime = FPlatformTime::Seconds();
 		FDateTime Result = LowerLevel->GetTimeStamp(Filename);
 		OpStat->Duration += FPlatformTime::Seconds() * 1000 - OpStat->LastOpTime;
@@ -340,7 +341,7 @@ public:
 	virtual void		SetTimeStamp(const TCHAR* Filename, FDateTime DateTime) override
 	{
 		StatsType* FileStat = CreateStat( Filename );
-		FProfiledFileStatsOp* OpStat = FileStat->CreateOpStat( FProfiledFileStatsOp::SetTimeStamp );
+		FProfiledFileStatsOp* OpStat = FileStat->CreateOpStat( FProfiledFileStatsOp::EOpType::SetTimeStamp );
 		double StartTime = FPlatformTime::Seconds();
 		LowerLevel->SetTimeStamp(Filename, DateTime);
 		OpStat->Duration += FPlatformTime::Seconds() * 1000 - OpStat->LastOpTime;
@@ -348,7 +349,7 @@ public:
 	virtual FDateTime	GetAccessTimeStamp(const TCHAR* Filename) override
 	{
 		StatsType* FileStat = CreateStat( Filename );
-		FProfiledFileStatsOp* OpStat = FileStat->CreateOpStat( FProfiledFileStatsOp::GetTimeStamp );
+		FProfiledFileStatsOp* OpStat = FileStat->CreateOpStat( FProfiledFileStatsOp::EOpType::GetTimeStamp );
 		double StartTime = FPlatformTime::Seconds();
 		FDateTime Result = LowerLevel->GetAccessTimeStamp(Filename);
 		OpStat->Duration += FPlatformTime::Seconds() * 1000 - OpStat->LastOpTime;
@@ -357,7 +358,7 @@ public:
 	virtual FString	GetFilenameOnDisk(const TCHAR* Filename) override
 	{
 		StatsType* FileStat = CreateStat(Filename);
-		FProfiledFileStatsOp* OpStat = FileStat->CreateOpStat(FProfiledFileStatsOp::GetFilenameOnDisk);
+		FProfiledFileStatsOp* OpStat = FileStat->CreateOpStat( FProfiledFileStatsOp::EOpType::GetFilenameOnDisk );
 		double StartTime = FPlatformTime::Seconds();
 		FString Result = LowerLevel->GetFilenameOnDisk(Filename);
 		OpStat->Duration += FPlatformTime::Seconds() * 1000 - OpStat->LastOpTime;
@@ -366,7 +367,7 @@ public:
 	virtual IFileHandle*	OpenRead(const TCHAR* Filename, bool bAllowWrite = false) override
 	{
 		StatsType* FileStat = CreateStat( Filename );
-		FProfiledFileStatsOp* OpStat = FileStat->CreateOpStat( FProfiledFileStatsOp::OpenRead );
+		FProfiledFileStatsOp* OpStat = FileStat->CreateOpStat( FProfiledFileStatsOp::EOpType::OpenRead );
 		IFileHandle* Result = LowerLevel->OpenRead(Filename, bAllowWrite);
 		OpStat->Duration += FPlatformTime::Seconds() * 1000.0 - OpStat->LastOpTime;
 		return Result ? (new TProfiledFileHandle< StatsType >( Result, Filename, FileStat )) : Result;
@@ -374,7 +375,7 @@ public:
 	virtual IFileHandle*	OpenWrite(const TCHAR* Filename, bool bAppend = false, bool bAllowRead = false) override
 	{
 		StatsType* FileStat = CreateStat( Filename );
-		FProfiledFileStatsOp* OpStat = FileStat->CreateOpStat( FProfiledFileStatsOp::OpenWrite );
+		FProfiledFileStatsOp* OpStat = FileStat->CreateOpStat( FProfiledFileStatsOp::EOpType::OpenWrite );
 		IFileHandle* Result = LowerLevel->OpenWrite(Filename, bAppend, bAllowRead);
 		OpStat->Duration += FPlatformTime::Seconds() * 1000.0 - OpStat->LastOpTime;
 		return Result ? (new TProfiledFileHandle< StatsType >( Result, Filename, FileStat )) : Result;
@@ -383,7 +384,7 @@ public:
 	virtual bool		DirectoryExists(const TCHAR* Directory) override
 	{
 		StatsType* FileStat = CreateStat( Directory );
-		FProfiledFileStatsOp* OpStat = FileStat->CreateOpStat( FProfiledFileStatsOp::Exists );
+		FProfiledFileStatsOp* OpStat = FileStat->CreateOpStat( FProfiledFileStatsOp::EOpType::Exists );
 		bool Result = LowerLevel->DirectoryExists(Directory);
 		OpStat->Duration += FPlatformTime::Seconds() * 1000.0 - OpStat->LastOpTime;
 		return Result;
@@ -391,7 +392,7 @@ public:
 	virtual bool		CreateDirectory(const TCHAR* Directory) override
 	{
 		StatsType* FileStat = CreateStat( Directory );
-		FProfiledFileStatsOp* OpStat = FileStat->CreateOpStat( FProfiledFileStatsOp::Create );
+		FProfiledFileStatsOp* OpStat = FileStat->CreateOpStat( FProfiledFileStatsOp::EOpType::Create );
 		bool Result = LowerLevel->CreateDirectory(Directory);
 		OpStat->Duration += FPlatformTime::Seconds() * 1000.0 - OpStat->LastOpTime;
 		return Result;
@@ -399,7 +400,7 @@ public:
 	virtual bool		DeleteDirectory(const TCHAR* Directory) override
 	{
 		StatsType* FileStat = CreateStat( Directory );
-		FProfiledFileStatsOp* OpStat = FileStat->CreateOpStat( FProfiledFileStatsOp::Delete );
+		FProfiledFileStatsOp* OpStat = FileStat->CreateOpStat( FProfiledFileStatsOp::EOpType::Delete );
 		bool Result = LowerLevel->DeleteDirectory(Directory);
 		OpStat->Duration += FPlatformTime::Seconds() * 1000.0 - OpStat->LastOpTime;
 		return Result;
@@ -408,7 +409,7 @@ public:
 	virtual bool		IterateDirectory(const TCHAR* Directory, IPlatformFile::FDirectoryVisitor& Visitor) override
 	{
 		StatsType* FileStat = CreateStat( Directory );
-		FProfiledFileStatsOp* OpStat = FileStat->CreateOpStat( FProfiledFileStatsOp::Iterate );
+		FProfiledFileStatsOp* OpStat = FileStat->CreateOpStat( FProfiledFileStatsOp::EOpType::Iterate );
 		bool Result = LowerLevel->IterateDirectory( Directory, Visitor );
 		OpStat->Duration += FPlatformTime::Seconds() * 1000.0 - OpStat->LastOpTime;
 		return Result;
@@ -416,7 +417,7 @@ public:
 	virtual bool		IterateDirectoryRecursively(const TCHAR* Directory, IPlatformFile::FDirectoryVisitor& Visitor) override
 	{
 		StatsType* FileStat = CreateStat( Directory );
-		FProfiledFileStatsOp* OpStat = FileStat->CreateOpStat( FProfiledFileStatsOp::Iterate );
+		FProfiledFileStatsOp* OpStat = FileStat->CreateOpStat( FProfiledFileStatsOp::EOpType::Iterate );
 		bool Result = LowerLevel->IterateDirectoryRecursively( Directory, Visitor );
 		OpStat->Duration += FPlatformTime::Seconds() * 1000.0 - OpStat->LastOpTime;
 		return Result;
@@ -424,7 +425,7 @@ public:
 	virtual bool		DeleteDirectoryRecursively(const TCHAR* Directory) override
 	{
 		StatsType* FileStat = CreateStat( Directory );
-		FProfiledFileStatsOp* OpStat = FileStat->CreateOpStat( FProfiledFileStatsOp::Delete );
+		FProfiledFileStatsOp* OpStat = FileStat->CreateOpStat( FProfiledFileStatsOp::EOpType::Delete );
 		bool Result = LowerLevel->DeleteDirectoryRecursively( Directory );
 		OpStat->Duration += FPlatformTime::Seconds() * 1000.0 - OpStat->LastOpTime;
 		return Result;
@@ -432,7 +433,7 @@ public:
 	virtual bool		CopyFile(const TCHAR* To, const TCHAR* From) override
 	{
 		StatsType* FileStat = CreateStat( From );
-		FProfiledFileStatsOp* OpStat = FileStat->CreateOpStat( FProfiledFileStatsOp::Copy );
+		FProfiledFileStatsOp* OpStat = FileStat->CreateOpStat( FProfiledFileStatsOp::EOpType::Copy );
 		bool Result = LowerLevel->CopyFile( To, From );
 		OpStat->Duration += FPlatformTime::Seconds() * 1000.0 - OpStat->LastOpTime;
 		return Result;
diff --git a/Engine/Source/Runtime/Engine/Private/LevelTick.cpp b/Engine/Source/Runtime/Engine/Private/LevelTick.cpp
--- a/Engine/Source/Runtime/Engine/Private/LevelTick.cpp
+++ b/Engine/Source/Runtime/Engine/Private/LevelTick.cpp
@@ -954,39 +954,39 @@ static class FFileProfileWrapperExec: private FSelfRegisteringExec
 					FString EventName;
 					switch( FileOpStat->Type )
 					{
-						case FProfiledFileStatsOp::Tell:
+						case FProfiledFileStatsOp::EOpType::Tell:
 							EventName = TEXT("Tell"); break;
-						case FProfiledFileStatsOp::Seek:
+						case FProfiledFileStatsOp::EOpType::Seek:
 							EventName = TEXT("Seek"); break;
-						case FProfiledFileStatsOp::Read:
+						case FProfiledFileStatsOp::EOpType::Read:
 							EventName = FString::Printf( TEXT("Read (%lld)"), FileOpStat->Bytes ); break;
-						case FProfiledFileStatsOp::Write:
+						case FProfiledFileStatsOp::EOpType::Write:
 							EventName = FString::Printf( TEXT("Write (%lld)"), FileOpStat->Bytes ); break;
-						case FProfiledFileStatsOp::Size:
+						case FProfiledFileStatsOp::EOpType::Size:
 							EventName = TEXT("Size"); break;
-						case FProfiledFileStatsOp::OpenRead:
+						case FProfiledFileStatsOp::EOpType::OpenRead:
 							EventName = TEXT("OpenRead"); break;
-						case FProfiledFileStatsOp::OpenWrite:
+						case FProfiledFileStatsOp::EOpType::OpenWrite:
 							EventName = TEXT("OpenWrite"); break;
-						case FProfiledFileStatsOp::Exists:
+						case FProfiledFileStatsOp::EOpType::Exists:
 							EventName = TEXT("Exists"); break;
-						case FProfiledFileStatsOp::Delete:
+						case FProfiledFileStatsOp::EOpType::Delete:
 							EventName = TEXT("Delete"); break;
-						case FProfiledFileStatsOp::Move:
+						case FProfiledFileStatsOp::EOpType::Move:
 							EventName = TEXT("Move"); break;
-						case FProfiledFileStatsOp::IsReadOnly:
+						case FProfiledFileStatsOp::EOpType::IsReadOnly:
 							EventName = TEXT("IsReadOnly"); break;
-						case FProfiledFileStatsOp::SetReadOnly:
+						case FProfiledFileStatsOp::EOpType::SetReadOnly:
 							EventName = TEXT("SetReadOnly"); break;
-						case FProfiledFileStatsOp::GetTimeStamp:
+						case FProfiledFileStatsOp::EOpType::GetTimeStamp:
 							EventName = TEXT("GetTimeStamp"); break;
-						case FProfiledFileStatsOp::SetTimeStamp:
+						case FProfiledFileStatsOp::EOpType::SetTimeStamp:
 							EventName = TEXT("SetTimeStamp"); break;
-						case FProfiledFileStatsOp::Create:
+						case FProfiledFileStatsOp::EOpType::Create:
 							EventName = TEXT("Create"); break;
-						case FProfiledFileStatsOp::Copy:
+						case FProfiledFileStatsOp::EOpType::Copy:
 							EventName = TEXT("Copy"); break;
-						case FProfiledFileStatsOp::Iterate:
+						case FProfiledFileStatsOp::EOpType::Iterate:
 							EventName = TEXT("Iterate"); break;
 						default:
 							EventName = TEXT("Unknown"); break;
