ensureOnce and ensureOnceMsgf, which only log/break the first time they are hit.
Unification of ensureOnce macros.
EnsureNotFalse and EnsureNotFalseFormatted transformed into a single OptionallyLogFormattedEnsureMessageReturningFalse.

diff --git a/Engine/Source/Runtime/Core/Private/Misc/OutputDevice.cpp b/Engine/Source/Runtime/Core/Private/Misc/OutputDevice.cpp
--- a/Engine/Source/Runtime/Core/Private/Misc/OutputDevice.cpp
+++ b/Engine/Source/Runtime/Core/Private/Misc/OutputDevice.cpp
@@ -361,18 +361,17 @@ void VARARGS FDebug::AssertFailed(const ANSICHAR* Expr, const ANSICHAR* File, in
 }
 
 #if DO_CHECK || DO_GUARD_SLOW
-bool VARARGS FDebug::EnsureNotFalseFormatted( bool bExpressionResult, const ANSICHAR* Expr, const ANSICHAR* File, int32 Line, const TCHAR* FormattedMsg, ... )
+bool VARARGS FDebug::OptionallyLogFormattedEnsureMessageReturningFalse( bool bLog, const ANSICHAR* Expr, const ANSICHAR* File, int32 Line, const TCHAR* FormattedMsg, ... )
 {
-	const int32 TempStrSize = 4096;
-	TCHAR TempStr[ TempStrSize ];
-	GET_VARARGS( TempStr, TempStrSize, TempStrSize - 1, FormattedMsg, FormattedMsg );
-
-	if( bExpressionResult == 0 )
+	if (bLog)
 	{
+		const int32 TempStrSize = 4096;
+		TCHAR TempStr[ TempStrSize ];
+		GET_VARARGS( TempStr, TempStrSize, TempStrSize - 1, FormattedMsg, FormattedMsg );
 		EnsureFailed( Expr, File, Line, TempStr );
 	}
 	
-	return bExpressionResult;
+	return false;
 }
 #endif
 
diff --git a/Engine/Source/Runtime/Core/Public/Misc/AssertionMacros.h b/Engine/Source/Runtime/Core/Public/Misc/AssertionMacros.h
--- a/Engine/Source/Runtime/Core/Public/Misc/AssertionMacros.h
+++ b/Engine/Source/Runtime/Core/Public/Misc/AssertionMacros.h
@@ -3,6 +3,7 @@
 #pragma once
 
 #include "HAL/Platform.h"
+#include "HAL/PlatformMisc.h"
 #include "Misc/Build.h"
 
 
@@ -112,22 +113,42 @@
 
 #if DO_CHECK && !USING_CODE_ANALYSIS // The Visual Studio 2013 analyzer doesn't understand these complex conditionals
 
-	#define ensure( InExpression ) \
-		FDebug::EnsureNotFalse( ( ( InExpression ) != 0 ), #InExpression, __FILE__, __LINE__ )
+	namespace UE4Asserts_Private
+	{
+		// This is used by ensureOnce to generate a bool per instance
+		// by passing a lambda which will uniquely instantiate the template.
+		template <typename Type>
+		bool TrueOnFirstCallOnly(const Type&)
+		{
+			static bool bValue = true;
+			bool Result = bValue;
+			bValue = false;
+			return Result;
+		}
 
-	#define ensureMsg( InExpression, InMsg ) \
-		(((InExpression) == 0) ? FDebug::EnsureNotFalse(false, #InExpression, __FILE__, __LINE__, InMsg ) : true)
+		FORCEINLINE bool OptionallyDebugBreakAndPromptForRemoteReturningFalse(bool bBreak)
+		{
+			if (bBreak)
+			{
+				FPlatformMisc::DebugBreakAndPromptForRemoteReturningFalse();
+			}
+			return false;
+		}
+	}
 
-	#define ensureMsgf( InExpression, InFormat, ... ) \
-		(((InExpression) == 0) ? ( FDebug::EnsureNotFalseFormatted(false, #InExpression, __FILE__, __LINE__, InFormat, ##__VA_ARGS__ ) || FPlatformMisc::DebugBreakAndPromptForRemoteReturningFalse()) : true)
+	#define ensure(         InExpression                ) ((InExpression) != 0 || FDebug::OptionallyLogFormattedEnsureMessageReturningFalse(true,                                          #InExpression, __FILE__, __LINE__, TEXT("")               ) || FPlatformMisc::DebugBreakAndPromptForRemoteReturningFalse())
+	#define ensureMsg(      InExpression, InMsg         ) ((InExpression) != 0 || FDebug::OptionallyLogFormattedEnsureMessageReturningFalse(true,                                          #InExpression, __FILE__, __LINE__, InMsg                  ) || FPlatformMisc::DebugBreakAndPromptForRemoteReturningFalse())
+	#define ensureMsgf(     InExpression, InFormat, ... ) ((InExpression) != 0 || FDebug::OptionallyLogFormattedEnsureMessageReturningFalse(true,                                          #InExpression, __FILE__, __LINE__, InFormat, ##__VA_ARGS__) || FPlatformMisc::DebugBreakAndPromptForRemoteReturningFalse())
+	#define ensureOnce(     InExpression                ) ((InExpression) != 0 || FDebug::OptionallyLogFormattedEnsureMessageReturningFalse(UE4Asserts_Private::TrueOnFirstCallOnly([]{}), #InExpression, __FILE__, __LINE__, TEXT("")               ) || UE4Asserts_Private::OptionallyDebugBreakAndPromptForRemoteReturningFalse(UE4Asserts_Private::TrueOnFirstCallOnly([]{})))
+	#define ensureOnceMsgf( InExpression, InFormat, ... ) ((InExpression) != 0 || FDebug::OptionallyLogFormattedEnsureMessageReturningFalse(UE4Asserts_Private::TrueOnFirstCallOnly([]{}), #InExpression, __FILE__, __LINE__, InFormat, ##__VA_ARGS__) || UE4Asserts_Private::OptionallyDebugBreakAndPromptForRemoteReturningFalse(UE4Asserts_Private::TrueOnFirstCallOnly([]{})))
 
 #else	// DO_CHECK
 
-	#define ensure( InExpression ) ( ( InExpression ) != 0 )
-
-	#define ensureMsg( InExpression, InMsg ) ( ( InExpression ) != 0 )
-
-	#define ensureMsgf( InExpression, InFormat, ... ) ( ( InExpression ) != 0 )
+	#define ensure(         InExpression                ) ((InExpression) != 0)
+	#define ensureMsg(      InExpression, InMsg         ) ((InExpression) != 0)
+	#define ensureMsgf(     InExpression, InFormat, ... ) ((InExpression) != 0)
+	#define ensureOnce(     InExpression                ) ((InExpression) != 0)
+	#define ensureOnceMsgf( InExpression, InFormat, ... ) ((InExpression) != 0)
 
 #endif	// DO_CHECK
 
diff --git a/Engine/Source/Runtime/Core/Public/Misc/OutputDevice.h b/Engine/Source/Runtime/Core/Public/Misc/OutputDevice.h
--- a/Engine/Source/Runtime/Core/Public/Misc/OutputDevice.h
+++ b/Engine/Source/Runtime/Core/Public/Misc/OutputDevice.h
@@ -399,40 +399,17 @@ struct CORE_API FDebug
 	static void EnsureFailed( const ANSICHAR* Expr, const ANSICHAR* File, int32 Line, const TCHAR* Msg );
 
 	/**
-	 * Triggers a (potentially) non-fatal, non-intrusive error if bExpressionResult is zero
+	 * Logs an error if bLog is true, and returns false.  Takes a formatted string.
 	 *
-	 * @param	bExpressionResult	Expression that will trigger an error if resolves to zero
-	 * @param	Expr	Code expression ANSI string (#code)
-	 * @param	File	File name ANSI string (__FILE__)
-	 * @param	Line	Line number (__LINE__)
-	 * @param	Msg		Optional informative error message text
-	 */
-    static FORCEINLINE bool EnsureNotFalse( bool bExpressionResult, const ANSICHAR* Expr, const ANSICHAR* File, int32 Line, const TCHAR* Msg = TEXT( "" ) )
-	{
-	    if( bExpressionResult == 0 )
-	    {
-			EnsureFailed(Expr, File, Line, Msg);
-			if (!FPlatformMisc::IsDebuggerPresent())
-			{
-				FPlatformMisc::PromptForRemoteDebugging(false);
-			}
-
-			FPlatformMisc::DebugBreak();
-	    }
-    
-    	return bExpressionResult;
-	}
-
-	/**
-	 * Triggers a (potentially) non-fatal, non-intrusive error if bExpressionResult is zero.  Takes a formatted string.
-	 *
-	 * @param	bExpressionResult	Expression that will trigger an error if resolves to zero
+	 * @param	bLog	Log if true.
 	 * @param	Expr	Code expression ANSI string (#code)
 	 * @param	File	File name ANSI string (__FILE__)
 	 * @param	Line	Line number (__LINE__)
 	 * @param	FormattedMsg	Informative error message text with variable args
+	 *
+	 * @return false in all cases.
 	 */
-	static bool VARARGS EnsureNotFalseFormatted(bool bExpressionResult, const ANSICHAR* Expr, const ANSICHAR* File, int32 Line, const TCHAR* FormattedMsg, ...);
+	static bool VARARGS OptionallyLogFormattedEnsureMessageReturningFalse(bool bLog, const ANSICHAR* Expr, const ANSICHAR* File, int32 Line, const TCHAR* FormattedMsg, ...);
 #endif // DO_CHECK || DO_GUARD_SLOW
 };
 
