Function name change to allow Crash Reporter and minidump tool to continue to recognise ensure macros.

diff --git a/Engine/Source/Runtime/Core/Private/Misc/OutputDevice.cpp b/Engine/Source/Runtime/Core/Private/Misc/OutputDevice.cpp
--- a/Engine/Source/Runtime/Core/Private/Misc/OutputDevice.cpp
+++ b/Engine/Source/Runtime/Core/Private/Misc/OutputDevice.cpp
@@ -361,7 +361,7 @@ void VARARGS FDebug::AssertFailed(const ANSICHAR* Expr, const ANSICHAR* File, in
 }
 
 #if DO_CHECK || DO_GUARD_SLOW
-bool VARARGS FDebug::OptionallyLogFormattedEnsureMessageReturningFalse( bool bLog, const ANSICHAR* Expr, const ANSICHAR* File, int32 Line, const TCHAR* FormattedMsg, ... )
+bool VARARGS FDebug::EnsureNotFalse_OptionallyLogFormattedEnsureMessageReturningFalse( bool bLog, const ANSICHAR* Expr, const ANSICHAR* File, int32 Line, const TCHAR* FormattedMsg, ... )
 {
 	if (bLog)
 	{
diff --git a/Engine/Source/Runtime/Core/Public/Misc/AssertionMacros.h b/Engine/Source/Runtime/Core/Public/Misc/AssertionMacros.h
--- a/Engine/Source/Runtime/Core/Public/Misc/AssertionMacros.h
+++ b/Engine/Source/Runtime/Core/Public/Misc/AssertionMacros.h
@@ -136,11 +136,11 @@
 		}
 	}
 
-	#define ensure(         InExpression                ) ((InExpression) != 0 || FDebug::OptionallyLogFormattedEnsureMessageReturningFalse(true,                                          #InExpression, __FILE__, __LINE__, TEXT("")               ) || FPlatformMisc::DebugBreakAndPromptForRemoteReturningFalse())
-	#define ensureMsg(      InExpression, InMsg         ) ((InExpression) != 0 || FDebug::OptionallyLogFormattedEnsureMessageReturningFalse(true,                                          #InExpression, __FILE__, __LINE__, InMsg                  ) || FPlatformMisc::DebugBreakAndPromptForRemoteReturningFalse())
-	#define ensureMsgf(     InExpression, InFormat, ... ) ((InExpression) != 0 || FDebug::OptionallyLogFormattedEnsureMessageReturningFalse(true,                                          #InExpression, __FILE__, __LINE__, InFormat, ##__VA_ARGS__) || FPlatformMisc::DebugBreakAndPromptForRemoteReturningFalse())
-	#define ensureOnce(     InExpression                ) ((InExpression) != 0 || FDebug::OptionallyLogFormattedEnsureMessageReturningFalse(UE4Asserts_Private::TrueOnFirstCallOnly([]{}), #InExpression, __FILE__, __LINE__, TEXT("")               ) || UE4Asserts_Private::OptionallyDebugBreakAndPromptForRemoteReturningFalse(UE4Asserts_Private::TrueOnFirstCallOnly([]{})))
-	#define ensureOnceMsgf( InExpression, InFormat, ... ) ((InExpression) != 0 || FDebug::OptionallyLogFormattedEnsureMessageReturningFalse(UE4Asserts_Private::TrueOnFirstCallOnly([]{}), #InExpression, __FILE__, __LINE__, InFormat, ##__VA_ARGS__) || UE4Asserts_Private::OptionallyDebugBreakAndPromptForRemoteReturningFalse(UE4Asserts_Private::TrueOnFirstCallOnly([]{})))
+	#define ensure(         InExpression                ) ((InExpression) != 0 || FDebug::EnsureNotFalse_OptionallyLogFormattedEnsureMessageReturningFalse(true,                                          #InExpression, __FILE__, __LINE__, TEXT("")               ) || FPlatformMisc::DebugBreakAndPromptForRemoteReturningFalse())
+	#define ensureMsg(      InExpression, InMsg         ) ((InExpression) != 0 || FDebug::EnsureNotFalse_OptionallyLogFormattedEnsureMessageReturningFalse(true,                                          #InExpression, __FILE__, __LINE__, InMsg                  ) || FPlatformMisc::DebugBreakAndPromptForRemoteReturningFalse())
+	#define ensureMsgf(     InExpression, InFormat, ... ) ((InExpression) != 0 || FDebug::EnsureNotFalse_OptionallyLogFormattedEnsureMessageReturningFalse(true,                                          #InExpression, __FILE__, __LINE__, InFormat, ##__VA_ARGS__) || FPlatformMisc::DebugBreakAndPromptForRemoteReturningFalse())
+	#define ensureOnce(     InExpression                ) ((InExpression) != 0 || FDebug::EnsureNotFalse_OptionallyLogFormattedEnsureMessageReturningFalse(UE4Asserts_Private::TrueOnFirstCallOnly([]{}), #InExpression, __FILE__, __LINE__, TEXT("")               ) || UE4Asserts_Private::OptionallyDebugBreakAndPromptForRemoteReturningFalse(UE4Asserts_Private::TrueOnFirstCallOnly([]{})))
+	#define ensureOnceMsgf( InExpression, InFormat, ... ) ((InExpression) != 0 || FDebug::EnsureNotFalse_OptionallyLogFormattedEnsureMessageReturningFalse(UE4Asserts_Private::TrueOnFirstCallOnly([]{}), #InExpression, __FILE__, __LINE__, InFormat, ##__VA_ARGS__) || UE4Asserts_Private::OptionallyDebugBreakAndPromptForRemoteReturningFalse(UE4Asserts_Private::TrueOnFirstCallOnly([]{})))
 
 #else	// DO_CHECK
 
diff --git a/Engine/Source/Runtime/Core/Public/Misc/OutputDevice.h b/Engine/Source/Runtime/Core/Public/Misc/OutputDevice.h
--- a/Engine/Source/Runtime/Core/Public/Misc/OutputDevice.h
+++ b/Engine/Source/Runtime/Core/Public/Misc/OutputDevice.h
@@ -408,8 +408,10 @@ struct CORE_API FDebug
 	 * @param	FormattedMsg	Informative error message text with variable args
 	 *
 	 * @return false in all cases.
+	 *
+	 * Note: this crazy name is to ensure that the crash reporter recognizes it, which checks for functions in the callstack starting with 'EnsureNotFalse'.
 	 */
-	static bool VARARGS OptionallyLogFormattedEnsureMessageReturningFalse(bool bLog, const ANSICHAR* Expr, const ANSICHAR* File, int32 Line, const TCHAR* FormattedMsg, ...);
+	static bool VARARGS EnsureNotFalse_OptionallyLogFormattedEnsureMessageReturningFalse(bool bLog, const ANSICHAR* Expr, const ANSICHAR* File, int32 Line, const TCHAR* FormattedMsg, ...);
 #endif // DO_CHECK || DO_GUARD_SLOW
 };
 
