Allow HTML5 platform to properly create directories.

diff --git a/Engine/Source/Runtime/Core/Private/HTML5/HTML5PlatformFile.cpp b/Engine/Source/Runtime/Core/Private/HTML5/HTML5PlatformFile.cpp
--- a/Engine/Source/Runtime/Core/Private/HTML5/HTML5PlatformFile.cpp
+++ b/Engine/Source/Runtime/Core/Private/HTML5/HTML5PlatformFile.cpp
@@ -352,8 +352,7 @@ public:
 #if PLATFORM_HTML5_WIN32
 		return _mkdir(TCHAR_TO_ANSI(*NormalizeFilename(Directory))) || (errno == EEXIST);
 #else 
-        // @todo 
-        return false;  
+		return mkdir(TCHAR_TO_ANSI(*NormalizeFilename(Directory)), S_IRWXG) || (errno == EEXIST);
 #endif 
 	}
 
