Added editor project setting to define whether to display units in imperial or metric (or not at all)

Also added a default input unit that is used when no units are specified on a text input.

This addresses UE-11863 and UE-9313

diff --git a/Engine/Source/Developer/Settings/Private/SettingsModule.cpp b/Engine/Source/Developer/Settings/Private/SettingsModule.cpp
--- a/Engine/Source/Developer/Settings/Private/SettingsModule.cpp
+++ b/Engine/Source/Developer/Settings/Private/SettingsModule.cpp
@@ -86,6 +86,7 @@ public:
 		ProjectSettingsContainer->Describe(LOCTEXT("ProjectSettingsSubMenuLabel", "Project Settings"), LOCTEXT("ProjectSettingsSubMenuToolTip", "Change the settings of the currently loaded project"), NAME_None);
 		ProjectSettingsContainer->DescribeCategory("Project", LOCTEXT("ProjectProjectCategoryName", "Project"), LOCTEXT("ProjectProjectCategoryDescription", "Project settings description text here"));
 		ProjectSettingsContainer->DescribeCategory("Engine", LOCTEXT("ProjectEngineCategoryName", "Engine"), LOCTEXT("ProjectEngineCategoryDescription", "Project settings description text here"));
+		ProjectSettingsContainer->DescribeCategory("Editor", LOCTEXT("ProjectEditorCategoryName", "Editor"), LOCTEXT("ProjectEditorCategoryDescription", "Project settings specific to the editor"));
 		ProjectSettingsContainer->DescribeCategory("Platforms", LOCTEXT("ProjectPlatformsCategoryName", "Platforms"), LOCTEXT("ProjectPlatformsCategoryDescription", "Platform settings description text here"));
 		ProjectSettingsContainer->DescribeCategory("Plugins", LOCTEXT("ProjectPluginsCategoryName", "Plugins"), LOCTEXT("ProjectPluginsCategoryDescription", "Plugins settings description text here"));
 	}
diff --git a/Engine/Source/Editor/DetailCustomizations/Private/ComponentTransformDetails.cpp b/Engine/Source/Editor/DetailCustomizations/Private/ComponentTransformDetails.cpp
--- a/Engine/Source/Editor/DetailCustomizations/Private/ComponentTransformDetails.cpp
+++ b/Engine/Source/Editor/DetailCustomizations/Private/ComponentTransformDetails.cpp
@@ -360,8 +360,9 @@ void FComponentTransformDetails::GenerateChildContent( IDetailChildrenBuilder& C
 	if(!bHideLocationField)
 	{
 		TSharedPtr<INumericTypeInterface<float>> TypeInterface;
-		if( FUnitConversion::bIsUnitDisplayEnabled )
+		if( FUnitConversion::Settings().ShouldDisplayUnits() )
 		{
+			UseDefaultInputUnits();
 			TypeInterface = SharedThis(this);
 		}
 
@@ -425,7 +426,7 @@ void FComponentTransformDetails::GenerateChildContent( IDetailChildrenBuilder& C
 	if(!bHideRotationField)
 	{
 		TSharedPtr<INumericTypeInterface<float>> TypeInterface;
-		if( FUnitConversion::bIsUnitDisplayEnabled )
+		if( FUnitConversion::Settings().ShouldDisplayUnits() )
 		{
 			TypeInterface = MakeShareable( new TNumericUnitTypeInterface<float>(EUnit::Degrees) );
 		}
@@ -583,17 +584,17 @@ void FComponentTransformDetails::CacheCommonLocationUnits()
 		}
 	};
 
-	if (CachedLocation.X.IsSet())
+	if (CachedLocation.X.IsSet() && CachedLocation.X.GetValue() != 0)
 	{
-		SetLargestUnit(FUnitConversion::QuantizeUnitsToBestFit(CachedLocation.X.GetValue(), EUnit::Centimeters).Units);
+		SetLargestUnit(FUnitConversion::QuantizeUnitsToBestFit(CachedLocation.X.GetValue(), Units).Units);
 	}
-	if (CachedLocation.Y.IsSet())
+	if (CachedLocation.Y.IsSet() && CachedLocation.Y.GetValue() != 0)
 	{
-		SetLargestUnit(FUnitConversion::QuantizeUnitsToBestFit(CachedLocation.Y.GetValue(), EUnit::Centimeters).Units);
+		SetLargestUnit(FUnitConversion::QuantizeUnitsToBestFit(CachedLocation.Y.GetValue(), Units).Units);
 	}
-	if (CachedLocation.Z.IsSet())
+	if (CachedLocation.Z.IsSet() && CachedLocation.Z.GetValue() != 0)
 	{
-		SetLargestUnit(FUnitConversion::QuantizeUnitsToBestFit(CachedLocation.Z.GetValue(), EUnit::Centimeters).Units);
+		SetLargestUnit(FUnitConversion::QuantizeUnitsToBestFit(CachedLocation.Z.GetValue(), Units).Units);
 	}
 
 	if (LargestUnit.IsSet())
@@ -602,8 +603,10 @@ void FComponentTransformDetails::CacheCommonLocationUnits()
 	}
 	else
 	{
-		VectorUnits	= EUnit::Centimeters;
+		VectorUnits	= DefaultInputUnits.GetValue();
 	}
+
+	VectorUnits = FUnitConversion::ConvertToGlobalDisplayRange(VectorUnits);
 }
 
 FNumericUnit<float> FComponentTransformDetails::QuantizeUnitsToBestFit(const float& InValue, EUnit InUnits) const
diff --git a/Engine/Source/Editor/EditorStyle/Private/EditorStyleClasses.cpp b/Engine/Source/Editor/EditorStyle/Private/EditorStyleClasses.cpp
--- a/Engine/Source/Editor/EditorStyle/Private/EditorStyleClasses.cpp
+++ b/Engine/Source/Editor/EditorStyle/Private/EditorStyleClasses.cpp
@@ -15,7 +15,6 @@ UEditorStyleSettings::UEditorStyleSettings( const FObjectInitializer& ObjectInit
 
 	bOpenAssetEditorTabsInNewWindow = true;
 
-	bDisplayPropertyUnits = true;
 	bShowFriendlyNames = true;
 	LogTimestampMode = ELogTimes::None;
 }
@@ -41,10 +40,6 @@ void UEditorStyleSettings::PostEditChangeProperty(struct FPropertyChangedEvent&
 	{
 		FSlateApplication::Get().EnableMenuAnimations(bEnableWindowAnimations);
 	}
-	else if (Name == GET_MEMBER_NAME_CHECKED(UEditorStyleSettings, bDisplayPropertyUnits))
-	{
-		FUnitConversion::bIsUnitDisplayEnabled = bDisplayPropertyUnits;
-	}
 
 //	if (!FUnrealEdMisc::Get().IsDeletePreferences())
 	{
@@ -54,10 +49,4 @@ void UEditorStyleSettings::PostEditChangeProperty(struct FPropertyChangedEvent&
 	SettingChangedEvent.Broadcast(Name);
 }
 
-
-void UEditorStyleSettings::PostInitProperties()
-{
-	Super::PostInitProperties();
-	FUnitConversion::bIsUnitDisplayEnabled = bDisplayPropertyUnits;
-}
 #endif
diff --git a/Engine/Source/Editor/EditorStyle/Public/Classes/EditorStyleSettings.h b/Engine/Source/Editor/EditorStyle/Public/Classes/EditorStyleSettings.h
--- a/Engine/Source/Editor/EditorStyle/Public/Classes/EditorStyleSettings.h
+++ b/Engine/Source/Editor/EditorStyle/Public/Classes/EditorStyleSettings.h
@@ -70,10 +70,6 @@ public:
 	UPROPERTY(EditAnywhere, config, Category=UserInterface, meta=(DisplayName="Show Friendly Variable Names"))
 	uint32 bShowFriendlyNames:1;
 
-	/** When set, numerical units will be displayed on applicable properties */
-	UPROPERTY(EditAnywhere, config, Category=UserInterface, meta=(DisplayName="Display Units on Numerical Properties"))
-	uint32 bDisplayPropertyUnits:1;
-
 	/** When enabled, the Editor Preferences and Project Settings menu items in the main menu will be expanded with sub-menus for each settings section. */
 	UPROPERTY(EditAnywhere, config, Category=UserInterface, AdvancedDisplay)
 	uint32 bExpandConfigurationMenus:1;
@@ -112,7 +108,6 @@ protected:
 
 #if WITH_EDITOR
 	virtual void PostEditChangeProperty( struct FPropertyChangedEvent& PropertyChangedEvent ) override;
-	virtual void PostInitProperties() override;
 #endif
 
 private:
diff --git a/Engine/Source/Editor/ProjectSettingsViewer/Private/ProjectSettingsViewerModule.cpp b/Engine/Source/Editor/ProjectSettingsViewer/Private/ProjectSettingsViewerModule.cpp
--- a/Engine/Source/Editor/ProjectSettingsViewer/Private/ProjectSettingsViewerModule.cpp
+++ b/Engine/Source/Editor/ProjectSettingsViewer/Private/ProjectSettingsViewerModule.cpp
@@ -18,6 +18,9 @@
 #include "Runtime/Engine/Classes/Sound/AudioSettings.h"
 #include "Runtime/AIModule/Classes/Navigation/CrowdManager.h"
 #include "Runtime/Engine/Classes/Animation/AnimationSettings.h"
+
+#include "Editor/UnrealEd/Public/Settings/EditorProjectSettings.h"
+
 #include "AISystem.h"
 
 #define LOCTEXT_NAMESPACE "FProjectSettingsViewerModule"
@@ -64,6 +67,7 @@ public:
 		{
 			RegisterEngineSettings(*SettingsModule);
 			RegisterProjectSettings(*SettingsModule);
+			RegisterEditorSettings(*SettingsModule);
 
 			SettingsModule->RegisterViewer("Project", *this);
 		}
@@ -268,6 +272,22 @@ protected:
 		);*/
 	}
 
+
+	/**
+	 * Registers Editor settings.
+	 *
+	 * @param SettingsModule A reference to the settings module.
+	 */
+	void RegisterEditorSettings( ISettingsModule& SettingsModule )
+	{
+		// Appearance settings
+		SettingsModule.RegisterSettings("Project", "Editor", "Appearance",
+			LOCTEXT("AppearanceSettingsName", "Appearance"),
+			LOCTEXT("AppearanceSettingsDescription", "Settings pertaining to the appearance of the editor"),
+			GetMutableDefault<UEditorProjectAppearanceSettings>()
+		);
+	}
+
 	/** Unregisters all settings. */
 	void UnregisterSettings()
 	{
@@ -276,7 +296,6 @@ protected:
 		if (SettingsModule != nullptr)
 		{
 			SettingsModule->UnregisterViewer("Project");
-
 			// engine settings
 			SettingsModule->UnregisterSettings("Project", "Engine", "General");
 			SettingsModule->UnregisterSettings("Project", "Engine", "CrowdManager");
@@ -294,6 +313,10 @@ protected:
 			SettingsModule->UnregisterSettings("Project", "Project", "Packaging");
 			SettingsModule->UnregisterSettings("Project", "Project", "SupportedPlatforms");
 			SettingsModule->UnregisterSettings("Project", "Project", "Movies");
+
+			// Editor settings
+			SettingsModule->UnregisterSettings("Editor", "Editor", "Appearance");
+
 //			SettingsModule->UnregisterSettings("Project", "Project", "GameSession");
 //			SettingsModule->UnregisterSettings("Project", "Project", "HUD");
 		}
diff --git a/Engine/Source/Editor/PropertyEditor/Private/UserInterface/PropertyEditor/SPropertyEditorNumeric.h b/Engine/Source/Editor/PropertyEditor/Private/UserInterface/PropertyEditor/SPropertyEditorNumeric.h
--- a/Engine/Source/Editor/PropertyEditor/Private/UserInterface/PropertyEditor/SPropertyEditorNumeric.h
+++ b/Engine/Source/Editor/PropertyEditor/Private/UserInterface/PropertyEditor/SPropertyEditorNumeric.h
@@ -94,7 +94,7 @@ public:
 
 		// Set up the correct type interface if we want to display units on the property editor
 		TSharedPtr<INumericTypeInterface<NumericType>> TypeInterface;
-		if (FUnitConversion::bIsUnitDisplayEnabled)
+		if (FUnitConversion::Settings().ShouldDisplayUnits())
 		{
 			const FString& Units = InPropertyEditor->GetProperty()->GetMetaData(TEXT("Units"));
 			auto PropertyUnits = FUnitConversion::UnitFromString(*Units);
@@ -107,7 +107,11 @@ public:
 				{
 					LexicalConversion::FromString(bAllowUnitRangeAdaption, *AllowRageAdaptionString);
 				}
-				TypeInterface = MakeShareable( new TNumericUnitTypeInterface<NumericType>(PropertyUnits.GetValue(), bAllowUnitRangeAdaption) );
+
+				// Create the type interface and set up the default input units if they are compatible
+				auto* Impl = new TNumericUnitTypeInterface<NumericType>(PropertyUnits.GetValue(), bAllowUnitRangeAdaption);
+				Impl->UseDefaultInputUnits();
+				TypeInterface = MakeShareable(Impl);
 			}
 		}
 
diff --git a/Engine/Source/Editor/UnrealEd/Private/Settings/EditorProjectSettings.cpp b/Engine/Source/Editor/UnrealEd/Private/Settings/EditorProjectSettings.cpp
new file mode 100644
--- /dev/null
+++ b/Engine/Source/Editor/UnrealEd/Private/Settings/EditorProjectSettings.cpp
@@ -0,0 +1,68 @@
+// Copyright 1998-2015 Epic Games, Inc. All Rights Reserved.
+
+#include "UnrealEd.h"
+#include "EditorProjectSettings.h"
+
+
+EUnit ConvertDefaultInputUnits(UEditorProjectAppearanceSettings::EDefaultLocationUnit In)
+{
+	typedef UEditorProjectAppearanceSettings::EDefaultLocationUnit EDefaultLocationUnit;
+
+	switch(In)
+	{
+	case EDefaultLocationUnit::Micrometers:		return EUnit::Micrometers;
+	case EDefaultLocationUnit::Millimeters:		return EUnit::Millimeters;
+	case EDefaultLocationUnit::Centimeters:		return EUnit::Centimeters;
+	case EDefaultLocationUnit::Meters:			return EUnit::Meters;
+	case EDefaultLocationUnit::Kilometers:		return EUnit::Kilometers;
+	case EDefaultLocationUnit::Inches:			return EUnit::Inches;
+	case EDefaultLocationUnit::Feet:			return EUnit::Feet;
+	case EDefaultLocationUnit::Yards:			return EUnit::Yards;
+	case EDefaultLocationUnit::Miles:			return EUnit::Miles;
+	default:									return EUnit::Centimeters;
+	}
+}
+
+EGlobalUnitDisplay ConvertGlobalUnitDisplay(EUnitDisplay In)
+{
+	switch(In)
+	{
+	case EUnitDisplay::None:		return EGlobalUnitDisplay::None;
+	case EUnitDisplay::Metric:		return EGlobalUnitDisplay::Metric;
+	case EUnitDisplay::Imperial:	return EGlobalUnitDisplay::Imperial;
+	default:				 		return EGlobalUnitDisplay::Metric;
+	}
+}
+
+UEditorProjectAppearanceSettings::UEditorProjectAppearanceSettings(const FObjectInitializer&)
+	: UnitDisplay(EUnitDisplay::Metric)
+	, DefaultInputUnits(EDefaultLocationUnit::Centimeters)
+{
+
+}
+
+void UEditorProjectAppearanceSettings::PostEditChangeProperty( struct FPropertyChangedEvent& PropertyChangedEvent )
+{
+	Super::PostEditChangeProperty(PropertyChangedEvent);
+
+	const FName Name = PropertyChangedEvent.Property ? PropertyChangedEvent.Property->GetFName() : NAME_None;
+
+	if (Name == GET_MEMBER_NAME_CHECKED(UEditorProjectAppearanceSettings, UnitDisplay))
+	{
+		FUnitConversion::Settings().SetGlobalUnitDisplay(ConvertGlobalUnitDisplay(UnitDisplay));
+	}
+	else if (Name == GET_MEMBER_NAME_CHECKED(UEditorProjectAppearanceSettings, DefaultInputUnits))
+	{
+		FUnitConversion::Settings().SetDefaultInputUnit(ConvertDefaultInputUnits(DefaultInputUnits));
+	}
+
+	SaveConfig();
+}
+
+void UEditorProjectAppearanceSettings::PostInitProperties()
+{
+	Super::PostInitProperties();
+
+	FUnitConversion::Settings().SetGlobalUnitDisplay(ConvertGlobalUnitDisplay(UnitDisplay));
+	FUnitConversion::Settings().SetDefaultInputUnit(ConvertDefaultInputUnits(DefaultInputUnits));
+}
\ No newline at end of file
diff --git a/Engine/Source/Editor/UnrealEd/Public/Settings/EditorProjectSettings.h b/Engine/Source/Editor/UnrealEd/Public/Settings/EditorProjectSettings.h
new file mode 100644
--- /dev/null
+++ b/Engine/Source/Editor/UnrealEd/Public/Settings/EditorProjectSettings.h
@@ -0,0 +1,44 @@
+// Copyright 1998-2015 Epic Games, Inc. All Rights Reserved.
+
+#pragma once
+
+#include "UnitConversion.h"
+
+#include "EditorProjectSettings.generated.h"
+
+/** UENUM to define the specific set of allowable unit types */
+UENUM()
+enum class EUnitDisplay : uint8
+{
+	None, Metric, Imperial
+};
+
+/** Editor project appearance settings. Stored in default config, per-project */
+UCLASS(config=Editor, defaultconfig)
+class UNREALED_API UEditorProjectAppearanceSettings : public UObject
+{
+public:
+	GENERATED_BODY()
+	UEditorProjectAppearanceSettings(const FObjectInitializer&);
+
+	/** UENUM to define the specific set of allowable default units */
+	UENUM()
+	enum class EDefaultLocationUnit : uint8
+	{
+		Micrometers, Millimeters, Centimeters, Meters, Kilometers,
+		Inches, Feet, Yards, Miles
+	};
+
+protected:
+	/** Called when a property on this object is changed */
+	virtual void PostEditChangeProperty( struct FPropertyChangedEvent& PropertyChangedEvent ) override;
+	virtual void PostInitProperties() override;
+
+private:
+	
+	UPROPERTY(EditAnywhere, config, Category=Units, meta=(Tooltip="Specifies how to present applicable property units on editor interfaces."))
+	EUnitDisplay UnitDisplay;
+
+	UPROPERTY(EditAnywhere, config, Category=Units, meta=(Tooltip="The default units to use for location and distance measurements on editor user interfaces when not explicitly specified by the user."))
+	EDefaultLocationUnit DefaultInputUnits;
+};
\ No newline at end of file
diff --git a/Engine/Source/Runtime/Core/Private/Math/UnitConversion.cpp b/Engine/Source/Runtime/Core/Private/Math/UnitConversion.cpp
--- a/Engine/Source/Runtime/Core/Private/Math/UnitConversion.cpp
+++ b/Engine/Source/Runtime/Core/Private/Math/UnitConversion.cpp
@@ -113,14 +113,95 @@ const FUnitConversion::EUnitType FUnitConversion::Types[] = {
 	EUnitType::Time,		EUnitType::Time,		EUnitType::Time,		EUnitType::Time,		EUnitType::Time,		EUnitType::Time,		EUnitType::Time,
 };
 
-/** Globally check whether we should display units in on user interfaces */
-bool FUnitConversion::bIsUnitDisplayEnabled = true;
+const EUnit MetricToImperial[] = {
+	/** Scalar distance/length units */
+	EUnit::Inches/*Micrometers*/, EUnit::Inches/*Millimeters*/, EUnit::Inches/*Centimeters*/, EUnit::Yards/*Meters*/, EUnit::Miles/*Kilometers*/,
+	EUnit::Inches, EUnit::Feet, EUnit::Yards, EUnit::Miles,
+	EUnit::Lightyears,
+
+	/** Angular units */
+	EUnit::Degrees, EUnit::Radians,
+
+	/** Speed units */
+	EUnit::MilesPerHour/*MetersPerSecond*/, EUnit::MilesPerHour/*KilometersPerHour*/, EUnit::MilesPerHour,
+
+	/** Temperature units */
+	EUnit::Farenheit/*Celsius*/, EUnit::Farenheit, EUnit::Kelvin,
+
+	/** Mass units */
+	EUnit::Ounces/*Micrograms*/, EUnit::Ounces/*Milligrams*/, EUnit::Ounces/*Grams*/, EUnit::Pounds/*Kilograms*/, EUnit::Stones/*MetricTons*/,
+	EUnit::Ounces, EUnit::Pounds, EUnit::Stones,
+
+	/** Force units */
+	EUnit::PoundsForce/*Newtons*/, EUnit::PoundsForce, EUnit::PoundsForce/*KilogramsForce*/,
+
+	/** Frequency units */
+	EUnit::Hertz, EUnit::Kilohertz, EUnit::Megahertz, EUnit::Gigahertz, EUnit::RevolutionsPerMinute,
+
+	/** Data Size units */
+	EUnit::Bytes, EUnit::Kilobytes, EUnit::Megabytes, EUnit::Gigabytes, EUnit::Terabytes,
+
+	/** Luminous flux units */
+	EUnit::Lumens,
+
+	/** Time units */
+	EUnit::Milliseconds, EUnit::Seconds, EUnit::Minutes, EUnit::Hours, EUnit::Days, EUnit::Months, EUnit::Years
+};
+EUnit ImperialToMetric[] = {
+	/** Scalar distance/length units */
+	EUnit::Micrometers, EUnit::Millimeters, EUnit::Centimeters, EUnit::Meters, EUnit::Kilometers,
+	EUnit::Centimeters/*Inches*/, EUnit::Meters/*Feet*/, EUnit::Meters/*Yards*/, EUnit::Kilometers/*Miles*/,
+	EUnit::Lightyears,
+
+	/** Angular units */
+	EUnit::Degrees, EUnit::Radians,
+
+	/** Speed units */
+	EUnit::MetersPerSecond, EUnit::KilometersPerHour, EUnit::KilometersPerHour/*MilesPerHour*/,
+
+	/** Temperature units */
+	EUnit::Celsius, EUnit::Celsius/*Farenheit*/, EUnit::Kelvin,
+
+	/** Mass units */
+	EUnit::Micrograms, EUnit::Milligrams, EUnit::Grams, EUnit::Kilograms, EUnit::MetricTons,
+	EUnit::Grams/*Ounces*/, EUnit::Kilograms/*Pounds*/, EUnit::MetricTons/*Stones*/,
+
+	/** Force units */
+	EUnit::Newtons, EUnit::Newtons/*PoundsForce*/, EUnit::KilogramsForce,
+
+	/** Frequency units */
+	EUnit::Hertz, EUnit::Kilohertz, EUnit::Megahertz, EUnit::Gigahertz, EUnit::RevolutionsPerMinute,
+
+	/** Data Size units */
+	EUnit::Bytes, EUnit::Kilobytes, EUnit::Megabytes, EUnit::Gigabytes, EUnit::Terabytes,
+
+	/** Luminous flux units */
+	EUnit::Lumens,
+
+	/** Time units */
+	EUnit::Milliseconds, EUnit::Seconds, EUnit::Minutes, EUnit::Hours, EUnit::Days, EUnit::Months, EUnit::Years,
+};
+
+FUnitSettings::FUnitSettings()
+	: GlobalUnitDisplay(EGlobalUnitDisplay::Metric)
+	, DefaultInputUnit(EUnit::Centimeters)
+{
+}
+
+/** Get the global settings for unit conversion/display */
+FUnitSettings& FUnitConversion::Settings()
+{
+	static TUniquePtr<FUnitSettings> Settings(new FUnitSettings);
+	return *Settings;
+}
 
 /** Get the unit abbreviation the specified unit type */
 const TCHAR* FUnitConversion::GetUnitDisplayString(EUnit Unit)
 {
 	static_assert(ARRAY_COUNT(FUnitConversion::Types) == (uint8)EUnit::Unspecified, "Type array does not match size of unit enum");
 	static_assert(ARRAY_COUNT(FUnitConversion::DisplayStrings) == (uint8)EUnit::Unspecified, "Display String array does not match size of unit enum");
+	static_assert(ARRAY_COUNT(MetricToImperial) == (uint8)EUnit::Unspecified, "MetricToImperial array does not match size of unit enum");
+	static_assert(ARRAY_COUNT(ImperialToMetric) == (uint8)EUnit::Unspecified, "ImperialToMetric array does not match size of unit enum");
 	
 	if (Unit != EUnit::Unspecified)
 	{
@@ -148,6 +229,21 @@ TOptional<EUnit> FUnitConversion::UnitFromString(const TCHAR* UnitString)
 	return TOptional<EUnit>();
 }
 
+EUnit FUnitConversion::ConvertToGlobalDisplayRange(EUnit InUnit)
+{
+	if (InUnit == EUnit::Unspecified)
+	{
+		return EUnit::Unspecified;
+	}
+
+	switch(Settings().GetGlobalUnitDisplay())
+	{
+	case EGlobalUnitDisplay::Metric:	return ImperialToMetric[(int8)InUnit];
+	case EGlobalUnitDisplay::Imperial:	return MetricToImperial[(int8)InUnit];
+	default:							return InUnit;
+	}
+}
+
 double FUnitConversion::DistanceUnificationFactor(EUnit From)
 {
 	// Convert to meters
diff --git a/Engine/Source/Runtime/Core/Public/Math/UnitConversion.h b/Engine/Source/Runtime/Core/Public/Math/UnitConversion.h
--- a/Engine/Source/Runtime/Core/Public/Math/UnitConversion.h
+++ b/Engine/Source/Runtime/Core/Public/Math/UnitConversion.h
@@ -45,12 +45,51 @@ enum class EUnit
 	Unspecified
 };
 
+enum class EGlobalUnitDisplay
+{
+	None, Metric, Imperial
+};
+
 template<typename NumericType> struct FNumericUnit;
 
+/** Unit settings accessed globally through FUnitConversion::Settings() */
+class CORE_API FUnitSettings
+{
+public:
+	
+	FUnitSettings();
+
+	/** Check whether unit display is globally enabled or disabled */
+	bool ShouldDisplayUnits() { return GlobalUnitDisplay != EGlobalUnitDisplay::None; }
+
+	/** Get/set the global unit display to either mertic, imperial, or none */
+	EGlobalUnitDisplay GetGlobalUnitDisplay() { return GlobalUnitDisplay; }
+	void SetGlobalUnitDisplay(EGlobalUnitDisplay InGlobalUnitDisplay) { GlobalUnitDisplay = InGlobalUnitDisplay; SettingChangedEvent.Broadcast(); }
+	
+	/** Default unit to use for text input when no unit was specified. Can be EUnit::Unspecified. Check for compatibility with property units using AreUnitsCompatible */
+	EUnit GetDefaultInputUnit() { return DefaultInputUnit; }
+	void SetDefaultInputUnit(EUnit InDefaultInputUnit) { DefaultInputUnit = InDefaultInputUnit; SettingChangedEvent.Broadcast(); }
+
+	/** Returns an event delegate that is executed when a display setting has changed. (GlobalUnitDisplay or DefaultInputUnits) */
+	DECLARE_EVENT(FUnitSettings, FDisplaySettingChanged);
+	FDisplaySettingChanged& OnDisplaySettingsChanged() { return SettingChangedEvent; }
+
+private:
+
+	/** Globally check how we should display units on user interfaces */
+	EGlobalUnitDisplay GlobalUnitDisplay;
+
+	/** Default unit to use for text input when no unit was specified. Can be EUnit::Unspecified. Check for compatibility with property units using AreUnitsCompatible */
+	EUnit DefaultInputUnit;
+	
+	/** Holds an event delegate that is executed when a display setting has changed. */
+	FDisplaySettingChanged SettingChangedEvent;
+};
+
 struct CORE_API FUnitConversion
 {
-	/** Globally check whether we should display units in on user interfaces */
-	static bool bIsUnitDisplayEnabled;
+	/** Get the global settings for unit conversion/display */
+	static FUnitSettings& Settings();
 
 	/** Check whether it is possible to convert a number between the two specified units */
 	static bool AreUnitsCompatible(EUnit From, EUnit To)
@@ -112,6 +151,9 @@ struct CORE_API FUnitConversion
 	/** Helper function to find a unit from a string (name or display string) */
 	static TOptional<EUnit> UnitFromString(const TCHAR* UnitString);
 
+	/** Perform metric <-> imperial conversion on the specified unit according to the GlobalUnitDisplay setting */
+	static EUnit ConvertToGlobalDisplayRange(EUnit InUnit);
+
 private:
 
 	/** Find the common quantization factor for the specified distance unit. Quantizes to Meters. */
diff --git a/Engine/Source/Runtime/Slate/Public/Widgets/Input/NumericTypeInterface.h b/Engine/Source/Runtime/Slate/Public/Widgets/Input/NumericTypeInterface.h
--- a/Engine/Source/Runtime/Slate/Public/Widgets/Input/NumericTypeInterface.h
+++ b/Engine/Source/Runtime/Slate/Public/Widgets/Input/NumericTypeInterface.h
@@ -10,7 +10,7 @@ struct INumericTypeInterface
 
 	/** Convert the type to/from a string */
 	virtual FString ToString(const NumericType& Value) const = 0;
-	virtual TOptional<NumericType> FromString(const FString& InString, const NumericType& InCurrentValue) = 0;
+	virtual TOptional<NumericType> FromString(const FString& InString) = 0;
 
 	/** Check whether the typed character is valid */
 	virtual bool IsCharacterValid(TCHAR InChar) const = 0;
@@ -25,7 +25,7 @@ struct TDefaultNumericTypeInterface : INumericTypeInterface<NumericType>
 	{
 		return LexicalConversion::ToSanitizedString(Value);
 	}
-	virtual TOptional<NumericType> FromString(const FString& InString, const NumericType& InCurrentValue) override
+	virtual TOptional<NumericType> FromString(const FString& InString) override
 	{
 		NumericType NewValue;
 		bool bEvalResult = LexicalConversion::TryParseString( NewValue, *InString );
@@ -74,21 +74,32 @@ struct TNumericUnitTypeInterface : TDefaultNumericTypeInterface<NumericType>
 	/** The units that we should treat this value as */
 	EUnit Units;
 
+	/** Optional units to use for input when no units where specified */
+	TOptional<EUnit> DefaultInputUnits;
+
 	/** True if we are allowed to display this unit to the user as a more relevant related unit type (e.g. 10000 m shows as 10 km) */
 	bool bAllowUnitRangeAdaption;
 
 	/** Constructor */
 	TNumericUnitTypeInterface(EUnit InUnits, bool bInAllowUnitRangeAdaption = true);
+	~TNumericUnitTypeInterface();
+
+	/** Set up this interface to use the global default input units */
+	void UseDefaultInputUnits();
 
 	/** Convert this type to a string */
 	virtual FString ToString(const NumericType& Value) const override;
 
 	/** Attempt to parse a numeral with our units from the specified string. */
-	virtual TOptional<NumericType> FromString(const FString& ValueString, const NumericType& InCurrentValue) override;
+	virtual TOptional<NumericType> FromString(const FString& ValueString) override;
 
 	/** Called by this class to quantize the specified value to the best unit */
 	virtual FNumericUnit<NumericType> QuantizeUnitsToBestFit(const NumericType& InValue, EUnit InUnits) const;
 
 	/** Check whether the specified typed character is valid */
 	virtual bool IsCharacterValid(TCHAR InChar) const override;
+
+private:
+	/** Called when the global unit settings have changed, if this type interface is using the default input units */
+	void OnGlobalUnitSettingChanged();
 };
\ No newline at end of file
diff --git a/Engine/Source/Runtime/Slate/Public/Widgets/Input/NumericUnitTypeInterface.inl b/Engine/Source/Runtime/Slate/Public/Widgets/Input/NumericUnitTypeInterface.inl
--- a/Engine/Source/Runtime/Slate/Public/Widgets/Input/NumericUnitTypeInterface.inl
+++ b/Engine/Source/Runtime/Slate/Public/Widgets/Input/NumericUnitTypeInterface.inl
@@ -9,9 +9,42 @@ TNumericUnitTypeInterface<NumericType>::TNumericUnitTypeInterface(EUnit InUnits,
 	: Units(InUnits), bAllowUnitRangeAdaption(bInAllowUnitRangeAdaption)
 {}
 
+template<typename NumericType>
+TNumericUnitTypeInterface<NumericType>::~TNumericUnitTypeInterface()
+{
+	FUnitConversion::Settings().OnDisplaySettingsChanged().RemoveAll(this);
+}
+
+
+template<typename NumericType>
+void TNumericUnitTypeInterface<NumericType>::UseDefaultInputUnits()
+{
+	FUnitConversion::Settings().OnDisplaySettingsChanged().AddRaw(this, &TNumericUnitTypeInterface<NumericType>::OnGlobalUnitSettingChanged);
+	OnGlobalUnitSettingChanged();
+}
+
+template<typename NumericType>
+void TNumericUnitTypeInterface<NumericType>::OnGlobalUnitSettingChanged()
+{
+	const EUnit Default = FUnitConversion::Settings().GetDefaultInputUnit();
+	if (FUnitConversion::AreUnitsCompatible(Default, Units))
+	{
+		DefaultInputUnits = Default;
+	}
+	else
+	{
+		DefaultInputUnits.Reset();
+	}
+}
+
 template<typename NumericType>
 FNumericUnit<NumericType> TNumericUnitTypeInterface<NumericType>::QuantizeUnitsToBestFit(const NumericType& InValue, EUnit InUnits) const
 {
+	// Use the default input units for 0
+	if (InValue == 0 && DefaultInputUnits.IsSet() && FUnitConversion::AreUnitsCompatible(DefaultInputUnits.GetValue(), InUnits))
+	{
+		return FNumericUnit<NumericType>(0, FUnitConversion::ConvertToGlobalDisplayRange(DefaultInputUnits.GetValue()));
+	}
 	return FUnitConversion::QuantizeUnitsToBestFit(InValue, InUnits);
 }
 
@@ -20,12 +53,19 @@ FString TNumericUnitTypeInterface<NumericType>::ToString(const NumericType& Valu
 {
 	using namespace LexicalConversion;
 
-	const FNumericUnit<NumericType> Default(Value, Units);
-	return ToSanitizedString(bAllowUnitRangeAdaption ? QuantizeUnitsToBestFit(Default.Value, Default.Units) : Default);
+	NumericType ValueToUse = Value;
+	EUnit DisplayUnits = FUnitConversion::ConvertToGlobalDisplayRange(Units);
+
+	if (DisplayUnits != Units)
+	{
+		ValueToUse = FUnitConversion::Convert(ValueToUse, Units, DisplayUnits);
+	}
+	
+	return ToSanitizedString(bAllowUnitRangeAdaption ? QuantizeUnitsToBestFit(ValueToUse, DisplayUnits) : FNumericUnit<NumericType>(ValueToUse, DisplayUnits));
 }
 
 template<typename NumericType>
-TOptional<NumericType> TNumericUnitTypeInterface<NumericType>::FromString(const FString& InString, const NumericType& InCurrentValue)
+TOptional<NumericType> TNumericUnitTypeInterface<NumericType>::FromString(const FString& InString)
 {
 	using namespace LexicalConversion;
 
@@ -36,10 +76,10 @@ TOptional<NumericType> TNumericUnitTypeInterface<NumericType>::FromString(const
 	{
 		// Convert the number into the correct units
 		EUnit SourceUnits = NewValue.Units;
-		if (SourceUnits == EUnit::Unspecified && bAllowUnitRangeAdaption)
+		if (SourceUnits == EUnit::Unspecified && DefaultInputUnits.IsSet())
 		{
-			// Use the display units of the current value
-			SourceUnits = QuantizeUnitsToBestFit(InCurrentValue, Units).Units;
+			// Use the default supplied input units
+			SourceUnits = DefaultInputUnits.GetValue();
 		}
 		return FUnitConversion::Convert(NewValue.Value, SourceUnits, Units);
 	}
diff --git a/Engine/Source/Runtime/Slate/Public/Widgets/Input/SNumericEntryBox.h b/Engine/Source/Runtime/Slate/Public/Widgets/Input/SNumericEntryBox.h
--- a/Engine/Source/Runtime/Slate/Public/Widgets/Input/SNumericEntryBox.h
+++ b/Engine/Source/Runtime/Slate/Public/Widgets/Input/SNumericEntryBox.h
@@ -380,7 +380,7 @@ private:
 		{
 			if( bCommit )
 			{
-				TOptional<NumericType> NumericValue = Interface->FromString(NewValue.ToString(), OnGetValueForSpinBox());
+				TOptional<NumericType> NumericValue = Interface->FromString(NewValue.ToString());
 				if( NumericValue.IsSet() )
 				{
 					OnValueCommitted.ExecuteIfBound(NumericValue.GetValue(), CommitInfo );
diff --git a/Engine/Source/Runtime/Slate/Public/Widgets/Input/SSpinBox.h b/Engine/Source/Runtime/Slate/Public/Widgets/Input/SSpinBox.h
--- a/Engine/Source/Runtime/Slate/Public/Widgets/Input/SSpinBox.h
+++ b/Engine/Source/Runtime/Slate/Public/Widgets/Input/SSpinBox.h
@@ -538,7 +538,7 @@ protected:
 			ExitTextMode();
 		}
 
-		TOptional<NumericType> NewValue = Interface->FromString(NewText.ToString(), ValueAttribute.Get());
+		TOptional<NumericType> NewValue = Interface->FromString(NewText.ToString());
 		if (NewValue.IsSet())
 		{
 			CommitValue( NewValue.GetValue(), CommittedViaTypeIn, CommitInfo );
