Add check for negative deserialized string size.

diff --git a/Engine/Source/Runtime/Core/Private/Containers/String.cpp b/Engine/Source/Runtime/Core/Private/Containers/String.cpp
--- a/Engine/Source/Runtime/Core/Private/Containers/String.cpp
+++ b/Engine/Source/Runtime/Core/Private/Containers/String.cpp
@@ -1006,9 +1006,18 @@ FArchive& operator<<( FArchive& Ar, FString& A )
 			SaveNum = -SaveNum;
 		}
 
-		// Protect against network packets allocating too much memory
+		// If SaveNum is still less than 0, they must have passed in MIN_INT. Archive is corrupted.
+		if (SaveNum < 0))
+		{
+			Ar.ArIsError = 1;
+			Ar.ArIsCriticalError = 1;
+			UE_LOG(LogNetSerialization, Error, TEXT("Archive is corrupted"));
+			return Ar;
+		}
+
 		auto MaxSerializeSize = Ar.GetMaxSerializeSize();
-		if (MaxSerializeSize > 0 && ( SaveNum > MaxSerializeSize || SaveNum < 0 ))		// If SaveNum is still less than 0, they must have passed in MIN_INT
+		// Protect against network packets allocating too much memory
+		if ((MaxSerializeSize > 0) && (SaveNum > MaxSerializeSize))
 		{
 			Ar.ArIsError         = 1;
 			Ar.ArIsCriticalError = 1;
