Typo fix.

diff --git a/Engine/Source/Runtime/Core/Private/Containers/String.cpp b/Engine/Source/Runtime/Core/Private/Containers/String.cpp
--- a/Engine/Source/Runtime/Core/Private/Containers/String.cpp
+++ b/Engine/Source/Runtime/Core/Private/Containers/String.cpp
@@ -1007,7 +1007,7 @@ FArchive& operator<<( FArchive& Ar, FString& A )
 		}
 
 		// If SaveNum is still less than 0, they must have passed in MIN_INT. Archive is corrupted.
-		if (SaveNum < 0))
+		if (SaveNum < 0)
 		{
 			Ar.ArIsError = 1;
 			Ar.ArIsCriticalError = 1;
