Stats - Added stats group for target platform, added DECLARE_PTR_STAT

diff --git a/Engine/Source/Runtime/Core/Public/Misc/CoreStats.h b/Engine/Source/Runtime/Core/Public/Misc/CoreStats.h
--- a/Engine/Source/Runtime/Core/Public/Misc/CoreStats.h
+++ b/Engine/Source/Runtime/Core/Public/Misc/CoreStats.h
@@ -91,3 +91,10 @@ DECLARE_CYCLE_STAT_EXTERN(TEXT("Convert Rotator to Quat"), STAT_MathConvertRotat
 DECLARE_CYCLE_STAT_EXTERN(TEXT("Convert Quat to Rotator"), STAT_MathConvertQuatToRotator, STATGROUP_Math, CORE_API);
 
 DECLARE_CYCLE_STAT_EXTERN(TEXT("CPU Stall - Sleep"), STAT_Sleep, STATGROUP_CPUStalls, CORE_API);
+
+
+DECLARE_CYCLE_STAT_EXTERN( TEXT( "CPU Stall - Wait For Event" ), STAT_EventWait, STATGROUP_CPUStalls, );
+
+/** The id will be stored as uint64 in the stat message. */
+DECLARE_PTR_STAT_EXTERN( TEXT( "CPU Stall - Wait For Event with ID" ), STAT_EventWaitWithId, STATGROUP_CPUStalls, );
+DECLARE_PTR_STAT_EXTERN( TEXT( "CPU Stall - Trigger For Event with ID" ), STAT_EventTriggerWithId, STATGROUP_CPUStalls, );
\ No newline at end of file
