Stats - Added stats group for target platform, added DECLARE_PTR_STAT

diff --git a/Engine/Source/Runtime/Core/Public/Stats/Stats.h b/Engine/Source/Runtime/Core/Public/Stats/Stats.h
--- a/Engine/Source/Runtime/Core/Public/Stats/Stats.h
+++ b/Engine/Source/Runtime/Core/Public/Stats/Stats.h
@@ -386,6 +386,7 @@ FORCEINLINE void StatsMasterEnableSubtract(int32 Value = 1)
 #define DECLARE_DWORD_COUNTER_STAT(CounterName,StatId,GroupId)
 #define DECLARE_FLOAT_ACCUMULATOR_STAT(CounterName,StatId,GroupId)
 #define DECLARE_DWORD_ACCUMULATOR_STAT(CounterName,StatId,GroupId)
+#define DECLARE_PTR_STAT(CounterName,StatId,GroupId)
 #define DECLARE_MEMORY_STAT(CounterName,StatId,GroupId)
 #define DECLARE_MEMORY_STAT_POOL(CounterName,StatId,GroupId,Pool)
 #define DECLARE_CYCLE_STAT_EXTERN(CounterName,StatId,GroupId, API)
@@ -393,6 +394,7 @@ FORCEINLINE void StatsMasterEnableSubtract(int32 Value = 1)
 #define DECLARE_DWORD_COUNTER_STAT_EXTERN(CounterName,StatId,GroupId, API)
 #define DECLARE_FLOAT_ACCUMULATOR_STAT_EXTERN(CounterName,StatId,GroupId, API)
 #define DECLARE_DWORD_ACCUMULATOR_STAT_EXTERN(CounterName,StatId,GroupId, API)
+#define DECLARE_PTR_STAT_EXTERN(CounterName,StatId,GroupId, API)
 #define DECLARE_MEMORY_STAT_EXTERN(CounterName,StatId,GroupId, API)
 #define DECLARE_MEMORY_STAT_POOL_EXTERN(CounterName,StatId,GroupId,Pool, API)
 
diff --git a/Engine/Source/Runtime/Core/Public/Stats/Stats2.h b/Engine/Source/Runtime/Core/Public/Stats/Stats2.h
--- a/Engine/Source/Runtime/Core/Public/Stats/Stats2.h
+++ b/Engine/Source/Runtime/Core/Public/Stats/Stats2.h
@@ -5,6 +5,7 @@
 #include "LockFreeList.h"
 #include "LockFreeFixedSizeAllocator.h"
 #include "ChunkedArray.h"
+#include "Function.h"
 
 class FThreadStats;
 
@@ -223,7 +224,7 @@ struct EStatOperation
 		/** This is not a regular stat operation, but just a special message marker to determine that we encountered a special data in the stat file. */
 		SpecialMessageMarker,
 		/** Set operation. */
-		Set, 
+		Set,
 		/** Clear operation. */
 		Clear,
 		/** Add operation. */
@@ -330,7 +331,7 @@ FORCEINLINE uint32 FromPackedCallCountDuration_CallCount(int64 Both)
 
 FORCEINLINE uint32 FromPackedCallCountDuration_Duration(int64 Both)
 {
-	return uint32(Both & 0xffffffff);
+	return uint32(Both & MAX_uint32);
 }
 
 /**
@@ -1726,6 +1727,11 @@ struct FStat_##StatName\
 	DECLARE_STAT(CounterName,StatId,GroupId,EStatDataType::ST_int64, false, false, FPlatformMemory::MCR_Invalid); \
 	static DEFINE_STAT(StatId)
 
+/** This is a fake stat, mostly used to implement memory message or other custom stats that don't easily fit into the system. */
+#define DECLARE_PTR_STAT(CounterName,StatId,GroupId)\
+	DECLARE_STAT(CounterName,StatId,GroupId,EStatDataType::ST_Ptr, false, false, FPlatformMemory::MCR_Invalid); \
+	static DEFINE_STAT(StatId)
+
 #define DECLARE_MEMORY_STAT(CounterName,StatId,GroupId) \
 	DECLARE_STAT(CounterName,StatId,GroupId,EStatDataType::ST_int64, false, false, FPlatformMemory::MCR_Physical); \
 	static DEFINE_STAT(StatId)
@@ -1753,9 +1759,16 @@ struct FStat_##StatName\
 #define DECLARE_DWORD_ACCUMULATOR_STAT_EXTERN(CounterName,StatId,GroupId, API) \
 	DECLARE_STAT(CounterName,StatId,GroupId,EStatDataType::ST_int64, false, false, FPlatformMemory::MCR_Invalid); \
 	extern API DEFINE_STAT(StatId);
+
+/** This is a fake stat, mostly used to implement memory message or other custom stats that don't easily fit into the system. */
+#define DECLARE_PTR_STAT_EXTERN(CounterName,StatId,GroupId, API) \
+	DECLARE_STAT(CounterName,StatId,GroupId,EStatDataType::ST_Ptr, false, false, FPlatformMemory::MCR_Invalid); \
+	extern API DEFINE_STAT(StatId)
+
 #define DECLARE_MEMORY_STAT_EXTERN(CounterName,StatId,GroupId, API) \
 	DECLARE_STAT(CounterName,StatId,GroupId,EStatDataType::ST_int64, false, false, FPlatformMemory::MCR_Physical); \
 	extern API DEFINE_STAT(StatId);
+
 #define DECLARE_MEMORY_STAT_POOL_EXTERN(CounterName,StatId,GroupId,Pool, API) \
 	DECLARE_STAT(CounterName,StatId,GroupId,EStatDataType::ST_int64, false, false, Pool); \
 	extern API DEFINE_STAT(StatId);
@@ -1967,6 +1980,7 @@ DECLARE_STATS_GROUP(TEXT("Slate"), STATGROUP_Slate, STATCAT_Advanced);
 DECLARE_STATS_GROUP(TEXT("Stat System"),STATGROUP_StatSystem, STATCAT_Advanced);
 DECLARE_STATS_GROUP(TEXT("Streaming Details"),STATGROUP_StreamingDetails, STATCAT_Advanced);
 DECLARE_STATS_GROUP(TEXT("Streaming"),STATGROUP_Streaming, STATCAT_Advanced);
+DECLARE_STATS_GROUP(TEXT("Target Platform"),STATGROUP_TargetPlatform, STATCAT_Advanced);
 DECLARE_STATS_GROUP(TEXT("Text"),STATGROUP_Text, STATCAT_Advanced);
 DECLARE_STATS_GROUP(TEXT("Threading"),STATGROUP_Threading, STATCAT_Advanced);
 DECLARE_STATS_GROUP(TEXT("Threads"),STATGROUP_Threads, STATCAT_Advanced);
