Allow platforms to inspect or modify transtions during LoadMap.
UEPLAT-685 WIP

diff --git a/Engine/Source/Runtime/Core/Public/GenericPlatform/GenericPlatformMisc.h b/Engine/Source/Runtime/Core/Public/GenericPlatform/GenericPlatformMisc.h
--- a/Engine/Source/Runtime/Core/Public/GenericPlatform/GenericPlatformMisc.h
+++ b/Engine/Source/Runtime/Core/Public/GenericPlatform/GenericPlatformMisc.h
@@ -225,7 +225,7 @@ struct CORE_API FGenericPlatformMisc
 		*Result = 0;
 	}
 
-	
+
 	/**
 	 * Sets an environment variable to the local process's environment
 	 *
@@ -346,7 +346,7 @@ struct CORE_API FGenericPlatformMisc
 	static void SetMemoryWarningHandler(void (* Handler)(const FGenericMemoryWarningContext& Context))
 	{
 	}
-	
+
 	/**
 	 *	Pumps Windows messages.
 	 *	@param bFromMainLoop if true, this is from the main loop, otherwise we are spinning waiting for the render thread
@@ -385,11 +385,11 @@ protected:
 
 	/**
 	* Retrieves some standard key code mappings (usually called by a subclass's GetCharKeyMap)
-	 *
-	 * @param OutKeyMap Key map to add to.
-	 * @param bMapUppercaseKeys If true, will map A, B, C, etc to EKeys::A, EKeys::B, EKeys::C
-	 * @param bMapLowercaseKeys If true, will map a, b, c, etc to EKeys::A, EKeys::B, EKeys::C
-	 */
+	*
+	* @param OutKeyMap Key map to add to.
+	* @param bMapUppercaseKeys If true, will map A, B, C, etc to EKeys::A, EKeys::B, EKeys::C
+	* @param bMapLowercaseKeys If true, will map a, b, c, etc to EKeys::A, EKeys::B, EKeys::C
+	*/
 	static uint32 GetStandardPrintableKeyMap(uint16* KeyCodes, FString* KeyNames, uint32 MaxMappings, bool bMapUppercaseKeys, bool bMapLowercaseKeys);
 
 public:
@@ -441,7 +441,7 @@ public:
 	 */
 	static bool GetStoredValue(const FString& InStoreId, const FString& InSectionName, const FString& InKeyName, FString& OutValue);
 
-	 /** Sends a message to a remote tool, and debugger consoles */
+	/** Sends a message to a remote tool, and debugger consoles */
 	static void LowLevelOutputDebugString(const TCHAR *Message);
 	static void VARARGS LowLevelOutputDebugStringf(const TCHAR *Format, ... );
 
@@ -819,6 +819,12 @@ public:
 	 */
 	static const TCHAR* GetEngineMode();
 
+	/**
+	* Execute plaform dependent pre load map actions
+	*/
+	static const void PreLoadMap(FString&, FString&, void*)
+	{}
+
 #if !UE_BUILD_SHIPPING
 protected:
 	/** Whether the user should be prompted to allow for a remote debugger to be attached */
diff --git a/Engine/Source/Runtime/Engine/Private/UnrealEngine.cpp b/Engine/Source/Runtime/Engine/Private/UnrealEngine.cpp
--- a/Engine/Source/Runtime/Engine/Private/UnrealEngine.cpp
+++ b/Engine/Source/Runtime/Engine/Private/UnrealEngine.cpp
@@ -8781,6 +8781,8 @@ bool UEngine::LoadMap( FWorldContext& WorldContext, FURL URL, class UPendingNetG
 	NETWORK_PROFILER(GNetworkProfiler.TrackSessionChange(true,URL));
 	MALLOC_PROFILER( FMallocProfiler::SnapshotMemoryLoadMapStart( URL.Map ) );
 	Error = TEXT("");
+
+	FGenericPlatformMisc::PreLoadMap(URL.Map, WorldContext.LastURL.Map, nullptr);
 	
 	// make sure level streaming isn't frozen
 	if (WorldContext.World())
