Back out revision 63 from //depot/UE4/Engine/Source/Runtime/Core/Private/Stats/StatsCommand.cpp

diff --git a/Engine/Source/Runtime/Core/Private/Stats/StatsCommand.cpp b/Engine/Source/Runtime/Core/Private/Stats/StatsCommand.cpp
--- a/Engine/Source/Runtime/Core/Private/Stats/StatsCommand.cpp
+++ b/Engine/Source/Runtime/Core/Private/Stats/StatsCommand.cpp
@@ -125,7 +125,9 @@ void DumpHistoryFrame(FStatsThreadState const& StatsData, int64 TargetFrame, flo
 			}
 			UE_LOG(LogStats, Log, TEXT("  %s"), *FStatsUtils::DebugPrint(Meta));
 		}
-
+	}
+	if (DumpCull == 0.0f)
+	{
 		UE_LOG(LogStats, Log, TEXT("Exclusive aggregate stack data---------------"));
 		TArray<FStatMessage> Stats;
 		StatsData.GetExclusiveAggregateStackStats(TargetFrame, Stats);
