Changed FDateTime UnixTimeStamp methods to use int64 rather than int32. Made relevant changes where the methods were used.

diff --git a/Engine/Source/Runtime/Core/Private/Internationalization/ICUText.cpp b/Engine/Source/Runtime/Core/Private/Internationalization/ICUText.cpp
--- a/Engine/Source/Runtime/Core/Private/Internationalization/ICUText.cpp
+++ b/Engine/Source/Runtime/Core/Private/Internationalization/ICUText.cpp
@@ -27,7 +27,7 @@ FText FText::AsDate(const FDateTime& DateTime, const EDateTimeStyle::Type DateSt
 	checkf(I18N.IsInitialized() == true, TEXT("FInternationalization is not initialized. An FText formatting method was likely used in static object initialization - this is not supported."));
 	FCulturePtr Culture = TargetCulture.IsValid() ? TargetCulture : I18N.GetCurrentCulture();
 
-	int32 UNIXTimestamp = DateTime.ToUnixTimestamp();
+	int64 UNIXTimestamp = DateTime.ToUnixTimestamp();
 	UDate ICUDate = static_cast<double>(UNIXTimestamp) * U_MILLIS_PER_SECOND;
 
 	UErrorCode ICUStatus = U_ZERO_ERROR;
@@ -50,7 +50,7 @@ FText FText::AsTime(const FDateTime& DateTime, const EDateTimeStyle::Type TimeSt
 	checkf(I18N.IsInitialized() == true, TEXT("FInternationalization is not initialized. An FText formatting method was likely used in static object initialization - this is not supported."));
 	FCulturePtr Culture = TargetCulture.IsValid() ? TargetCulture : I18N.GetCurrentCulture();
 
-	int32 UNIXTimestamp = DateTime.ToUnixTimestamp();
+	int64 UNIXTimestamp = DateTime.ToUnixTimestamp();
 	UDate ICUDate = static_cast<double>(UNIXTimestamp) * U_MILLIS_PER_SECOND;
 
 	UErrorCode ICUStatus = U_ZERO_ERROR;
@@ -97,7 +97,7 @@ FText FText::AsDateTime(const FDateTime& DateTime, const EDateTimeStyle::Type Da
 	checkf(I18N.IsInitialized() == true, TEXT("FInternationalization is not initialized. An FText formatting method was likely used in static object initialization - this is not supported."));
 	FCulturePtr Culture = TargetCulture.IsValid() ? TargetCulture : I18N.GetCurrentCulture();
 
-	int32 UNIXTimestamp = DateTime.ToUnixTimestamp();
+	int64 UNIXTimestamp = DateTime.ToUnixTimestamp();
 	UDate ICUDate = static_cast<double>(UNIXTimestamp) * U_MILLIS_PER_SECOND;
 
 	UErrorCode ICUStatus = U_ZERO_ERROR;
diff --git a/Engine/Source/Runtime/Core/Private/Tests/Misc/DateTimeTest.cpp b/Engine/Source/Runtime/Core/Private/Tests/Misc/DateTimeTest.cpp
--- a/Engine/Source/Runtime/Core/Private/Tests/Misc/DateTimeTest.cpp
+++ b/Engine/Source/Runtime/Core/Private/Tests/Misc/DateTimeTest.cpp
@@ -8,10 +8,10 @@ IMPLEMENT_SIMPLE_AUTOMATION_TEST(FDateTimeTest, "Core.Misc.DateTime", EAutomatio
 
 bool FDateTimeTest::RunTest( const FString& Parameters )
 {
-	const int32 UnixEpochTimestamp = 0;
-	const int32 UnixBillenniumTimestamp = 1000000000;
-	const int32 UnixOnesTimestamp = 1111111111;
-	const int32 UnixDecimalSequenceTimestamp = 1234567890;
+	const int64 UnixEpochTimestamp = 0;
+	const int64 UnixBillenniumTimestamp = 1000000000;
+	const int64 UnixOnesTimestamp = 1111111111;
+	const int64 UnixDecimalSequenceTimestamp = 1234567890;
 
 	const FDateTime UnixEpoch = FDateTime::FromUnixTimestamp(UnixEpochTimestamp);
 	const FDateTime UnixBillennium = FDateTime::FromUnixTimestamp(UnixBillenniumTimestamp);
diff --git a/Engine/Source/Runtime/Core/Public/Misc/DateTime.h b/Engine/Source/Runtime/Core/Public/Misc/DateTime.h
--- a/Engine/Source/Runtime/Core/Public/Misc/DateTime.h
+++ b/Engine/Source/Runtime/Core/Public/Misc/DateTime.h
@@ -492,9 +492,9 @@ public:
 	 * @return Time of day.
 	 * @see FromUnixTimestamp
 	 */
-	CORE_API int32 ToUnixTimestamp() const
+	CORE_API int64 ToUnixTimestamp() const
 	{
-		return static_cast<int32>((Ticks - FDateTime(1970, 1, 1).Ticks) / ETimespan::TicksPerSecond);
+		return (Ticks - FDateTime(1970, 1, 1).Ticks) / ETimespan::TicksPerSecond;
 	}
 
 public:
@@ -537,9 +537,9 @@ public:
 	 * @return Gregorian date and time.
 	 * @see ToUnixTimestamp
 	 */
-	static FDateTime FromUnixTimestamp( int32 UnixTime )
+	static FDateTime FromUnixTimestamp( int64 UnixTime )
 	{
-		return FDateTime(1970, 1, 1) + FTimespan(static_cast<int64>(UnixTime) * ETimespan::TicksPerSecond);
+		return FDateTime(1970, 1, 1) + FTimespan(UnixTime * ETimespan::TicksPerSecond);
 	}
 
 	/**
