Linux: make -Wparentheses-equality non-fatal.

- Constructs like if (foo=1) will still break the build, but if ((foo == 1)) will not (won't be suppressed completely, but will produce a visible warning).

diff --git a/Engine/Source/Runtime/Core/Public/Linux/LinuxPlatformCompilerPreSetup.h b/Engine/Source/Runtime/Core/Public/Linux/LinuxPlatformCompilerPreSetup.h
--- a/Engine/Source/Runtime/Core/Public/Linux/LinuxPlatformCompilerPreSetup.h
+++ b/Engine/Source/Runtime/Core/Public/Linux/LinuxPlatformCompilerPreSetup.h
@@ -55,3 +55,4 @@
 // since they cannot be caught by all compilers that we support. They are deemed to be relatively safe to be ignored, at least until all SDKs/toolchains start supporting them.
 #pragma clang diagnostic warning "-Winconsistent-missing-override"
 #pragma clang diagnostic warning "-Wreorder"
+#pragma clang diagnostic warning "-Wparentheses-equality"
