Fix for my last malformed 'fix'....

diff --git a/Engine/Source/Runtime/Core/Private/Stats/StatsCommand.cpp b/Engine/Source/Runtime/Core/Private/Stats/StatsCommand.cpp
--- a/Engine/Source/Runtime/Core/Private/Stats/StatsCommand.cpp
+++ b/Engine/Source/Runtime/Core/Private/Stats/StatsCommand.cpp
@@ -125,14 +125,12 @@ void DumpHistoryFrame(FStatsThreadState const& StatsData, int64 TargetFrame, flo
 			}
 			UE_LOG(LogStats, Log, TEXT("  %s"), *FStatsUtils::DebugPrint(Meta));
 		}
-	}
-	if (DumpCull == 0.0f)
-	{
+
 		UE_LOG(LogStats, Log, TEXT("Exclusive aggregate stack data---------------"));
-		TArray<FStatMessage> Stats;
+		Stats.Empty();
 		StatsData.GetExclusiveAggregateStackStats(TargetFrame, Stats);
 		Stats.Sort(FGroupSort());
-		FName LastGroup = NAME_None;
+		LastGroup = NAME_None;
 		for (int32 Index = 0; Index < Stats.Num(); Index++)
 		{
 			FStatMessage const& Meta = Stats[Index];
