Fixing overriding ini name with -def*ini= in the commandline

diff --git a/Engine/Source/Runtime/Core/Private/Misc/ConfigCacheIni.cpp b/Engine/Source/Runtime/Core/Private/Misc/ConfigCacheIni.cpp
--- a/Engine/Source/Runtime/Core/Private/Misc/ConfigCacheIni.cpp
+++ b/Engine/Source/Runtime/Core/Private/Misc/ConfigCacheIni.cpp
@@ -2682,46 +2682,27 @@ static bool GenerateDestIniFile(FConfigFile& DestConfigFile, const FString& Dest
 
 /**
  * Calculates the name of the source (default) .ini file for a given base (ie Engine, Game, etc)
- * 
+ *
+ * @param ConfigDir Path to ini file
+ * @param Prefix Prefix for the ini filename (Default/Platform name), may include additional sub-directories
  * @param IniBaseName Base name of the .ini (Engine, Game)
- * @param PlatformName Name of the platform to get the .ini path for (NULL means to use the current platform)
- * @param GameName The name of the game to get the .ini path for (if NULL, uses current)
  * 
  * @return Standardized .ini filename
  */
-static FString GetSourceIniFilename(const TCHAR* BaseIniName, const TCHAR* PlatformName, const TCHAR* GameName)
+static FString GetSourceIniFilename(const TCHAR* ConfigDir, const TCHAR* Prefix, const TCHAR* BaseIniName)
 {
-	checkf(GameName==NULL, TEXT("Specifying a GameName in GetSourceIniFilename requires appSourceConfigDir to take a GameName parameter"));
-
-	// figure out what to look for on the commandline for an override
-	FString CommandLineSwitch = FString::Printf(TEXT("DEF%sINI="), BaseIniName);
-
-	// if it's not found on the commandline, then generate it
 	FString IniFilename;
+	
+#if !UE_BUILD_SHIPPING
+	// Figure out what to look for on the commandline for an override. Disabled in shipping builds for security reasons
+	const FString CommandLineSwitch = FString::Printf(TEXT("DEF%sINI="), BaseIniName);	
 	if (FParse::Value(FCommandLine::Get(), *CommandLineSwitch, IniFilename) == false)
+#endif
 	{
-		FString Name(PlatformName ? PlatformName : ANSI_TO_TCHAR(FPlatformProperties::PlatformName()));
-
-		// figure out where in the Config dir the source ini is
-		FString Prefix;
-		if (Name.StartsWith(TEXT("Windows"))) //@todo hack, either move to the platform abstraction or make Windows like all other platforms
-		{
-			// @todo: Support -simmobile using Mobile .ini's?
-			// 	if (FParse::Param( FCommandLine::Get(), TEXT("simmobile"))) return TEXT("Mobile/")
-			// for Windows, just use the default directly
-			Prefix = TEXT("Default");
-		}
-		else
-		{
-			// other platforms are Platform/Platform (like PS3/PS3Engine.ini)
-			Prefix = FString::Printf(TEXT("%s/%s"), *Name, *Name);
-		}
-
-		// put it all together
-		IniFilename = FString::Printf(TEXT("%s%s%s.ini"), *FPaths::SourceConfigDir(), *Prefix, BaseIniName);
+		// if it's not found on the commandline, then generate it	
+		IniFilename = FString(ConfigDir) / FString::Printf(TEXT("%s%s.ini"), Prefix, BaseIniName);
 	}
 
-	// standardize it!
 	FPaths::MakeStandardFilename(IniFilename);
 	return IniFilename;
 }
@@ -2766,7 +2747,7 @@ static void GetSourceIniHierarchyFilenames(const TCHAR* InBaseIniName, const TCH
 
 	// [[[[ PROJECT SETTINGS ]]]]
 	// Game/Config/Default* ini
-	OutHierarchy.Add(FIniFilename(FString::Printf(TEXT("%sDefault%s.ini"), SourceConfigDir, InBaseIniName), bRequireDefaultIni));
+	OutHierarchy.Add(FIniFilename(GetSourceIniFilename(SourceConfigDir, TEXT("Default"), InBaseIniName), bRequireDefaultIni));
 	// Game/Config/DedicatedServer* ini (knowingly NOT in EngineConfiguration.cs because this is a runtime only check)
 	if (IsRunningDedicatedServer())
 	{
