Fixed text automation test that was comparing actual results to inaccurate expected results.

diff --git a/Engine/Source/Runtime/Core/Private/Tests/Internationalization/TextTest.cpp b/Engine/Source/Runtime/Core/Private/Tests/Internationalization/TextTest.cpp
--- a/Engine/Source/Runtime/Core/Private/Tests/Internationalization/TextTest.cpp
+++ b/Engine/Source/Runtime/Core/Private/Tests/Internationalization/TextTest.cpp
@@ -448,7 +448,7 @@ bool FTextTest::RunTest (const FString& Parameters)
 			FText AsCurrencyTest1 = FText::AsCurrency(100.25);
 
 			FDateTime DateTimeInfo(2080, 8, 20, 9, 33, 22);
-			FText AsDateTimeTest1 = FText::AsDateTime(DateTimeInfo, EDateTimeStyle::Default, EDateTimeStyle::Default, TEXT("CST"));
+			FText AsDateTimeTest1 = FText::AsDateTime(DateTimeInfo, EDateTimeStyle::Default, EDateTimeStyle::Default, TEXT("UTC"));
 
 			// FormattedTestLayer2 must be updated when adding or removing from this block. Further, below, 
 			// verifying the LEET translated string must be changed to reflect what the new string looks like.
@@ -479,7 +479,7 @@ bool FTextTest::RunTest (const FString& Parameters)
 				// When changes are made to FormattedTestLayer2, please pull out the newly translated LEET string and update the below if-statement to keep the test passing!
 				FString LEETTranslatedString = FormattedTestLayer2.ToString();
 
-				FString DesiredOutput = FString(TEXT("\x2021") TEXT("\xAB") TEXT("\xAB") TEXT("L0r3m") TEXT("\xBB") TEXT(" \"L0r3m 1p$um\" ") TEXT("\xAB") TEXT("Ip$um") TEXT("\xBB") TEXT("\xBB") TEXT(" | ") TEXT("\xAB") TEXT("\xAB") TEXT("L0r3m") TEXT("\xBB") TEXT(" \"L0r3m 1p$um\" ") TEXT("\xAB") TEXT("Ip$um") TEXT("\xBB") TEXT("\xBB") TEXT(" | ") TEXT("\xAB") TEXT("5.5421") TEXT("\xBB") TEXT(" | ") TEXT("\xAB") TEXT("5010.89221") TEXT("\xBB") TEXT(" | ") TEXT("\xAB") TEXT("Jul 14, 1944, 10:05:06 PM") TEXT("\xBB") TEXT(" | ") TEXT("\xAB") TEXT("92%") TEXT("\xBB") TEXT(" | ") TEXT("\xAB") TEXT("\xA4") TEXT("\xA0") TEXT("100.25") TEXT("\xBB") TEXT("\x2021"));
+				FString DesiredOutput = FString(TEXT("\x2021") TEXT("\xAB") TEXT("\xAB") TEXT("L0r3m") TEXT("\xBB") TEXT(" \"L0r3m 1p$um\" ") TEXT("\xAB") TEXT("Ip$um") TEXT("\xBB") TEXT("\xBB") TEXT(" | ") TEXT("\xAB") TEXT("\xAB") TEXT("L0r3m") TEXT("\xBB") TEXT(" \"L0r3m 1p$um\" ") TEXT("\xAB") TEXT("Ip$um") TEXT("\xBB") TEXT("\xBB") TEXT(" | ") TEXT("\xAB") TEXT("5.5421") TEXT("\xBB") TEXT(" | ") TEXT("\xAB") TEXT("5010.89221") TEXT("\xBB") TEXT(" | ") TEXT("\xAB") TEXT("Aug 20, 2080, 9:33:22 AM") TEXT("\xBB") TEXT(" | ") TEXT("\xAB") TEXT("92%") TEXT("\xBB") TEXT(" | ") TEXT("\xAB") TEXT("\xA4") TEXT("\xA0") TEXT("100.25") TEXT("\xBB") TEXT("\x2021"));
 				// Convert the baked string into an FText, which will be leetified, then compare it to the rebuilt FText
 				if(LEETTranslatedString != DesiredOutput)
 				{
@@ -514,7 +514,7 @@ bool FTextTest::RunTest (const FString& Parameters)
 					AddError( TEXT("Currency Output=") + AsCurrencyTest1.ToString() );
 				}
 
-				if(AsDateTimeTest1.CompareTo(FText::AsDateTime(DateTimeInfo, EDateTimeStyle::Default, EDateTimeStyle::Default, TEXT("CST"))) != 0)
+				if(AsDateTimeTest1.CompareTo(FText::AsDateTime(DateTimeInfo, EDateTimeStyle::Default, EDateTimeStyle::Default, TEXT("UTC"))) != 0)
 				{
 					AddError( TEXT("AsDateTimeTest1 did not rebuild correctly in French-Canadian") );
 					AddError( TEXT("DateTime Output=") + AsDateTimeTest1.ToString() );
