Android: make -Wparentheses-equality non-fatal.

- Constructs like if (foo=1) will still break the build, but if ((foo == 1)) will not (won't be suppressed completely, but will produce a visible warning).

#ue4
#android

diff --git a/Engine/Source/Runtime/Core/Public/Android/AndroidPlatformCompilerPreSetup.h b/Engine/Source/Runtime/Core/Public/Android/AndroidPlatformCompilerPreSetup.h
--- a/Engine/Source/Runtime/Core/Public/Android/AndroidPlatformCompilerPreSetup.h
+++ b/Engine/Source/Runtime/Core/Public/Android/AndroidPlatformCompilerPreSetup.h
@@ -49,4 +49,9 @@
 #ifndef EMIT_CUSTOM_WARNING_AT_LINE
 #define EMIT_CUSTOM_WARNING_AT_LINE(Line, Warning) \
 	_Pragma(PREPROCESSOR_TO_STRING(message(WARNING_LOCATION(Line) Warning)))
-#endif // EMIT_CUSTOM_WARNING_AT_LINE
\ No newline at end of file
+#endif // EMIT_CUSTOM_WARNING_AT_LINE
+
+// Make certain warnings always be warnings, even despite -Werror.
+// Rationale: we don't want to suppress those as there are plans to address them (e.g. UE-12341), but breaking builds due to these warnings is very expensive
+// since they cannot be caught by all compilers that we support. They are deemed to be relatively safe to be ignored, at least until all SDKs/toolchains start supporting them.
+#pragma clang diagnostic warning "-Wparentheses-equality"
