Set AssetPath for local files to fix issue with Get/SetTimeStamp on Android
#ue4
#android

diff --git a/Engine/Source/Runtime/Core/Private/Android/AndroidFile.cpp b/Engine/Source/Runtime/Core/Private/Android/AndroidFile.cpp
--- a/Engine/Source/Runtime/Core/Private/Android/AndroidFile.cpp
+++ b/Engine/Source/Runtime/Core/Private/Android/AndroidFile.cpp
@@ -1603,6 +1603,7 @@ private:
 			{
 				// Absolute paths are only local.
 				LocalPath = AndroidPath;
+				AssetPath = AndroidPath;
 			}
 			else
 			{
