UEPLAT-593 implement lifetime event callbacks as game thread tasks
#ios

diff --git a/Engine/Source/Runtime/Core/Private/Async/TaskGraph.cpp b/Engine/Source/Runtime/Core/Private/Async/TaskGraph.cpp
--- a/Engine/Source/Runtime/Core/Private/Async/TaskGraph.cpp
+++ b/Engine/Source/Runtime/Core/Private/Async/TaskGraph.cpp
@@ -1202,6 +1202,11 @@ void FTaskGraphInterface::Shutdown()
 	TaskGraphImplementationSingleton = NULL;
 }
 
+bool FTaskGraphInterface::IsRunning()
+{
+    return TaskGraphImplementationSingleton != NULL;
+}
+
 FTaskGraphInterface& FTaskGraphInterface::Get()
 {
 	checkThreadGraph(TaskGraphImplementationSingleton);
diff --git a/Engine/Source/Runtime/Core/Private/IOS/IOSAppDelegate.cpp b/Engine/Source/Runtime/Core/Private/IOS/IOSAppDelegate.cpp
--- a/Engine/Source/Runtime/Core/Private/IOS/IOSAppDelegate.cpp
+++ b/Engine/Source/Runtime/Core/Private/IOS/IOSAppDelegate.cpp
@@ -541,8 +541,12 @@ void InstallSignalHandlers()
 	 Sent when the application is about to move from active to inactive state. This can occur for certain types of temporary interruptions (such as an incoming phone call or SMS message) or when the user quits the application and it begins the transition to the background state.
 	 Use this method to pause ongoing tasks, disable timers, and throttle down OpenGL ES frame rates. Games should use this method to pause the game.
 	 */
-
-	FCoreDelegates::ApplicationWillDeactivateDelegate.Broadcast();
+    FGraphEventRef ResignTask = FFunctionGraphTask::CreateAndDispatchWhenReady([&]()
+	{
+		FCoreDelegates::ApplicationWillDeactivateDelegate.Broadcast();
+    }, TStatId(), NULL, ENamedThreads::GameThread);
+    FTaskGraphInterface::Get().WaitUntilTaskCompletes(ResignTask);
+    
 	[self ToggleSuspend:true];
 	[self ToggleAudioSession:false];
 }
@@ -553,7 +557,7 @@ void InstallSignalHandlers()
 	 Use this method to release shared resources, save user data, invalidate timers, and store enough application state information to restore your application to its current state in case it is terminated later. 
 	 If your application supports background execution, this method is called instead of applicationWillTerminate: when the user quits.
 	 */
-	FCoreDelegates::ApplicationWillEnterBackgroundDelegate.Broadcast();
+    FCoreDelegates::ApplicationWillEnterBackgroundDelegate.Broadcast();
 }
 
 - (void)applicationWillEnterForeground:(UIApplication *)application
@@ -561,7 +565,7 @@ void InstallSignalHandlers()
 	/*
 	 Called as part of the transition from the background to the inactive state; here you can undo many of the changes made on entering the background.
 	 */
-	FCoreDelegates::ApplicationHasEnteredForegroundDelegate.Broadcast();
+    FCoreDelegates::ApplicationHasEnteredForegroundDelegate.Broadcast();
 }
 
 - (void)applicationDidBecomeActive:(UIApplication *)application
@@ -569,10 +573,17 @@ void InstallSignalHandlers()
 	/*
 	 Restart any tasks that were paused (or not yet started) while the application was inactive. If the application was previously in the background, optionally refresh the user interface.
 	 */
-
     [self ToggleSuspend:false];
 	[self ToggleAudioSession:true];
-	FCoreDelegates::ApplicationHasReactivatedDelegate.Broadcast();
+
+    if (FTaskGraphInterface::IsRunning())
+    {
+        FGraphEventRef ResignTask = FFunctionGraphTask::CreateAndDispatchWhenReady([&]()
+        {
+            FCoreDelegates::ApplicationHasReactivatedDelegate.Broadcast();
+        }, TStatId(), NULL, ENamedThreads::GameThread);
+        FTaskGraphInterface::Get().WaitUntilTaskCompletes(ResignTask);
+    }
 }
 
 - (void)applicationWillTerminate:(UIApplication *)application
@@ -623,10 +634,10 @@ void InstallSignalHandlers()
  * @param Controller The Controller object to animate off the screen
  * @param bShouldAnimate YES to slide down, NO to hide immediately
  */
--(void)HideController:(UIViewController*)Controller Animated : (BOOL)bShouldAnimate
+-(void)HideController:(UIViewController*)Controller Animated:(BOOL)bShouldAnimate
 {
-	// slide it off
-	[Controller dismissViewControllerAnimated : bShouldAnimate completion : nil];
+    // slide it off
+    [Controller dismissViewControllerAnimated : bShouldAnimate completion : nil];
 
 	// stop drawing the 3D world for faster UI speed
 	//FViewport::SetGameRenderingEnabled(true);
diff --git a/Engine/Source/Runtime/Core/Public/Async/TaskGraphInterfaces.h b/Engine/Source/Runtime/Core/Public/Async/TaskGraphInterfaces.h
--- a/Engine/Source/Runtime/Core/Public/Async/TaskGraphInterfaces.h
+++ b/Engine/Source/Runtime/Core/Public/Async/TaskGraphInterfaces.h
@@ -121,7 +121,11 @@ public:
 	 *	Explicit start call to shutdown the system. This is unlikely to work unless the system is idle.
 	**/
 	static CORE_API void Shutdown();
-	/** 
+    /**
+     *	Check to see if the system is running.
+     **/
+    static CORE_API bool IsRunning();
+	/**
 	 *	Singleton for the system
 	 *	@return a reference to the task graph system
 	**/
@@ -1177,6 +1181,68 @@ public:
 	}
 };
 
+/** Task class for lambda based tasks. **/
+class FFunctionGraphTask : public FCustomStatIDGraphTaskBase
+{
+public:
+    /** Function to run **/
+    TFunction<void()> Function;
+    /** Thread to run the function on **/
+    const ENamedThreads::Type			DesiredThread;
+    
+public:
+    ENamedThreads::Type GetDesiredThread()
+    {
+        return DesiredThread;
+    }
+    static ESubsequentsMode::Type GetSubsequentsMode() { return ESubsequentsMode::TrackSubsequents; }
+    
+    void DoTask(ENamedThreads::Type CurrentThread, const FGraphEventRef& MyCompletionGraphEvent)
+    {
+        Function();
+    }
+    /**
+     * Task constructor
+     * @param InFunction - function to execute when the prerequisites are complete
+     *	@param StatId The stat id for this task.
+     * @param InDesiredThread - Thread to run on
+     **/
+    FFunctionGraphTask(TFunction<void()>& InFunction, const TStatId StatId, ENamedThreads::Type InDesiredThread)
+    : FCustomStatIDGraphTaskBase(StatId)
+    , Function(InFunction)
+    , DesiredThread(InDesiredThread)
+    {
+    }
+    
+    /**
+     * Create a task and dispatch it when the prerequisites are complete
+     * @param InTaskDelegate - delegate to execute when the prerequisites are complete
+     * @param InStatId - StatId of task for debugging or analysis tools
+     * @param InPrerequisites - Handles for prerequisites for this task, can be NULL if there are no prerequisites
+     * @param InDesiredThread - Thread to run on
+     * @return completion handle for the new task
+     **/
+    static FGraphEventRef CreateAndDispatchWhenReady(TFunction<void()> InFunction, const TStatId InStatId, const FGraphEventArray* InPrerequisites = NULL, ENamedThreads::Type InDesiredThread = ENamedThreads::AnyThread)
+    {
+        return TGraphTask<FFunctionGraphTask>::CreateTask(InPrerequisites).ConstructAndDispatchWhenReady(InFunction, InStatId, InDesiredThread);
+    }
+    /**
+     * Create a task and dispatch it when the prerequisites are complete
+     * @param InTaskDelegate - delegate to execute when the prerequisites are complete
+     * @param InStatId - StatId of task for debugging or analysis tools
+     * @param InPrerequisite - Handle for a single prerequisite for this task
+     * @param InDesiredThread - Thread to run on
+     * @return completion handle for the new task
+     **/
+    static FGraphEventRef CreateAndDispatchWhenReady(TFunction<void()> InFunction, const TStatId&& InStatId, const FGraphEventRef& InPrerequisite, ENamedThreads::Type InDesiredThread = ENamedThreads::AnyThread)
+    {
+        FGraphEventArray Prerequisites;
+        check(InPrerequisite.GetReference());
+        Prerequisites.Add(InPrerequisite);
+        return CreateAndDispatchWhenReady(InFunction, InStatId, &Prerequisites, InDesiredThread);
+    }
+};
+
 /**
  * List of tasks that can be "joined" into one task which can be waited on or used as a prerequisite.
  * Note, these are FGraphEventRef's, but we manually manage the reference count instead of using a smart pointer
