PR983 - iOS app delegate retains launch options
#ios

diff --git a/Engine/Source/Runtime/Core/Private/IOS/IOSAppDelegate.cpp b/Engine/Source/Runtime/Core/Private/IOS/IOSAppDelegate.cpp
--- a/Engine/Source/Runtime/Core/Private/IOS/IOSAppDelegate.cpp
+++ b/Engine/Source/Runtime/Core/Private/IOS/IOSAppDelegate.cpp
@@ -378,6 +378,9 @@ void InstallSignalHandlers()
 
 - (BOOL)application:(UIApplication *)application didFinishLaunchingWithOptions:(NSDictionary*)launchOptions
 {
+	// save launch options
+	self.launchOptions = launchOptions;
+
 	// use the status bar orientation to properly determine landscape vs portrait
 	self.bDeviceInPortraitMode = UIInterfaceOrientationIsPortrait([[UIApplication sharedApplication] statusBarOrientation]);
 	printf("========= This app is in %s mode\n", self.bDeviceInPortraitMode ? "PORTRAIT" : "LANDSCAPE");
diff --git a/Engine/Source/Runtime/Core/Public/IOS/IOSAppDelegate.h b/Engine/Source/Runtime/Core/Public/IOS/IOSAppDelegate.h
--- a/Engine/Source/Runtime/Core/Public/IOS/IOSAppDelegate.h
+++ b/Engine/Source/Runtime/Core/Public/IOS/IOSAppDelegate.h
@@ -76,6 +76,10 @@ UITextFieldDelegate>
 
 /** True if we need to reset the idle timer */
 @property (readonly) bool bResetIdleTimer;
+
+/** initial launch options */
+@property(retain) NSDictionary* launchOptions;
+
 /**
  * @return the single app delegate object
  */
