PR #838: Improve TSharedPtr-on-UObject error message. (Contributed by [REDACTED_PERSON])

diff --git a/Engine/Source/Runtime/Core/Public/Templates/SharedPointer.h b/Engine/Source/Runtime/Core/Public/Templates/SharedPointer.h
--- a/Engine/Source/Runtime/Core/Public/Templates/SharedPointer.h
+++ b/Engine/Source/Runtime/Core/Public/Templates/SharedPointer.h
@@ -465,7 +465,7 @@ template< class ObjectType, ESPMode Mode >
 class TSharedPtr
 {
 	// TSharedPtrs with UObjects are illegal.
-	static_assert(!TPointerIsConvertibleFromTo<ObjectType, const UObjectBase>::Value, "You cannot use TSharedPtr with UObjects.");
+	static_assert(!TPointerIsConvertibleFromTo<ObjectType, const UObjectBase>::Value, "You cannot use TSharedPtr or TWeakPtr with UObjects. Consider a UPROPERTY() pointer or TWeakObjectPtr.");
 
 public:
 
