Stats - Generalized method for dynamic stats

diff --git a/Engine/Source/Runtime/Core/Private/Stats/Stats2.cpp b/Engine/Source/Runtime/Core/Private/Stats/Stats2.cpp
--- a/Engine/Source/Runtime/Core/Private/Stats/Stats2.cpp
+++ b/Engine/Source/Runtime/Core/Private/Stats/Stats2.cpp
@@ -412,17 +412,19 @@ public:
 		--PendingCount;
 		TStatIdData* Result = PendingStatIds;
 
+		const FString StatDescription = InDescription ? InDescription : StatShortName.GetPlainNameString();
+
 		// Get the wide stat description.
-		const int32 StatDescLen = FCString::Strlen( InDescription ) + 1;
+		const int32 StatDescLen = StatDescription.Len() + 1;
 		// We are leaking this. @see STAT_StatDescMemory
 		WIDECHAR* StatDescWide = new WIDECHAR[StatDescLen];
-		TCString<WIDECHAR>::Strcpy( StatDescWide, StatDescLen, StringCast<WIDECHAR>( InDescription ).Get() );
+		TCString<WIDECHAR>::Strcpy( StatDescWide, StatDescLen, StringCast<WIDECHAR>( *StatDescription ).Get() );
 		Result->WideString = reinterpret_cast<uint64>(StatDescWide);
 
 		// Get the ansi stat description.
 		// We are leaking this. @see STAT_StatDescMemory
 		ANSICHAR* StatDescAnsi = new ANSICHAR[StatDescLen];
-		TCString<ANSICHAR>::Strcpy( StatDescAnsi, StatDescLen, StringCast<ANSICHAR>( InDescription ).Get() );
+		TCString<ANSICHAR>::Strcpy( StatDescAnsi, StatDescLen, StringCast<ANSICHAR>( *StatDescription ).Get() );
 		Result->AnsiString = reinterpret_cast<uint64>(StatDescAnsi);
 
 		MemoryCounter.Add( StatDescLen*(sizeof( ANSICHAR ) + sizeof( WIDECHAR )) );
diff --git a/Engine/Source/Runtime/Core/Public/Stats/Stats.h b/Engine/Source/Runtime/Core/Public/Stats/Stats.h
--- a/Engine/Source/Runtime/Core/Public/Stats/Stats.h
+++ b/Engine/Source/Runtime/Core/Public/Stats/Stats.h
@@ -434,4 +434,40 @@ FORCEINLINE void StatsMasterEnableSubtract(int32 Value = 1)
 #define GET_STATFNAME(Stat) (FName())
 #define GET_STATDESCRIPTION(Stat) (nullptr)
 
-#endif
\ No newline at end of file
+#endif
+
+
+/** Helper class used to generate dynamic stat ids. */
+struct FDynamicStats
+{
+	/**
+	* Create a new stat id and registers it with the stats system.
+	* This is the only way to create dynamic stat ids at runtime.
+	* Can be used only with FScopeCycleCounters.
+	*
+	* Example:
+	*	FDynamicStats::CreateStatId<STAT_GROUP_TO_FStatGroup( STATGROUP_UObjects )>( FString::Printf(TEXT("MyDynamicStat_%i"),Index) )
+	*/
+	template< typename TStatGroup >
+	static TStatId CreateStatId( const FString& StatNameOrDescription )
+	{
+#if	STATS
+		const FName StatName = FName( *StatNameOrDescription );
+		FStartupMessages::Get().AddMetadata( StatName, nullptr,
+											 TStatGroup::GetGroupName(),
+											 TStatGroup::GetGroupCategory(),
+											 TStatGroup::GetDescription(),
+											 true, EStatDataType::ST_int64, true );
+
+		TStatId StatID = IStatGroupEnableManager::Get().GetHighPerformanceEnableForStat( StatName,
+			TStatGroup::GetGroupName(),
+			TStatGroup::GetGroupCategory(),
+			TStatGroup::DefaultEnable,
+			true, EStatDataType::ST_int64, nullptr, true );
+
+		return StatID;
+#endif // STATS
+
+		return TStatId();
+	}
+};
\ No newline at end of file
