Disaplay filename when reporting errors about writing to file. Better error reporting for FArchiveFileWriterGeneric

diff --git a/Engine/Source/Runtime/Core/Private/HAL/FileManagerGeneric.cpp b/Engine/Source/Runtime/Core/Private/HAL/FileManagerGeneric.cpp
--- a/Engine/Source/Runtime/Core/Private/HAL/FileManagerGeneric.cpp
+++ b/Engine/Source/Runtime/Core/Private/HAL/FileManagerGeneric.cpp
@@ -80,7 +80,7 @@ public:
 	{
 		Pos += Length;
 	}
-	virtual FString GetArchiveName() const { return TEXT( "FArchiveFileWriterDummy" ); }
+	virtual FString GetArchiveName() const override { return TEXT( "FArchiveFileWriterDummy" ); }
 protected:
 	int64             Pos;
 };
@@ -561,7 +561,7 @@ void FArchiveFileReaderGeneric::Seek( int64 InPos )
 	{
 		TCHAR ErrorBuffer[1024];
 		ArIsError = true;
-		UE_LOG( LogFileManager, Error, TEXT( "SetFilePointer Failed %lld/%lld: %lld %s" ), InPos, Size, Pos, FPlatformMisc::GetSystemErrorMessage( ErrorBuffer, 1024, 0 ) );
+		UE_LOG(LogFileManager, Error, TEXT("SetFilePointer on %s Failed %lld/%lld: %lld %s"), *Filename, InPos, Size, Pos, FPlatformMisc::GetSystemErrorMessage(ErrorBuffer, 1024, 0));
 	}
 	Pos         = InPos;
 	BufferBase  = Pos;
@@ -715,7 +715,7 @@ void FArchiveFileWriterGeneric::Seek( int64 InPos )
 	if( !SeekLowLevel( InPos ) )
 	{
 		ArIsError = true;
-		UE_LOG( LogFileManager, Error, TEXT("%s"), TEXT("Error seeking file") );
+		LogWriteError(TEXT("Error seeking file"));
 	}
 	Pos = InPos;
 }
@@ -726,7 +726,7 @@ bool FArchiveFileWriterGeneric::Close()
 	if( !CloseLowLevel() )
 	{
 		ArIsError = true;
-		UE_LOG( LogFileManager, Error, TEXT("%s"), TEXT("Error writing to file") );
+		LogWriteError(TEXT("Error closing file"));
 	}
 	return !ArIsError;
 }
@@ -740,7 +740,7 @@ void FArchiveFileWriterGeneric::Serialize( void* V, int64 Length )
 		if( !WriteLowLevel( (uint8*)V, Length ) )
 		{
 			ArIsError = true;
-			UE_LOG( LogFileManager, Error, TEXT("%s"), TEXT("Error writing to file") );
+			LogWriteError(TEXT("Error writing to file"));
 		}
 	}
 	else
@@ -772,12 +772,17 @@ void FArchiveFileWriterGeneric::Flush()
 		if( !WriteLowLevel( Buffer, BufferCount ) )
 		{
 			ArIsError = true;
-			UE_LOG( LogFileManager, Error, TEXT("%s"), TEXT("Error writing to file") );
+			LogWriteError(TEXT("Error flushing file"));
 		}
 		BufferCount = 0;
 	}
 }
 
+void FArchiveFileWriterGeneric::LogWriteError(const TCHAR* Message)
+{
+	TCHAR ErrorBuffer[1024];
+	UE_LOG(LogFileManager, Error, TEXT("%s: %s (%s)"), Message, *Filename, FPlatformMisc::GetSystemErrorMessage(ErrorBuffer, 1024, 0));
+}
 //---
 
 
diff --git a/Engine/Source/Runtime/Core/Public/HAL/FileManagerGeneric.h b/Engine/Source/Runtime/Core/Public/HAL/FileManagerGeneric.h
--- a/Engine/Source/Runtime/Core/Public/HAL/FileManagerGeneric.h
+++ b/Engine/Source/Runtime/Core/Public/HAL/FileManagerGeneric.h
@@ -222,6 +222,10 @@ public:
 	virtual bool Close() final;
 	virtual void Serialize( void* V, int64 Length ) final;
 	virtual void Flush() final;
+	virtual FString GetArchiveName() const override
+	{
+		return Filename;
+	}
 
 protected:
 	/** 
@@ -243,6 +247,14 @@ protected:
 	**/
 	virtual bool WriteLowLevel(const uint8* Src, int64 CountToWrite);
 
+	/** 
+	 * Logs I/O error
+	 * It is important to not call any platform API functions after the error occurred and before 
+	 * calling this functions as the system error code may get reset and will not be properly
+	 * logged in this message.
+	 * @param Message Brief description of what went wrong
+	 */
+	void LogWriteError(const TCHAR* Message);
 
 	/** Filename for debugging purposes */
 	FString			Filename;
