Core - Misc tweaks to async loading blocking time stat

diff --git a/Engine/Source/Runtime/Core/Private/Misc/Core.cpp b/Engine/Source/Runtime/Core/Private/Misc/Core.cpp
--- a/Engine/Source/Runtime/Core/Private/Misc/Core.cpp
+++ b/Engine/Source/Runtime/Core/Private/Misc/Core.cpp
@@ -306,7 +306,7 @@ DEFINE_STAT(STAT_AsyncIO_CanceledReadSize);
 DEFINE_STAT(STAT_AsyncIO_OutstandingReadCount);
 DEFINE_STAT(STAT_AsyncIO_OutstandingReadSize);
 DEFINE_STAT(STAT_AsyncIO_UncompressorWaitTime);
-DEFINE_STAT(STAT_AsyncIO_MainThreadBlockTime);
+DEFINE_STAT(STAT_AsyncIO_AsyncLoadingBlockingTime);
 DEFINE_STAT(STAT_AsyncIO_AsyncPackagePrecacheWaitTime);
 DEFINE_STAT(STAT_AsyncIO_Bandwidth);
 DEFINE_STAT(STAT_AsyncIO_PlatformReadTime);
diff --git a/Engine/Source/Runtime/Core/Public/Misc/CoreStats.h b/Engine/Source/Runtime/Core/Public/Misc/CoreStats.h
--- a/Engine/Source/Runtime/Core/Public/Misc/CoreStats.h
+++ b/Engine/Source/Runtime/Core/Public/Misc/CoreStats.h
@@ -79,7 +79,7 @@ DECLARE_DWORD_ACCUMULATOR_STAT_EXTERN(TEXT("Outstanding read count"),STAT_AsyncI
 DECLARE_MEMORY_STAT_EXTERN(TEXT("Outstanding read size"),STAT_AsyncIO_OutstandingReadSize,STATGROUP_AsyncIO, CORE_API);
 DECLARE_FLOAT_ACCUMULATOR_STAT_EXTERN(TEXT("Platform read time"),STAT_AsyncIO_PlatformReadTime,STATGROUP_AsyncIO, CORE_API);
 DECLARE_FLOAT_ACCUMULATOR_STAT_EXTERN(TEXT("Uncompressor wait time"),STAT_AsyncIO_UncompressorWaitTime,STATGROUP_AsyncIO, CORE_API);
-DECLARE_FLOAT_ACCUMULATOR_STAT_EXTERN(TEXT("Main thread block time"),STAT_AsyncIO_MainThreadBlockTime,STATGROUP_AsyncIO, CORE_API);
+DECLARE_FLOAT_ACCUMULATOR_STAT_EXTERN(TEXT("Async loading block time"),STAT_AsyncIO_AsyncLoadingBlockingTime,STATGROUP_AsyncIO, CORE_API);
 DECLARE_FLOAT_ACCUMULATOR_STAT_EXTERN(TEXT("Async package precache wait time"),STAT_AsyncIO_AsyncPackagePrecacheWaitTime,STATGROUP_AsyncIO, CORE_API);
 DECLARE_FLOAT_COUNTER_STAT_EXTERN(TEXT("Bandwidth (MByte/ sec)"),STAT_AsyncIO_Bandwidth,STATGROUP_AsyncIO, CORE_API);
 
diff --git a/Engine/Source/Runtime/CoreUObject/Private/Serialization/AsyncLoading.cpp b/Engine/Source/Runtime/CoreUObject/Private/Serialization/AsyncLoading.cpp
--- a/Engine/Source/Runtime/CoreUObject/Private/Serialization/AsyncLoading.cpp
+++ b/Engine/Source/Runtime/CoreUObject/Private/Serialization/AsyncLoading.cpp
@@ -542,7 +542,7 @@ void FAsyncLoadingThread::TickAsyncThread(bool bUseTimeLimit, bool bUseFullTimeL
 	}
 
 	// Update stats
-	SET_FLOAT_STAT( STAT_AsyncIO_MainThreadBlockTime, FPlatformTime::ToSeconds( BlockingCycles.GetValue() ) );
+	SET_FLOAT_STAT( STAT_AsyncIO_AsyncLoadingBlockingTime, FPlatformTime::ToSeconds( BlockingCycles.GetValue() ) );
 	BlockingCycles.Set( 0 );
 }
 
@@ -2103,7 +2103,7 @@ void FArchiveAsync::Serialize( void* Data, int64 Count )
 
 	// Update stats if we were blocked.
 #if STATS
-	if( bIOBlocked && IsInGameThread() )
+	if( bIOBlocked )
 	{
 		const int32 BlockingCycles = int32(FPlatformTime::Cycles() - StartCycles);
 		FAsyncLoadingThread::BlockingCycles.Add( BlockingCycles );
