Attempt to harden code against rare Mac crashes in IME due to text blocks moving between windows.

diff --git a/Engine/Source/Runtime/Core/Private/Mac/MacTextInputMethodSystem.cpp b/Engine/Source/Runtime/Core/Private/Mac/MacTextInputMethodSystem.cpp
--- a/Engine/Source/Runtime/Core/Private/Mac/MacTextInputMethodSystem.cpp
+++ b/Engine/Source/Runtime/Core/Private/Mac/MacTextInputMethodSystem.cpp
@@ -33,7 +33,7 @@ namespace
 		
 	private:
 		TWeakPtr<ITextInputMethodContext> Context;
-		TSharedPtr<FGenericWindow> ContextWindow;
+		TWeakPtr<FGenericWindow> ContextWindow;
 	};
 	
 	void FTextInputMethodChangeNotifier::SetContextWindow(TSharedPtr<FGenericWindow> Window)
@@ -47,7 +47,7 @@ namespace
 		{
 			ContextWindow = Context.Pin()->GetWindow();
 		}
-		return ContextWindow;
+		return ContextWindow.Pin();
 	}
 	
 	void FTextInputMethodChangeNotifier::NotifyLayoutChanged(const ELayoutChangeType ChangeType)
