Collect the correct values for maximum available physical memory & currently used virtual memory on OS X.

diff --git a/Engine/Source/Runtime/Core/Private/Mac/MacPlatformMemory.cpp b/Engine/Source/Runtime/Core/Private/Mac/MacPlatformMemory.cpp
--- a/Engine/Source/Runtime/Core/Private/Mac/MacPlatformMemory.cpp
+++ b/Engine/Source/Runtime/Core/Private/Mac/MacPlatformMemory.cpp
@@ -9,6 +9,9 @@
 #include "MallocAnsi.h"
 #include "MallocBinned.h"
 
+#include <sys/param.h>
+#include <sys/mount.h>
+
 void FMacPlatformMemory::Init()
 {
 	const FPlatformMemoryConstants& MemoryConstants = FPlatformMemory::GetConstants();
@@ -48,18 +51,29 @@ FPlatformMemoryStats FMacPlatformMemory::GetStats()
 	mach_msg_type_number_t StatsSize = sizeof(Stats);
 	host_statistics(mach_host_self(), HOST_VM_INFO, (host_info_t)&Stats, &StatsSize);
 	uint64_t FreeMem = Stats.free_count * MemoryConstants.PageSize;
-	uint64_t UsedMem = (Stats.active_count + Stats.inactive_count + Stats.wire_count) * MemoryConstants.PageSize;
-
 	MemoryStats.AvailablePhysical = FreeMem;
-	MemoryStats.AvailableVirtual = 0;
-
-	MemoryStats.UsedVirtual = 0;
+	
+	// Get swap file info
+	xsw_usage SwapUsage;
+	SIZE_T Size = sizeof(SwapUsage);
+	sysctlbyname("vm.swapusage", &SwapUsage, &Size, NULL, 0);
+	MemoryStats.AvailableVirtual = FreeMem + SwapUsage.xsu_avail;
 
 	// Just get memory information for the process and report the working set instead
-	task_basic_info_64_data_t TaskInfo;
-	mach_msg_type_number_t TaskInfoCount = TASK_BASIC_INFO_COUNT;
-	task_info( mach_task_self(), TASK_BASIC_INFO, (task_info_t)&TaskInfo, &TaskInfoCount );
+	mach_task_basic_info_data_t TaskInfo;
+	mach_msg_type_number_t TaskInfoCount = MACH_TASK_BASIC_INFO_COUNT;
+	task_info( mach_task_self(), MACH_TASK_BASIC_INFO, (task_info_t)&TaskInfo, &TaskInfoCount );
 	MemoryStats.UsedPhysical = TaskInfo.resident_size;
+	if(MemoryStats.UsedPhysical > MemoryStats.PeakUsedPhysical)
+	{
+		MemoryStats.PeakUsedPhysical = MemoryStats.UsedPhysical;
+	}
+	MemoryStats.UsedVirtual = TaskInfo.virtual_size;
+	if(MemoryStats.UsedVirtual > MemoryStats.PeakUsedVirtual)
+	{
+		MemoryStats.PeakUsedVirtual = MemoryStats.UsedVirtual;
+	}
+	
 
 	return MemoryStats;
 }
@@ -82,17 +96,13 @@ const FPlatformMemoryConstants& FMacPlatformMemory::GetConstants()
 		sysctlbyname("vm.swapusage", &SwapUsage, &Size, NULL, 0);
 
 		// Get memory.
-		vm_statistics Stats;
-		mach_msg_type_number_t StatsSize = sizeof(Stats);
-		host_statistics(mach_host_self(), HOST_VM_INFO, (host_info_t)&Stats, &StatsSize);
-		uint64_t FreeMem = Stats.free_count * PageSize;
-		uint64_t UsedMem = (Stats.active_count + Stats.inactive_count + Stats.wire_count) * PageSize;
-		uint64_t TotalPhys = FreeMem + UsedMem;
-		uint64_t TotalPageFile = SwapUsage.xsu_total;
-		uint64_t TotalVirtual = TotalPhys + TotalPageFile;
-
-		MemoryConstants.TotalPhysical = TotalPhys;
-		MemoryConstants.TotalVirtual = TotalVirtual;
+		int64 AvailablePhysical = 0;
+		int Mib[] = {CTL_HW, HW_MEMSIZE};
+		size_t Length = sizeof(int64);
+		sysctl(Mib, 2, &AvailablePhysical, &Length, NULL, 0);
+		
+		MemoryConstants.TotalPhysical = AvailablePhysical;
+		MemoryConstants.TotalVirtual = AvailablePhysical + SwapUsage.xsu_total;
 		MemoryConstants.PageSize = (uint32)PageSize;
 
 		MemoryConstants.TotalPhysicalGB = (MemoryConstants.TotalPhysical + 1024 * 1024 * 1024 - 1) / 1024 / 1024 / 1024;
