fix for iOS compile error from PC
#ios

diff --git a/Engine/Source/Runtime/Core/Private/IOS/IOSAppDelegate.cpp b/Engine/Source/Runtime/Core/Private/IOS/IOSAppDelegate.cpp
--- a/Engine/Source/Runtime/Core/Private/IOS/IOSAppDelegate.cpp
+++ b/Engine/Source/Runtime/Core/Private/IOS/IOSAppDelegate.cpp
@@ -7,6 +7,7 @@
 #include "ModuleBoilerplate.h"
 #include "CallbackDevice.h"
 #include "IOSView.h"
+#include "TaskGraphInterfaces.h"
 
 #include <AudioToolbox/AudioToolbox.h>
 #include <AVFoundation/AVAudioSession.h>
