On losing & gaining application focus on OS X validate that the window numbers we are storing actually correspond to valid NSWindows, otherwise subsequent use of the window number might crash inside AppKit. This may help alleviate [UE-13034].

diff --git a/Engine/Source/Runtime/Core/Private/Mac/MacApplication.cpp b/Engine/Source/Runtime/Core/Private/Mac/MacApplication.cpp
--- a/Engine/Source/Runtime/Core/Private/Mac/MacApplication.cpp
+++ b/Engine/Source/Runtime/Core/Private/Mac/MacApplication.cpp
@@ -1000,14 +1000,19 @@ void FMacApplication::OnApplicationDidBecomeActive()
 
 	if (SavedWindowsOrder.Num() > 0)
 	{
-		[[NSApp windowWithWindowNumber:SavedWindowsOrder[0].WindowNumber] orderWindow:NSWindowAbove relativeTo:0];
+		NSWindow* TopWindow = [NSApp windowWithWindowNumber:SavedWindowsOrder[0].WindowNumber];
+		if ( TopWindow )
+		{
+			[TopWindow orderWindow:NSWindowAbove relativeTo:0];
+		}
 		for (int32 Index = 1; Index < SavedWindowsOrder.Num(); Index++)
 		{
 			const FSavedWindowOrderInfo& Info = SavedWindowsOrder[Index];
 			NSWindow* Window = [NSApp windowWithWindowNumber:Info.WindowNumber];
-			if (Window)
+			if (Window && TopWindow)
 			{
-				[Window orderWindow:NSWindowBelow relativeTo:SavedWindowsOrder[Index - 1].WindowNumber];
+				[Window orderWindow:NSWindowBelow relativeTo:[TopWindow windowNumber]];
+				TopWindow = Window;
 			}
 		}
 	}
@@ -1038,21 +1043,26 @@ void FMacApplication::OnApplicationWillResignActive()
 
 	if (SavedWindowsOrder.Num() > 0)
 	{
-		[[NSApp windowWithWindowNumber:SavedWindowsOrder[0].WindowNumber] orderWindow:NSWindowAbove relativeTo:0];
+		NSWindow* TopWindow = [NSApp windowWithWindowNumber:SavedWindowsOrder[0].WindowNumber];
+		if ( TopWindow )
+		{
+			[TopWindow orderWindow:NSWindowAbove relativeTo:0];
+		}
 		for (int32 Index = 1; Index < SavedWindowsOrder.Num(); Index++)
 		{
 			const FSavedWindowOrderInfo& Info = SavedWindowsOrder[Index];
 			NSWindow* Window = [NSApp windowWithWindowNumber:Info.WindowNumber];
 			if (Window)
 			{
-				if (SavedWindowsOrder[Index - 1].Level == Info.Level)
+				if (SavedWindowsOrder[Index - 1].Level == Info.Level && TopWindow)
 				{
-					[Window orderWindow:NSWindowBelow relativeTo:SavedWindowsOrder[Index - 1].WindowNumber];
+					[Window orderWindow:NSWindowBelow relativeTo:[TopWindow windowNumber]];
 				}
 				else
 				{
 					[Window orderWindow:NSWindowAbove relativeTo:0];
 				}
+				TopWindow = Window;
 			}
 		}
 	}
