On OS X when exiting the game thread, block and set the game-thread ID to to the Cocoa main-thread ID so that static C++ destructors that check IsInGameThread won't assert. This change also disables sudden termination support which can cause data corruption if the user reboots their Mac without exiting UE4 cleanly as the OS provides no notifications, events or signals when sudden termination is permitted - even if you are in the middle of a file I/O operation.

diff --git a/Engine/Source/Runtime/Core/Private/Mac/CocoaThread.cpp b/Engine/Source/Runtime/Core/Private/Mac/CocoaThread.cpp
--- a/Engine/Source/Runtime/Core/Private/Mac/CocoaThread.cpp
+++ b/Engine/Source/Runtime/Core/Private/Mac/CocoaThread.cpp
@@ -359,10 +359,18 @@ FCocoaRunLoopSource* FCocoaRunLoopSource::GameRunLoopSource = nullptr;
 	
 	if (GIsRequestingExit)
 	{
-		MainThreadCall(^{
+		dispatch_sync(dispatch_get_main_queue(), ^{
 			GGameThreadId = FPlatformTLS::GetCurrentThreadId();
 			[NSApp replyToApplicationShouldTerminate:YES];
-		}, NSDefaultRunLoopMode, false);
+			[[NSProcessInfo processInfo] enableSuddenTermination];
+		});
+	}
+	else
+	{
+		dispatch_sync(dispatch_get_main_queue(), ^{
+			GGameThreadId = FPlatformTLS::GetCurrentThreadId();
+			[[NSProcessInfo processInfo] enableSuddenTermination];
+		});
 	}
 	
 	[self release];
@@ -429,6 +437,10 @@ void GameThreadCall(dispatch_block_t Block, NSArray* SendModes, bool const bWait
 void RunGameThread(id Target, SEL Selector)
 {
 	SCOPED_AUTORELEASE_POOL;
+	
+	// @todo: Proper support for sudden termination (OS termination of application without any events, notifications or signals) - which presently can assert, crash or corrupt. At present we can't deal with this and we need to ensure that we don't permit the system to kill us without warning by disabling & enabling support around operations which must be committed to disk atomically.
+	[[NSProcessInfo processInfo] disableSuddenTermination];
+	
 #if MAC_SEPARATE_GAME_THREAD
 	// Register main run loop source
 	FCocoaRunLoopSource::RegisterMainRunLoop(CFRunLoopGetCurrent());
