Stats - Generalized method for dynamic stats (update)

diff --git a/Engine/Source/Runtime/Core/Public/Stats/Stats.h b/Engine/Source/Runtime/Core/Public/Stats/Stats.h
--- a/Engine/Source/Runtime/Core/Public/Stats/Stats.h
+++ b/Engine/Source/Runtime/Core/Public/Stats/Stats.h
@@ -445,6 +445,9 @@ struct FDynamicStats
 	* This is the only way to create dynamic stat ids at runtime.
 	* Can be used only with FScopeCycleCounters.
 	*
+	* Store the created stat id.
+	* Expensive method, avoid calling that method every frame.
+	*
 	* Example:
 	*	FDynamicStats::CreateStatId<STAT_GROUP_TO_FStatGroup( STATGROUP_UObjects )>( FString::Printf(TEXT("MyDynamicStat_%i"),Index) )
 	*/
@@ -452,14 +455,23 @@ struct FDynamicStats
 	static TStatId CreateStatId( const FString& StatNameOrDescription )
 	{
 #if	STATS
-		const FName StatName = FName( *StatNameOrDescription );
-		FStartupMessages::Get().AddMetadata( StatName, nullptr,
-											 TStatGroup::GetGroupName(),
-											 TStatGroup::GetGroupCategory(),
-											 TStatGroup::GetDescription(),
-											 true, EStatDataType::ST_int64, true );
-
-		TStatId StatID = IStatGroupEnableManager::Get().GetHighPerformanceEnableForStat( StatName,
+		return CreateStatId( FName( *StatNameOrDescription ) );
+#endif // STATS
+
+		return TStatId();
+	}
+
+	template< typename TStatGroup >
+	static TStatId CreateStatId( const FName StatNameOrDescription )
+	{
+#if	STATS
+		FStartupMessages::Get().AddMetadata( StatNameOrDescription, nullptr,
+			TStatGroup::GetGroupName(),
+			TStatGroup::GetGroupCategory(),
+			TStatGroup::GetDescription(),
+			true, EStatDataType::ST_int64, true );
+
+		TStatId StatID = IStatGroupEnableManager::Get().GetHighPerformanceEnableForStat( StatNameOrDescription,
 			TStatGroup::GetGroupName(),
 			TStatGroup::GetGroupCategory(),
 			TStatGroup::DefaultEnable,
