Stats - Generalized method for dynamic stats (fixed previous check-in)

diff --git a/Engine/Source/Runtime/Core/Public/Stats/Stats.h b/Engine/Source/Runtime/Core/Public/Stats/Stats.h
--- a/Engine/Source/Runtime/Core/Public/Stats/Stats.h
+++ b/Engine/Source/Runtime/Core/Public/Stats/Stats.h
@@ -455,7 +455,7 @@ struct FDynamicStats
 	static TStatId CreateStatId( const FString& StatNameOrDescription )
 	{
 #if	STATS
-		return CreateStatId( FName( *StatNameOrDescription ) );
+		return CreateStatId<TStatGroup>( FName( *StatNameOrDescription ) );
 #endif // STATS
 
 		return TStatId();
