Fix stats bug

diff --git a/Engine/Source/Runtime/Core/Private/Stats/Stats2.cpp b/Engine/Source/Runtime/Core/Private/Stats/Stats2.cpp
--- a/Engine/Source/Runtime/Core/Private/Stats/Stats2.cpp
+++ b/Engine/Source/Runtime/Core/Private/Stats/Stats2.cpp
@@ -368,24 +368,19 @@ public:
 			{
 				UE_LOG(LogStatGroupEnableManager, Fatal, TEXT("Stat group %s was was defined both on and off by default."), *Group.ToString());
 			}
-			TStatIdData** StatFound = Found->NamesInThisGroup.Find(Stat);
+			TStatIdData** StatFound = Found->NamesInThisGroup.Find( Stat );
 			TStatIdData** StatFoundAlways = Found->AlwaysEnabledNamesInThisGroup.Find( Stat );
-			if (StatFound)
+			if( StatFound )
 			{
-				TStatIdData** StatFound = Found->NamesInThisGroup.Find( Stat );
-				TStatIdData** StatFoundAlways = Found->AlwaysEnabledNamesInThisGroup.Find( Stat );
-				if( StatFound )
+				if( StatFoundAlways )
 				{
-					if( StatFoundAlways )
-					{
-						UE_LOG( LogStatGroupEnableManager, Fatal, TEXT( "Stat %s is both always enabled and not always enabled, so it was used for two different things." ), *Stat.ToString() );
-					}
-					return TStatId( *StatFound );
-				}
-				else if( StatFoundAlways )
-				{
-					return TStatId( *StatFoundAlways );
+					UE_LOG( LogStatGroupEnableManager, Fatal, TEXT( "Stat %s is both always enabled and not always enabled, so it was used for two different things." ), *Stat.ToString() );
 				}
+				return TStatId( *StatFound );
+			}
+			else if( StatFoundAlways )
+			{
+				return TStatId( *StatFoundAlways );
 			}
 		}
 		else
