Always make sure FRotator angles are in the range [0,360) before compressing them. Fixes UE-9488.

diff --git a/Engine/Source/Runtime/Core/Public/Math/Rotator.h b/Engine/Source/Runtime/Core/Public/Math/Rotator.h
--- a/Engine/Source/Runtime/Core/Public/Math/Rotator.h
+++ b/Engine/Source/Runtime/Core/Public/Math/Rotator.h
@@ -559,6 +559,9 @@ FORCEINLINE float FRotator::NormalizeAxis( float Angle )
 
 FORCEINLINE uint8 FRotator::CompressAxisToByte( float Angle )
 {
+	// Make sure angle is in the range [0,360)
+	Angle = ClampAxis(Angle);
+
 	// map [0->360) to [0->256) and mask off any winding
 	return FMath::RoundToInt(Angle * 256.f / 360.f) & 0xFF;
 }
@@ -573,6 +576,9 @@ FORCEINLINE float FRotator::DecompressAxisFromByte( uint16 Angle )
 
 FORCEINLINE uint16 FRotator::CompressAxisToShort( float Angle )
 {
+	// Make sure angle is in the range [0,360)
+	Angle = ClampAxis(Angle);
+
 	// map [0->360) to [0->65536) and mask off any winding
 	return FMath::RoundToInt(Angle * 65536.f / 360.f) & 0xFFFF;
 }
