PR #955: Fix FFileManagerGeneric::FindFiles() to also match provided pattern with directories (Contributed by [REDACTED_PERSON])

diff --git a/Engine/Source/Runtime/Core/Private/HAL/FileManagerGeneric.cpp b/Engine/Source/Runtime/Core/Private/HAL/FileManagerGeneric.cpp
--- a/Engine/Source/Runtime/Core/Private/HAL/FileManagerGeneric.cpp
+++ b/Engine/Source/Runtime/Core/Private/HAL/FileManagerGeneric.cpp
@@ -374,8 +374,8 @@ void FFileManagerGeneric::FindFiles( TArray<FString>& Result, const TCHAR* InFil
 		}
 		virtual bool Visit( const TCHAR* FilenameOrDirectory, bool bIsDirectory )
 		{
-			if( ( bIsDirectory && bDirectories ) ||
-				( !bIsDirectory && bFiles && FPaths::GetCleanFilename(FilenameOrDirectory).MatchesWildcard( WildCard ) ) )
+			if (((bIsDirectory && bDirectories) || (!bIsDirectory && bFiles))
+				&& FPaths::GetCleanFilename(FilenameOrDirectory).MatchesWildcard(WildCard))
 			{
 				new( Result ) FString( FPaths::GetCleanFilename(FilenameOrDirectory) );
 			}
