stat cmds triggered from blueprints in shipping no longer hang

#change knock on effect from CL#2062088. There appears to be an issue with starting a FSimpleDelegateGraphTask on the GameThread (which happens in shipping as STATS isn't defined), when this happens WaitUntilTaskCompletes waits forever which causes the game to hang. There may be an issue in Core with starting a thread on a thread which you want to yeild to that same thread(?) not sure... As a workaround, we avoid threading the request when SHIPPING isn't defined. Previous implementations of this function just didn't process stats at all when they weren't defined, however we need engine stats to still filter through which are on the same system which _shouldn't_ need to be threaded anyway as they are lightweight

#branch UE4

diff --git a/Engine/Source/Runtime/Core/Private/Stats/StatsCommand.cpp b/Engine/Source/Runtime/Core/Private/Stats/StatsCommand.cpp
--- a/Engine/Source/Runtime/Core/Private/Stats/StatsCommand.cpp
+++ b/Engine/Source/Runtime/Core/Private/Stats/StatsCommand.cpp
@@ -1456,8 +1456,9 @@ bool DirectStatsCommand(const TCHAR* Cmd, bool bBlockForCompletion /*= false*/,
 		check(IsInGameThread());
 		if( !bIsEmpty )
 		{
-			ENamedThreads::Type ThreadType = ENamedThreads::GameThread;
+			const FString FullCmd = FString(Cmd) + AddArgs;
 #if STATS
+			ENamedThreads::Type ThreadType = ENamedThreads::GameThread;
 			if (FPlatformProcess::SupportsMultithreading())
 			{
 				ThreadType = ENamedThreads::StatsThread;
@@ -1466,13 +1467,13 @@ bool DirectStatsCommand(const TCHAR* Cmd, bool bBlockForCompletion /*= false*/,
 			// make sure these are initialized on the game thread
 			FHUDGroupGameThreadRenderer::Get();
 			FStatGroupGameThreadNotifier::Get();
-#endif
+
 			DECLARE_CYCLE_STAT(TEXT("FSimpleDelegateGraphTask.StatCmd"),
 				STAT_FSimpleDelegateGraphTask_StatCmd,
 				STATGROUP_TaskGraphTasks);
 
 			FGraphEventRef CompleteHandle = FSimpleDelegateGraphTask::CreateAndDispatchWhenReady(
-				FSimpleDelegateGraphTask::FDelegate::CreateStatic(&StatCmd, FString(Cmd) + AddArgs),
+				FSimpleDelegateGraphTask::FDelegate::CreateStatic(&StatCmd, FullCmd),
 				GET_STATID(STAT_FSimpleDelegateGraphTask_StatCmd), NULL, ThreadType
 			);
 			if (bBlockForCompletion)
@@ -1480,6 +1481,10 @@ bool DirectStatsCommand(const TCHAR* Cmd, bool bBlockForCompletion /*= false*/,
 				FTaskGraphInterface::Get().WaitUntilTaskCompletes(CompleteHandle);
 				GLog->FlushThreadedLogs();
 			}
+#else
+			// If stats aren't enabled, broadcast so engine stats can still be triggered
+			StatCmd(FullCmd);
+#endif
 		}
 	}
 	return bResult;
