PR #1025: Remove unnecessary file exists check (Contributed by [REDACTED_PERSON])

diff --git a/Engine/Source/Runtime/Core/Private/HAL/FileManagerGeneric.cpp b/Engine/Source/Runtime/Core/Private/HAL/FileManagerGeneric.cpp
--- a/Engine/Source/Runtime/Core/Private/HAL/FileManagerGeneric.cpp
+++ b/Engine/Source/Runtime/Core/Private/HAL/FileManagerGeneric.cpp
@@ -414,12 +414,6 @@ double FFileManagerGeneric::GetFileAgeSeconds( const TCHAR* Filename )
 
 FDateTime FFileManagerGeneric::GetTimeStamp(const TCHAR* Filename)
 {
-	// if the file doesn't exist, use a dummy DateTime
-	if (!GetLowLevel().FileExists(Filename))
-	{
-		return FDateTime::MinValue();
-	}
-
 	// ask low level for timestamp
 	return GetLowLevel().GetTimeStamp(Filename);
 }
@@ -438,12 +432,6 @@ bool FFileManagerGeneric::SetTimeStamp(const TCHAR* Filename, FDateTime DateTime
 
 FDateTime FFileManagerGeneric::GetAccessTimeStamp( const TCHAR* Filename )
 {
-	// if the file doesn't exist, use a dummy DateTime
-	if (!GetLowLevel().FileExists(Filename))
-	{
-		return FDateTime::MinValue();
-	}
-
 	// ask low level for timestamp
 	return GetLowLevel().GetAccessTimeStamp(Filename);
 }
