LogFile wrapper improvements:

- reduced verbosity to make it more usable for file open/close debugging
- added open file handle tracking to be able to dump all open file handles via 'filelogdump' console command

diff --git a/Engine/Source/Runtime/Core/Private/HAL/IPlatformFileLogWrapper.cpp b/Engine/Source/Runtime/Core/Private/HAL/IPlatformFileLogWrapper.cpp
--- a/Engine/Source/Runtime/Core/Private/HAL/IPlatformFileLogWrapper.cpp
+++ b/Engine/Source/Runtime/Core/Private/HAL/IPlatformFileLogWrapper.cpp
@@ -5,3 +5,83 @@
 
 bool bSuppressFileLog = false;
 DEFINE_LOG_CATEGORY(LogPlatformFile);
+
+
+#if !UE_BUILD_SHIPPING
+class FFileLogExec : private FSelfRegisteringExec
+{
+	FLoggedPlatformFile& PlatformFile;
+
+public:
+
+	FFileLogExec(FLoggedPlatformFile& InPlatformFile)
+		: PlatformFile(InPlatformFile)
+	{}
+
+	/** Console commands **/
+	virtual bool Exec(UWorld* InWorld, const TCHAR* Cmd, FOutputDevice& Ar) override
+	{
+		if (FParse::Command(&Cmd, TEXT("LogFileDump")))
+		{
+			PlatformFile.HandleDumpCommand(Cmd, Ar);
+			return true;
+		}
+		return false;
+	}
+};
+static TAutoPtr<FFileLogExec> GFileLogExec;
+
+#endif // !UE_BUILD_SHIPPING
+
+
+bool FLoggedPlatformFile::ShouldBeUsed(IPlatformFile* Inner, const TCHAR* CmdLine) const
+{
+	return FParse::Param(CmdLine, TEXT("FileLog"));
+}
+
+bool FLoggedPlatformFile::Initialize(IPlatformFile* Inner, const TCHAR* CommandLineParam)
+{
+	// Inner is required.
+	check(Inner != nullptr);
+	LowerLevel = Inner;
+
+#if !UE_BUILD_SHIPPING
+	GFileLogExec = new FFileLogExec(*this);
+#endif
+
+	return !!LowerLevel;
+}
+
+#if !UE_BUILD_SHIPPING
+void FLoggedPlatformFile::HandleDumpCommand(const TCHAR* Cmd, FOutputDevice& Ar)
+{
+	FScopeLock LogFileLock(&LogFileCritical);
+	bSuppressFileLog = true;
+	Ar.Logf(TEXT("Open file handles: %d"), OpenHandles.Num());
+	for (auto& HandlePair : OpenHandles)
+	{
+		Ar.Logf(TEXT("%s: %d"), *HandlePair.Key, HandlePair.Value);
+	}	
+	bSuppressFileLog = false;
+}
+#endif
+
+FLoggedFileHandle::FLoggedFileHandle(IFileHandle* InFileHandle, const TCHAR* InFilename, FLoggedPlatformFile& InOwner)
+: FileHandle(InFileHandle)
+, Filename(InFilename)
+#if !UE_BUILD_SHIPPING
+, PlatformFile(InOwner)
+#endif
+{
+#if !UE_BUILD_SHIPPING
+	PlatformFile.OnHandleOpen(Filename);
+#endif
+}
+
+FLoggedFileHandle::~FLoggedFileHandle()
+{
+#if !UE_BUILD_SHIPPING
+	PlatformFile.OnHandleClosed(Filename);
+#endif
+	FILE_LOG(LogPlatformFile, Log, TEXT("Close %s"), *Filename);
+}
\ No newline at end of file
diff --git a/Engine/Source/Runtime/Core/Public/HAL/IPlatformFileLogWrapper.h b/Engine/Source/Runtime/Core/Public/HAL/IPlatformFileLogWrapper.h
--- a/Engine/Source/Runtime/Core/Public/HAL/IPlatformFileLogWrapper.h
+++ b/Engine/Source/Runtime/Core/Public/HAL/IPlatformFileLogWrapper.h
@@ -18,78 +18,85 @@ extern bool bSuppressFileLog;
 		bSuppressFileLog = false; \
 	}
 
+class FLoggedPlatformFile;
 
 class CORE_API FLoggedFileHandle : public IFileHandle
 {
 	TAutoPtr<IFileHandle>	FileHandle;
 	FString					Filename;
+#if !UE_BUILD_SHIPPING
+	FLoggedPlatformFile& PlatformFile;
+#endif
 public:
 
-	FLoggedFileHandle(IFileHandle* InFileHandle, const TCHAR* InFilename)
-		: FileHandle(InFileHandle)
-		, Filename(InFilename)
-	{
-	}
+	FLoggedFileHandle(IFileHandle* InFileHandle, const TCHAR* InFilename, FLoggedPlatformFile& InOwner);
+	virtual ~FLoggedFileHandle();
 
 	virtual int64		Tell() override
 	{
-		FILE_LOG(LogPlatformFile, Log, TEXT("Tell %s"), *Filename);
+		FILE_LOG(LogPlatformFile, VeryVerbose, TEXT("Tell %s"), *Filename);
 		double StartTime = FPlatformTime::Seconds();
 		int64 Result = FileHandle->Tell();
 		float ThisTime = 1000.0f * float(FPlatformTime::Seconds() - StartTime);
-		FILE_LOG(LogPlatformFile, Log, TEXT("Tell return %lld [%fms]"), Result, ThisTime);
+		FILE_LOG(LogPlatformFile, VeryVerbose, TEXT("Tell return %lld [%fms]"), Result, ThisTime);
 		return Result;
 	}
 	virtual bool		Seek(int64 NewPosition) override
 	{
-		FILE_LOG(LogPlatformFile, Log, TEXT("Seek %s %lld"), *Filename, NewPosition);
+		FILE_LOG(LogPlatformFile, VeryVerbose, TEXT("Seek %s %lld"), *Filename, NewPosition);
 		double StartTime = FPlatformTime::Seconds();
 		bool Result = FileHandle->Seek(NewPosition);
 		float ThisTime = 1000.0f * float(FPlatformTime::Seconds() - StartTime);
-		FILE_LOG(LogPlatformFile, Log, TEXT("Seek return %d [%fms]"), int32(Result), ThisTime);
+		FILE_LOG(LogPlatformFile, VeryVerbose, TEXT("Seek return %d [%fms]"), int32(Result), ThisTime);
 		return Result;
 	}
 	virtual bool		SeekFromEnd(int64 NewPositionRelativeToEnd) override
 	{
-		FILE_LOG(LogPlatformFile, Log, TEXT("SeekFromEnd %s %lld"), *Filename, NewPositionRelativeToEnd);
+		FILE_LOG(LogPlatformFile, VeryVerbose, TEXT("SeekFromEnd %s %lld"), *Filename, NewPositionRelativeToEnd);
 		double StartTime = FPlatformTime::Seconds();
 		bool Result = FileHandle->SeekFromEnd(NewPositionRelativeToEnd);
 		float ThisTime = 1000.0f * float(FPlatformTime::Seconds() - StartTime);
-		FILE_LOG(LogPlatformFile, Log, TEXT("SeekFromEnd return %d [%fms]"), int32(Result), ThisTime);
+		FILE_LOG(LogPlatformFile, VeryVerbose, TEXT("SeekFromEnd return %d [%fms]"), int32(Result), ThisTime);
 		return Result;
 	}
 	virtual bool		Read(uint8* Destination, int64 BytesToRead) override
 	{
-		FILE_LOG(LogPlatformFile, Log, TEXT("Read %s %lld"), *Filename, BytesToRead);
+		FILE_LOG(LogPlatformFile, VeryVerbose, TEXT("Read %s %lld"), *Filename, BytesToRead);
 		double StartTime = FPlatformTime::Seconds();
 		bool Result = FileHandle->Read(Destination, BytesToRead);
 		float ThisTime = 1000.0f * float(FPlatformTime::Seconds() - StartTime);
-		FILE_LOG(LogPlatformFile, Log, TEXT("Read return %d [%fms]"), int32(Result), ThisTime);
+		FILE_LOG(LogPlatformFile, VeryVerbose, TEXT("Read return %d [%fms]"), int32(Result), ThisTime);
 		return Result;
 	}
 	virtual bool		Write(const uint8* Source, int64 BytesToWrite) override
 	{
-		FILE_LOG(LogPlatformFile, Log, TEXT("Write %s %lld"), *Filename, BytesToWrite);
+		FILE_LOG(LogPlatformFile, VeryVerbose, TEXT("Write %s %lld"), *Filename, BytesToWrite);
 		double StartTime = FPlatformTime::Seconds();
 		bool Result = FileHandle->Write(Source, BytesToWrite);
 		float ThisTime = 1000.0f * float(FPlatformTime::Seconds() - StartTime);
-		FILE_LOG(LogPlatformFile, Log, TEXT("Write return %d [%fms]"), int32(Result), ThisTime);
+		FILE_LOG(LogPlatformFile, VeryVerbose, TEXT("Write return %d [%fms]"), int32(Result), ThisTime);
 		return Result;
 	}
 	virtual int64		Size() override
 	{
-		FILE_LOG(LogPlatformFile, Log, TEXT("Size %s"), *Filename);
+		FILE_LOG(LogPlatformFile, Verbose, TEXT("Size %s"), *Filename);
 		double StartTime = FPlatformTime::Seconds();
 		int64 Result = FileHandle->Size();
 		float ThisTime = 1000.0f * float(FPlatformTime::Seconds() - StartTime);
-		FILE_LOG(LogPlatformFile, Log, TEXT("Size return %lld [%fms]"), Result, ThisTime);
+		FILE_LOG(LogPlatformFile, Verbose, TEXT("Size return %lld [%fms]"), Result, ThisTime);
 		return Result;
 	}
 };
 
 class CORE_API FLoggedPlatformFile : public IPlatformFile
 {
-	IPlatformFile*		LowerLevel;
+	IPlatformFile* LowerLevel;
+
+#if !UE_BUILD_SHIPPING
+	FCriticalSection LogFileCritical;
+	TMap<FString, int32> OpenHandles;
+#endif
+
 public:
 	static const TCHAR* GetTypeName()
 	{
@@ -101,17 +108,9 @@ public:
 	{
 	}
 
-	virtual bool ShouldBeUsed(IPlatformFile* Inner, const TCHAR* CmdLine) const override
-	{
-		return FParse::Param(CmdLine, TEXT("FileLog"));
-	}
-	virtual bool Initialize(IPlatformFile* Inner, const TCHAR* CommandLineParam) override
-	{
-		// Inner is required.
-		check(Inner != nullptr);
-		LowerLevel = Inner;
-		return !!LowerLevel;
-	}
+	virtual bool ShouldBeUsed(IPlatformFile* Inner, const TCHAR* CmdLine) const override;
+
+	virtual bool Initialize(IPlatformFile* Inner, const TCHAR* CommandLineParam) override;
 
 	IPlatformFile* GetLowerLevel() override
 	{
@@ -219,7 +218,7 @@ public:
 		IFileHandle* Result = LowerLevel->OpenRead(Filename, bAllowWrite);
 		float ThisTime = 1000.0f * float(FPlatformTime::Seconds() - StartTime);
 		FILE_LOG(LogPlatformFile, Log, TEXT("OpenRead return %llx [%fms]"), uint64(Result), ThisTime);
-		return Result ? (new FLoggedFileHandle(Result, Filename)) : Result;
+		return Result ? (new FLoggedFileHandle(Result, Filename, *this)) : Result;
 	}
 	virtual IFileHandle*	OpenWrite(const TCHAR* Filename, bool bAppend = false, bool bAllowRead = false) override
 	{
@@ -228,7 +227,7 @@ public:
 		IFileHandle* Result = LowerLevel->OpenWrite(Filename, bAppend, bAllowRead);
 		float ThisTime = 1000.0f * float(FPlatformTime::Seconds() - StartTime);
 		FILE_LOG(LogPlatformFile, Log, TEXT("OpenWrite return %llx [%fms]"), uint64(Result), ThisTime);
-		return Result ? (new FLoggedFileHandle(Result, Filename)) : Result;
+		return Result ? (new FLoggedFileHandle(Result, Filename, *this)) : Result;
 	}
 
 	virtual bool		DirectoryExists(const TCHAR* Directory) override
@@ -269,11 +268,11 @@ public:
 		}
 		virtual bool Visit(const TCHAR* FilenameOrDirectory, bool bIsDirectory)
 		{
-			FILE_LOG(LogPlatformFile, Log, TEXT("Visit %s %d"), FilenameOrDirectory, int32(bIsDirectory));
+			FILE_LOG(LogPlatformFile, Verbose, TEXT("Visit %s %d"), FilenameOrDirectory, int32(bIsDirectory));
 			double StartTime = FPlatformTime::Seconds();
 			bool Result = Visitor.Visit(FilenameOrDirectory, bIsDirectory);
 			float ThisTime = 1000.0f * float(FPlatformTime::Seconds() - StartTime);
-			FILE_LOG(LogPlatformFile, Log, TEXT("Visit return %d [%fms]"), int32(Result), ThisTime);
+			FILE_LOG(LogPlatformFile, Verbose, TEXT("Visit return %d [%fms]"), int32(Result), ThisTime);
 			return Result;
 		}
 	};
@@ -316,4 +315,23 @@ public:
 		FILE_LOG(LogPlatformFile, Log, TEXT("CopyFile return %d [%fms]"), int32(Result), ThisTime);
 		return Result;
 	}
+
+#if !UE_BUILD_SHIPPING
+	void OnHandleOpen(const FString& Filename)
+	{
+		FScopeLock LogFileLock(&LogFileCritical);
+		int32& NumOpenHandles = OpenHandles.FindOrAdd(Filename);
+		NumOpenHandles++;
+	}
+	void OnHandleClosed(const FString& Filename)
+	{
+		FScopeLock LogFileLock(&LogFileCritical);
+		int32& NumOpenHandles = OpenHandles.FindChecked(Filename);
+		if (--NumOpenHandles == 0)
+		{
+			OpenHandles.Remove(Filename);
+		}
+	}
+	void HandleDumpCommand(const TCHAR* Cmd, FOutputDevice& Ar);
+#endif
 };
