Editor: Fixing the build by preventing UHT from parsing a header it cannot currently understand (removed UENUM from ETextFilterComparisonOperation)

diff --git a/Engine/Source/Runtime/Core/Public/Misc/TextFilter.h b/Engine/Source/Runtime/Core/Public/Misc/TextFilter.h
--- a/Engine/Source/Runtime/Core/Public/Misc/TextFilter.h
+++ b/Engine/Source/Runtime/Core/Public/Misc/TextFilter.h
@@ -3,7 +3,6 @@
 
 #include "IFilter.h"
 
-UENUM()
 enum class ETextFilterComparisonOperation : uint8
 {
 	Equal,
