Giving the source ini hierarchy an enum for lookups.

Removed a no longer relevant check for "Default" in the filename checK for a default ini file.

diff --git a/Engine/Source/Runtime/Core/Private/Misc/ConfigCacheIni.cpp b/Engine/Source/Runtime/Core/Private/Misc/ConfigCacheIni.cpp
--- a/Engine/Source/Runtime/Core/Private/Misc/ConfigCacheIni.cpp
+++ b/Engine/Source/Runtime/Core/Private/Misc/ConfigCacheIni.cpp
@@ -19,7 +19,7 @@
 DEFINE_LOG_CATEGORY(LogConfig);
 namespace 
 {
-	FString GenerateHierarchyCacheKey(const TArray<FIniFilename>& IniHierarchy, const FString& IniPath, const FString& BaseIniName)
+	FString GenerateHierarchyCacheKey(const FConfigFileHierarchy& IniHierarchy, const FString& IniPath, const FString& BaseIniName)
 	{
 #if !INI_CACHE
 		return TEXT("");
@@ -31,13 +31,13 @@ namespace
 		KeyLen += BaseIniName.Len();
 		for (const auto& Ini : IniHierarchy)
 		{
-			KeyLen += Ini.Filename.Len();
+			KeyLen += Ini.Value.Filename.Len();
 		}
 		HierKey.Reserve(KeyLen);
 		HierKey += BaseIniName;
 		for (const auto& Ini : IniHierarchy)
 		{
-			HierKey += Ini.Filename;
+			HierKey += Ini.Value.Filename;
 		}
 		HierKey += IniPath;
 
@@ -668,7 +668,7 @@ static void OverrideFromCommandline(FConfigFile* File, const FString& Filename)
  * @param ConfigFile - This is the FConfigFile which will have the contents of the .ini loaded into and Combined()
  *
  **/
-static bool LoadIniFileHierarchy(const TArray<FIniFilename>& HierarchyToLoad, FConfigFile& ConfigFile, const bool bUseCache)
+static bool LoadIniFileHierarchy(const FConfigFileHierarchy& HierarchyToLoad, FConfigFile& ConfigFile, const bool bUseCache)
 {
 	// This shouldn't be getting called if seekfree is enabled on console.
 	check(!GUseSeekFreeLoading || !FPlatformProperties::RequiresCookedData());
@@ -683,9 +683,9 @@ static bool LoadIniFileHierarchy(const TArray<FIniFilename>& HierarchyToLoad, FC
 	{
 		// If no inis exist or only engine (Base*.ini) inis exist, don't load anything
 		int32 NumExistingOptionalInis = 0;
-		for( int32 IniIndex = 0; IniIndex < HierarchyToLoad.Num(); IniIndex++ )
+		for( const auto& HierarchyIt : HierarchyToLoad )
 		{
-			const FIniFilename& IniToLoad = HierarchyToLoad[ IniIndex ];
+			const FIniFilename& IniToLoad = HierarchyIt.Value;
 			if (IniToLoad.bRequired == false &&
 				 (!IsUsingLocalIniFile(*IniToLoad.Filename, NULL) || IFileManager::Get().FileSize(*IniToLoad.Filename) >= 0))
 			{
@@ -699,16 +699,16 @@ static bool LoadIniFileHierarchy(const TArray<FIniFilename>& HierarchyToLoad, FC
 		}
 	}
 
-	int32 FirstCacheIndex = 0;
+	EConfigFileHierarchy FirstCacheIndex = EConfigFileHierarchy::AbsoluteBase;
 #if INI_CACHE
-	if ( bUseCache )
+	if (bUseCache && HierarchyCache.Num() > 0)
 	{
 		// Find the last value in the hierarchy that is cached. We can start the load from there
-		for( int32 IniIndex = 0; IniIndex < HierarchyToLoad.Num(); IniIndex++ )
+		for (const auto& HierarchyIt : HierarchyToLoad)
 		{
-			if (HierarchyCache.Find(HierarchyToLoad[IniIndex].CacheKey)) 
+			if (HierarchyCache.Find(HierarchyIt.Value.CacheKey)) 
 			{
-				FirstCacheIndex = IniIndex;
+				FirstCacheIndex = HierarchyIt.Key;
 			}
 		}
 	}
@@ -717,66 +717,67 @@ static bool LoadIniFileHierarchy(const TArray<FIniFilename>& HierarchyToLoad, FC
 	TArray<FDateTime> TimestampsOfInis;
 
 	// Traverse ini list back to front, merging along the way.
-	for( int32 IniIndex = FirstCacheIndex; IniIndex < HierarchyToLoad.Num(); IniIndex++ )
+	for (const auto& HierarchyIt : HierarchyToLoad)
 	{
-		const FIniFilename& IniToLoad = HierarchyToLoad[ IniIndex ];
-		const FString& IniFileName = IniToLoad.Filename;
-		bool bDoProcess = true;
-#if INI_CACHE
-		bool bShouldCache = IniToLoad.CacheKey.Len() > 0;
-		bShouldCache &= bUseCache;
-		if ( bShouldCache ) // if we are forcing a load don't mess with the cache
+		if (FirstCacheIndex <= HierarchyIt.Key)
 		{
-			auto* CachedConfigFile = HierarchyCache.Find(IniToLoad.CacheKey);
-			if (CachedConfigFile) 
+			const FIniFilename& IniToLoad = HierarchyIt.Value;
+			const FString& IniFileName = IniToLoad.Filename;
+			bool bDoProcess = true;
+#if INI_CACHE
+			bool bShouldCache = IniToLoad.CacheKey.Len() > 0;
+			bShouldCache &= bUseCache;
+			if ( bShouldCache ) // if we are forcing a load don't mess with the cache
 			{
-				ConfigFile = *CachedConfigFile;
-				bDoProcess = false;
+				auto* CachedConfigFile = HierarchyCache.Find(IniToLoad.CacheKey);
+				if (CachedConfigFile) 
+				{
+					ConfigFile = *CachedConfigFile;
+					bDoProcess = false;
+				}
+				ConfigFile.CacheKey = IniToLoad.CacheKey;
+			}
+			else
+			{
+				ConfigFile.CacheKey = TEXT("");
 			}
-			ConfigFile.CacheKey = IniToLoad.CacheKey;
-		}
-		else
-		{
-			ConfigFile.CacheKey = TEXT("");
-		}
 #endif
-		if (bDoProcess)
-		{
-			// Spit out friendly error if there was a problem locating .inis (e.g. bad command line parameter or missing folder, ...).
-			if (IsUsingLocalIniFile(*IniFileName, NULL) && (IFileManager::Get().FileSize(*IniFileName) < 0))
+			if (bDoProcess)
 			{
-				if (IniToLoad.bRequired)
+				// Spit out friendly error if there was a problem locating .inis (e.g. bad command line parameter or missing folder, ...).
+				if (IsUsingLocalIniFile(*IniFileName, NULL) && (IFileManager::Get().FileSize(*IniFileName) < 0))
 				{
-					//UE_LOG(LogConfig, Error, TEXT("Couldn't locate '%s' which is required to run '%s'"), *IniToLoad.Filename, FApp::GetGameName() );
-					return false;
-				}
-				else
-				{
-#if INI_CACHE
-					// missing file just add the current config file to the cache
-					if ( bShouldCache )
+					if (IniToLoad.bRequired)
 					{
-							HierarchyCache.Add(IniToLoad.CacheKey, ConfigFile);
+						//UE_LOG(LogConfig, Error, TEXT("Couldn't locate '%s' which is required to run '%s'"), *IniToLoad.Filename, FApp::GetGameName() );
+						return false;
 					}
+					else
+					{
+#if INI_CACHE
+						// missing file just add the current config file to the cache
+						if ( bShouldCache )
+						{
+								HierarchyCache.Add(IniToLoad.CacheKey, ConfigFile);
+						}
 #endif
-					continue;
+						continue;
+					}
 				}
-			}
 
-			bool bDoEmptyConfig = false;
-			bool bDoCombine = (IniIndex != 0);
-			bool bDoWrite = false;
-			//UE_LOG(LogConfig, Log,  TEXT( "Combining configFile: %s" ), *IniList(IniIndex) );
-			ProcessIniContents(*HierarchyToLoad.Last().Filename, *IniFileName, &ConfigFile, bDoEmptyConfig, bDoCombine, bDoWrite);
+				bool bDoEmptyConfig = false;
+				bool bDoCombine = (HierarchyIt.Key != EConfigFileHierarchy::AbsoluteBase);
+				bool bDoWrite = false;
+				//UE_LOG(LogConfig, Log,  TEXT( "Combining configFile: %s" ), *IniList(IniIndex) );
+				ProcessIniContents(*(HierarchyIt.Value.Filename), *IniFileName, &ConfigFile, bDoEmptyConfig, bDoCombine, bDoWrite);
 #if INI_CACHE
-			if ( bShouldCache )
-			{
-				HierarchyCache.Add(IniToLoad.CacheKey, ConfigFile);
-			}
+				if ( bShouldCache )
+				{
+					HierarchyCache.Add(IniToLoad.CacheKey, ConfigFile);
+				}
 #endif
+			}
 		}
-
-
 	}
 
 	// Set this configs files source ini hierarchy to show where it was loaded from.
@@ -947,8 +948,7 @@ bool FConfigFile::Write( const FString& Filename, bool bDoRemoteWrite/* = true*/
 
 				// If we are writing to a default config file and this property is an array, we need to be careful to remove those from higher up the hierarchy
 				const bool bIsADefaultIniWrite = !Filename.Contains( FPaths::GameSavedDir() ) 
-					&& !Filename.Contains( FPaths::EngineSavedDir() ) 
-					&& FPaths::GetBaseFilename(Filename).StartsWith( TEXT( "Default" ) );
+					&& !Filename.Contains( FPaths::EngineSavedDir() );
 
 				// Check if the property matches the source configs. We do not wanna write it out if so.
 				if ((bIsADefaultIniWrite || bDifferentNumberOfElements || !DoesConfigPropertyValueMatch(SourceConfigFile, SectionName, PropertyName, PropertyValue)) && !bOptionIsFromCommandline)
@@ -1289,12 +1289,21 @@ void FConfigFile::ProcessPropertyAndWriteForDefaults( const TArray< FString >& I
 		{
 			// Build a config file out of this default configs hierarchy.
 			FConfigCacheIni Hierarchy(EConfigCacheType::Temporary);
-			const FString& LastFileInHierarchy = SourceIniHierarchy.Last().Filename;
 
+			EConfigFileHierarchy HighestFileIndex = EConfigFileHierarchy::AbsoluteBase;
+			TArray<EConfigFileHierarchy> ExistingEntries;
+			SourceIniHierarchy.GetKeys(ExistingEntries);
+			for (const EConfigFileHierarchy& NextEntry : ExistingEntries)
+			{
+				HighestFileIndex = NextEntry > HighestFileIndex ? NextEntry : HighestFileIndex;
+			}
+
+			const FString& LastFileInHierarchy = SourceIniHierarchy.FindChecked(HighestFileIndex).Filename;
 			FConfigFile& DefaultConfigFile = Hierarchy.Add(LastFileInHierarchy, FConfigFile());
-			for (TArray<FIniFilename>::TIterator HierarchyFileIt(SourceIniHierarchy); HierarchyFileIt; ++HierarchyFileIt)
+
+			for (const auto& HierarchyFileIt : SourceIniHierarchy)
 			{
-				DefaultConfigFile.Combine(HierarchyFileIt->Filename);
+				DefaultConfigFile.Combine(HierarchyFileIt.Value.Filename);
 			}
 
 			// Remove any array elements from the default configs hierearchy, we will add these in below
@@ -2576,7 +2585,7 @@ static void LoadAnIniFile(const FString& FilenameToLoad, FConfigFile& ConfigFile
  * @param DestIniFilename			The name of the destination .ini file - it's loaded and checked against source (e.g. Engine.ini)
  * @param SourceIniFilename			The name of the source .ini file (e.g. DefaultEngine.ini)
  */
-static bool GenerateDestIniFile(FConfigFile& DestConfigFile, const FString& DestIniFilename, const TArray<FIniFilename>& SourceIniHierarchy, bool bAllowGeneratedINIs, const bool bUseHierarchyCache)
+static bool GenerateDestIniFile(FConfigFile& DestConfigFile, const FString& DestIniFilename, const FConfigFileHierarchy& SourceIniHierarchy, bool bAllowGeneratedINIs, const bool bUseHierarchyCache)
 {
 	bool bResult = LoadIniFileHierarchy(SourceIniHierarchy, *DestConfigFile.SourceConfigFile, bUseHierarchyCache);
 	if ( bResult == false )
@@ -2713,11 +2722,10 @@ static FString GetSourceIniFilename(const TCHAR* ConfigDir, const TCHAR* Prefix,
  * @param InBaseIniName Ini name.
  * @param InPlatformName Platform name, NULL means to use the current platform
  * @param InGameName Game name, NULL means to use the current
- * @param OutHierarchy An array which is to receive the generated hierachy of ini filenames.
+ * @param OutHierarchy An array which is to receive the generated hierarchy of ini filenames.
  */
-static void GetSourceIniHierarchyFilenames(const TCHAR* InBaseIniName, const TCHAR* InPlatformName, const TCHAR* InGameName, const TCHAR* EngineConfigDir, const TCHAR* SourceConfigDir, TArray<FIniFilename>& OutHierarchy, bool bRequireDefaultIni)
+static void GetSourceIniHierarchyFilenames(const TCHAR* InBaseIniName, const TCHAR* InPlatformName, const TCHAR* InGameName, const TCHAR* EngineConfigDir, const TCHAR* SourceConfigDir, FConfigFileHierarchy& OutHierarchy, bool bRequireDefaultIni)
 {
-	
 	/**************************************************
 	 **** CRITICAL NOTES
 	 **** If you change this function, you need to also change EnumerateCrossPlatformIniFileNames() in EngineConfiguration.cs!!!
@@ -2727,11 +2735,11 @@ static void GetSourceIniHierarchyFilenames(const TCHAR* InBaseIniName, const TCH
 
 	// [[[[ ENGINE DEFAULTS ]]]]
 	// Engine/Config/Base.ini (included in every ini type, required)
-	OutHierarchy.Add( FIniFilename(FString::Printf(TEXT("%sBase.ini"), EngineConfigDir), true) );
+	OutHierarchy.Add(EConfigFileHierarchy::AbsoluteBase, FIniFilename(FString::Printf(TEXT("%sBase.ini"), EngineConfigDir), true));
 	// Engine/Config/Base* ini
-	OutHierarchy.Add( FIniFilename(FString::Printf(TEXT("%sBase%s.ini"), EngineConfigDir, InBaseIniName), false) );
+	OutHierarchy.Add(EConfigFileHierarchy::EngineDirBase, FIniFilename(FString::Printf(TEXT("%sBase%s.ini"), EngineConfigDir, InBaseIniName), false));
 	// Engine/Config/NotForLicensees/Base* ini
-	OutHierarchy.Add( FIniFilename(FString::Printf(TEXT("%sNotForLicensees/Base%s.ini"), EngineConfigDir, InBaseIniName), false) );
+	OutHierarchy.Add(EConfigFileHierarchy::EngineDirBase_NotForLicensees, FIniFilename(FString::Printf(TEXT("%sNotForLicensees/Base%s.ini"), EngineConfigDir, InBaseIniName), false));
 	// Engine/Config/NoRedist/Base* ini
 	/////
 	// NOTE: 4.7
@@ -2747,46 +2755,48 @@ static void GetSourceIniHierarchyFilenames(const TCHAR* InBaseIniName, const TCH
 
 	// [[[[ PROJECT SETTINGS ]]]]
 	// Game/Config/Default* ini
-	OutHierarchy.Add(FIniFilename(GetSourceIniFilename(SourceConfigDir, TEXT("Default"), InBaseIniName), bRequireDefaultIni));
+	OutHierarchy.Add(EConfigFileHierarchy::GameDirDefault, FIniFilename(GetSourceIniFilename(SourceConfigDir, TEXT("Default"), InBaseIniName), bRequireDefaultIni));
 	// Game/Config/DedicatedServer* ini (knowingly NOT in EngineConfiguration.cs because this is a runtime only check)
 	if (IsRunningDedicatedServer())
 	{
-		OutHierarchy.Add(FIniFilename(FString::Printf(TEXT("%s/DedicatedServer%s.ini"), SourceConfigDir, InBaseIniName), false));
+		OutHierarchy.Add(EConfigFileHierarchy::GameDirDedicatedServer, FIniFilename(FString::Printf(TEXT("%s/DedicatedServer%s.ini"), SourceConfigDir, InBaseIniName), false));
 	}
 	// Game/Config/NotForLicensees/Default* ini
-	OutHierarchy.Add(FIniFilename(FString::Printf(TEXT("%sNotForLicensees/Default%s.ini"), SourceConfigDir, InBaseIniName), false));
+	OutHierarchy.Add(EConfigFileHierarchy::GameDirDefault_NotForLicensees, FIniFilename(FString::Printf(TEXT("%sNotForLicensees/Default%s.ini"), SourceConfigDir, InBaseIniName), false));
 	// Game/Config/NoRedist/Default* ini (Checkpointed here before the platform ini divergence)
 	FString HierarchyCheckpointPath = FString::Printf(TEXT("%sNoRedist/Default%s.ini"), SourceConfigDir, InBaseIniName);
-	OutHierarchy.Add(FIniFilename(HierarchyCheckpointPath, false, GenerateHierarchyCacheKey(OutHierarchy, HierarchyCheckpointPath, InBaseIniName)));
+	OutHierarchy.Add(EConfigFileHierarchy::GameDirDefault_NoRedist, FIniFilename(HierarchyCheckpointPath, false, GenerateHierarchyCacheKey(OutHierarchy, HierarchyCheckpointPath, InBaseIniName)));
 	
 	if (PlatformName.Len() > 0)
 	{
 		// [[[[ PLATFORM DEFAULTS AND PROJECT SETTINGS ]]]]
 		// Engine/Config/Platform/Platform* ini
-		OutHierarchy.Add( FIniFilename(FString::Printf(TEXT("%s%s/%s%s.ini"), EngineConfigDir, *PlatformName, *PlatformName, InBaseIniName), false) );
+		OutHierarchy.Add(EConfigFileHierarchy::EngineDir_Platform, FIniFilename(FString::Printf(TEXT("%s%s/%s%s.ini"), EngineConfigDir, *PlatformName, *PlatformName, InBaseIniName), false));
 		// Engine/Config/NotForLicensees/Platform/Platform* ini
-		OutHierarchy.Add( FIniFilename(FString::Printf(TEXT("%sNotForLicensees/%s/%s%s.ini"), EngineConfigDir, *PlatformName, *PlatformName, InBaseIniName), false) );
+		OutHierarchy.Add(EConfigFileHierarchy::EngineDir_Platform_NotForLicensees, FIniFilename(FString::Printf(TEXT("%sNotForLicensees/%s/%s%s.ini"), EngineConfigDir, *PlatformName, *PlatformName, InBaseIniName), false));
 		// Engine/Config/NoRedist/Platform/Platform* ini
-		OutHierarchy.Add( FIniFilename(FString::Printf(TEXT("%sNoRedist/%s/%s%s.ini"), EngineConfigDir, *PlatformName, *PlatformName, InBaseIniName), false) );
+		OutHierarchy.Add(EConfigFileHierarchy::EngineDir_Platform_NoRedist, FIniFilename(FString::Printf(TEXT("%sNoRedist/%s/%s%s.ini"), EngineConfigDir, *PlatformName, *PlatformName, InBaseIniName), false));
 		// Game/Config/Platform/Platform* ini
-		OutHierarchy.Add( FIniFilename(FString::Printf(TEXT("%s%s/%s%s.ini"), SourceConfigDir, *PlatformName, *PlatformName, InBaseIniName), false) );
+		OutHierarchy.Add(EConfigFileHierarchy::GameDir_Platform, FIniFilename(FString::Printf(TEXT("%s%s/%s%s.ini"), SourceConfigDir, *PlatformName, *PlatformName, InBaseIniName), false));
 		// Game/Config/NotForLicensee/Platform/Platform* ini
-		OutHierarchy.Add( FIniFilename(FString::Printf(TEXT("%sNotForLicensees/%s/%s%s.ini"), SourceConfigDir, *PlatformName, *PlatformName, InBaseIniName), false) );
+		OutHierarchy.Add(EConfigFileHierarchy::GameDir_Platform_NotForLicensees, FIniFilename(FString::Printf(TEXT("%sNotForLicensees/%s/%s%s.ini"), SourceConfigDir, *PlatformName, *PlatformName, InBaseIniName), false));
 		// Game/Config/NoRedist/Platform/Platform* ini
-		OutHierarchy.Add( FIniFilename(FString::Printf(TEXT("%sNoRedist/%s/%s%s.ini"), SourceConfigDir, *PlatformName, *PlatformName, InBaseIniName), false) );
+		OutHierarchy.Add(EConfigFileHierarchy::GameDir_Platform_NoRedist, FIniFilename(FString::Printf(TEXT("%sNoRedist/%s/%s%s.ini"), SourceConfigDir, *PlatformName, *PlatformName, InBaseIniName), false));
 	}
 
 	// [[[[ GLOBAL USER OVERRIDES ]]]]
 	// <AppData>/UE4/EngineConfig/User* ini
-	OutHierarchy.Add(FIniFilename(FPaths::Combine(FPlatformProcess::UserSettingsDir(), *FString::Printf(TEXT("Unreal Engine/Engine/Config/User%s.ini"), InBaseIniName)), false));
+	OutHierarchy.Add(EConfigFileHierarchy::UserSettingsDir_EngineDir_User, FIniFilename(FPaths::Combine(FPlatformProcess::UserSettingsDir(), *FString::Printf(TEXT("Unreal Engine/Engine/Config/User%s.ini"), InBaseIniName)), false));
 	// <Documents>/UE4/EngineConfig/User* ini
-	OutHierarchy.Add(FIniFilename(FPaths::Combine(FPlatformProcess::UserDir(), *FString::Printf(TEXT("Unreal Engine/Engine/Config/User%s.ini"), InBaseIniName)), false));
+	OutHierarchy.Add(EConfigFileHierarchy::UserDir_User, FIniFilename(FPaths::Combine(FPlatformProcess::UserDir(), *FString::Printf(TEXT("Unreal Engine/Engine/Config/User%s.ini"), InBaseIniName)), false));
 
 	// [[[[ PROJECT USER OVERRIDES ]]]]
 	// Game/Config/User* ini (Checkpointed here at the end)
 	HierarchyCheckpointPath = FString::Printf(TEXT("%sUser%s.ini"), SourceConfigDir, InBaseIniName);
-	OutHierarchy.Add(FIniFilename(HierarchyCheckpointPath, false, GenerateHierarchyCacheKey(OutHierarchy, HierarchyCheckpointPath, InBaseIniName)));
-	
+	OutHierarchy.Add(EConfigFileHierarchy::GameDir_User, FIniFilename(HierarchyCheckpointPath, false, GenerateHierarchyCacheKey(OutHierarchy, HierarchyCheckpointPath, InBaseIniName)));
+
+	// Ensure the hierarchy is in the correct order.
+	OutHierarchy.KeySort([](const EConfigFileHierarchy& A, const EConfigFileHierarchy& B) { return (A < B); });
 }
 
 /**
@@ -3114,11 +3124,15 @@ void FConfigFile::UpdateSections(const TCHAR* DiskFilename, const TCHAR* IniRoot
 		GetSourceIniHierarchyFilenames(IniRootName, NULL, NULL, *FPaths::EngineConfigDir(), *FPaths::SourceConfigDir(), SourceIniHierarchy, false);
 
 		// now chop off this file and any after it
-		for (int32 FileIndex = 0; FileIndex < SourceIniHierarchy.Num(); FileIndex++)
+		for (const auto& HierarchyFileIt : SourceIniHierarchy)
 		{
-			if (SourceIniHierarchy[FileIndex].Filename == DiskFilename)
+			if (HierarchyFileIt.Value.Filename == DiskFilename)
 			{
-				SourceIniHierarchy.RemoveAt(FileIndex, SourceIniHierarchy.Num() - FileIndex);
+				for (uint32 FileIdx = (uint32)HierarchyFileIt.Key; FileIdx < (uint32)EConfigFileHierarchy::NumHierarchyFiles; FileIdx++)
+				{
+					SourceIniHierarchy.Remove((EConfigFileHierarchy)FileIdx);
+				}
+
 				break;
 			}
 		}
diff --git a/Engine/Source/Runtime/Core/Public/Misc/ConfigCacheIni.h b/Engine/Source/Runtime/Core/Public/Misc/ConfigCacheIni.h
--- a/Engine/Source/Runtime/Core/Public/Misc/ConfigCacheIni.h
+++ b/Engine/Source/Runtime/Core/Public/Misc/ConfigCacheIni.h
@@ -50,6 +50,53 @@ struct FConfigCommandlineOverride
 };
 #endif // !UE_BUILD_SHIPPING
 
+// Possible entries in a config hierarchy
+enum class EConfigFileHierarchy : uint8
+{
+	// Engine/Config/Base.ini
+	AbsoluteBase = 0,
+
+	// Engine/Config/*.ini
+	EngineDirBase,
+	// Engine/Config/NotForLicensees/*.ini
+	EngineDirBase_NotForLicensees,
+	// Engine/Config/NoRedist/*.ini -Not supported at this time.
+	EngineDirBase_NoRedist,
+
+	// Game/Config/*.ini
+	GameDirDefault,
+	// Game/Config/DedicatedServer*.ini
+	GameDirDedicatedServer,
+	// Game/Config/NotForLicensees/*.ini
+	GameDirDefault_NotForLicensees,
+	// Game/Config/NoRedist*.ini
+	GameDirDefault_NoRedist,
+
+	// Engine/Config/PlatformName/PlatformName*.ini
+	EngineDir_Platform,
+	// Engine/Config/NotForLicensees/PlatformName/PlatformName*.ini
+	EngineDir_Platform_NotForLicensees,
+	// Engine/Config/NoRedist/PlatformName/PlatformName*.ini
+	EngineDir_Platform_NoRedist,
+
+	// Game/Config/PlatformName/PlatformName*.ini
+	GameDir_Platform,
+	// Game/Config/NotForLicensees/PlatformName/PlatformName*.ini
+	GameDir_Platform_NotForLicensees,
+	// Game/Config/NoRedist/PlatformName/PlatformName*.ini
+	GameDir_Platform_NoRedist,
+
+	// <UserSettingsDir|AppData>/Unreal Engine/Engine/Config/User*.ini
+	UserSettingsDir_EngineDir_User,
+	// <UserDir|Documents>/Unreal Engine/Engine/Config/User*.ini
+	UserDir_User,
+	// Game/Config/User*.ini
+	GameDir_User,
+
+	// Number of config files in hierarchy.
+	NumHierarchyFiles,
+};
+typedef TMap<EConfigFileHierarchy, FIniFilename> FConfigFileHierarchy;
 
 // One config file.
 
@@ -62,7 +109,7 @@ public:
 	FName Name;
 
 	// The collection of source files which were used to generate this file.
-	TArray< FIniFilename > SourceIniHierarchy;
+	FConfigFileHierarchy SourceIniHierarchy;
 
 	/** The untainted config file which contains the coalesced base/default options. I.e. No Saved/ options*/
 	FConfigFile* SourceConfigFile;
