Stats - Fixed an issue with encoding/decoding group category

diff --git a/Engine/Source/Runtime/Core/Private/Stats/Stats2.cpp b/Engine/Source/Runtime/Core/Private/Stats/Stats2.cpp
--- a/Engine/Source/Runtime/Core/Private/Stats/Stats2.cpp
+++ b/Engine/Source/Runtime/Core/Private/Stats/Stats2.cpp
@@ -566,9 +566,9 @@ FName FStatNameAndInfo::ToLongName(FName InStatName, char const* InGroup, char c
 	}
 	if (InCategory)
 	{
-		LongName += TEXT("////");
+		LongName += TEXT("####");
 		LongName += InCategory;
-		LongName += TEXT("////");
+		LongName += TEXT("####");
 	}
 	return FName(*LongName);
 }
@@ -588,10 +588,15 @@ FName FStatNameAndInfo::GetShortNameFrom(FName InLongName)
 		}
 		Input = Input.RightChop(IndexEnd + 2);
 	}
-	const int32 IndexEnd = Input.Find(TEXT("///"), ESearchCase::CaseSensitive);
-	if (IndexEnd != INDEX_NONE)
+	const int32 DescIndexEnd = Input.Find(TEXT("///"), ESearchCase::CaseSensitive);
+	if (DescIndexEnd != INDEX_NONE)
 	{
-		Input = Input.Left(IndexEnd);
+		Input = Input.Left(DescIndexEnd);
+	}
+	const int32 CategoryIndexEnd = Input.Find( TEXT( "####" ), ESearchCase::CaseSensitive );
+	if( DescIndexEnd == INDEX_NONE && CategoryIndexEnd != INDEX_NONE )
+	{
+		Input = Input.Left(CategoryIndexEnd);
 	}
 	return FName(*Input);
 }
@@ -634,11 +639,11 @@ FName FStatNameAndInfo::GetGroupCategoryFrom(FName InLongName)
 {
 	FString Input(InLongName.ToString());
 
-	const int32 IndexStart = Input.Find(TEXT("///////"), ESearchCase::CaseSensitive, ESearchDir::FromEnd);
+	const int32 IndexStart = Input.Find(TEXT("####"), ESearchCase::CaseSensitive, ESearchDir::FromEnd);
 	if (IndexStart != INDEX_NONE)
 	{
-		Input = Input.RightChop(IndexStart + 7);
-		const int32 IndexEnd = Input.Find(TEXT("////"), ESearchCase::CaseSensitive);
+		Input = Input.RightChop(IndexStart + 4);
+		const int32 IndexEnd = Input.Find(TEXT("####"), ESearchCase::CaseSensitive);
 		if (IndexEnd != INDEX_NONE)
 		{
 			return FName(*Input.Left(IndexEnd));
