Add a PRAGMA_DISABLE_SHADOW_WARNING to easily wrap third party code that needs shadow variables disabled
Move PRAGMA_POP out of deprecation section and in to general usability

diff --git a/Engine/Source/Runtime/Core/Public/Android/AndroidPlatformCompilerPreSetup.h b/Engine/Source/Runtime/Core/Public/Android/AndroidPlatformCompilerPreSetup.h
--- a/Engine/Source/Runtime/Core/Public/Android/AndroidPlatformCompilerPreSetup.h
+++ b/Engine/Source/Runtime/Core/Public/Android/AndroidPlatformCompilerPreSetup.h
@@ -40,15 +40,22 @@
 	#define PRAGMA_ENABLE_DEPRECATION_WARNINGS \
 		_Pragma ("clang diagnostic push") \
 		_Pragma ("clang diagnostic warning \"-Wdeprecated-declarations\"")
+#endif // DISABLE_DEPRECATION
+
+#ifndef PRAGMA_DISABLE_SHADOW_VARIABLE_WARNINGS
+	#define PRAGMA_DISABLE_SHADOW_VARIABLE_WARNINGS \
+		_Pragma ("clang diagnostic push") \
+		_Pragma ("clang diagnostic ignored \"-Wshadow\"")
+#endif // PRAGMA_DISABLE_SHADOW_VARIABLE_WARNINGS
 
+#ifndef PRAGMA_POP
 	#define PRAGMA_POP \
 		_Pragma("clang diagnostic pop")
-
-#endif // DISABLE_DEPRECATION
+#endif // PRAGMA_POP
 
 #ifndef EMIT_CUSTOM_WARNING_AT_LINE
-#define EMIT_CUSTOM_WARNING_AT_LINE(Line, Warning) \
-	_Pragma(PREPROCESSOR_TO_STRING(message(WARNING_LOCATION(Line) Warning)))
+	#define EMIT_CUSTOM_WARNING_AT_LINE(Line, Warning) \
+		_Pragma(PREPROCESSOR_TO_STRING(message(WARNING_LOCATION(Line) Warning)))
 #endif // EMIT_CUSTOM_WARNING_AT_LINE
 
 // Make certain warnings always be warnings, even despite -Werror.
diff --git a/Engine/Source/Runtime/Core/Public/HTML5/HTML5PlatformCompilerPreSetup.h b/Engine/Source/Runtime/Core/Public/HTML5/HTML5PlatformCompilerPreSetup.h
--- a/Engine/Source/Runtime/Core/Public/HTML5/HTML5PlatformCompilerPreSetup.h
+++ b/Engine/Source/Runtime/Core/Public/HTML5/HTML5PlatformCompilerPreSetup.h
@@ -77,11 +77,19 @@
 	#define PRAGMA_ENABLE_DEPRECATION_WARNINGS \
 		_Pragma ("clang diagnostic push") \
 		_Pragma ("clang diagnostic warning \"-Wdeprecated-declarations\"")
+#endif 
+#endif // DISABLE_DEPRECATION
+
+#ifndef PRAGMA_DISABLE_SHADOW_VARIABLE_WARNINGS
+	#define PRAGMA_DISABLE_SHADOW_VARIABLE_WARNINGS \
+		_Pragma ("clang diagnostic push") \
+		_Pragma ("clang diagnostic ignored \"-Wshadow\"")
+#endif // PRAGMA_DISABLE_SHADOW_VARIABLE_WARNINGS
 
+#ifndef PRAGMA_POP
 	#define PRAGMA_POP \
 		_Pragma("clang diagnostic pop")
-#endif 
-#endif // DISABLE_DEPRECATION
+#endif // PRAGMA_POP
 
 #ifndef EMIT_CUSTOM_WARNING_AT_LINE
 #define EMIT_CUSTOM_WARNING_AT_LINE(Line, Warning) \
diff --git a/Engine/Source/Runtime/Core/Public/IOS/IOSPlatformCompilerPreSetup.h b/Engine/Source/Runtime/Core/Public/IOS/IOSPlatformCompilerPreSetup.h
--- a/Engine/Source/Runtime/Core/Public/IOS/IOSPlatformCompilerPreSetup.h
+++ b/Engine/Source/Runtime/Core/Public/IOS/IOSPlatformCompilerPreSetup.h
@@ -76,9 +76,18 @@
 		_Pragma ("clang diagnostic push") \
 		_Pragma ("clang diagnostic warning \"-Wdeprecated-declarations\"")
 
+#endif // DISABLE_DEPRECATION
+
+#ifndef PRAGMA_DISABLE_SHADOW_VARIABLE_WARNINGS
+	#define PRAGMA_DISABLE_SHADOW_VARIABLE_WARNINGS \
+		_Pragma ("clang diagnostic push") \
+		_Pragma ("clang diagnostic ignored \"-Wshadow\"")
+#endif // PRAGMA_DISABLE_SHADOW_VARIABLE_WARNINGS
+
+#ifndef PRAGMA_POP
 	#define PRAGMA_POP \
 		_Pragma("clang diagnostic pop")
-#endif // DISABLE_DEPRECATION
+#endif // PRAGMA_POP
 
 #ifndef EMIT_CUSTOM_WARNING_AT_LINE
 #define EMIT_CUSTOM_WARNING_AT_LINE(Line, Warning) \
diff --git a/Engine/Source/Runtime/Core/Public/Linux/LinuxPlatformCompilerPreSetup.h b/Engine/Source/Runtime/Core/Public/Linux/LinuxPlatformCompilerPreSetup.h
--- a/Engine/Source/Runtime/Core/Public/Linux/LinuxPlatformCompilerPreSetup.h
+++ b/Engine/Source/Runtime/Core/Public/Linux/LinuxPlatformCompilerPreSetup.h
@@ -40,10 +40,18 @@
 	#define PRAGMA_ENABLE_DEPRECATION_WARNINGS \
 		_Pragma ("clang diagnostic push") \
 		_Pragma ("clang diagnostic warning \"-Wdeprecated-declarations\"")
+#endif // DISABLE_DEPRECATION
+
+#ifndef PRAGMA_DISABLE_SHADOW_VARIABLE_WARNINGS
+	#define PRAGMA_DISABLE_SHADOW_VARIABLE_WARNINGS \
+		_Pragma ("clang diagnostic push") \
+		_Pragma ("clang diagnostic ignored \"-Wshadow\"")
+#endif // PRAGMA_DISABLE_SHADOW_VARIABLE_WARNINGS
 
+#ifndef PRAGMA_POP
 	#define PRAGMA_POP \
 		_Pragma("clang diagnostic pop")
-#endif // DISABLE_DEPRECATION
+#endif // PRAGMA_POP
 
 #ifndef EMIT_CUSTOM_WARNING_AT_LINE
 #define EMIT_CUSTOM_WARNING_AT_LINE(Line, Warning) \
diff --git a/Engine/Source/Runtime/Core/Public/Mac/MacPlatformCompilerPreSetup.h b/Engine/Source/Runtime/Core/Public/Mac/MacPlatformCompilerPreSetup.h
--- a/Engine/Source/Runtime/Core/Public/Mac/MacPlatformCompilerPreSetup.h
+++ b/Engine/Source/Runtime/Core/Public/Mac/MacPlatformCompilerPreSetup.h
@@ -75,10 +75,18 @@
 	#define PRAGMA_ENABLE_DEPRECATION_WARNINGS \
 		_Pragma ("clang diagnostic push") \
 		_Pragma ("clang diagnostic warning \"-Wdeprecated-declarations\"")
+#endif // DISABLE_DEPRECATION
+
+#ifndef PRAGMA_DISABLE_SHADOW_VARIABLE_WARNINGS
+	#define PRAGMA_DISABLE_SHADOW_VARIABLE_WARNINGS \
+		_Pragma ("clang diagnostic push") \
+		_Pragma ("clang diagnostic ignored \"-Wshadow\"")
+#endif // PRAGMA_DISABLE_SHADOW_VARIABLE_WARNINGS
 
+#ifndef PRAGMA_POP
 	#define PRAGMA_POP \
 		_Pragma("clang diagnostic pop")
-#endif // DISABLE_DEPRECATION
+#endif // PRAGMA_POP
 
 #ifndef EMIT_CUSTOM_WARNING_AT_LINE
 #define EMIT_CUSTOM_WARNING_AT_LINE(Line, Warning) \
diff --git a/Engine/Source/Runtime/Core/Public/WinRT/WinRTPlatformCompilerPreSetup.h b/Engine/Source/Runtime/Core/Public/WinRT/WinRTPlatformCompilerPreSetup.h
--- a/Engine/Source/Runtime/Core/Public/WinRT/WinRTPlatformCompilerPreSetup.h
+++ b/Engine/Source/Runtime/Core/Public/WinRT/WinRTPlatformCompilerPreSetup.h
@@ -78,9 +78,6 @@
 		__pragma (warning(disable:4995)) \
 		__pragma (warning(disable:4996))
 
-	#define PRAGMA_POP \
-		__pragma(warning(pop))
-
 	#define PRAGMA_ENABLE_DEPRECATION_WARNINGS \
 		__pragma (warning(push)) \
 		__pragma (warning(default:4995)) \
@@ -88,6 +85,16 @@
 
 #endif // DISABLE_DEPRECATION
 
+#ifndef PRAGMA_DISABLE_SHADOW_VARIABLE_WARNINGS
+	#define PRAGMA_DISABLE_SHADOW_VARIABLE_WARNINGS \
+		__pragma (warning(push))
+#endif // PRAGMA_DISABLE_SHADOW_VARIABLE_WARNINGS
+
+#ifndef PRAGMA_POP
+	#define PRAGMA_POP \
+		__pragma(warning(pop))
+#endif // PRAGMA_POP
+
 #if _MSC_VER
 #define EMIT_CUSTOM_WARNING_AT_LINE(Line, Warning) \
 	__pragma(message(WARNING_LOCATION(Line) ": warning C4996: " Warning))
diff --git a/Engine/Source/Runtime/Core/Public/Windows/WindowsPlatformCompilerPreSetup.h b/Engine/Source/Runtime/Core/Public/Windows/WindowsPlatformCompilerPreSetup.h
--- a/Engine/Source/Runtime/Core/Public/Windows/WindowsPlatformCompilerPreSetup.h
+++ b/Engine/Source/Runtime/Core/Public/Windows/WindowsPlatformCompilerPreSetup.h
@@ -43,9 +43,6 @@
 		__pragma (warning(disable:4995)) \
 		__pragma (warning(disable:4996))
 
-	#define PRAGMA_POP \
-		__pragma(warning(pop))
-
 	#define PRAGMA_ENABLE_DEPRECATION_WARNINGS \
 		__pragma (warning(push)) \
 		__pragma (warning(default:4995)) \
@@ -53,6 +50,16 @@
 
 #endif // DISABLE_DEPRECATION
 
+#ifndef PRAGMA_DISABLE_SHADOW_VARIABLE_WARNINGS
+	#define PRAGMA_DISABLE_SHADOW_VARIABLE_WARNINGS \
+		__pragma (warning(push))
+#endif // PRAGMA_DISABLE_SHADOW_VARIABLE_WARNINGS
+
+#ifndef PRAGMA_POP
+	#define PRAGMA_POP \
+		__pragma(warning(pop))
+#endif // PRAGMA_POP
+
 #if _MSC_VER
 #define EMIT_CUSTOM_WARNING_AT_LINE(Line, Warning) \
 	__pragma(message(WARNING_LOCATION(Line) ": warning C4996: " Warning))
