added console commands and exec commands to "help"

diff --git a/Engine/Source/Runtime/Core/Private/Misc/Parse.cpp b/Engine/Source/Runtime/Core/Private/Misc/Parse.cpp
--- a/Engine/Source/Runtime/Core/Private/Misc/Parse.cpp
+++ b/Engine/Source/Runtime/Core/Private/Misc/Parse.cpp
@@ -134,7 +134,7 @@ void ConsoleCommandLibrary_DumpLibraryHTML(UWorld* InWorld, FExec& SubSystem, co
 			// generated by
 			LazyPrintf.PushParam(TEXT("Unreal Engine 4 'DumpConsoleCommands'"));
 			// version
-			LazyPrintf.PushParam(TEXT("0.93"));
+			LazyPrintf.PushParam(TEXT("0.94"));
 			// date
 			LazyPrintf.PushParam(*FDateTime::Now().ToString());
 
@@ -144,17 +144,27 @@ void ConsoleCommandLibrary_DumpLibraryHTML(UWorld* InWorld, FExec& SubSystem, co
 			{
 				const FString& Name = *It;
 
-				auto CVar = IConsoleManager::Get().FindConsoleVariable(*Name);
+				auto Element = IConsoleManager::Get().FindConsoleObject(*Name);
 
-				if(CVar)
+				if (Element)
 				{
-					FString Help = CVar->GetHelp();
+					// console command or variable
+
+					FString Help = Element->GetHelp();
 
 					Help = Help.ReplaceCharWithEscapedChar();
 
 					//{name: "r.SetRes", help:"To change the screen/window resolution."},
 					FString DataLine = FString::Printf(TEXT("{name: \"%s\", help:\"%s\"},\r\n"), *Name, *Help);
 
+					AllData += DataLine;
+				}
+				else
+				{
+					// Exec command (better we change them to use the new method as it has better help and is more convenient to use)
+					//{name: "", help:"To change the screen/window resolution."},
+					FString DataLine = FString::Printf(TEXT("{name: \"%s\", help:\"Sorry: Exec commands have no help\"},\r\n"), *Name);
+
 					AllData += DataLine;
 				}
 			}
diff --git a/Engine/Source/Runtime/Core/Public/HAL/ConsoleManager.h b/Engine/Source/Runtime/Core/Public/HAL/ConsoleManager.h
--- a/Engine/Source/Runtime/Core/Public/HAL/ConsoleManager.h
+++ b/Engine/Source/Runtime/Core/Public/HAL/ConsoleManager.h
@@ -76,6 +76,7 @@ public:
 	virtual IConsoleCommand* RegisterConsoleCommand(const TCHAR* Name, const TCHAR* Help, const FConsoleCommandWithWorldAndArgsDelegate& Command, uint32 Flags) override;
 	virtual IConsoleCommand* RegisterConsoleCommand(const TCHAR* Name, const TCHAR* Help, const FConsoleCommandWithOutputDeviceDelegate& Command, uint32 Flags) override;
 	virtual IConsoleCommand* RegisterConsoleCommand(const TCHAR* Name, const TCHAR* Help, uint32 Flags) override;
+	virtual IConsoleObject* FindConsoleObject(const TCHAR* Name) const;
 	virtual IConsoleVariable* FindConsoleVariable(const TCHAR* Name) const override;
 	virtual void ForEachConsoleObject(const FConsoleObjectVisitor& Visitor, const TCHAR* ThatStartsWith) const override;
 	virtual bool ProcessUserConsoleInput(const TCHAR* InInput, FOutputDevice& Ar, UWorld* InWorld) override;
@@ -126,13 +127,6 @@ private: // ----------------------------------------------------
 	 */
 	static FString GetTextSection(const TCHAR* &It);
 
-	/**
-	 * Find a console variable or command
-	 * @param Name must not be 0
-	 * @return 0 if the object wasn't found
-	 */
-	IConsoleObject* FindConsoleObject(const TCHAR* Name) const;
-
 	/** same as FindConsoleObject() but ECVF_CreatedFromIni are not filtered out (for internal use) */
 	IConsoleObject* FindConsoleObjectUnfiltered(const TCHAR* Name) const;
 
diff --git a/Engine/Source/Runtime/Core/Public/HAL/IConsoleManager.h b/Engine/Source/Runtime/Core/Public/HAL/IConsoleManager.h
--- a/Engine/Source/Runtime/Core/Public/HAL/IConsoleManager.h
+++ b/Engine/Source/Runtime/Core/Public/HAL/IConsoleManager.h
@@ -493,6 +493,13 @@ struct CORE_API IConsoleManager
 	 */
 	virtual IConsoleVariable* FindConsoleVariable(const TCHAR* Name) const = 0;
 
+	/**
+	* Find a console variable or command
+	* @param Name must not be 0
+	* @return 0 if the object wasn't found
+	*/
+	virtual IConsoleObject* FindConsoleObject(const TCHAR* Name) const = 0;
+
 	/**
 	 * Find a typed console variable (faster access to the value, no virtual function call)
 	 * @param Name must not be 0
