Linux: do not filter out events that don't need a window (UE-10036).

diff --git a/Engine/Source/Runtime/Core/Private/Linux/LinuxApplication.cpp b/Engine/Source/Runtime/Core/Private/Linux/LinuxApplication.cpp
--- a/Engine/Source/Runtime/Core/Private/Linux/LinuxApplication.cpp
+++ b/Engine/Source/Runtime/Core/Private/Linux/LinuxApplication.cpp
@@ -211,13 +211,14 @@ void FLinuxApplication::ProcessDeferredMessage( SDL_Event Event )
 	}
 
 	// get pointer to window that received this event
-	TSharedPtr< FLinuxWindow > CurrentEventWindow = FindEventWindow(&Event);
+	bool bWindowlessEvent = false;
+	TSharedPtr< FLinuxWindow > CurrentEventWindow = FindEventWindow(&Event, bWindowlessEvent);
 
 	if (CurrentEventWindow.IsValid())
 	{
 		NativeWindow = CurrentEventWindow->GetHWnd();
 	}
-	if (!NativeWindow && Event.type != SDL_USEREVENT)
+	if (!NativeWindow && !bWindowlessEvent)
 	{
 		return;
 	}
@@ -1010,9 +1011,11 @@ TCHAR FLinuxApplication::ConvertChar( SDL_Keysym Keysym )
     return Char;
 }
 
-TSharedPtr< FLinuxWindow > FLinuxApplication::FindEventWindow( SDL_Event* Event )
+TSharedPtr< FLinuxWindow > FLinuxApplication::FindEventWindow(SDL_Event* Event, bool& bOutWindowlessEvent)
 {
 	uint16 WindowID = 0;
+	bOutWindowlessEvent = false;
+
 	switch (Event->type)
 	{
 		case SDL_TEXTINPUT:
@@ -1035,11 +1038,12 @@ TSharedPtr< FLinuxWindow > FLinuxApplication::FindEventWindow( SDL_Event* Event
 		case SDL_MOUSEWHEEL:
 			WindowID = Event->wheel.windowID;
 			break;
-
 		case SDL_WINDOWEVENT:
 			WindowID = Event->window.windowID;
 			break;
+
 		default:
+			bOutWindowlessEvent = true;
 			return TSharedPtr< FLinuxWindow >(nullptr);
 	}
 
diff --git a/Engine/Source/Runtime/Core/Public/Linux/LinuxApplication.h b/Engine/Source/Runtime/Core/Public/Linux/LinuxApplication.h
--- a/Engine/Source/Runtime/Core/Public/Linux/LinuxApplication.h
+++ b/Engine/Source/Runtime/Core/Public/Linux/LinuxApplication.h
@@ -80,7 +80,15 @@ private:
 
 	TCHAR ConvertChar( SDL_Keysym Keysym );
 
-	TSharedPtr< FLinuxWindow > FindEventWindow( SDL_Event *pEvent );
+	/**
+	 * Finds a window associated with the event (if there is such an association).
+	 *
+	 * @param Event event in question
+	 * @param bOutWindowlessEvent whether or not the event needs to have a window at all
+	 *
+	 * @return pointer to the window, if any
+	 */
+	TSharedPtr< FLinuxWindow > FindEventWindow(SDL_Event *Event, bool& bOutWindowlessEvent);
 
 	void UpdateMouseCaptureWindow( SDL_HWindow TargetWindow );
 
