fixed r.DumpingMovie

diff --git a/Engine/Source/Runtime/Core/Private/HAL/ConsoleManager.cpp b/Engine/Source/Runtime/Core/Private/HAL/ConsoleManager.cpp
--- a/Engine/Source/Runtime/Core/Private/HAL/ConsoleManager.cpp
+++ b/Engine/Source/Runtime/Core/Private/HAL/ConsoleManager.cpp
@@ -1480,6 +1480,13 @@ void FConsoleManager::Test()
 // These don't belong here, but they belong here more than they belong in launch engine loop.
 void CreateConsoleVariables()
 {
+	// this registeres to a reference, so we cannot use TAutoConsoleVariable
+	IConsoleManager::Get().RegisterConsoleVariableRef(TEXT("r.DumpingMovie"),
+		GIsDumpingMovie,
+		TEXT("Allows to dump each rendered frame to disk (slow fps, names MovieFrame..).\n")
+		TEXT("<=0:off (default), <0:remains on, >0:remains on for n frames (n is the number specified)"),
+		ECVF_Cheat);
+
 
 #if !(UE_BUILD_SHIPPING || UE_BUILD_TEST)
 
@@ -1551,13 +1558,6 @@ static TAutoConsoleVariable<int32> CVarLimitRenderingFeatures(
 	TEXT(" <=0:off, order is defined in code (can be documented here when we settled on an order)"),
 	ECVF_Cheat | ECVF_RenderThreadSafe);
 
-static TAutoConsoleVariable<int32> CVarDumpingMovie(
-	TEXT("r.DumpingMovie"),
-	GIsDumpingMovie,
-	TEXT("Allows to dump each rendered frame to disk (slow fps, names MovieFrame..).\n")
-	TEXT("<=0:off (default), <0:remains on, >0:remains on for n frames (n is the number specified)"),
-	ECVF_Cheat);
-
 #endif // !(UE_BUILD_SHIPPING || UE_BUILD_TEST)
 
 static TAutoConsoleVariable<int32> CVarUniformBufferPooling(
