Core - Stats based memory profiler should start on all platforms (UE-9906)

diff --git a/Engine/Source/Runtime/Core/Private/GenericPlatform/GenericPlatformMisc.cpp b/Engine/Source/Runtime/Core/Private/GenericPlatform/GenericPlatformMisc.cpp
--- a/Engine/Source/Runtime/Core/Private/GenericPlatform/GenericPlatformMisc.cpp
+++ b/Engine/Source/Runtime/Core/Private/GenericPlatform/GenericPlatformMisc.cpp
@@ -875,9 +875,4 @@ FString FGenericPlatformMisc::GetOperatingSystemId()
 {
 	// not implemented by default.
 	return FString();
-}
-
-const TCHAR* FGenericPlatformMisc::GetNativeCommandLine()
-{
-	return TEXT("");
-}
+}
\ No newline at end of file
diff --git a/Engine/Source/Runtime/Core/Private/HAL/UnrealMemory.cpp b/Engine/Source/Runtime/Core/Private/HAL/UnrealMemory.cpp
--- a/Engine/Source/Runtime/Core/Private/HAL/UnrealMemory.cpp
+++ b/Engine/Source/Runtime/Core/Private/HAL/UnrealMemory.cpp
@@ -14,7 +14,6 @@
 #include "MallocProfiler.h"
 #include "MallocThreadSafeProxy.h"
 #include "MallocCrash.h"
-#include "StatsMallocProfilerProxy.h"
 
 
 /** Helper function called on first allocation to create and initialize GMalloc */
@@ -40,17 +39,6 @@ void GCreateMalloc()
 	{
 		GMalloc = new FMallocThreadSafeProxy( GMalloc );
 	}
-
-	
-#if	STATS
-	const bool bCreateMallocProfiler = FParse::Param( FPlatformMisc::GetNativeCommandLine(), TEXT( "MemoryProfiler" ) );
-
-	// Create the stats malloc profiler proxy.
-	if( bCreateMallocProfiler )
-	{
-		GMalloc = FStatsMallocProfilerProxy::Get();
-	}
-#endif // STATS
 }
 
 void* FMemory::Malloc( SIZE_T Count, uint32 Alignment ) 
diff --git a/Engine/Source/Runtime/Core/Private/Stats/Stats2.cpp b/Engine/Source/Runtime/Core/Private/Stats/Stats2.cpp
--- a/Engine/Source/Runtime/Core/Private/Stats/Stats2.cpp
+++ b/Engine/Source/Runtime/Core/Private/Stats/Stats2.cpp
@@ -1066,9 +1066,7 @@ void FThreadStats::CheckForCollectingStartupStats()
 	}
 
 	// Now we can safely enable malloc profiler.
-	// @TODO yrx 2014-12-01 Investigate if we can enable it earlier.
-	const bool bEnableMallocProfiler = FParse::Param( FCommandLine::Get(), TEXT("MemoryProfiler") );
-	if( bEnableMallocProfiler )
+	if( FStatsMallocProfilerProxy::HasMemoryProfilerToken() )
 	{
 		// Enable all available groups and enable malloc profiler.
 		IStatGroupEnableManager::Get().StatGroupEnableManagerCommand( TEXT("all") );
diff --git a/Engine/Source/Runtime/Core/Private/Stats/StatsCommand.cpp b/Engine/Source/Runtime/Core/Private/Stats/StatsCommand.cpp
--- a/Engine/Source/Runtime/Core/Private/Stats/StatsCommand.cpp
+++ b/Engine/Source/Runtime/Core/Private/Stats/StatsCommand.cpp
@@ -1256,13 +1256,16 @@ static void StatCmd(FString InCmd)
 		FCommandStatsFile::Get().Stop();
 		FThreadStats::DisableRawStats();
 
-		if( FStatsMallocProfilerProxy::Get()->GetState() )
+		if( FStatsMallocProfilerProxy::HasMemoryProfilerToken() )
 		{
-			// Disable memory profiler and restore default stats groups.
-			FStatsMallocProfilerProxy::Get()->SetState( false );		
-			IStatGroupEnableManager::Get().StatGroupEnableManagerCommand( TEXT("default") );
+			if( FStatsMallocProfilerProxy::Get()->GetState() )
+			{
+				// Disable memory profiler and restore default stats groups.
+				FStatsMallocProfilerProxy::Get()->SetState( false );
+				IStatGroupEnableManager::Get().StatGroupEnableManagerCommand( TEXT( "default" ) );
+			}
 		}
-
+		
 		Stats.ResetStatsForRawStats();
 
 		// Disable displaying the raw stats memory overhead.
diff --git a/Engine/Source/Runtime/Core/Private/Stats/StatsMallocProfilerProxy.cpp b/Engine/Source/Runtime/Core/Private/Stats/StatsMallocProfilerProxy.cpp
--- a/Engine/Source/Runtime/Core/Private/Stats/StatsMallocProfilerProxy.cpp
+++ b/Engine/Source/Runtime/Core/Private/Stats/StatsMallocProfilerProxy.cpp
@@ -47,10 +47,17 @@ FStatsMallocProfilerProxy* FStatsMallocProfilerProxy::Get()
 	if( Instance == nullptr )
 	{
 		Instance = new FStatsMallocProfilerProxy( GMalloc );
+		// Initialize stats metadata.
+		// We need to do here, after all hardcoded names have been initialized.
+		Instance->InitializeStatsMetadata();
 	}
 	return Instance;
 }
 
+bool FStatsMallocProfilerProxy::HasMemoryProfilerToken()
+{
+	return FParse::Param( FCommandLine::Get(), TEXT( "MemoryProfiler" ) );
+}
 
 void FStatsMallocProfilerProxy::SetState( bool bNewState )
 {
diff --git a/Engine/Source/Runtime/Core/Private/UObject/UnrealNames.cpp b/Engine/Source/Runtime/Core/Private/UObject/UnrealNames.cpp
--- a/Engine/Source/Runtime/Core/Private/UObject/UnrealNames.cpp
+++ b/Engine/Source/Runtime/Core/Private/UObject/UnrealNames.cpp
@@ -816,10 +816,6 @@ void FName::StaticInit()
 		}
 	}
 #endif
-
-	// Initialize stats metadata.
-	// We need to do here, after all hardcoded names have been initialized.
-	GMalloc->InitializeStatsMetadata();
 }
 
 bool& FName::GetIsInitialized()
diff --git a/Engine/Source/Runtime/Core/Private/Windows/WindowsPlatformMisc.cpp b/Engine/Source/Runtime/Core/Private/Windows/WindowsPlatformMisc.cpp
--- a/Engine/Source/Runtime/Core/Private/Windows/WindowsPlatformMisc.cpp
+++ b/Engine/Source/Runtime/Core/Private/Windows/WindowsPlatformMisc.cpp
@@ -1933,11 +1933,6 @@ uint32 FWindowsPlatformMisc::GetLastError()
 	return (uint32)::GetLastError();
 }
 
-const TCHAR* FWindowsPlatformMisc::GetNativeCommandLine()
-{
-	return ::GetCommandLine();
-}
-
 bool FWindowsPlatformMisc::CoInitialize()
 {
 	HRESULT hr = ::CoInitialize(NULL);
diff --git a/Engine/Source/Runtime/Core/Public/GenericPlatform/GenericPlatformMisc.h b/Engine/Source/Runtime/Core/Public/GenericPlatform/GenericPlatformMisc.h
--- a/Engine/Source/Runtime/Core/Public/GenericPlatform/GenericPlatformMisc.h
+++ b/Engine/Source/Runtime/Core/Public/GenericPlatform/GenericPlatformMisc.h
@@ -267,6 +267,7 @@ struct CORE_API FGenericPlatformMisc
 	 */
 	static FString GetUniqueDeviceId();
 
+	// @TODO yrx 2015-02-24 Remove
 	/** Submits a crash report to a central server (release builds only) */
 	static void SubmitErrorReport( const TCHAR* InErrorHist, EErrorReportMode::Type InMode );
 
@@ -412,11 +413,6 @@ public:
 	{
 	}
 
-	/**
-	 * @return the executable's command line
-	 */
-	static const TCHAR* GetNativeCommandLine();
-
 	/** 
 	 *	Set the value for the given section and key in the platform specific key->value store
 	 *  Note: The key->value store is user-specific, but may be used to share data between different applications for the same user
diff --git a/Engine/Source/Runtime/Core/Public/Stats/StatsMallocProfilerProxy.h b/Engine/Source/Runtime/Core/Public/Stats/StatsMallocProfilerProxy.h
--- a/Engine/Source/Runtime/Core/Public/Stats/StatsMallocProfilerProxy.h
+++ b/Engine/Source/Runtime/Core/Public/Stats/StatsMallocProfilerProxy.h
@@ -38,6 +38,11 @@ public:
 
 	static CORE_API FStatsMallocProfilerProxy* Get();
 
+	/**
+	 * @return true, if the command line has the memory profiler token
+	 */
+	static CORE_API bool HasMemoryProfilerToken();
+
 	void SetState( bool bNewState );
 
 	bool GetState() const
diff --git a/Engine/Source/Runtime/Core/Public/Windows/WindowsPlatformMisc.h b/Engine/Source/Runtime/Core/Public/Windows/WindowsPlatformMisc.h
--- a/Engine/Source/Runtime/Core/Public/Windows/WindowsPlatformMisc.h
+++ b/Engine/Source/Runtime/Core/Public/Windows/WindowsPlatformMisc.h
@@ -98,7 +98,6 @@ struct CORE_API FWindowsPlatformMisc
 	static void LoadStartupModules();
 
 	static uint32 GetLastError();
-	static const TCHAR* GetNativeCommandLine();
 	static void RaiseException( uint32 ExceptionCode );
 	static bool SetStoredValue(const FString& InStoreId, const FString& InSectionName, const FString& InKeyName, const FString& InValue);
 	static bool GetStoredValue(const FString& InStoreId, const FString& InSectionName, const FString& InKeyName, FString& OutValue);
