Stats - Removed redundant code

diff --git a/Engine/Source/Runtime/Core/Private/Stats/StatsMallocProfilerProxy.cpp b/Engine/Source/Runtime/Core/Private/Stats/StatsMallocProfilerProxy.cpp
--- a/Engine/Source/Runtime/Core/Private/Stats/StatsMallocProfilerProxy.cpp
+++ b/Engine/Source/Runtime/Core/Private/Stats/StatsMallocProfilerProxy.cpp
@@ -7,21 +7,13 @@
 #include "StatsData.h"
 #include "StatsMallocProfilerProxy.h"
 
-#define DECLARE_MEMORY_PTR_STAT(CounterName,StatId,GroupId)\
-	DECLARE_STAT(CounterName,StatId,GroupId,EStatDataType::ST_Ptr, false, false, FPlatformMemory::MCR_Invalid); \
-	static DEFINE_STAT(StatId)
-
-#define DECLARE_MEMORY_SIZE_STAT(CounterName,StatId,GroupId)\
-	DECLARE_STAT(CounterName,StatId,GroupId,EStatDataType::ST_int64, false, false, FPlatformMemory::MCR_Invalid); \
-	static DEFINE_STAT(StatId)
-
 /** Fake stat group and memory stats. */
 DECLARE_STATS_GROUP(TEXT("Memory Profiler"), STATGROUP_MemoryProfiler, STATCAT_Advanced);
 
-DECLARE_MEMORY_PTR_STAT(TEXT("Memory Free Ptr"),				STAT_Memory_FreePtr,				STATGROUP_MemoryProfiler);
-DECLARE_MEMORY_PTR_STAT(TEXT("Memory Alloc Ptr"),				STAT_Memory_AllocPtr,				STATGROUP_MemoryProfiler);
-DECLARE_MEMORY_SIZE_STAT(TEXT("Memory Alloc Size"),				STAT_Memory_AllocSize,				STATGROUP_MemoryProfiler);
-DECLARE_MEMORY_SIZE_STAT(TEXT("Memory Operation Sequence Tag"),	STAT_Memory_OperationSequenceTag,	STATGROUP_MemoryProfiler);
+DECLARE_PTR_STAT( TEXT( "Memory Free Ptr" ), STAT_Memory_FreePtr, STATGROUP_MemoryProfiler );
+DECLARE_PTR_STAT( TEXT( "Memory Alloc Ptr" ), STAT_Memory_AllocPtr, STATGROUP_MemoryProfiler );
+DECLARE_MEMORY_STAT( TEXT( "Memory Alloc Size" ), STAT_Memory_AllocSize, STATGROUP_MemoryProfiler );
+DECLARE_MEMORY_STAT( TEXT( "Memory Operation Sequence Tag" ), STAT_Memory_OperationSequenceTag, STATGROUP_MemoryProfiler );
 
 /** Stats for memory usage by the profiler. */
 DECLARE_DWORD_COUNTER_STAT(TEXT("Profiler AllocPtr Calls"),	STAT_Memory_AllocPtr_Calls,	STATGROUP_StatSystem);
@@ -181,6 +173,7 @@ void* FStatsMallocProfilerProxy::Malloc( SIZE_T Size, uint32 Alignment )
 
 void* FStatsMallocProfilerProxy::Realloc( void* OldPtr, SIZE_T NewSize, uint32 Alignment )
 {
+	// @TODO yrx 2015-02-18 Add TrackRealloc to optimize
 	const int32 FreeSequenceTag = MemorySequenceTag.Increment();
 	TrackFree( OldPtr, FreeSequenceTag );
 
