Add VectorDivide to UnrealMathDirectX.h

diff --git a/Engine/Source/Runtime/Core/Public/Math/UnrealMathDirectX.h b/Engine/Source/Runtime/Core/Public/Math/UnrealMathDirectX.h
--- a/Engine/Source/Runtime/Core/Public/Math/UnrealMathDirectX.h
+++ b/Engine/Source/Runtime/Core/Public/Math/UnrealMathDirectX.h
@@ -266,6 +266,15 @@ FORCEINLINE float VectorGetComponent( VectorRegister Vec, uint32 ComponentIndex
  */
 #define VectorMultiply( Vec1, Vec2 )	DirectX::XMVectorMultiply( Vec1, Vec2 )
 
+ /**
+ * Divides two vectors (component-wise) and returns the result.
+ *
+ * @param Vec1	1st vector
+ * @param Vec2	2nd vector
+ * @return		VectorRegister( Vec1.x/Vec2.x, Vec1.y/Vec2.y, Vec1.z/Vec2.z, Vec1.w/Vec2.w )
+ */
+#define VectorDivide( Vec1, Vec2 )	DirectX::XMVectorDivide( Vec1, Vec2 )
+
 /**
  * Multiplies two vectors (component-wise), adds in the third vector and returns the result.
  *
