CIS fix
#ios

diff --git a/Engine/Source/Runtime/Core/Public/Math/UnrealMathFPU.h b/Engine/Source/Runtime/Core/Public/Math/UnrealMathFPU.h
--- a/Engine/Source/Runtime/Core/Public/Math/UnrealMathFPU.h
+++ b/Engine/Source/Runtime/Core/Public/Math/UnrealMathFPU.h
@@ -288,6 +288,23 @@ FORCEINLINE VectorRegister VectorMultiplyAdd( const VectorRegister& Vec1, const
 	return Vec;
 }
 
+/**
+* Divides two vectors (component-wise) and returns the result.
+*
+* @param Vec1	1st vector
+* @param Vec2	2nd vector
+* @return		VectorRegister( Vec1.x/Vec2.x, Vec1.y/Vec2.y, Vec1.z/Vec2.z, Vec1.w/Vec2.w )
+*/
+FORCEINLINE VectorRegister VectorDivide(const VectorRegister& Vec1, const VectorRegister& Vec2)
+{
+	VectorRegister Vec;
+	Vec.V[0] = Vec1.V[0] / Vec2.V[0];
+	Vec.V[1] = Vec1.V[1] / Vec2.V[1];
+	Vec.V[2] = Vec1.V[2] / Vec2.V[2];
+	Vec.V[3] = Vec1.V[3] / Vec2.V[3];
+	return Vec;
+}
+
 /**
  * Calculates the dot3 product of two vectors and returns a vector with the result in all 4 components.
  * Only really efficient on Xbox 360.
