another CIS fix
#ios

diff --git a/Engine/Source/Runtime/Core/Public/Math/UnrealMathNeon.h b/Engine/Source/Runtime/Core/Public/Math/UnrealMathNeon.h
--- a/Engine/Source/Runtime/Core/Public/Math/UnrealMathNeon.h
+++ b/Engine/Source/Runtime/Core/Public/Math/UnrealMathNeon.h
@@ -325,7 +325,6 @@ FORCEINLINE VectorRegister VectorMultiplyAdd( VectorRegister Vec1, VectorRegiste
 	return vmlaq_f32( Vec3, Vec1, Vec2 );
 }
 
-
 /**
  * Calculates the dot3 product of two vectors and returns a vector with the result in all 4 components.
  * Only really efficient on Xbox 360.
@@ -587,6 +586,19 @@ FORCEINLINE VectorRegister VectorReciprocalAccurate(const VectorRegister& Vec)
 	return vmulq_f32(vrecpsq_f32(Vec, Reciprocal), Reciprocal);
 }
 
+/**
+* Divides two vectors (component-wise) and returns the result.
+*
+* @param Vec1	1st vector
+* @param Vec2	2nd vector
+* @return		VectorRegister( Vec1.x/Vec2.x, Vec1.y/Vec2.y, Vec1.z/Vec2.z, Vec1.w/Vec2.w )
+*/
+FORCEINLINE VectorRegister VectorDivide(VectorRegister Vec1, VectorRegister Vec2)
+{
+	VectorRegister x = VectorReciprocalAccurate(Vec2);
+	return VectorMultiply(Vec1, x);
+}
+
 /**
 * Normalize vector
 *
