Fix for deadlock at exit in Mac editor

diff --git a/Engine/Source/Runtime/Core/Private/Mac/CocoaThread.cpp b/Engine/Source/Runtime/Core/Private/Mac/CocoaThread.cpp
--- a/Engine/Source/Runtime/Core/Private/Mac/CocoaThread.cpp
+++ b/Engine/Source/Runtime/Core/Private/Mac/CocoaThread.cpp
@@ -14,6 +14,7 @@ NSString* UE4CloseEventMode = @"UE4CloseEventMode";
 NSString* UE4IMEEventMode = @"UE4IMEEventMode";
 
 static FCocoaGameThread* GCocoaGameThread = nil;
+static uint32 GMainThreadId = 0;
 
 class FCocoaRunLoopSource;
 @interface FCocoaRunLoopSourceInfo : NSObject
@@ -357,22 +358,26 @@ FCocoaRunLoopSource* FCocoaRunLoopSource::GameRunLoopSource = nullptr;
 	
 	[super main];
 	
+	// We have exited the game thread, so any UE4 code running now should treat the Main thread
+	// as the game thread, so we don't crash in static destructors.
+	GGameThreadId = GMainThreadId;
+	
+	// Tell the main thread we are OK to quit, but don't wait for it.
 	if (GIsRequestingExit)
 	{
-		dispatch_sync(dispatch_get_main_queue(), ^{
-			GGameThreadId = FPlatformTLS::GetCurrentThreadId();
+		MainThreadCall(^{
 			[NSApp replyToApplicationShouldTerminate:YES];
 			[[NSProcessInfo processInfo] enableSuddenTermination];
-		});
+		}, NSDefaultRunLoopMode, false);
 	}
 	else
 	{
-		dispatch_sync(dispatch_get_main_queue(), ^{
-			GGameThreadId = FPlatformTLS::GetCurrentThreadId();
+		MainThreadCall(^{
 			[[NSProcessInfo processInfo] enableSuddenTermination];
-		});
+		}, NSDefaultRunLoopMode, false);
 	}
 	
+	// And now it is time to die.
 	[self release];
 }
 
@@ -441,6 +446,8 @@ void RunGameThread(id Target, SEL Selector)
 	// @todo: Proper support for sudden termination (OS termination of application without any events, notifications or signals) - which presently can assert, crash or corrupt. At present we can't deal with this and we need to ensure that we don't permit the system to kill us without warning by disabling & enabling support around operations which must be committed to disk atomically.
 	[[NSProcessInfo processInfo] disableSuddenTermination];
 	
+	GMainThreadId = FPlatformTLS::GetCurrentThreadId();
+	
 #if MAC_SEPARATE_GAME_THREAD
 	// Register main run loop source
 	FCocoaRunLoopSource::RegisterMainRunLoop(CFRunLoopGetCurrent());
diff --git a/Engine/Source/Runtime/Core/Private/Mac/MacPlatformMisc.cpp b/Engine/Source/Runtime/Core/Private/Mac/MacPlatformMisc.cpp
--- a/Engine/Source/Runtime/Core/Private/Mac/MacPlatformMisc.cpp
+++ b/Engine/Source/Runtime/Core/Private/Mac/MacPlatformMisc.cpp
@@ -399,7 +399,9 @@ void FMacPlatformMisc::PlatformTearDown()
 {
 	if (CommandletActivity)
 	{
-		[[NSProcessInfo processInfo] endActivity:CommandletActivity];
+		MainThreadCall(^{
+			[[NSProcessInfo processInfo] endActivity:CommandletActivity];
+		}, NSDefaultRunLoopMode, false);
 		CommandletActivity = nil;
 	}
 	FApplePlatformSymbolication::EnableCoreSymbolication(false);
