Core - Buffer overrun in FWindowsPlatformMisc::SubmitErrorReport (UECORE-127)

diff --git a/Engine/Source/Runtime/Core/Private/Windows/WindowsPlatformMisc.cpp b/Engine/Source/Runtime/Core/Private/Windows/WindowsPlatformMisc.cpp
--- a/Engine/Source/Runtime/Core/Private/Windows/WindowsPlatformMisc.cpp
+++ b/Engine/Source/Runtime/Core/Private/Windows/WindowsPlatformMisc.cpp
@@ -581,6 +581,8 @@ void FWindowsPlatformMisc::SubmitErrorReport( const TCHAR* InErrorHist, EErrorRe
 			return;
 		}
 
+		const uint32 MAX_STRING_LEN = 256;
+
 		TCHAR ReportDumpVersion[] = TEXT("3");
 
 		FString ReportDumpPath;
@@ -592,60 +594,60 @@ void FWindowsPlatformMisc::SubmitErrorReport( const TCHAR* InErrorHist, EErrorRe
 		FArchive * AutoReportFile = IFileManager::Get().CreateFileWriter(*ReportDumpPath, FILEWRITE_EvenIfReadOnly);
 		if (AutoReportFile != NULL)
 		{
-			TCHAR CompName[256];
-			FCString::Strcpy(CompName, FPlatformProcess::ComputerName());
-			TCHAR UserName[256];
-			FCString::Strcpy(UserName, FPlatformProcess::UserName());
-			TCHAR GameName[256];
-			FCString::Strcpy(GameName, *FString::Printf(TEXT("%s %s"), TEXT(BRANCH_NAME), FApp::GetGameName()));
-			TCHAR PlatformName[32];
+			TCHAR CompName[MAX_STRING_LEN];
+			FCString::Strncpy(CompName, FPlatformProcess::ComputerName(), MAX_STRING_LEN);
+			TCHAR UserName[MAX_STRING_LEN];
+			FCString::Strncpy(UserName, FPlatformProcess::UserName(), MAX_STRING_LEN);
+			TCHAR GameName[MAX_STRING_LEN];
+			FCString::Strncpy(GameName, *FString::Printf(TEXT("%s %s"), TEXT(BRANCH_NAME), FApp::GetGameName()), MAX_STRING_LEN);
+			TCHAR PlatformName[MAX_STRING_LEN];
 #if PLATFORM_64BITS
-			FCString::Strcpy(PlatformName, TEXT("PC 64-bit"));
+			FCString::Strncpy(PlatformName, TEXT("PC 64-bit"), MAX_STRING_LEN);
 #else	//PLATFORM_64BITS
-			FCString::Strcpy(PlatformName, TEXT("PC 32-bit"));
+			FCString::Strncpy(PlatformName, TEXT("PC 32-bit"), MAX_STRING_LEN);
 #endif	//PLATFORM_64BITS
-			TCHAR CultureName[10];
-			FCString::Strcpy(CultureName, *FInternationalization::Get().GetDefaultCulture()->GetName());
-			TCHAR SystemTime[256];
-			FCString::Strcpy(SystemTime, *FDateTime::Now().ToString());
-			TCHAR EngineVersionStr[32];
-			FCString::Strcpy(EngineVersionStr, *GEngineVersion.ToString());
-
-			TCHAR ChangelistVersionStr[32];
+			TCHAR CultureName[MAX_STRING_LEN];
+			FCString::Strncpy(CultureName, *FInternationalization::Get().GetDefaultCulture()->GetName(), MAX_STRING_LEN);
+			TCHAR SystemTime[MAX_STRING_LEN];
+			FCString::Strncpy(SystemTime, *FDateTime::Now().ToString(), MAX_STRING_LEN);
+			TCHAR EngineVersionStr[MAX_STRING_LEN];
+			FCString::Strncpy(EngineVersionStr, *GEngineVersion.ToString(), 256 );
+
+			TCHAR ChangelistVersionStr[MAX_STRING_LEN];
 			int32 ChangelistFromCommandLine = 0;
 			const bool bFoundAutomatedBenchMarkingChangelist = FParse::Value( FCommandLine::Get(), TEXT("-gABC="), ChangelistFromCommandLine );
 			if( bFoundAutomatedBenchMarkingChangelist == true )
 			{
-				FCString::Strcpy(ChangelistVersionStr, *FString::FromInt(ChangelistFromCommandLine));
+				FCString::Strncpy(ChangelistVersionStr, *FString::FromInt(ChangelistFromCommandLine), MAX_STRING_LEN);
 			}
 			// we are not passing in the changelist to use so use the one that was stored in the ObjectVersion
 			else
 			{
-				FCString::Strcpy(ChangelistVersionStr, *FString::FromInt(GEngineVersion.GetChangelist()));
+				FCString::Strncpy(ChangelistVersionStr, *FString::FromInt(GEngineVersion.GetChangelist()), MAX_STRING_LEN);
 			}
 
 			TCHAR CmdLine[2048];
-			FCString::Strcpy(CmdLine, FCommandLine::Get());
+			FCString::Strncpy(CmdLine, FCommandLine::Get(),ARRAY_COUNT(CmdLine));
 			TCHAR BaseDir[260];
-			FCString::Strcpy(BaseDir, FPlatformProcess::BaseDir());
+			FCString::Strncpy(BaseDir, FPlatformProcess::BaseDir(), ARRAY_COUNT(BaseDir));
 			TCHAR separator = 0;
 
-			TCHAR EngineMode[64];
+			TCHAR EngineMode[MAX_STRING_LEN];
 			if( IsRunningCommandlet() )
 			{
-				FCString::Strcpy(EngineMode, TEXT("Commandlet"));
+				FCString::Strncpy(EngineMode, TEXT("Commandlet"), MAX_STRING_LEN);
 			}
 			else if( GIsEditor )
 			{
-				FCString::Strcpy(EngineMode, TEXT("Editor"));
+				FCString::Strncpy(EngineMode, TEXT("Editor"), MAX_STRING_LEN);
 			}
 			else if( IsRunningDedicatedServer() )
 			{
-				FCString::Strcpy(EngineMode, TEXT("Server"));
+				FCString::Strncpy(EngineMode, TEXT("Server"), MAX_STRING_LEN);
 			}
 			else
 			{
-				FCString::Strcpy(EngineMode, TEXT("Game"));
+				FCString::Strncpy(EngineMode, TEXT("Game"), MAX_STRING_LEN);
 			}
 
 			//build the report dump file
