TUniquePtr::IsValid added.
Construction, assignment and comparison against nullptr.
Comment fixes.

diff --git a/Engine/Source/Runtime/Core/Public/Templates/UniquePtr.h b/Engine/Source/Runtime/Core/Public/Templates/UniquePtr.h
--- a/Engine/Source/Runtime/Core/Public/Templates/UniquePtr.h
+++ b/Engine/Source/Runtime/Core/Public/Templates/UniquePtr.h
@@ -38,6 +38,14 @@ public:
 	{
 	}
 
+	/**
+	 * nullptr constructor - initializes the TUniquePtr to null.
+	 */
+	FORCEINLINE TUniquePtr(TYPE_OF_NULLPTR)
+		: Ptr(nullptr)
+	{
+	}
+
 	/**
 	 * Move constructor
 	 */
@@ -92,6 +100,19 @@ public:
 		return *this;
 	}
 
+	/**
+	 * Nullptr assignment operator
+	 */
+	FORCEINLINE TUniquePtr& operator=(TYPE_OF_NULLPTR)
+	{
+		// We delete last, because we don't want odd side effects if the destructor of T relies on the state of this
+		T* OldPtr = Ptr;
+		Ptr = nullptr;
+		delete OldPtr;
+
+		return *this;
+	}
+
 	/**
 	 * Destructor
 	 */
@@ -100,14 +121,24 @@ public:
 		delete Ptr;
 	}
 
+	/**
+	 * Tests if the TUniquePtr currently owns an object.
+	 *
+	 * @return true if the TUniquePtr currently owns an object, false otherwise.
+	 */
+	bool IsValid() const
+	{
+		return Ptr != nullptr;
+	}
+
 	/**
 	 * operator bool
 	 *
-	 * @return true if the TUniquePtr currently owns an object, convertible-to-false otherwise.
+	 * @return true if the TUniquePtr currently owns an object, false otherwise.
 	 */
 	FORCEINLINE_EXPLICIT_OPERATOR_BOOL() const
 	{
-		return !!Ptr;
+		return IsValid();
 	}
 
 	/**
@@ -117,7 +148,7 @@ public:
 	 */
 	FORCEINLINE bool operator!() const
 	{
-		return !Ptr;
+		return !IsValid();
 	}
 
 	/**
@@ -163,13 +194,13 @@ public:
 	}
 
 	/**
-	 * Replaces any owned object with aGives the TUniquePtr a new object to own, destroying any previously-owned object.
+	 * Gives the TUniquePtr a new object to own, destroying any previously-owned object.
 	 *
 	 * @param InPtr A pointer to the object to take ownership of.
 	 */
 	FORCEINLINE void Reset(T* InPtr = nullptr)
 	{
-		// We delete last, because we don't want odd side effects if the destructor of T relies on the state of this or Other
+		// We delete last, because we don't want odd side effects if the destructor of T relies on the state of this
 		T* OldPtr = Ptr;
 		Ptr = InPtr;
 		delete OldPtr;
@@ -221,6 +252,42 @@ FORCEINLINE bool operator!=(const TUniquePtr<T>& Lhs, const TUniquePtr<T>& Rhs)
 	return Lhs.Get() != Rhs.Get();
 }
 
+/**
+ * Equality comparison operator against nullptr.
+ *
+ * @param Lhs The TUniquePtr to compare.
+ *
+ * @return true if the TUniquePtr is null, false otherwise.
+ */
+template <typename T>
+FORCEINLINE bool operator==(const TUniquePtr<T>& Lhs, TYPE_OF_NULLPTR)
+{
+	return !Lhs.IsValid();
+}
+template <typename T>
+FORCEINLINE bool operator==(TYPE_OF_NULLPTR, const TUniquePtr<T>& Rhs)
+{
+	return !Rhs.IsValid();
+}
+
+/**
+ * Inequality comparison operator against nullptr.
+ *
+ * @param Rhs The TUniquePtr to compare.
+ *
+ * @return true if the TUniquePtr is not null, false otherwise.
+ */
+template <typename T>
+FORCEINLINE bool operator!=(const TUniquePtr<T>& Lhs, TYPE_OF_NULLPTR)
+{
+	return Lhs.IsValid();
+}
+template <typename T>
+FORCEINLINE bool operator!=(TYPE_OF_NULLPTR, const TUniquePtr<T>& Rhs)
+{
+	return Rhs.IsValid();
+}
+
 // Trait which allows TUniquePtr to be default constructed by memsetting to zero.
 template <typename T>
 struct TIsZeroConstructType<TUniquePtr<T>>
