GitHub 1057 : Added support for variadic templates to TGraphTask

#1057

diff --git a/Engine/Source/Runtime/Core/Public/Async/TaskGraphInterfaces.h b/Engine/Source/Runtime/Core/Public/Async/TaskGraphInterfaces.h
--- a/Engine/Source/Runtime/Core/Public/Async/TaskGraphInterfaces.h
+++ b/Engine/Source/Runtime/Core/Public/Async/TaskGraphInterfaces.h
@@ -501,7 +501,7 @@ public:
 	{
 		// The arguments are useful for setting up other tasks. 
 		// Do work here, probably using SomeArgument.
-		MyCompletionEvent->DontCompleteUntil(TGraphTask<FSomeChildTask>::CreateTask(NULL,CurrentThread).ConstructAndDispatchWhenReady());
+		MyCompletionGraphEvent->DontCompleteUntil(TGraphTask<FSomeChildTask>::CreateTask(NULL,CurrentThread).ConstructAndDispatchWhenReady());
 	}
 };
 **/
@@ -522,6 +522,15 @@ public:
 	class FConstructor
 	{
 	public:
+	#if PLATFORM_COMPILER_HAS_VARIADIC_TEMPLATES
+		/** Passthrough internal task constructor and dispatch. Note! Generally speaking references will not pass through; use pointers */
+		template<typename...T>
+		FGraphEventRef ConstructAndDispatchWhenReady(T&&... Args)
+		{
+			new ((void *)&Owner->TaskStorage) TTask(Forward<T>(Args)...);
+			return Owner->Setup(Prerequisites, CurrentThreadIfKnown);
+		}
+	#else
 		/** Passthrough internal task constructor and dispatch. */
 		FGraphEventRef ConstructAndDispatchWhenReady()
 		{
@@ -581,8 +590,18 @@ public:
 		{
 			new ((void *)&Owner->TaskStorage) TTask(Forward<T1>(Arg1), Forward<T2>(Arg2), Forward<T3>(Arg3), Forward<T4>(Arg4), Forward<T5>(Arg5), Forward<T6>(Arg6), Forward<T7>(Arg7), Forward<T8>(Arg8), Forward<T9>(Arg9));
 			return Owner->Setup(Prerequisites, CurrentThreadIfKnown);
-		}
+		}		
+	#endif
 
+	#if PLATFORM_COMPILER_HAS_VARIADIC_TEMPLATES
+		/** Passthrough internal task constructor and hold. */
+		template<typename...T>
+		FGraphEventRef ConstructAndHold(T&&... Args)
+		{
+			new ((void *)&Owner->TaskStorage) TTask(Forward<T>(Args)...);
+			return Owner->Hold(Prerequisites, CurrentThreadIfKnown);
+		}
+	#else
 		/** Passthrough internal task constructor and hold. */
 		TGraphTask* ConstructAndHold()
 		{
@@ -643,6 +662,7 @@ public:
 			new ((void *)&Owner->TaskStorage) TTask(Forward<T1>(Arg1), Forward<T2>(Arg2), Forward<T3>(Arg3), Forward<T4>(Arg4), Forward<T5>(Arg5), Forward<T6>(Arg6), Forward<T7>(Arg7), Forward<T8>(Arg8), Forward<T9>(Arg9));
 			return Owner->Hold(Prerequisites, CurrentThreadIfKnown);
 		}
+	#endif
 
 	private:
 		friend class TGraphTask;
