Return type fix.

diff --git a/Engine/Source/Runtime/Core/Public/Async/TaskGraphInterfaces.h b/Engine/Source/Runtime/Core/Public/Async/TaskGraphInterfaces.h
--- a/Engine/Source/Runtime/Core/Public/Async/TaskGraphInterfaces.h
+++ b/Engine/Source/Runtime/Core/Public/Async/TaskGraphInterfaces.h
@@ -596,7 +596,7 @@ public:
 	#if PLATFORM_COMPILER_HAS_VARIADIC_TEMPLATES
 		/** Passthrough internal task constructor and hold. */
 		template<typename...T>
-		FGraphEventRef ConstructAndHold(T&&... Args)
+		TGraphTask* ConstructAndHold(T&&... Args)
 		{
 			new ((void *)&Owner->TaskStorage) TTask(Forward<T>(Args)...);
 			return Owner->Hold(Prerequisites, CurrentThreadIfKnown);
