Massive cleanup of text localization code and its uses. REVIEW API CHANGES FOR RELEASE NOTES.

diff --git a/Engine/Source/Editor/Kismet/Private/FindInBlueprintManager.cpp b/Engine/Source/Editor/Kismet/Private/FindInBlueprintManager.cpp
--- a/Engine/Source/Editor/Kismet/Private/FindInBlueprintManager.cpp
+++ b/Engine/Source/Editor/Kismet/Private/FindInBlueprintManager.cpp
@@ -62,7 +62,7 @@ namespace FiBSerializationHelpers
 	const FString Serialize(Type& InValue, bool bInIncludeSize)
 	{
 		TArray<uint8> SerializedData;
-		FMemoryWriter Ar(SerializedData, /*bIsPersistent=*/ true);
+		FMemoryWriter Ar(SerializedData);
 
 		Ar << InValue;
 		Ar.Close();
@@ -72,7 +72,7 @@ namespace FiBSerializationHelpers
 		if(bInIncludeSize)
 		{
 			SerializedData.Empty();
-			FMemoryWriter ArWithLength(SerializedData, /*bIsPersistent=*/ true);
+			FMemoryWriter ArWithLength(SerializedData);
 			int32 Length = Result.Len();
 			ArWithLength << Length;
 
@@ -102,7 +102,7 @@ namespace FiBSerializationHelpers
 	{
 		TArray<uint8> DerivedData;
 		DecodeFromStream(InStream, sizeof(Type), DerivedData);
-		FMemoryReader SizeOfDataAr(DerivedData, /*bIsPersistent=*/ true);
+		FMemoryReader SizeOfDataAr(DerivedData);
 
 		Type ReturnValue;
 		SizeOfDataAr << ReturnValue;
@@ -115,7 +115,7 @@ namespace FiBSerializationHelpers
 	{
 		TArray<uint8> DerivedData;
 		DecodeFromStream(InStream, InBytes, DerivedData);
-		FMemoryReader SizeOfDataAr(DerivedData, /*bIsPersistent=*/ true);
+		FMemoryReader SizeOfDataAr(DerivedData);
 
 		Type ReturnValue;
 		SizeOfDataAr << ReturnValue;
@@ -1370,7 +1370,7 @@ FText FFindInBlueprintSearchManager::ConvertHexStringToFText(FString InHexString
 
 	FText ResultText;
 
-	FMemoryReader Ar(SerializedData, /*bIsPersistent=*/ true);
+	FMemoryReader Ar(SerializedData);
 	Ar << ResultText;
 	Ar.Close();
 
@@ -1380,7 +1380,7 @@ FText FFindInBlueprintSearchManager::ConvertHexStringToFText(FString InHexString
 FString FFindInBlueprintSearchManager::ConvertFTextToHexString(FText InValue)
 {
 	TArray<uint8> SerializedData;
-	FMemoryWriter Ar(SerializedData, /*bIsPersistent=*/ true);
+	FMemoryWriter Ar(SerializedData);
 
 	Ar << InValue;
 	Ar.Close();
diff --git a/Engine/Source/Editor/KismetCompiler/Private/KismetCompilerVMBackend.cpp b/Engine/Source/Editor/KismetCompiler/Private/KismetCompilerVMBackend.cpp
--- a/Engine/Source/Editor/KismetCompiler/Private/KismetCompilerVMBackend.cpp
+++ b/Engine/Source/Editor/KismetCompiler/Private/KismetCompilerVMBackend.cpp
@@ -1,4 +1,4 @@
-// Copyright 1998-2015 Epic Games, Inc. All Rights Reserved.
+﻿// Copyright 1998-2015 Epic Games, Inc. All Rights Reserved.
 
 /*=============================================================================
 	KismetCompilerVMBackend.cpp
@@ -376,8 +376,8 @@ public:
 				Writer << EX_TextConst;
 				
 				EmitStringLiteral(FTextInspector::GetSourceString(Term->TextLiteral)? *FTextInspector::GetSourceString(Term->TextLiteral) : TEXT(""));
-				EmitStringLiteral(FTextInspector::GetKey(Term->TextLiteral)? *FTextInspector::GetKey(Term->TextLiteral) : TEXT(""));
-				EmitStringLiteral(FTextInspector::GetNamespace(Term->TextLiteral)? *FTextInspector::GetNamespace(Term->TextLiteral) : TEXT(""));
+				EmitStringLiteral(FTextInspector::GetKey(Term->TextLiteral).Get(TEXT("")));
+				EmitStringLiteral(FTextInspector::GetNamespace(Term->TextLiteral).Get(TEXT("")));
 			}
 			else if (CoerceProperty->IsA(UFloatProperty::StaticClass()))
 			{
diff --git a/Engine/Source/Editor/PropertyEditor/Private/PropertyCustomizationHelpers.cpp b/Engine/Source/Editor/PropertyEditor/Private/PropertyCustomizationHelpers.cpp
--- a/Engine/Source/Editor/PropertyEditor/Private/PropertyCustomizationHelpers.cpp
+++ b/Engine/Source/Editor/PropertyEditor/Private/PropertyCustomizationHelpers.cpp
@@ -16,6 +16,7 @@
 #include "SPropertyEditorClass.h"
 #include "SPropertyEditorInteractiveActorPicker.h"
 #include "SHyperlink.h"
+#include "SWidgetSwitcher.h"
 
 #define LOCTEXT_NAMESPACE "PropertyCustomizationHelpers"
 
@@ -1054,162 +1055,180 @@ void FMaterialList::AddMaterialItem( FDetailWidgetRow& Row, int32 CurrentSlot, c
 
 TSharedRef<SWidget> PropertyCustomizationHelpers::MakeTextLocalizationButton(const TSharedRef<IPropertyHandle>& InPropertyHandle)
 {
-	const auto& HandleGenerateKeyClicked = [=]() -> FReply 
+	class STextPropertyLocalizationMenuContent : public SCompoundWidget
 	{
-		FText DisplayText;
-		if (InPropertyHandle->GetValueAsDisplayText(DisplayText) == FPropertyAccess::Success)
-		{
-			TSharedRef< FString, ESPMode::ThreadSafe > DisplayString = FTextInspector::GetSharedDisplayString(DisplayText);
-			TSharedPtr< FString, ESPMode::ThreadSafe > Namespace;
-			TSharedPtr< FString, ESPMode::ThreadSafe > Key;
-			FTextLocalizationManager::Get().FindKeyNamespaceFromDisplayString(DisplayString, Namespace, Key);
-
-			if (Key.IsValid())
-			{
-				InPropertyHandle->NotifyPreChange();
-				*Key.Get() = FGuid::NewGuid().ToString();
-				InPropertyHandle->NotifyPostChange();
-			}
-		}
-
-		return FReply::Unhandled();
-	};
+		SLATE_BEGIN_ARGS(STextPropertyLocalizationMenuContent) {}
+		SLATE_END_ARGS()
 
-	const auto& HandleKeyTextCommitted = [=](const FText& InText, ETextCommit::Type InCommitType)
-	{
-		FText DisplayText;
-		if (InPropertyHandle->GetValueAsDisplayText(DisplayText) == FPropertyAccess::Success)
+	public:
+		STextPropertyLocalizationMenuContent()
+			: HasNamespaceAndKey(false)
 		{
-			TSharedRef< FString, ESPMode::ThreadSafe > DisplayString = FTextInspector::GetSharedDisplayString(DisplayText);
-			TSharedPtr< FString, ESPMode::ThreadSafe > Namespace;
-			TSharedPtr< FString, ESPMode::ThreadSafe > Key;
-			FTextLocalizationManager::Get().FindKeyNamespaceFromDisplayString(DisplayString, Namespace, Key);
-
-			if (Key.IsValid())
-			{
-				InPropertyHandle->NotifyPreChange();
-				*Key.Get() = InText.ToString();
-				InPropertyHandle->NotifyPostChange();
-			}
 		}
-	};
 
-	const auto& HandleNamespaceTextCommitted = [=](const FText& InText, ETextCommit::Type InCommitType)
-	{
-		FText DisplayText;
-		if (InPropertyHandle->GetValueAsDisplayText(DisplayText) == FPropertyAccess::Success)
+		void Construct(const FArguments& InArgs, const TSharedRef<IPropertyHandle>& InPropertyHandle)
 		{
-			TSharedRef< FString, ESPMode::ThreadSafe > DisplayString = FTextInspector::GetSharedDisplayString(DisplayText);
-			TSharedPtr< FString, ESPMode::ThreadSafe > Namespace;
-			TSharedPtr< FString, ESPMode::ThreadSafe > Key;
-			FTextLocalizationManager::Get().FindKeyNamespaceFromDisplayString(DisplayString, Namespace, Key);
+			PropertyHandle = InPropertyHandle;
 
-			if (Namespace.IsValid())
+			FText DisplayText;
+			if (PropertyHandle->GetValueAsDisplayText(DisplayText) == FPropertyAccess::Success)
 			{
-				InPropertyHandle->NotifyPreChange();
-				*Namespace.Get() = InText.ToString();
-				InPropertyHandle->NotifyPostChange();
+				const FTextDisplayStringRef DisplayString = FTextInspector::GetSharedDisplayString(DisplayText);
+				CacheNamespaceAndKey(DisplayString);
 			}
-		}
-	};
 
-	const auto& GetLocalizationMenuContent = [=]() -> TSharedRef<SWidget>
-	{
-		FMenuBuilder MenuContentBuilder(true, NULL);
-		{
-			FText DisplayText;
-			if (InPropertyHandle->GetValueAsDisplayText(DisplayText) == FPropertyAccess::Success)
+			FMenuBuilder MenuContentBuilder(true, NULL);
 			{
-				TSharedRef< FString, ESPMode::ThreadSafe > DisplayString = FTextInspector::GetSharedDisplayString(DisplayText);
-				TSharedPtr< FString, ESPMode::ThreadSafe > Namespace;
-				TSharedPtr< FString, ESPMode::ThreadSafe > Key;
-				FTextLocalizationManager::Get().FindKeyNamespaceFromDisplayString(DisplayString, Namespace, Key);
-
 				MenuContentBuilder.BeginSection(TEXT("Localization"), LOCTEXT("LocalizationSectionHeading", "Localization"));
 				{
-					if (!Namespace.IsValid() && !Key.IsValid())
-					{
-						MenuContentBuilder.AddWidget(SNew(SSpacer), LOCTEXT("NoLocWarning", "No localization available"));
-					}
-					else
-					{
-						TSharedRef<SGridPanel> GridPanel = SNew(SGridPanel);
+					TSharedPtr<SGridPanel> GridPanel;
+					TSharedRef<SWidgetSwitcher> WidgetSwitcher = SNew(SWidgetSwitcher)
+						.WidgetIndex_Lambda( [this](){return HasNamespaceAndKey ? 0 : 1;} )
+						+SWidgetSwitcher::Slot()
+						[
+							SAssignNew(GridPanel, SGridPanel)
+						]
+						+SWidgetSwitcher::Slot()
+						[
+							SNew(STextBlock)
+							.Text(LOCTEXT("NoLocWarning", "No localization information available."))
+						];
 
-						if (Key.IsValid())
-						{
-							GridPanel->AddSlot(0, 0)
-								.VAlign(VAlign_Center)
-								.Padding(1.0f)
-								[
-									SNew(STextBlock)
-									.Text(LOCTEXT("KeyLabel", "Key"))
-									.ToolTipText(LOCTEXT("KeyTooltip", "The localization key of the text property"))
-								];
+					GridPanel->AddSlot(0, 0)
+						.VAlign(VAlign_Center)
+						.Padding(1.0f)
+						[
+							SNew(STextBlock)
+							.Text(LOCTEXT("KeyLabel", "Key"))
+							.ToolTipText(LOCTEXT("KeyTooltip", "The localization key of the text property."))
+						];
 
-							auto KeyLambda = [=](){ return FText::FromString(*Key.Get()); };
+					const auto& GetKeyAsText = [this](){ return KeyAsText; };
 
-							GridPanel->AddSlot(1, 0)
-								.VAlign(VAlign_Center)
-								.Padding(1.0f)
-								[
-									SNew(SHorizontalBox)
-									+SHorizontalBox::Slot()
-									.FillWidth(1.0f)
-									[
-										SNew(SEditableTextBox)
-										.Text_Lambda(KeyLambda)
-										.OnTextCommitted(FOnTextCommitted::CreateLambda(HandleKeyTextCommitted))
-									]
-									+ SHorizontalBox::Slot()
-									.AutoWidth()
-									.HAlign(HAlign_Center)
-									.Padding(2.0f, 0.0f, 0.0f, 0.0f)
-									[
-										SNew(SButton)
-										.ToolTipText(LOCTEXT("RefreshKeyTooltip", "Generate a new random key"))
-										.ButtonStyle(FEditorStyle::Get(), "HoverHintOnly")
-										.OnClicked(FOnClicked::CreateLambda(HandleGenerateKeyClicked))
-										.Content()
-										[
-											SNew(SImage)
-											.Image(FEditorStyle::GetBrush("PropertyWindow.Button_Refresh"))
-										]
-									]
-								];
-						}
-
-						if (Namespace.IsValid())
-						{
-							auto NamespaceLambda = [=](){ return FText::FromString(*Namespace.Get()); };
-
-							GridPanel->AddSlot(0, 1)
-								.VAlign(VAlign_Center)
-								.Padding(1.0f, 1.0f, 5.0f, 1.0f)
-								[
-									SNew(STextBlock)
-									.Text(LOCTEXT("NamespaceLabel", "Namespace"))
-									.ToolTipText(LOCTEXT("NamespaceTooltip", "The localization namespace of the text property"))
-								];
-
-							GridPanel->AddSlot(1, 1)
-								.VAlign(VAlign_Center)
-								.Padding(1.0f)
+					GridPanel->AddSlot(1, 0)
+						.VAlign(VAlign_Center)
+						.Padding(1.0f)
+						[
+							SNew(SHorizontalBox)
+							+SHorizontalBox::Slot()
+							.FillWidth(1.0f)
+							[
+								SNew(SEditableTextBox)
+								.Text_Lambda(GetKeyAsText)
+								.OnTextCommitted(FOnTextCommitted::CreateSP(this, &STextPropertyLocalizationMenuContent::HandleKeyTextCommitted))
+							]
+							+ SHorizontalBox::Slot()
+							.AutoWidth()
+							.HAlign(HAlign_Center)
+							.Padding(2.0f, 0.0f, 0.0f, 0.0f)
+							[
+								SNew(SButton)
+								.ToolTipText(LOCTEXT("RefreshKeyTooltip", "Generate a new random key."))
+								.ButtonStyle(FEditorStyle::Get(), "HoverHintOnly")
+								.OnClicked(FOnClicked::CreateSP(this, &STextPropertyLocalizationMenuContent::HandleGenerateKeyClicked))
+								.Content()
 								[
-									SNew(SEditableTextBox)
-									.Text_Lambda(NamespaceLambda)
-									.OnTextCommitted(FOnTextCommitted::CreateLambda(HandleNamespaceTextCommitted))
-								];
-						}
+									SNew(SImage)
+									.Image(FEditorStyle::GetBrush("PropertyWindow.Button_Refresh"))
+								]
+							]
+						];
 
-						MenuContentBuilder.AddWidget(GridPanel, FText::GetEmpty());
-					}
-				}
+					GridPanel->AddSlot(0, 1)
+						.VAlign(VAlign_Center)
+						.Padding(1.0f, 1.0f, 5.0f, 1.0f)
+						[
+							SNew(STextBlock)
+							.Text(LOCTEXT("NamespaceLabel", "Namespace"))
+							.ToolTipText(LOCTEXT("NamespaceTooltip", "The localization namespace of the text property."))
+						];
+
+					const auto& GetNamespaceAsText = [this](){ return NamespaceAsText; };
+
+					GridPanel->AddSlot(1, 1)
+						.VAlign(VAlign_Center)
+						.Padding(1.0f)
+						[
+							SNew(SEditableTextBox)
+							.Text_Lambda(GetNamespaceAsText)
+							.OnTextCommitted(FOnTextCommitted::CreateSP(this, &STextPropertyLocalizationMenuContent::HandleNamespaceTextCommitted))
+						];
 
+					MenuContentBuilder.AddWidget(WidgetSwitcher, FText::GetEmpty());
+				}
 				MenuContentBuilder.EndSection();
 			}
+
+			ChildSlot
+				[
+					MenuContentBuilder.MakeWidget()
+				];
+		}
+
+	private:
+		FReply HandleGenerateKeyClicked() 
+		{
+			FText DisplayText;
+			if (PropertyHandle->GetValueAsDisplayText(DisplayText) == FPropertyAccess::Success)
+			{
+				const FTextDisplayStringRef DisplayString = FTextInspector::GetSharedDisplayString(DisplayText);
+				CacheNamespaceAndKey(DisplayString);
+				PropertyHandle->NotifyPreChange();
+				FTextLocalizationManager::Get().UpdateDisplayString(DisplayString, *DisplayString, CachedNamespace, FGuid::NewGuid().ToString());
+				PropertyHandle->NotifyPostChange();
+				CacheNamespaceAndKey(DisplayString);
+			}
+
+			return FReply::Handled();
+		}
+
+		void HandleKeyTextCommitted(const FText& NewKeyAsText, ETextCommit::Type InCommitType)
+		{
+			FText DisplayText;
+			if (PropertyHandle->GetValueAsDisplayText(DisplayText) == FPropertyAccess::Success)
+			{
+				const FTextDisplayStringRef DisplayString = FTextInspector::GetSharedDisplayString(DisplayText);
+				CacheNamespaceAndKey(DisplayString);
+				PropertyHandle->NotifyPreChange();
+				FTextLocalizationManager::Get().UpdateDisplayString(DisplayString, *DisplayString, CachedNamespace, NewKeyAsText.ToString());
+				PropertyHandle->NotifyPostChange();
+				CacheNamespaceAndKey(DisplayString);
+			}
 		}
 
-		return MenuContentBuilder.MakeWidget();
+		void HandleNamespaceTextCommitted(const FText& NewNamespaceAsText, ETextCommit::Type InCommitType)
+		{
+			FText DisplayText;
+			if (PropertyHandle->GetValueAsDisplayText(DisplayText) == FPropertyAccess::Success)
+			{
+				const FTextDisplayStringRef DisplayString = FTextInspector::GetSharedDisplayString(DisplayText);
+				CacheNamespaceAndKey(DisplayString);
+				PropertyHandle->NotifyPreChange();
+				FTextLocalizationManager::Get().UpdateDisplayString(DisplayString, *DisplayString, NewNamespaceAsText.ToString(), CachedKey);
+				PropertyHandle->NotifyPostChange();
+				CacheNamespaceAndKey(DisplayString);
+			}
+		}
+
+		void CacheNamespaceAndKey(const FTextDisplayStringRef& DisplayString)
+		{
+			HasNamespaceAndKey = FTextLocalizationManager::Get().FindNamespaceAndKeyFromDisplayString(DisplayString, CachedNamespace, CachedKey);
+			NamespaceAsText = FText::FromString(CachedNamespace);
+			KeyAsText = FText::FromString(CachedKey);
+		}
+
+	private:
+		TSharedPtr<IPropertyHandle> PropertyHandle;
+		bool HasNamespaceAndKey;
+		FText NamespaceAsText;
+		FText KeyAsText;
+		FString CachedNamespace;
+		FString CachedKey;
+	};
+
+	const auto& GetLocalizationMenuContent = [=]() -> TSharedRef<SWidget>
+	{
+		return SNew(STextPropertyLocalizationMenuContent, InPropertyHandle);
 	};
 
 	return SNew(SComboButton)
diff --git a/Engine/Source/Editor/TranslationEditor/Private/TranslationDataManager.cpp b/Engine/Source/Editor/TranslationEditor/Private/TranslationDataManager.cpp
--- a/Engine/Source/Editor/TranslationEditor/Private/TranslationDataManager.cpp
+++ b/Engine/Source/Editor/TranslationEditor/Private/TranslationDataManager.cpp
@@ -571,7 +571,7 @@ void FTranslationDataManager::PreviewAllTranslationsInEditor()
 	FJsonInternationalizationArchiveSerializer LocalizationArchiveSerializer;
 	FJsonInternationalizationManifestSerializer LocalizationManifestSerializer;
 
-	FTextLocalizationManager::Get().RegenerateResources(ConfigFilePath, LocalizationArchiveSerializer, LocalizationManifestSerializer);
+	FTextLocalizationManager::Get().LoadFromManifestAndArchives(ConfigFilePath, LocalizationArchiveSerializer, LocalizationManifestSerializer);
 }
 
 void FTranslationDataManager::PopulateSearchResultsUsingFilter(const FString& SearchFilter)
diff --git a/Engine/Source/Editor/TranslationEditor/Private/TranslationPickerEditWindow.cpp b/Engine/Source/Editor/TranslationEditor/Private/TranslationPickerEditWindow.cpp
--- a/Engine/Source/Editor/TranslationEditor/Private/TranslationPickerEditWindow.cpp
+++ b/Engine/Source/Editor/TranslationEditor/Private/TranslationPickerEditWindow.cpp
@@ -179,8 +179,8 @@ void STranslationPickerEditWidget::Construct(const FArguments& InArgs)
 	bool bShouldGatherForLocalization = FTextInspector::ShouldGatherForLocalization(PickedText);
 
 	// Get all the data we need and format it properly
-	const FString* NamespaceString = FTextInspector::GetNamespace(PickedText);
-	const FString* KeyString = FTextInspector::GetKey(PickedText);
+	TOptional<FString> NamespaceString = FTextInspector::GetNamespace(PickedText);
+	TOptional<FString> KeyString = FTextInspector::GetKey(PickedText);
 	const FString* SourceString = FTextInspector::GetSourceString(PickedText);
 	const FString& TranslationString = FTextInspector::GetDisplayString(PickedText);
 	FString LocresFullPath;
@@ -188,16 +188,16 @@ void STranslationPickerEditWidget::Construct(const FArguments& InArgs)
 	FString ManifestAndArchiveNameString;
 	if (NamespaceString && KeyString)
 	{
-		TSharedPtr<FString, ESPMode::ThreadSafe> TextTableName = FTextLocalizationManager::Get().GetTableName(*NamespaceString, *KeyString);
-		LocresFullPath = *TextTableName;
-		if (TextTableName.IsValid())
+		FString LocResId;
+		if (FTextLocalizationManager::Get().GetLocResID(NamespaceString.GetValue(), KeyString.GetValue(), LocResId))
 		{
-			ManifestAndArchiveNameString = FPaths::GetBaseFilename(*TextTableName);
+			LocresFullPath = *LocResId;
+			ManifestAndArchiveNameString = FPaths::GetBaseFilename(*LocResId);
 		}
 	}
 
-	FText Namespace = NamespaceString != nullptr ? FText::FromString(*NamespaceString) : FText::GetEmpty();
-	FText Key = KeyString != nullptr ? FText::FromString(*KeyString) : FText::GetEmpty();
+	FText Namespace = FText::FromString(NamespaceString.Get(TEXT("")));
+	FText Key = FText::FromString(KeyString.Get(TEXT("")));
 	FText Source = SourceString != nullptr ? FText::FromString(*SourceString) : FText::GetEmpty();
 	FText ManifestAndArchiveName = FText::FromString(ManifestAndArchiveNameString);
 	FText Translation = FText::FromString(TranslationString);
@@ -209,13 +209,13 @@ void STranslationPickerEditWidget::Construct(const FArguments& InArgs)
 
 	// Save the necessary data in UTranslationUnit for later.  This is what we pass to TranslationDataManager to save our edits
 	TranslationUnit = NewObject<UTranslationUnit>();
-	TranslationUnit->Namespace = NamespaceString != nullptr ? *NamespaceString : "";
+	TranslationUnit->Namespace = NamespaceString.Get(TEXT(""));
 	TranslationUnit->Source = SourceString != nullptr ? *SourceString : "";
 	TranslationUnit->Translation = TranslationString;
 	TranslationUnit->LocresPath = LocresFullPath;
 	
 	// Can only save if we have all the required information
-	bool bHasRequiredLocalizationInfo = NamespaceString != nullptr && SourceString != nullptr && LocresFullPath.Len() > 0;
+	bool bHasRequiredLocalizationInfo = NamespaceString.IsSet() && SourceString != nullptr && LocresFullPath.Len() > 0;
 
 	TSharedRef<SHorizontalBox> LocalizationInfoAndSaveButtonSlot = SNew(SHorizontalBox);
 
diff --git a/Engine/Source/Editor/UnrealEd/Private/Commandlets/GatherTextFromAssetsCommandlet.cpp b/Engine/Source/Editor/UnrealEd/Private/Commandlets/GatherTextFromAssetsCommandlet.cpp
--- a/Engine/Source/Editor/UnrealEd/Private/Commandlets/GatherTextFromAssetsCommandlet.cpp
+++ b/Engine/Source/Editor/UnrealEd/Private/Commandlets/GatherTextFromAssetsCommandlet.cpp
@@ -566,25 +566,27 @@ void UGatherTextFromAssetsCommandlet::ProcessTextProperty(UTextProperty* InTextP
 	NewEntry.SourceString = Data->GetSourceString();
 	NewEntry.Status = EAssetTextGatherStatus::None;
 
-	TSharedPtr< FString, ESPMode::ThreadSafe > Namespace;
-	TSharedPtr< FString, ESPMode::ThreadSafe > Key;
-	FTextLocalizationManager::Get().FindKeyNamespaceFromDisplayString(Data->DisplayString, Namespace, Key);
+	FString Namespace;
+	FString Key;
+	const bool FoundNamespaceAndKey = FTextLocalizationManager::Get().FindNamespaceAndKeyFromDisplayString(Data->DisplayString, Namespace, Key);
 
 	// Check if text is localizable and check for missing key. Unlocalizable texts don't need a key.
-	if( ( !Key.IsValid() || Key->IsEmpty() ) && Data->ShouldGatherForLocalization() )
+	if( ( !FoundNamespaceAndKey || Key.IsEmpty() ) && Data->ShouldGatherForLocalization() )
 	{
 		NewEntry.Status = EAssetTextGatherStatus::MissingKey;
 
 		// Fix missing key if allowed.
 		if (bFixBroken)
 		{
-			// Create key if needed.
-			if( !( Key.IsValid() ) )
+			// Generate new GUID for key.
+			if (FoundNamespaceAndKey)
 			{
-				Key = MakeShareable( new FString() );
+				FTextLocalizationManager::Get().UpdateDisplayString(Data->DisplayString, *Data->DisplayString, Namespace, FGuid::NewGuid().ToString());
+			}
+			else
+			{
+				FTextLocalizationManager::Get().AddDisplayString(Data->DisplayString, TEXT(""), FGuid::NewGuid().ToString());
 			}
-			// Generate new GUID for key.
-			*(Key) = FGuid::NewGuid().ToString();
 
 			// Fixed.
 			NewEntry.Status = EAssetTextGatherStatus::MissingKey_Resolved;
@@ -592,17 +594,17 @@ void UGatherTextFromAssetsCommandlet::ProcessTextProperty(UTextProperty* InTextP
 	}
 
 	// Must have valid key to test for identity conflicts. Even if a text doesn't require localization, if it conflicts, it should be warned about.
-	if( Key.IsValid() && !( Key->IsEmpty() ) )
+	if( FoundNamespaceAndKey && !Key.IsEmpty() )
 	{
 		// Find existing entry from manifest or manifest dependencies.
 		FContext SearchContext;
-		SearchContext.Key = *Key;
-		TSharedPtr< FManifestEntry > ExistingEntry = ManifestInfo->GetManifest()->FindEntryByContext( Namespace.IsValid() ? *Namespace : TEXT(""), SearchContext );
+		SearchContext.Key = Key;
+		TSharedPtr< FManifestEntry > ExistingEntry = ManifestInfo->GetManifest()->FindEntryByContext( Namespace, SearchContext );
 
 		if( !ExistingEntry.IsValid() )
 		{
 			FString FileInfo;
-			ExistingEntry = ManifestInfo->FindDependencyEntrybyContext( Namespace.IsValid() ? *Namespace : TEXT(""), SearchContext, FileInfo );
+			ExistingEntry = ManifestInfo->FindDependencyEntrybyContext( Namespace, SearchContext, FileInfo );
 		}
 
 		// Entry already exists, check for conflict.
@@ -614,7 +616,7 @@ void UGatherTextFromAssetsCommandlet::ProcessTextProperty(UTextProperty* InTextP
 			if (bFixBroken)
 			{
 				// Generate new GUID for key.
-				*(Key) = FGuid::NewGuid().ToString();
+				FTextLocalizationManager::Get().UpdateDisplayString(Data->DisplayString, *Data->DisplayString, Namespace, FGuid::NewGuid().ToString());
 
 				// Fixed.
 				NewEntry.Status = EAssetTextGatherStatus::IdentityConflict_Resolved;
@@ -641,11 +643,11 @@ void UGatherTextFromAssetsCommandlet::ProcessTextProperty(UTextProperty* InTextP
 				{
 					UTextProperty* TextProp	=  Cast<UTextProperty>( *(PropIt) );
 					FText* DataCDO = TextProp->ContainerPtrToValuePtr<FText>( CDO );
-					TSharedPtr< FString, ESPMode::ThreadSafe > NamespaceCDO;
-					TSharedPtr< FString, ESPMode::ThreadSafe > KeyCDO;
-					FTextLocalizationManager::Get().FindKeyNamespaceFromDisplayString(DataCDO->DisplayString, NamespaceCDO, KeyCDO);
+					FString NamespaceCDO;
+					FString KeyCDO;
+					const bool FoundCDONamespaceAndKey = FTextLocalizationManager::Get().FindNamespaceAndKeyFromDisplayString(DataCDO->DisplayString, NamespaceCDO, KeyCDO);
 
-					if( KeyCDO == Key || ( KeyCDO.Get() && Key.Get() && ( *KeyCDO == *Key ) ) )
+					if( (!FoundCDONamespaceAndKey && !FoundNamespaceAndKey) || (FoundCDONamespaceAndKey && FoundNamespaceAndKey && ( KeyCDO == Key ) ) )
 					{
 						SrcLocation = CDO->GetPathName() + TEXT(".") + TextProp->GetName();
 						break;
@@ -655,16 +657,16 @@ void UGatherTextFromAssetsCommandlet::ProcessTextProperty(UTextProperty* InTextP
 		}
 
 		FContext Context;
-		Context.Key = *Key;
+		Context.Key = Key;
 		Context.SourceLocation = SrcLocation;
 
 		FString EntryDescription = FString::Printf( TEXT("In %s"), *Object->GetFullName());
-		ManifestInfo->AddEntry(EntryDescription, Namespace.Get() ? *Namespace : TEXT(""), NewEntry.SourceString.Get() ? *(NewEntry.SourceString) : TEXT(""), Context );
+		ManifestInfo->AddEntry(EntryDescription, Namespace, NewEntry.SourceString.Get() ? *(NewEntry.SourceString) : TEXT(""), Context );
 	}
 
 	// Add to conflict tracker.
-	FConflictTracker::FKeyTable& KeyTable = ConflictTracker.Namespaces.FindOrAdd( Namespace.IsValid() ? *Namespace : TEXT("") );
-	FConflictTracker::FEntryArray& EntryArray = KeyTable.FindOrAdd( Key.IsValid() ? *Key : TEXT("") );
+	FConflictTracker::FKeyTable& KeyTable = ConflictTracker.Namespaces.FindOrAdd(Namespace);
+	FConflictTracker::FEntryArray& EntryArray = KeyTable.FindOrAdd(Key);
 	EntryArray.Add(NewEntry);
 
 	OutFixed = (NewEntry.Status == EAssetTextGatherStatus::MissingKey_Resolved || NewEntry.Status == EAssetTextGatherStatus::IdentityConflict_Resolved);
diff --git a/Engine/Source/Editor/UnrealEd/Private/UnrealEdSrv.cpp b/Engine/Source/Editor/UnrealEd/Private/UnrealEdSrv.cpp
--- a/Engine/Source/Editor/UnrealEd/Private/UnrealEdSrv.cpp
+++ b/Engine/Source/Editor/UnrealEd/Private/UnrealEdSrv.cpp
@@ -800,7 +800,7 @@ bool UUnrealEdEngine::Exec( UWorld* InWorld, const TCHAR* Stream, FOutputDevice&
 			FJsonInternationalizationArchiveSerializer ArchiveSerializer;
 			FJsonInternationalizationManifestSerializer ManifestSerializer;
 
-			FTextLocalizationManager::Get().RegenerateResources(ConfigFilePath, ArchiveSerializer, ManifestSerializer);
+			FTextLocalizationManager::Get().LoadFromManifestAndArchives(ConfigFilePath, ArchiveSerializer, ManifestSerializer);
 		}
 	}
 #endif
@@ -3027,7 +3027,7 @@ bool UUnrealEdEngine::Exec_Mode( const TCHAR* Str, FOutputDevice& Ar )
 		FJsonInternationalizationArchiveSerializer ArchiveSerializer;
 		FJsonInternationalizationManifestSerializer ManifestSerializer;
 
-		FTextLocalizationManager::Get().RegenerateResources(ConfigFilePath, ArchiveSerializer, ManifestSerializer);
+		FTextLocalizationManager::Get().LoadFromManifestAndArchives(ConfigFilePath, ArchiveSerializer, ManifestSerializer);
 	}
 #endif
 
diff --git a/Engine/Source/Runtime/Core/Private/Internationalization/Text.cpp b/Engine/Source/Runtime/Core/Private/Internationalization/Text.cpp
--- a/Engine/Source/Runtime/Core/Private/Internationalization/Text.cpp
+++ b/Engine/Source/Runtime/Core/Private/Internationalization/Text.cpp
@@ -28,20 +28,20 @@ bool FTextInspector::ShouldGatherForLocalization(const FText& Text)
 	return Text.ShouldGatherForLocalization();
 }
 
-const FString* FTextInspector::GetNamespace(const FText& Text)
+TOptional<FString> FTextInspector::GetNamespace(const FText& Text)
 {
-	TSharedPtr< FString, ESPMode::ThreadSafe > Namespace;
-	TSharedPtr< FString, ESPMode::ThreadSafe > Key;
-	FTextLocalizationManager::Get().FindKeyNamespaceFromDisplayString(Text.DisplayString, Namespace, Key);
-	return Namespace.Get();
+	FString Namespace;
+	FString Key;
+	const bool WasNamespaceAndKeyFound = FTextLocalizationManager::Get().FindNamespaceAndKeyFromDisplayString(Text.DisplayString, Namespace, Key);
+	return WasNamespaceAndKeyFound ? TOptional<FString>(Namespace) : TOptional<FString>();
 }
 
-const FString* FTextInspector::GetKey(const FText& Text)
+TOptional<FString> FTextInspector::GetKey(const FText& Text)
 {
-	TSharedPtr< FString, ESPMode::ThreadSafe > Namespace;
-	TSharedPtr< FString, ESPMode::ThreadSafe > Key;
-	FTextLocalizationManager::Get().FindKeyNamespaceFromDisplayString(Text.DisplayString, Namespace, Key);
-	return Key.Get();
+	FString Namespace;
+	FString Key;
+	const bool WasNamespaceAndKeyFound = FTextLocalizationManager::Get().FindNamespaceAndKeyFromDisplayString(Text.DisplayString, Namespace, Key);
+	return WasNamespaceAndKeyFound ? TOptional<FString>(Key) : TOptional<FString>();
 }
 
 const FString* FTextInspector::GetSourceString(const FText& Text)
@@ -54,7 +54,7 @@ const FString& FTextInspector::GetDisplayString(const FText& Text)
 	return Text.DisplayString.Get();
 }
 
-const TSharedRef<FString, ESPMode::ThreadSafe> FTextInspector::GetSharedDisplayString(const FText& Text)
+const FTextDisplayStringRef FTextInspector::GetSharedDisplayString(const FText& Text)
 {
 	return Text.DisplayString;
 }
@@ -296,7 +296,7 @@ FText::FText( FString InSourceString )
 }
 
 FText::FText( FString InSourceString, FString InNamespace, FString InKey, int32 InFlags )
-	: DisplayString( FTextLocalizationManager::Get().GetString(InNamespace, InKey, &InSourceString) )
+	: DisplayString( FTextLocalizationManager::Get().GetDisplayString(InNamespace, InKey, &InSourceString) )
 	, Flags(InFlags)
 {
 	History = MakeShareable(new FTextHistory_Base(InSourceString));
@@ -530,7 +530,7 @@ FString FText::GetInvariantTimeZone()
 
 bool FText::FindText( const FString& Namespace, const FString& Key, FText& OutText, const FString* const SourceString )
 {
-	TSharedPtr< FString, ESPMode::ThreadSafe > FoundString = FTextLocalizationManager::Get().FindString( Namespace, Key );
+	TSharedPtr< FString, ESPMode::ThreadSafe > FoundString = FTextLocalizationManager::Get().FindDisplayString( Namespace, Key );
 
 	if ( FoundString.IsValid() )
 	{
@@ -557,19 +557,14 @@ CORE_API FArchive& operator<<( FArchive& Ar, FText& Value )
 		// Namespaces and keys are no longer stored in the FText, we need to read them in and discard
 		if( Ar.UE4Ver() >= VER_UE4_ADDED_NAMESPACE_AND_KEY_DATA_TO_FTEXT )
 		{
-			TSharedPtr< FString, ESPMode::ThreadSafe > Namespace;
-			TSharedPtr< FString, ESPMode::ThreadSafe > Key;
+			FString Namespace;
+			FString Key;
 
-			Namespace = MakeShareable( new FString() );
-			Ar << *Namespace;
-
-			Key = MakeShareable( new FString() );
-			Ar << *Key;
+			Ar << Namespace;
+			Ar << Key;
 
 			// Get the DisplayString using the namespace, key, and source string.
-			const FString& NamespaceStr = Namespace.Get() ? *Namespace : TEXT("");
-			const FString& KeyStr = Key.Get() ? *Key : TEXT("");
-			Value.DisplayString = FTextLocalizationManager::Get().GetString(NamespaceStr, KeyStr, &SourceStringToImplantIntoHistory);
+			Value.DisplayString = FTextLocalizationManager::Get().GetDisplayString(Namespace, Key, &SourceStringToImplantIntoHistory);
 		}
 	}
 
@@ -858,14 +853,15 @@ bool FTextSnapshot::IsDisplayStringEqualTo(const FText& InText) const
 }
 
 FScopedTextIdentityPreserver::FScopedTextIdentityPreserver(FText& InTextToPersist)
-	: Flags(InTextToPersist.Flags)
+	: HadFoundNamespaceAndKey(false)
+	, Flags(InTextToPersist.Flags)
 	, TextToPersist(InTextToPersist)
 {
 	// Empty display strings can't have a namespace or key.
 	if (GIsEditor && !InTextToPersist.DisplayString->IsEmpty())
 	{
 		// Save off namespace and key to be restored later.
-		FTextLocalizationManager::Get().FindKeyNamespaceFromDisplayString(InTextToPersist.DisplayString, Namespace, Key);
+		HadFoundNamespaceAndKey = FTextLocalizationManager::Get().FindNamespaceAndKeyFromDisplayString(InTextToPersist.DisplayString, Namespace, Key);
 	}
 }
 
@@ -874,7 +870,7 @@ FScopedTextIdentityPreserver::~FScopedTextIdentityPreserver()
 	// Never persist identities in non-editor situations
 	// If we don't have a key, then the old identity wasn't valid and shouldn't be preserved.
 	// Never persist identities for immutable (i.e. code LOCTEXT declared) text.
-	if(GIsEditor && Key.IsValid() && (Flags & ETextFlag::Immutable) == 0)
+	if(GIsEditor && HadFoundNamespaceAndKey && (Flags & ETextFlag::Immutable) == 0)
 	{
 		// Get the text's new source string.
 		const FString* SourceString = FTextInspector::GetSourceString(TextToPersist);
@@ -886,7 +882,7 @@ FScopedTextIdentityPreserver::~FScopedTextIdentityPreserver()
 		TextToPersist.History = MakeShareable(new FTextHistory_Base(*SourceString));
 
 		// Create/update the display string instance for this identity in the text localization manager...
-		const TSharedRef<FString, ESPMode::ThreadSafe> DisplayString = FTextLocalizationManager::Get().GetString(Namespace.IsValid() ? *Namespace : TEXT(""), *Key, SourceString);
+		const FTextDisplayStringRef DisplayString = FTextLocalizationManager::Get().GetDisplayString(Namespace, Key, SourceString);
 
 		// ... and set it so that the text's history's serialization will properly map this display string instance back to the identity we intended to preserve and restore.
 		TextToPersist.DisplayString = DisplayString;
diff --git a/Engine/Source/Runtime/Core/Private/Internationalization/TextHistory.cpp b/Engine/Source/Runtime/Core/Private/Internationalization/TextHistory.cpp
--- a/Engine/Source/Runtime/Core/Private/Internationalization/TextHistory.cpp
+++ b/Engine/Source/Runtime/Core/Private/Internationalization/TextHistory.cpp
@@ -3,6 +3,7 @@
 #include "CorePrivatePCH.h"
 
 #include "TextHistory.h"
+#include "PropertyPortflags.h"
 
 ///////////////////////////////////////
 // FTextHistory
@@ -10,14 +11,14 @@
 /** Base class for all FText history types */
 
 FTextHistory::FTextHistory()
-	: Revision(FTextLocalizationManager::Get().GetHeadCultureRevision())
+	: Revision(FTextLocalizationManager::Get().GetTextRevision())
 {
 
 }
 
 bool FTextHistory::IsOutOfDate()
 {
-	return Revision < FTextLocalizationManager::Get().GetHeadCultureRevision();
+	return Revision < FTextLocalizationManager::Get().GetTextRevision();
 }
 
 TSharedPtr< FString, ESPMode::ThreadSafe > FTextHistory::GetSourceString() const
@@ -31,7 +32,7 @@ void FTextHistory::GetSourceTextsFromFormatHistory(FText Text, TArray<FText>& Ou
 	OutSourceTexts.Add(Text);
 }
 
-void FTextHistory::SerializeForDisplayString(FArchive& Ar, TSharedRef<FString, ESPMode::ThreadSafe>& InOutDisplayString)
+void FTextHistory::SerializeForDisplayString(FArchive& Ar, FTextDisplayStringRef& InOutDisplayString)
 {
 	if(Ar.IsLoading())
 	{
@@ -51,7 +52,7 @@ void FTextHistory::Rebuild(TSharedRef< FString, ESPMode::ThreadSafe > InDisplayS
 
 	// FTextHistory_Base will never report being out-of-date, but we need to keep the history revision in sync
 	// with the head culture regardless so that FTextSnapshot::IdenticalTo still works correctly
-	Revision = FTextLocalizationManager::Get().GetHeadCultureRevision();
+	Revision = FTextLocalizationManager::Get().GetTextRevision();
 
 	if(bIsOutOfDate)
 	{
@@ -106,7 +107,7 @@ void FTextHistory_Base::Serialize( FArchive& Ar )
 	}
 }
 
-void FTextHistory_Base::SerializeForDisplayString(FArchive& Ar, TSharedRef<FString, ESPMode::ThreadSafe>& InOutDisplayString)
+void FTextHistory_Base::SerializeForDisplayString(FArchive& Ar, FTextDisplayStringRef& InOutDisplayString)
 {
 	if(Ar.IsLoading())
 	{
@@ -128,45 +129,33 @@ void FTextHistory_Base::SerializeForDisplayString(FArchive& Ar, TSharedRef<FStri
 		}
 
 		// Using the deserialized namespace and key, find the DisplayString.
-		InOutDisplayString = FTextLocalizationManager::Get().GetString(Namespace, Key, SourceString.Get());
+		InOutDisplayString = FTextLocalizationManager::Get().GetDisplayString(Namespace, Key, SourceString.Get());
 	}
 	else if(Ar.IsSaving())
 	{
-		TSharedPtr< FString, ESPMode::ThreadSafe > Namespace;
-		TSharedPtr< FString, ESPMode::ThreadSafe > Key;
-
-		FTextLocalizationManager::Get().FindKeyNamespaceFromDisplayString(InOutDisplayString, Namespace, Key);
+		FString Namespace;
+		FString Key;
 
-		// Serialize the Namespace, or an empty string if there is none
-		if( Namespace.IsValid() )
-		{
-			Ar << *(Namespace);
-		}
-		else
-		{
-			FString Empty;
-			Ar << Empty;
-		}
+		const bool FoundNamespaceAndKey = FTextLocalizationManager::Get().FindNamespaceAndKeyFromDisplayString(InOutDisplayString, Namespace, Key);
 
-		// Serialize the Key, if there is not one and this is in the editor, generate a new key.
-		if( Key.IsValid() )
+		// If this has no namespace or key, attempt to give it a GUID for a key and register it.
+		if (!FoundNamespaceAndKey && GIsEditor && (Ar.IsPersistent() && !Ar.HasAnyPortFlags(PPF_Duplicate)))
 		{
-			Ar << *(Key);
-		}
-		else
-		{
-			if(GIsEditor)
-			{
-				FString SerializedKey = FGuid::NewGuid().ToString();
-				Ar << SerializedKey;
-			}
-			else
+			Key = FGuid::NewGuid().ToString();
+			if (!FTextLocalizationManager::Get().AddDisplayString(InOutDisplayString, Namespace, Key))
 			{
-				FString Empty;
-				Ar << Empty;
+				// Could not add display string, reset namespace and key.
+				Namespace.Empty();
+				Key.Empty();
 			}
 		}
 
+		// Serialize the Namespace
+		Ar << Namespace;
+
+		// Serialize the Key
+		Ar << Key;
+
 		// Serialize the SourceString, or an empty string if there is none
 		if( SourceString.IsValid() )
 		{
diff --git a/Engine/Source/Runtime/Core/Private/Internationalization/TextHistory.h b/Engine/Source/Runtime/Core/Private/Internationalization/TextHistory.h
--- a/Engine/Source/Runtime/Core/Private/Internationalization/TextHistory.h
+++ b/Engine/Source/Runtime/Core/Private/Internationalization/TextHistory.h
@@ -34,7 +34,7 @@ public:
 	virtual void Serialize(FArchive& Ar) = 0;
 
 	/** Serializes data needed to get the FText's DisplayString */
-	virtual void SerializeForDisplayString(FArchive& Ar, TSharedRef<FString, ESPMode::ThreadSafe>& InOutDisplayString);
+	virtual void SerializeForDisplayString(FArchive& Ar, FTextDisplayStringRef& InOutDisplayString);
 
 	/** Returns TRUE if the Revision is out of date */
 	virtual bool IsOutOfDate();
@@ -65,14 +65,14 @@ public:
 	// Begin FTextHistory interface
 	virtual FText ToText(bool bInAsSource) const override;
 	virtual void Serialize(FArchive& Ar) override;
-	virtual void SerializeForDisplayString(FArchive& Ar, TSharedRef<FString, ESPMode::ThreadSafe>& InOutDisplayString) override;
+	virtual void SerializeForDisplayString(FArchive& Ar, FTextDisplayStringRef& InOutDisplayString) override;
 	virtual bool IsOutOfDate() override { return false; }
 	virtual TSharedPtr< FString, ESPMode::ThreadSafe > GetSourceString() const override;
 	// End FTextHistory interface
 
 private:
 	/** The source string for an FText */
-	TSharedPtr<FString, ESPMode::ThreadSafe> SourceString;
+	FTextDisplayStringPtr SourceString;
 };
 
 /** Handles history for FText::Format when passing named arguments */
diff --git a/Engine/Source/Runtime/Core/Private/Internationalization/TextLocalizationManager.cpp b/Engine/Source/Runtime/Core/Private/Internationalization/TextLocalizationManager.cpp
--- a/Engine/Source/Runtime/Core/Private/Internationalization/TextLocalizationManager.cpp
+++ b/Engine/Source/Runtime/Core/Private/Internationalization/TextLocalizationManager.cpp
@@ -12,8 +12,8 @@ static FString AccessedStringBeforeLocLoadedErrorMsg = TEXT("Can't access string
 
 void BeginInitTextLocalization()
 {
-	// Initialize FInternationalization before we bind to OnCultureChanged otherwise we can accidentally initialize
-	// twice since FInternationalization::Initialize sets the culture
+	// Initialize FInternationalization before we bind to OnCultureChanged, otherwise we can accidentally initialize
+	// twice since FInternationalization::Initialize sets the culture.
 	FInternationalization::Get();
 
 	FInternationalization::Get().OnCultureChanged().AddRaw( &(FTextLocalizationManager::Get()), &FTextLocalizationManager::OnCultureChanged );
@@ -119,105 +119,32 @@ void EndInitTextLocalization()
 			{
 				// Fallback to English.
 				UE_LOG(LogTextLocalizationManager, Log, TEXT("The requested culture ('%s') has no localization data; falling back to 'en' for localization and internationalization data."), *RequestedCultureName);
-				TargetCultureName = "en";
+				TargetCultureName = TEXT("en");
 			}
 		}
 
 		I18N.SetCurrentCulture(TargetCultureName);
 	}
 
-	FTextLocalizationManager::Get().LoadResources(ShouldLoadEditor, ShouldLoadGame);
+	FTextLocalizationManager::Get().LoadLocalizationResourcesForCulture(I18N.GetCurrentCulture()->GetName(), ShouldLoadEditor, ShouldLoadGame);
 	FTextLocalizationManager::Get().bIsInitialized = true;
 }
 
-FTextLocalizationManager& FTextLocalizationManager::Get()
-{
-	static FTextLocalizationManager* GTextLocalizationManager = NULL;
-	if( !GTextLocalizationManager )
-	{
-		GTextLocalizationManager = new FTextLocalizationManager;
-	}
-
-	return *GTextLocalizationManager;
-}
-
-const TSharedRef<FString, ESPMode::ThreadSafe>* FTextLocalizationManager::FTextLookupTable::GetString(const FString& Namespace, const FString& Key, const uint32 SourceStringHash) const
-{
-	const FKeyTable* const KeyTable = NamespaceTable.Find(Namespace);
-	if( KeyTable )
-	{
-		const FStringEntry* const Entry = KeyTable->Find(Key);
-		if( Entry )
-		{
-			return &(Entry->String);
-		}
-	}
-
-	return NULL;
-}
-
-void FTextLocalizationManager::FLocalizationEntryTracker::ReportCollisions() const
-{
-	for(auto i = Namespaces.CreateConstIterator(); i; ++i)
-	{
-		const FString& NamespaceName = i.Key();
-		const FLocalizationEntryTracker::FKeyTable& KeyTable = i.Value();
-		for(auto j = KeyTable.CreateConstIterator(); j; ++j)
-		{
-			const FString& KeyName = j.Key();
-			const FLocalizationEntryTracker::FEntryArray& EntryArray = j.Value();
-
-			bool bWasCollisionDetected = false;
-			for(int32 k = 0; k < EntryArray.Num(); ++k)
-			{
-				const FLocalizationEntryTracker::FEntry& LeftEntry = EntryArray[k];
-				for(int32 l = k + 1; l < EntryArray.Num(); ++l)
-				{
-					const FLocalizationEntryTracker::FEntry& RightEntry = EntryArray[l];
-					const bool bDoesSourceStringHashDiffer = LeftEntry.SourceStringHash != RightEntry.SourceStringHash;
-					const bool bDoesLocalizedStringDiffer = LeftEntry.LocalizedString != RightEntry.LocalizedString;
-					bWasCollisionDetected = bDoesSourceStringHashDiffer || bDoesLocalizedStringDiffer;
-				}
-			}
-
-			if(bWasCollisionDetected)
-			{
-				FString CollidingEntryListString;
-				for(int32 k = 0; k < EntryArray.Num(); ++k)
-				{
-					const FLocalizationEntryTracker::FEntry& Entry = EntryArray[k];
-
-					if( !(CollidingEntryListString.IsEmpty()) )
-					{
-						CollidingEntryListString += '\n';
-					}
-
-					CollidingEntryListString += FString::Printf( TEXT("Resource: (%s) Hash: (%d) String: (%s)"), *(Entry.TableName), Entry.SourceStringHash, *(Entry.LocalizedString) );
-				}
-
-				UE_LOG(LogTextLocalizationManager, Warning, TEXT("Loading localization resources contain conflicting entries for (Namespace:%s, Key:%s):\n%s"), *NamespaceName, *KeyName, *CollidingEntryListString);
-			}
-		}
-	}
-}
-
-void FTextLocalizationManager::FLocalizationEntryTracker::ReadFromDirectory(const FString& DirectoryPath)
+void FTextLocalizationManager::FLocalizationEntryTracker::LoadFromDirectory(const FString& DirectoryPath)
 {
 	// Find resources in the specified folder.
 	TArray<FString> ResourceFileNames;
 	IFileManager::Get().FindFiles(ResourceFileNames, *(DirectoryPath / TEXT("*.locres")), true, false);
 
-	// For each resource:
-	for(int32 ResourceIndex = 0; ResourceIndex < ResourceFileNames.Num(); ++ResourceIndex)
+	for (const FString& ResourceFileName : ResourceFileNames)
 	{
-		const FString ResourceFilePath = DirectoryPath / ResourceFileNames[ResourceIndex];
-		ReadFromFile( FPaths::ConvertRelativePathToFull(ResourceFilePath) );
+		LoadFromFile( FPaths::ConvertRelativePathToFull(DirectoryPath / ResourceFileName) );
 	}
 }
 
-bool FTextLocalizationManager::FLocalizationEntryTracker::ReadFromFile(const FString& FilePath)
+bool FTextLocalizationManager::FLocalizationEntryTracker::LoadFromFile(const FString& FilePath)
 {
-	FArchive* Reader = IFileManager::Get().CreateFileReader( *FilePath );
+	TScopedPointer<FArchive> Reader(IFileManager::Get().CreateFileReader( *FilePath ));
 	if( !Reader )
 	{
 		return false;
@@ -225,15 +152,14 @@ bool FTextLocalizationManager::FLocalizationEntryTracker::ReadFromFile(const FSt
 
 	Reader->SetForceUnicode(true);
 
-	ReadFromArchive(*Reader, FilePath);
+	LoadFromLocalizationResource(*Reader, FilePath);
 
 	bool Success = Reader->Close();
-	delete Reader;
 
 	return Success;
 }
 
-void FTextLocalizationManager::FLocalizationEntryTracker::ReadFromArchive(FArchive& Archive, const FString& Identifier)
+void FTextLocalizationManager::FLocalizationEntryTracker::LoadFromLocalizationResource(FArchive& Archive, const FString& LocalizationResourceIdentifier)
 {
 	// Read namespace count
 	uint32 NamespaceCount;
@@ -249,7 +175,7 @@ void FTextLocalizationManager::FLocalizationEntryTracker::ReadFromArchive(FArchi
 		uint32 KeyCount;
 		Archive << KeyCount;
 
-		FLocalizationEntryTracker::FKeyTable& KeyTable = Namespaces.FindOrAdd(*Namespace);
+		FLocalizationEntryTracker::FKeysTable& KeyTable = Namespaces.FindOrAdd(*Namespace);
 
 		for(uint32 j = 0; j < KeyCount; ++j)
 		{
@@ -260,7 +186,7 @@ void FTextLocalizationManager::FLocalizationEntryTracker::ReadFromArchive(FArchi
 			FLocalizationEntryTracker::FEntryArray& EntryArray = KeyTable.FindOrAdd( *Key );
 
 			FLocalizationEntryTracker::FEntry NewEntry;
-			NewEntry.TableName = Identifier;
+			NewEntry.LocResID = LocalizationResourceIdentifier;
 
 			// Read string entry.
 			Archive << NewEntry.SourceStringHash;
@@ -271,266 +197,102 @@ void FTextLocalizationManager::FLocalizationEntryTracker::ReadFromArchive(FArchi
 	}
 }
 
-void FTextLocalizationManager::LoadResources(const bool ShouldLoadEditor, const bool ShouldLoadGame)
+void FTextLocalizationManager::FLocalizationEntryTracker::DetectAndLogConflicts() const
 {
-	// Add one to the revision index, so all FText's refresh.
-	++HeadCultureRevisionIndex;
-
-	FInternationalization& I18N = FInternationalization::Get();
-
-#if ENABLE_LOC_TESTING
+	for (const auto& NamespaceEntry : Namespaces)
 	{
-		const FString& CultureName = I18N.GetCurrentCulture()->GetName();
-
-		if(CultureName == TEXT("LEET"))
+		const FString& NamespaceName = NamespaceEntry.Key;
+		const FKeysTable& KeyTable = NamespaceEntry.Value;
+		for (const auto& KeyEntry : KeyTable)
 		{
-			for(auto NamespaceIterator = LiveTable.NamespaceTable.CreateIterator(); NamespaceIterator; ++NamespaceIterator)
-			{
-				const FString& Namespace = NamespaceIterator.Key();
-				FTextLookupTable::FKeyTable& LiveKeyTable = NamespaceIterator.Value();
-				for(auto KeyIterator = LiveKeyTable.CreateIterator(); KeyIterator; ++KeyIterator)
-				{
-					const FString& Key = KeyIterator.Key();
-					FStringEntry& LiveStringEntry = KeyIterator.Value();
-					LiveStringEntry.bIsLocalized = true;
-					FInternationalization::Leetify( *LiveStringEntry.String );
-				}
-			}
-
-			return;
-		}
-	}
-#endif
-
-	TArray<FString> LocalizationPaths;
-	if(ShouldLoadGame)
-	{
-		LocalizationPaths += FPaths::GetGameLocalizationPaths();
-	}
-	if(ShouldLoadEditor)
-	{
-		LocalizationPaths += FPaths::GetEditorLocalizationPaths();
-		LocalizationPaths += FPaths::GetToolTipLocalizationPaths();
+			const FString& KeyName = KeyEntry.Key;
+			const FEntryArray& EntryArray = KeyEntry.Value;
 
-		bool bShouldLoadLocalizedPropertyNames = true;
-		if( !GConfig->GetBool( TEXT("Internationalization"), TEXT("ShouldLoadLocalizedPropertyNames"), bShouldLoadLocalizedPropertyNames, GEditorGameAgnosticIni ) )
-		{
-			GConfig->GetBool( TEXT("Internationalization"), TEXT("ShouldLoadLocalizedPropertyNames"), bShouldLoadLocalizedPropertyNames, GEngineIni );
-		}
-		if(bShouldLoadLocalizedPropertyNames)
-		{
-			LocalizationPaths += FPaths::GetPropertyNameLocalizationPaths();
-		}
-	}
-	LocalizationPaths += FPaths::GetEngineLocalizationPaths();
-
-	// Prioritized array of localization entry trackers.
-	TArray<FLocalizationEntryTracker> LocalizationEntryTrackers;
-
-	const auto MapCulturesToDirectories = [](const FString& LocalizationPath) -> TMap<FString, FString>
-	{
-		TMap<FString, FString> CultureToDirectoryMap;
-		IFileManager& FileManager = IFileManager::Get();
-
-		/* Visitor class used to enumerate directories of culture */
-		class FCultureDirectoryMapperVistor : public IPlatformFile::FDirectoryVisitor
-		{
-		public:
-			FCultureDirectoryMapperVistor( TMap<FString, FString>& OutCultureToDirectoryMap )
-				: CultureToDirectoryMap(OutCultureToDirectoryMap)
-			{
-			}
-
-			virtual bool Visit(const TCHAR* FilenameOrDirectory, bool bIsDirectory) override
+			bool WasConflictDetected = false;
+			for(int32 k = 0; k < EntryArray.Num(); ++k)
 			{
-				if(bIsDirectory)
+				const FEntry& LeftEntry = EntryArray[k];
+				for(int32 l = k + 1; l < EntryArray.Num(); ++l)
 				{
-					// UE localization resource folders use "en-US" style while ICU uses "en_US"
-					const FString LocalizationFolder = FPaths::GetCleanFilename(FilenameOrDirectory);
-					const FString CanonicalName = FCulture::GetCanonicalName(LocalizationFolder);
-					CultureToDirectoryMap.Add(CanonicalName, LocalizationFolder);
+					const FEntry& RightEntry = EntryArray[l];
+					const bool bDoesSourceStringHashDiffer = LeftEntry.SourceStringHash != RightEntry.SourceStringHash;
+					const bool bDoesLocalizedStringDiffer = LeftEntry.LocalizedString != RightEntry.LocalizedString;
+					WasConflictDetected = bDoesSourceStringHashDiffer || bDoesLocalizedStringDiffer;
 				}
-
-				return true;
-			}
-
-			/** Array to fill with the names of the UE localization folders available at the given path */
-			TMap<FString, FString>& CultureToDirectoryMap;
-		};
-
-		FCultureDirectoryMapperVistor CultureEnumeratorVistor(CultureToDirectoryMap);
-		FileManager.IterateDirectory(*LocalizationPath, CultureEnumeratorVistor);
-
-		return CultureToDirectoryMap;
-	};
-
-	TMap< FString, TMap<FString, FString> > LocalizationPathToCultureDirectoryMap;
-	for (const FString& LocalizationPath : LocalizationPaths)
-	{
-		LocalizationPathToCultureDirectoryMap.Add(LocalizationPath, MapCulturesToDirectories(LocalizationPath));
-	}
-
-	// Read culture localization resources.
-	TArray<FString> PrioritizedParentCultureNames = I18N.GetCurrentCulture()->GetPrioritizedParentCultureNames();
-
-	for (const FString& CultureName : PrioritizedParentCultureNames)
-	{
-		FLocalizationEntryTracker& CultureTracker = LocalizationEntryTrackers[LocalizationEntryTrackers.Add(FLocalizationEntryTracker())];
-		for (const FString& LocalizationPath : LocalizationPaths)
-		{
-			const FString* const Entry = LocalizationPathToCultureDirectoryMap[LocalizationPath].Find(FCulture::GetCanonicalName(CultureName));
-			if (Entry)
-			{
-				const FString CulturePath = LocalizationPath / (*Entry);
-
-				CultureTracker.ReadFromDirectory(CulturePath);
-			}
-		}
-		CultureTracker.ReportCollisions();
-	}
-
-	UpdateLiveTable(LocalizationEntryTrackers);
-}
-
-void FTextLocalizationManager::UpdateLiveTable(const TArray<FLocalizationEntryTracker>& LocalizationEntryTrackers, const bool FilterUpdatesByTableName)
-{
-	// Update existing localized entries/flag existing newly-unlocalized entries.
-	for(auto NamespaceIterator = LiveTable.NamespaceTable.CreateIterator(); NamespaceIterator; ++NamespaceIterator)
-	{
-		const FString& NamespaceName = NamespaceIterator.Key();
-		FTextLookupTable::FKeyTable& LiveKeyTable = NamespaceIterator.Value();
-		for(auto KeyIterator = LiveKeyTable.CreateIterator(); KeyIterator; ++KeyIterator)
-		{
-			const FString& KeyName = KeyIterator.Key();
-			FStringEntry& LiveStringEntry = KeyIterator.Value();
-
-			const FLocalizationEntryTracker::FEntry* UpdateEntry = NULL;
-
-			// Attempt to use resources in prioritized order until we find an entry.
-			for(int32 i = 0; i < LocalizationEntryTrackers.Num() && !UpdateEntry; ++i)
-			{
-				const FLocalizationEntryTracker& Tracker = LocalizationEntryTrackers[i];
-				const FLocalizationEntryTracker::FKeyTable* const UpdateKeyTable = Tracker.Namespaces.Find(NamespaceName);
-				const FLocalizationEntryTracker::FEntryArray* const UpdateEntryArray = UpdateKeyTable ? UpdateKeyTable->Find(KeyName) : NULL;
-				const FLocalizationEntryTracker::FEntry* Entry = UpdateEntryArray && UpdateEntryArray->Num() ? &((*UpdateEntryArray)[0]) : NULL;
-				UpdateEntry = (Entry && (!FilterUpdatesByTableName || LiveStringEntry.TableName == Entry->TableName)) ? Entry : NULL;
 			}
 
-			if( UpdateEntry )
-			{
-				// If an entry is unlocalized and the source hash differs, it suggests that the source hash changed - do not replace the display string.
-				if(LiveStringEntry.bIsLocalized || LiveStringEntry.SourceStringHash == UpdateEntry->SourceStringHash)
-				{
-					LiveStringEntry.bIsLocalized = true;
-					*(LiveStringEntry.String) = UpdateEntry->LocalizedString;
-					LiveStringEntry.SourceStringHash = UpdateEntry->SourceStringHash;
-				}
-			}
-			else if(!FilterUpdatesByTableName)
+			if(WasConflictDetected)
 			{
-				if ( !LiveStringEntry.bIsLocalized && LiveStringEntry.String.Get() == AccessedStringBeforeLocLoadedErrorMsg )
+				FString CollidingEntryListString;
+				for(int32 k = 0; k < EntryArray.Num(); ++k)
 				{
-					*(LiveStringEntry.String) = TEXT("");
-				}
+					const FEntry& Entry = EntryArray[k];
 
-				LiveStringEntry.bIsLocalized = false;
+					if( !(CollidingEntryListString.IsEmpty()) )
+					{
+						CollidingEntryListString += TEXT('\n');
+					}
 
-#if ENABLE_LOC_TESTING
-				const bool bShouldLEETIFYUnlocalizedString = FParse::Param(FCommandLine::Get(), TEXT("LEETIFYUnlocalized"));
-				if(bShouldLEETIFYUnlocalizedString )
-				{
-					FInternationalization::Leetify(*(LiveStringEntry.String));
+					CollidingEntryListString += FString::Printf( TEXT("Localization Resource: (%s) Source String Hash: (%d) Localized String: (%s)"), *(Entry.LocResID), Entry.SourceStringHash, *(Entry.LocalizedString) );
 				}
-#endif
-			}
-		}
-	}
-
-	// Add new entries. 
-	for(int32 i = 0; i < LocalizationEntryTrackers.Num(); ++i)
-	{
-		const FLocalizationEntryTracker& Tracker = LocalizationEntryTrackers[i];
-		for(auto NamespaceIterator = Tracker.Namespaces.CreateConstIterator(); NamespaceIterator; ++NamespaceIterator)
-		{
-			const FString& NamespaceName = NamespaceIterator.Key();
-			const FLocalizationEntryTracker::FKeyTable& NewKeyTable = NamespaceIterator.Value();
-			for(auto KeyIterator = NewKeyTable.CreateConstIterator(); KeyIterator; ++KeyIterator)
-			{
-				const FString& Key = KeyIterator.Key();
-				const FLocalizationEntryTracker::FEntryArray& NewEntryArray = KeyIterator.Value();
-				const FLocalizationEntryTracker::FEntry& NewEntry = NewEntryArray[0];
-
-				FTextLookupTable::FKeyTable& LiveKeyTable = LiveTable.NamespaceTable.FindOrAdd(NamespaceName);
-				FStringEntry* const LiveStringEntry = LiveKeyTable.Find(Key);
-				// Note: Anything we find in the table has already been updated above.
-				if( !LiveStringEntry )
-				{
-					FStringEntry NewLiveEntry(
-						true,													/*bIsLocalized*/
-						NewEntry.TableName,
-						NewEntry.SourceStringHash,								/*SourceStringHash*/
-						MakeShareable( new FString(NewEntry.LocalizedString) )	/*String*/
-						);
-					LiveKeyTable.Add( Key, NewLiveEntry );
 
-					ReverseLiveTable.Add(NewLiveEntry.String, FNamespaceKeyEntry( MakeShareable( new FString( NamespaceName ) ), MakeShareable( new FString( Key ) )));
-				}
+				UE_LOG(LogTextLocalizationManager, Warning, TEXT("Loaded localization resources contain conflicting entries for (Namespace:%s, Key:%s):\n%s"), *NamespaceName, *KeyName, *CollidingEntryListString);
+			}
 			}
 		}
-	}
-
-	OnTranslationsChanged().Broadcast();
 }
 
-void FTextLocalizationManager::OnCultureChanged()
+void FTextLocalizationManager::FDisplayStringLookupTable::Find(const FString& InNamespace, FKeysTable*& OutKeysTableForNamespace, const FString& InKey, FDisplayStringEntry*& OutDisplayStringEntry)
 {
-	const bool ShouldLoadEditor = bIsInitialized && WITH_EDITOR;
-	const bool ShouldLoadGame = bIsInitialized && FApp::IsGame();
-	LoadResources(ShouldLoadEditor, ShouldLoadGame);
+	// Find namespace's key table.
+	OutKeysTableForNamespace = NamespacesTable.Find( InNamespace );
+
+	// Find key table's entry.
+	OutDisplayStringEntry = OutKeysTableForNamespace ? OutKeysTableForNamespace->Find( InKey ) : nullptr;
 }
 
-TSharedPtr<FString, ESPMode::ThreadSafe> FTextLocalizationManager::FindString( const FString& Namespace, const FString& Key, const FString* const SourceString )
+void FTextLocalizationManager::FDisplayStringLookupTable::Find(const FString& InNamespace, const FKeysTable*& OutKeysTableForNamespace, const FString& InKey, const FDisplayStringEntry*& OutDisplayStringEntry) const
 {
-	FScopeLock ScopeLock( &SynchronizationObject );
-
 	// Find namespace's key table.
-	const FTextLookupTable::FKeyTable* LiveKeyTable = LiveTable.NamespaceTable.Find( Namespace );
+	OutKeysTableForNamespace = NamespacesTable.Find( InNamespace );
 
 	// Find key table's entry.
-	const FStringEntry* LiveEntry = LiveKeyTable ? LiveKeyTable->Find( Key ) : NULL;
+	OutDisplayStringEntry = OutKeysTableForNamespace ? OutKeysTableForNamespace->Find( InKey ) : nullptr;
+}
 
-	if( LiveEntry != NULL && ( !SourceString || LiveEntry->SourceStringHash == FCrc::StrCrc32(**SourceString) ) )
+FTextLocalizationManager& FTextLocalizationManager::Get()
+{
+	static FTextLocalizationManager* GTextLocalizationManager = nullptr;
+	if( !GTextLocalizationManager )
 	{
-		return LiveEntry->String;
+		GTextLocalizationManager = new FTextLocalizationManager;
 	}
 
-	return NULL;
+	return *GTextLocalizationManager;
 }
 
-TSharedPtr<FString, ESPMode::ThreadSafe> FTextLocalizationManager::GetTableName(const FString& Namespace, const FString& Key)
+FTextDisplayStringPtr FTextLocalizationManager::FindDisplayString( const FString& Namespace, const FString& Key, const FString* const SourceString )
 {
-	FScopeLock ScopeLock(&SynchronizationObject);
+	FScopeLock ScopeLock( &SynchronizationObject );
 
-	// Find namespace's key table.
-	const FTextLookupTable::FKeyTable* LiveKeyTable = LiveTable.NamespaceTable.Find(Namespace);
+	const FDisplayStringLookupTable::FKeysTable* LiveKeyTable = nullptr;
+	const FDisplayStringLookupTable::FDisplayStringEntry* LiveEntry = nullptr;
 
-	// Find key table's entry.
-	const FStringEntry* LiveEntry = LiveKeyTable ? LiveKeyTable->Find(Key) : NULL;
+	DisplayStringLookupTable.Find(Namespace, LiveKeyTable, Key, LiveEntry);
 
-	if (LiveEntry != nullptr)
+	if ( LiveEntry != nullptr && ( !SourceString || LiveEntry->SourceStringHash == FCrc::StrCrc32(**SourceString) ) )
 	{
-		return MakeShareable(new FString((LiveEntry->TableName)));
+		return LiveEntry->DisplayString;
 	}
 
-	return NULL;
+	return nullptr;
 }
 
-
-TSharedRef<FString, ESPMode::ThreadSafe> FTextLocalizationManager::GetString(const FString& Namespace, const FString& Key, const FString* const SourceString)
+FTextDisplayStringRef FTextLocalizationManager::GetDisplayString(const FString& Namespace, const FString& Key, const FString* const SourceString)
 {
 	FScopeLock ScopeLock( &SynchronizationObject );
 
-	// Hack fix for old assets that don't have namespace/key info
+	// Hack fix for old assets that don't have namespace/key info.
 	if(Namespace.IsEmpty() && Key.IsEmpty())
 	{
 		return MakeShareable( new FString( SourceString ? **SourceString : TEXT("") ) );
@@ -539,23 +301,21 @@ TSharedRef<FString, ESPMode::ThreadSafe> FTextLocalizationManager::GetString(con
 #if ENABLE_LOC_TESTING
 	const bool bShouldLEETIFYAll = bIsInitialized && FInternationalization::Get().GetCurrentCulture()->GetName() == TEXT("LEET");
 
-	// Attempt to set bShouldLEETIFYUnlocalizedString appropriately, only once, after the commandline is initialized and parsed
+	// Attempt to set bShouldLEETIFYUnlocalizedString appropriately, only once, after the commandline is initialized and parsed.
 	static bool bShouldLEETIFYUnlocalizedString = false;
 	{
 		static bool bHasParsedCommandLine = false;
 		if(!bHasParsedCommandLine && FCommandLine::IsInitialized())
-		{
+	{
 			bShouldLEETIFYUnlocalizedString = FParse::Param(FCommandLine::Get(), TEXT("LEETIFYUnlocalized"));
 			bHasParsedCommandLine = true;
 		}
 	}
 #endif
 
-	// Find namespace's key table.
-	FTextLookupTable::FKeyTable* LiveKeyTable = LiveTable.NamespaceTable.Find( Namespace );
-
-	// Find key table's entry.
-	FStringEntry* LiveEntry = LiveKeyTable ? LiveKeyTable->Find( Key ) : NULL;
+	FDisplayStringLookupTable::FKeysTable* LiveKeyTable = nullptr;
+	FDisplayStringLookupTable::FDisplayStringEntry* LiveEntry = nullptr;
+	DisplayStringLookupTable.Find(Namespace, LiveKeyTable, Key, LiveEntry);
 
 	// Entry is present.
 	if( LiveEntry )
@@ -567,20 +327,20 @@ TSharedRef<FString, ESPMode::ThreadSafe> FTextLocalizationManager::GetString(con
 		if( SourceStringHash != 0 && SourceStringHash != LiveEntry->SourceStringHash )
 		{
 			LiveEntry->SourceStringHash = SourceStringHash;
-			LiveEntry->String.Get() = SourceString ? **SourceString : TEXT("");
+			LiveEntry->DisplayString.Get() = SourceString ? **SourceString : TEXT("");
 
 #if ENABLE_LOC_TESTING
 			if( (bShouldLEETIFYAll || bShouldLEETIFYUnlocalizedString) && SourceString )
 			{
-				FInternationalization::Leetify(*LiveEntry->String);
-				if(*LiveEntry->String == *SourceString)
+				FInternationalization::Leetify(*LiveEntry->DisplayString);
+				if(*LiveEntry->DisplayString == *SourceString)
 				{
 					UE_LOG(LogTextLocalizationManager, Warning, TEXT("Leetify failed to alter a string (%s)."), **SourceString );
 				}
 			}
 #endif
 
-			UE_LOG(LogTextLocalizationManager, Verbose, TEXT("An attempt was made to get a localized string (Namespace:%s, Key:%s), but the source string hash does not match - the source string (%s) will be used."), *Namespace, *Key, **LiveEntry->String);
+			UE_LOG(LogTextLocalizationManager, Verbose, TEXT("An attempt was made to get a localized string (Namespace:%s, Key:%s), but the source string hash does not match - the source string (%s) will be used."), *Namespace, *Key, **LiveEntry->DisplayString);
 
 #if ENABLE_LOC_TESTING
 			LiveEntry->bIsLocalized = bShouldLEETIFYAll;
@@ -589,88 +349,180 @@ TSharedRef<FString, ESPMode::ThreadSafe> FTextLocalizationManager::GetString(con
 #endif
 		}
 
-		return LiveEntry->String;
+		return LiveEntry->DisplayString;
 	}
 	// Entry is absent.
 	else
 	{
 		// Don't log warnings about unlocalized strings if the system hasn't been initialized - we simply don't have localization data yet.
 		if( bIsInitialized )
-		{
+	{
 			UE_LOG(LogTextLocalizationManager, Verbose, TEXT("An attempt was made to get a localized string (Namespace:%s, Key:%s, Source:%s), but it did not exist."), *Namespace, *Key, SourceString ? **SourceString : TEXT(""));
 		}
 
-		const TSharedRef<FString, ESPMode::ThreadSafe> UnlocalizedString = MakeShareable( new FString( SourceString ? **SourceString : TEXT("") ) );
+		const FTextDisplayStringRef UnlocalizedString = MakeShareable( new FString( SourceString ? **SourceString : TEXT("") ) );
 
 		// If live-culture-swap is enabled or the system is uninitialized - make entries so that they can be updated when system is initialized or a culture swap occurs.
 		CA_SUPPRESS(6236)
-		CA_SUPPRESS(6316)
-		if( !(bIsInitialized) || ENABLE_LOC_TESTING )
-		{
-#if ENABLE_LOC_TESTING
-			if( (bShouldLEETIFYAll || bShouldLEETIFYUnlocalizedString) && SourceString )
+			CA_SUPPRESS(6316)
+			if( !(bIsInitialized) || ENABLE_LOC_TESTING )
 			{
-				FInternationalization::Leetify(*UnlocalizedString);
-				if(*UnlocalizedString == *SourceString)
+#if ENABLE_LOC_TESTING
+				if( (bShouldLEETIFYAll || bShouldLEETIFYUnlocalizedString) && SourceString )
 				{
-					UE_LOG(LogTextLocalizationManager, Warning, TEXT("Leetify failed to alter a string (%s)."), **SourceString );
+					FInternationalization::Leetify(*UnlocalizedString);
+					if(*UnlocalizedString == *SourceString)
+					{
+						UE_LOG(LogTextLocalizationManager, Warning, TEXT("Leetify failed to alter a string (%s)."), **SourceString );
+					}
 				}
-			}
 #endif
 
-			if ( UnlocalizedString->IsEmpty() )
-			{
-				if ( !bIsInitialized )
+				if ( UnlocalizedString->IsEmpty() )
 				{
-					*(UnlocalizedString) = AccessedStringBeforeLocLoadedErrorMsg;
+					if ( !bIsInitialized )
+		{
+						*(UnlocalizedString) = AccessedStringBeforeLocLoadedErrorMsg;
+					}
 				}
-			}
 
-			FStringEntry NewEntry(
+				FDisplayStringLookupTable::FDisplayStringEntry NewEntry(
 #if ENABLE_LOC_TESTING
-				bShouldLEETIFYAll						/*bIsLocalized*/
+					bShouldLEETIFYAll						/*bIsLocalized*/
 #else
-				false												/*bIsLocalized*/
+					false												/*bIsLocalized*/
 #endif
-				, TEXT("")
-				, SourceString ? FCrc::StrCrc32(**SourceString) : 0	/*SourceStringHash*/
-				, UnlocalizedString									/*String*/
-				);
+					, TEXT("")
+					, SourceString ? FCrc::StrCrc32(**SourceString) : 0	/*SourceStringHash*/
+					, UnlocalizedString									/*String*/
+					);
 
-			if( !LiveKeyTable )
-			{
-				LiveKeyTable = &(LiveTable.NamespaceTable.Add( Namespace, FTextLookupTable::FKeyTable() ));
-			}
+				if( !LiveKeyTable )
+				{
+					LiveKeyTable = &(DisplayStringLookupTable.NamespacesTable.Add( Namespace, FDisplayStringLookupTable::FKeysTable() ));
+				}
 
-			LiveKeyTable->Add( Key, NewEntry );
+				LiveKeyTable->Add( Key, NewEntry );
 
-			ReverseLiveTable.Add(NewEntry.String, FNamespaceKeyEntry( MakeShareable( new FString( Namespace ) ), MakeShareable( new FString( Key ) )));
+				NamespaceKeyLookupTable.Add(NewEntry.DisplayString, FNamespaceKeyEntry(Namespace, Key));
 
 		}
 
-		return UnlocalizedString;
+			return UnlocalizedString;
+	}
+}
+
+bool FTextLocalizationManager::GetLocResID(const FString& Namespace, const FString& Key, FString& OutLocResId)
+{
+	FScopeLock ScopeLock(&SynchronizationObject);
+
+	const FDisplayStringLookupTable::FKeysTable* LiveKeyTable = nullptr;
+	const FDisplayStringLookupTable::FDisplayStringEntry* LiveEntry = nullptr;
+	DisplayStringLookupTable.Find(Namespace, LiveKeyTable, Key, LiveEntry);
+
+	if (LiveEntry != nullptr)
+	{
+		OutLocResId = LiveEntry->LocResID;
+		return true;
 	}
+
+	return false;
 }
 
-void FTextLocalizationManager::FindKeyNamespaceFromDisplayString(TSharedRef<FString, ESPMode::ThreadSafe> InDisplayString, TSharedPtr<FString, ESPMode::ThreadSafe>& OutNamespace, TSharedPtr<FString, ESPMode::ThreadSafe>& OutKey)
+bool FTextLocalizationManager::FindNamespaceAndKeyFromDisplayString(const FTextDisplayStringRef& InDisplayString, FString& OutNamespace, FString& OutKey)
 {
 	FScopeLock ScopeLock( &SynchronizationObject );
 
-	FNamespaceKeyEntry* NamespaceKeyEntry = ReverseLiveTable.Find(InDisplayString);
+	FNamespaceKeyEntry* NamespaceKeyEntry = NamespaceKeyLookupTable.Find(InDisplayString);
 
-	if(NamespaceKeyEntry)
+	if (NamespaceKeyEntry)
 	{
 		OutNamespace = NamespaceKeyEntry->Namespace;
 		OutKey = NamespaceKeyEntry->Key;
 	}
+
+	return NamespaceKeyEntry != nullptr;
 }
 
+bool FTextLocalizationManager::AddDisplayString(const FTextDisplayStringRef& DisplayString, const FString& Namespace, const FString& Key)
+{
+	FScopeLock ScopeLock( &SynchronizationObject );
+
+	// Try to find existing entries.
+	FNamespaceKeyEntry* ReverseLiveTableEntry = NamespaceKeyLookupTable.Find(DisplayString);
+	FDisplayStringLookupTable::FKeysTable* KeysTableForExistingNamespace = nullptr;
+	FDisplayStringLookupTable::FDisplayStringEntry* ExistingDisplayStringEntry = nullptr;
+	DisplayStringLookupTable.Find(Namespace, KeysTableForExistingNamespace, Key, ExistingDisplayStringEntry);
+
+	// If there are any existing entries, they may cause a conflict, unless they're exactly the same as what we would be adding.
+	if ( (ExistingDisplayStringEntry && ExistingDisplayStringEntry->DisplayString != DisplayString) || // Namespace and key mustn't be associated with a different display string
+		 (ReverseLiveTableEntry && (ReverseLiveTableEntry->Namespace != Namespace || ReverseLiveTableEntry->Key != Key)) ) // Display string mustn't be associated with a different namespace and key.
+	{
+		return false;
+	}
+
+	// Add the necesary associations in both directions.
+	FDisplayStringLookupTable::FKeysTable& KeysTableForNamespace = DisplayStringLookupTable.NamespacesTable.FindOrAdd(Namespace);
+	KeysTableForNamespace.Add(Key, FDisplayStringLookupTable::FDisplayStringEntry(false, TEXT(""), FCrc::StrCrc32(*DisplayString.Get()), DisplayString));
+	NamespaceKeyLookupTable.Add(DisplayString, FNamespaceKeyEntry(Namespace, Key));
+
+	return true;
+}
 
-void FTextLocalizationManager::RegenerateResources( const FString& ConfigFilePath, IInternationalizationArchiveSerializer& ArchiveSerializer, IInternationalizationManifestSerializer& ManifestSerializer )
+bool FTextLocalizationManager::UpdateDisplayString(const FTextDisplayStringRef& DisplayString, const FString& Value, const FString& Namespace, const FString& Key)
 {
-	// Add one to the revision index, so all FText's refresh.
-	++HeadCultureRevisionIndex;
-	
+	FScopeLock ScopeLock( &SynchronizationObject );
+
+	// Get entry from reverse live table. Contains current namespace and key values.
+	FNamespaceKeyEntry& ReverseLiveTableEntry = NamespaceKeyLookupTable[DisplayString];
+
+	// Copy old live table entry over as new live table entry and destroy old live table entry if the namespace or key has changed.
+	if (ReverseLiveTableEntry.Namespace != Namespace || ReverseLiveTableEntry.Key != Key)
+	{
+		FDisplayStringLookupTable::FKeysTable& KeysTableForNewNamespace = DisplayStringLookupTable.NamespacesTable.FindOrAdd(Namespace);
+		FDisplayStringLookupTable::FDisplayStringEntry* NewDisplayStringEntry = KeysTableForNewNamespace.Find(Key);
+		if (NewDisplayStringEntry)
+		{
+			// Can not update, that namespace and key combination is already in use by another string.
+			return false;
+		}
+		else
+		{
+			// Get old namespace's keys table and old live table entry under old key.
+			FDisplayStringLookupTable::FKeysTable* KeysTableForOldNamespace = nullptr;
+			FDisplayStringLookupTable::FDisplayStringEntry* OldDisplayStringEntry = nullptr;
+			DisplayStringLookupTable.Find(ReverseLiveTableEntry.Namespace, KeysTableForOldNamespace, ReverseLiveTableEntry.Key, OldDisplayStringEntry);
+
+			// Copy old live table entry to new key in the new namespace's key table.
+			check(OldDisplayStringEntry);
+			KeysTableForNewNamespace.Add(Key, *OldDisplayStringEntry);
+
+			// Remove old live table entry and old key in the old namespace's key table.
+			check(KeysTableForOldNamespace);
+			KeysTableForOldNamespace->Remove(ReverseLiveTableEntry.Key);
+
+			// Remove old namespace if empty.
+			if(DisplayStringLookupTable.NamespacesTable[ReverseLiveTableEntry.Namespace].Num() == 0)
+			{
+				DisplayStringLookupTable.NamespacesTable.Remove(ReverseLiveTableEntry.Namespace);
+			}
+		}
+	}
+
+	// Update display string value.
+	*DisplayString = Value;
+
+	// Update entry from reverse live table.
+	ReverseLiveTableEntry.Namespace = Namespace;
+	ReverseLiveTableEntry.Key = Key;
+
+	return true;
+}
+
+void FTextLocalizationManager::LoadFromManifestAndArchives( const FString& ConfigFilePath, IInternationalizationArchiveSerializer& ArchiveSerializer, IInternationalizationManifestSerializer& ManifestSerializer )
+{
+	FScopeLock ScopeLock( &SynchronizationObject );
+
 	FInternationalization& I18N = FInternationalization::Get();
 
 	FString SectionName = TEXT("RegenerateResources");
@@ -790,12 +642,233 @@ void FTextLocalizationManager::RegenerateResources( const FString& ConfigFilePat
 		const FString ResourceFilePath = FPaths::ConvertRelativePathToFull(CulturePath / ResourceName);
 
 		FLocalizationEntryTracker& CultureTracker = LocalizationEntryTrackers[LocalizationEntryTrackers.Add(FLocalizationEntryTracker())];
-		CultureTracker.ReadFromArchive(MemoryReader, ResourceFilePath);
-		CultureTracker.ReportCollisions();
+		CultureTracker.LoadFromLocalizationResource(MemoryReader, ResourceFilePath);
+		CultureTracker.DetectAndLogConflicts();
 
 		MemoryReader.Close();
 	}
 
-	// Don't filter updates by table name, or we can't live preview strings that had no translation originally
-	UpdateLiveTable(LocalizationEntryTrackers);
+	UpdateFromLocalizations(LocalizationEntryTrackers);
+}
+
+void FTextLocalizationManager::OnCultureChanged()
+{
+	const bool ShouldLoadEditor = bIsInitialized && WITH_EDITOR;
+	const bool ShouldLoadGame = bIsInitialized && FApp::IsGame();
+	LoadLocalizationResourcesForCulture(FInternationalization::Get().GetCurrentCulture()->GetName(), ShouldLoadEditor, ShouldLoadGame);
+}
+
+void FTextLocalizationManager::LoadLocalizationResourcesForCulture(const FString& CultureName, const bool ShouldLoadEditor, const bool ShouldLoadGame)
+{
+#if ENABLE_LOC_TESTING
+	{
+		// The leet culture is fake. Just leet-ify existing strings.
+		if(CultureName == TEXT("LEET"))
+		{
+			for(auto NamespaceIterator = DisplayStringLookupTable.NamespacesTable.CreateIterator(); NamespaceIterator; ++NamespaceIterator)
+			{
+				const FString& Namespace = NamespaceIterator.Key();
+				FDisplayStringLookupTable::FKeysTable& LiveKeyTable = NamespaceIterator.Value();
+				for(auto KeyIterator = LiveKeyTable.CreateIterator(); KeyIterator; ++KeyIterator)
+				{
+					const FString& Key = KeyIterator.Key();
+					FDisplayStringLookupTable::FDisplayStringEntry& LiveStringEntry = KeyIterator.Value();
+					LiveStringEntry.bIsLocalized = true;
+					FInternationalization::Leetify( *LiveStringEntry.DisplayString );
+				}
+			}
+
+			// Early-out, there can be no localization resources to load for the fake leet culture.
+			DirtyTextRevision();
+			return;
+		}
+	}
+#endif
+
+	const FCulturePtr Culture = FInternationalization::Get().GetCulture(CultureName);
+
+	// Can't load localization resources for a culture that doesn't exist, early-out.
+	if (!Culture.IsValid())
+	{
+		return;
+	}
+
+	// Collect the localization paths to load from.
+	TArray<FString> LocalizationPaths;
+	if(ShouldLoadGame)
+	{
+		LocalizationPaths += FPaths::GetGameLocalizationPaths();
+	}
+	if(ShouldLoadEditor)
+	{
+		LocalizationPaths += FPaths::GetEditorLocalizationPaths();
+		LocalizationPaths += FPaths::GetToolTipLocalizationPaths();
+
+		bool bShouldLoadLocalizedPropertyNames = true;
+		if( !GConfig->GetBool( TEXT("Internationalization"), TEXT("ShouldLoadLocalizedPropertyNames"), bShouldLoadLocalizedPropertyNames, GEditorGameAgnosticIni ) )
+		{
+			GConfig->GetBool( TEXT("Internationalization"), TEXT("ShouldLoadLocalizedPropertyNames"), bShouldLoadLocalizedPropertyNames, GEngineIni );
+		}
+		if(bShouldLoadLocalizedPropertyNames)
+		{
+			LocalizationPaths += FPaths::GetPropertyNameLocalizationPaths();
+		}
+	}
+	LocalizationPaths += FPaths::GetEngineLocalizationPaths();
+
+	// Prioritized array of localization entry trackers.
+	TArray<FLocalizationEntryTracker> LocalizationEntryTrackers;
+
+	const auto MapCulturesToDirectories = [](const FString& LocalizationPath) -> TMap<FString, FString>
+	{
+		TMap<FString, FString> CultureToDirectoryMap;
+		IFileManager& FileManager = IFileManager::Get();
+
+		/* Visitor class used to enumerate directories of culture. */
+		class FCultureDirectoryMapperVistor : public IPlatformFile::FDirectoryVisitor
+		{
+		public:
+			FCultureDirectoryMapperVistor( TMap<FString, FString>& OutCultureToDirectoryMap )
+				: CultureToDirectoryMap(OutCultureToDirectoryMap)
+			{
+			}
+
+			virtual bool Visit(const TCHAR* FilenameOrDirectory, bool bIsDirectory) override
+			{
+				if(bIsDirectory)
+				{
+					// UE localization resource folders use "en-US" style while ICU uses "en_US".
+					const FString LocalizationFolder = FPaths::GetCleanFilename(FilenameOrDirectory);
+					const FString CanonicalName = FCulture::GetCanonicalName(LocalizationFolder);
+					CultureToDirectoryMap.Add(CanonicalName, LocalizationFolder);
+				}
+
+				return true;
+			}
+
+			/** Array to fill with the names of the UE localization folders available at the given path. */
+			TMap<FString, FString>& CultureToDirectoryMap;
+		};
+
+		FCultureDirectoryMapperVistor CultureEnumeratorVistor(CultureToDirectoryMap);
+		FileManager.IterateDirectory(*LocalizationPath, CultureEnumeratorVistor);
+
+		return CultureToDirectoryMap;
+	};
+
+	TMap< FString, TMap<FString, FString> > LocalizationPathToCultureDirectoryMap;
+	for (const FString& LocalizationPath : LocalizationPaths)
+	{
+		LocalizationPathToCultureDirectoryMap.Add(LocalizationPath, MapCulturesToDirectories(LocalizationPath));
+	}
+
+	// Read culture localization resources.
+	const TArray<FString> PrioritizedParentCultureNames = Culture->GetPrioritizedParentCultureNames();
+	for (const FString& ParentCultureName : PrioritizedParentCultureNames)
+	{
+		FLocalizationEntryTracker& CultureTracker = LocalizationEntryTrackers[LocalizationEntryTrackers.Add(FLocalizationEntryTracker())];
+		for (const FString& LocalizationPath : LocalizationPaths)
+		{
+			const FString* const Entry = LocalizationPathToCultureDirectoryMap[LocalizationPath].Find(FCulture::GetCanonicalName(ParentCultureName));
+			if (Entry)
+			{
+				const FString CulturePath = LocalizationPath / (*Entry);
+
+				CultureTracker.LoadFromDirectory(CulturePath);
+			}
+		}
+		CultureTracker.DetectAndLogConflicts();
+	}
+
+	// Replace localizations with those of the loaded localization resources.
+	UpdateFromLocalizations(LocalizationEntryTrackers);
+}
+
+void FTextLocalizationManager::UpdateFromLocalizations(const TArray<FLocalizationEntryTracker>& LocalizationEntryTrackers)
+{
+	// Update existing localized entries/flag existing newly-unlocalized entries.
+	for (auto& Namespace : DisplayStringLookupTable.NamespacesTable)
+	{
+		const FString& NamespaceName = Namespace.Key;
+		FDisplayStringLookupTable::FKeysTable& LiveKeyTable = Namespace.Value;
+		for(auto& Key : LiveKeyTable)
+		{
+			const FString& KeyName = Key.Key;
+			FDisplayStringLookupTable::FDisplayStringEntry& LiveStringEntry = Key.Value;
+
+			const FLocalizationEntryTracker::FEntry* SourceEntryForUpdate = nullptr;
+
+			// Attempt to use resources in prioritized order until we find an entry.
+			for(const FLocalizationEntryTracker& Tracker : LocalizationEntryTrackers)
+			{
+				const FLocalizationEntryTracker::FKeysTable* const UpdateKeyTable = Tracker.Namespaces.Find(NamespaceName);
+				const FLocalizationEntryTracker::FEntryArray* const UpdateEntryArray = UpdateKeyTable ? UpdateKeyTable->Find(KeyName) : nullptr;
+				const FLocalizationEntryTracker::FEntry* Entry = UpdateEntryArray && UpdateEntryArray->Num() ? &((*UpdateEntryArray)[0]) : nullptr;
+				if(Entry)
+				{
+					SourceEntryForUpdate = Entry;
+					break;
+				}
+			}
+
+			// If the source string hashes are are the same, we can replace the display string.
+			// Otherwise, it would suggest the source string has changed and the new localization may be based off of an old source string.
+			if (SourceEntryForUpdate && LiveStringEntry.SourceStringHash == SourceEntryForUpdate->SourceStringHash)
+			{
+				LiveStringEntry.bIsLocalized = true;
+				*(LiveStringEntry.DisplayString) = SourceEntryForUpdate->LocalizedString;
+			}
+			else
+			{
+				if ( !LiveStringEntry.bIsLocalized && LiveStringEntry.DisplayString.Get() == AccessedStringBeforeLocLoadedErrorMsg )
+				{
+					*(LiveStringEntry.DisplayString) = TEXT("");
+				}
+
+				LiveStringEntry.bIsLocalized = false;
+
+#if ENABLE_LOC_TESTING
+				const bool bShouldLEETIFYUnlocalizedString = FParse::Param(FCommandLine::Get(), TEXT("LEETIFYUnlocalized"));
+				if(bShouldLEETIFYUnlocalizedString )
+				{
+					FInternationalization::Leetify(*(LiveStringEntry.DisplayString));
+				}
+#endif
+			}
+		}
+	}
+
+	// Add new entries. 
+	for(const auto& Tracker : LocalizationEntryTrackers)
+	{
+		for(const auto& Namespace : Tracker.Namespaces)
+		{
+			const FString& NamespaceName = Namespace.Key;
+			const FLocalizationEntryTracker::FKeysTable& NewKeyTable = Namespace.Value;
+			for(const auto& Key : NewKeyTable)
+			{
+				const FString& KeyName = Key.Key;
+				const FLocalizationEntryTracker::FEntryArray& NewEntryArray = Key.Value;
+				const FLocalizationEntryTracker::FEntry& NewEntry = NewEntryArray[0];
+
+				FDisplayStringLookupTable::FKeysTable& LiveKeyTable = DisplayStringLookupTable.NamespacesTable.FindOrAdd(NamespaceName);
+				FDisplayStringLookupTable::FDisplayStringEntry* const LiveStringEntry = LiveKeyTable.Find(KeyName);
+				// Note: Anything we find in the table has already been updated above.
+				if( !LiveStringEntry )
+				{
+					FDisplayStringLookupTable::FDisplayStringEntry NewLiveEntry(
+						true,													/*bIsLocalized*/
+						NewEntry.LocResID,
+						NewEntry.SourceStringHash,								/*SourceStringHash*/
+						MakeShareable( new FString(NewEntry.LocalizedString) )	/*String*/
+						);
+					LiveKeyTable.Add( KeyName, NewLiveEntry );
+
+					NamespaceKeyLookupTable.Add(NewLiveEntry.DisplayString, FNamespaceKeyEntry(NamespaceName, KeyName));
+				}
+			}
+		}
+	}
+
+	DirtyTextRevision();
 }
diff --git a/Engine/Source/Runtime/Core/Private/Misc/ConfigCacheIni.cpp b/Engine/Source/Runtime/Core/Private/Misc/ConfigCacheIni.cpp
--- a/Engine/Source/Runtime/Core/Private/Misc/ConfigCacheIni.cpp
+++ b/Engine/Source/Runtime/Core/Private/Misc/ConfigCacheIni.cpp
@@ -105,7 +105,7 @@ class FTextFriendHelper
 		else
 		{
 			return FString::Printf( TEXT( "NSLOCTEXT(\"%s\",\"%s\",\"%s\")" ),
-				FTextInspector::GetNamespace(Text) ? **FTextInspector::GetNamespace(Text) : TEXT(""), FTextInspector::GetKey(Text) ? **FTextInspector::GetKey(Text) : TEXT(""), *Str );
+				*FTextInspector::GetNamespace(Text).Get(TEXT("")), *FTextInspector::GetKey(Text).Get(TEXT("")), *Str );
 		}
 #undef LOC_DEFINE_REGION
 	}
diff --git a/Engine/Source/Runtime/Core/Private/Tests/Internationalization/TextTest.cpp b/Engine/Source/Runtime/Core/Private/Tests/Internationalization/TextTest.cpp
--- a/Engine/Source/Runtime/Core/Private/Tests/Internationalization/TextTest.cpp
+++ b/Engine/Source/Runtime/Core/Private/Tests/Internationalization/TextTest.cpp
@@ -464,7 +464,7 @@ bool FTextTest::RunTest (const FString& Parameters)
 
 			{
 				// Serialize the full, bulky FText that is a composite of most of the other FTextHistories.
-				FMemoryWriter Ar(FormattedHistoryAsEnglish, /*bIsPersistent=*/ true);
+				FMemoryWriter Ar(FormattedHistoryAsEnglish);
 				Ar << FormattedTestLayer2;
 				Ar.Close();
 			}
@@ -523,7 +523,7 @@ bool FTextTest::RunTest (const FString& Parameters)
 				{
 					// Serialize the full, bulky FText that is a composite of most of the other FTextHistories.
 					// We don't care how this may be translated, we will be serializing this in as LEET.
-					FMemoryWriter Ar(FormattedHistoryAsFrenchCanadian, /*bIsPersistent=*/ true);
+					FMemoryWriter Ar(FormattedHistoryAsFrenchCanadian);
 					Ar << FormattedTestLayer2;
 					Ar.Close();
 				}
@@ -532,7 +532,7 @@ bool FTextTest::RunTest (const FString& Parameters)
 					FText InvariantFText = FText::FromString(InvariantString);
 
 					// Serialize an invariant FText
-					FMemoryWriter Ar(InvariantFTextData, /*bIsPersistent=*/ true);
+					FMemoryWriter Ar(InvariantFTextData);
 					Ar << InvariantFText;
 					Ar.Close();
 				}
@@ -546,12 +546,12 @@ bool FTextTest::RunTest (const FString& Parameters)
 			FText FormattedFrenchCanadianTextHistoryAsLeet;
 
 			{
-				FMemoryReader Ar(FormattedHistoryAsEnglish, /*bIsPersistent=*/ true);
+				FMemoryReader Ar(FormattedHistoryAsEnglish);
 				Ar << FormattedEnglishTextHistoryAsLeet;
 				Ar.Close();
 			}
 			{
-				FMemoryReader Ar(FormattedHistoryAsFrenchCanadian, /*bIsPersistent=*/ true);
+				FMemoryReader Ar(FormattedHistoryAsFrenchCanadian);
 				Ar << FormattedFrenchCanadianTextHistoryAsLeet;
 				Ar.Close();
 			}
@@ -580,7 +580,7 @@ bool FTextTest::RunTest (const FString& Parameters)
 
 				FText InvariantFText;
 
-				FMemoryReader Ar(InvariantFTextData, /*bIsPersistent=*/ true);
+				FMemoryReader Ar(InvariantFTextData);
 				Ar << InvariantFText;
 				Ar.Close();
 
diff --git a/Engine/Source/Runtime/Core/Public/Internationalization/Text.h b/Engine/Source/Runtime/Core/Public/Internationalization/Text.h
--- a/Engine/Source/Runtime/Core/Public/Internationalization/Text.h
+++ b/Engine/Source/Runtime/Core/Public/Internationalization/Text.h
@@ -4,6 +4,8 @@
 #include "Internationalization/CulturePointer.h"
 #include "Templates/SharedPointer.h"
 #include "UObject/NameTypes.h"
+#include "TextLocalizationManager.h"
+#include "Optional.h"
 
 struct FTimespan;
 struct FDateTime;
@@ -358,7 +360,7 @@ private:
 
 private:
 	/** The visible display string for this FText */
-	TSharedRef<FString, ESPMode::ThreadSafe> DisplayString;
+	FTextDisplayStringRef DisplayString;
 
 	/** The FText's history, to allow it to rebuild under a new culture */
 	TSharedPtr<class FTextHistory, ESPMode::ThreadSafe> History;
@@ -414,7 +416,7 @@ public:
 
 private:
 	/** A pointer to the visible display string for the FText we took a snapshot of (used for an efficient pointer compare) */
-	TSharedPtr<FString, ESPMode::ThreadSafe> DisplayStringPtr;
+	FTextDisplayStringPtr DisplayStringPtr;
 
 	/** Revision index of the history of the FText we took a snapshot of, or INDEX_NONE if there was no history */
 	int32 HistoryRevision;
@@ -431,11 +433,11 @@ private:
 
 public:
 	static bool ShouldGatherForLocalization(const FText& Text);
-	static const FString* GetNamespace(const FText& Text);
-	static const FString* GetKey(const FText& Text);
+	static TOptional<FString> GetNamespace(const FText& Text);
+	static TOptional<FString> GetKey(const FText& Text);
 	static const FString* GetSourceString(const FText& Text);
 	static const FString& GetDisplayString(const FText& Text);
-	static const TSharedRef<FString, ESPMode::ThreadSafe> GetSharedDisplayString(const FText& Text);
+	static const FTextDisplayStringRef GetSharedDisplayString(const FText& Text);
 	static int32 GetFlags(const FText& Text);
 };
 
@@ -582,8 +584,9 @@ public:
 	~FScopedTextIdentityPreserver();
 
 private:
-	TSharedPtr< FString, ESPMode::ThreadSafe > Namespace;
-	TSharedPtr< FString, ESPMode::ThreadSafe > Key;
+	bool HadFoundNamespaceAndKey;
+	FString Namespace;
+	FString Key;
 	int32 Flags;
 	FText& TextToPersist;
 };
diff --git a/Engine/Source/Runtime/Core/Public/Internationalization/TextLocalizationManager.h b/Engine/Source/Runtime/Core/Public/Internationalization/TextLocalizationManager.h
--- a/Engine/Source/Runtime/Core/Public/Internationalization/TextLocalizationManager.h
+++ b/Engine/Source/Runtime/Core/Public/Internationalization/TextLocalizationManager.h
@@ -5,126 +5,163 @@
 #include "IInternationalizationArchiveSerializer.h"
 #include "IInternationalizationManifestSerializer.h"
 
+typedef TSharedRef<FString, ESPMode::ThreadSafe> FTextDisplayStringRef;
+typedef TSharedPtr<FString, ESPMode::ThreadSafe> FTextDisplayStringPtr;
 
-namespace ELocalizationResourceSource
-{
-	enum Type
-	{
-		ManifestArchive,
-		LocalizationResource
-	};
-}
-
-class FTextLocalizationManager
+/** Singleton class that manages display strings for FText. */
+class CORE_API FTextLocalizationManager
 {
 	friend CORE_API void BeginInitTextLocalization();
 	friend CORE_API void EndInitTextLocalization();
 
 private:
-	struct FLocalizationEntryTracker
+	/** Utility class for loading text localizations from some number of localization resources, tracking conflicts between them. */
+	class FLocalizationEntryTracker
 	{
+	public:
+		/** Data struct for tracking a localization entry from a localization resource. */
 		struct FEntry
 		{
-			FString TableName;
+			FString LocResID;
 			uint32 SourceStringHash;
 			FString LocalizedString;
 		};
 
 		typedef TArray<FEntry> FEntryArray;
-		typedef TMap<FString, FEntryArray> FKeyTable;
-		typedef TMap<FString, FKeyTable> FNamespaceTable;
+		typedef TMap<FString, FEntryArray> FKeysTable;
+		typedef TMap<FString, FKeysTable> FNamespacesTable;
+
+		FNamespacesTable Namespaces;
 
-		FNamespaceTable Namespaces;
+	public:
+		/** Loads all text localizations from all localization resource files in the specified directory. */
+		void LoadFromDirectory(const FString& DirectoryPath);
+		/** Loads all text localizations from the specified localization resource file. */
+		bool LoadFromFile(const FString& FilePath);
+		/** Loads all text localizations from the specified localization resource archive, associating the entries with the specified identifier. */
+		void LoadFromLocalizationResource(FArchive& Archive, const FString& LocResID);
 
-		void ReportCollisions() const;
-		void ReadFromDirectory(const FString& DirectoryPath);
-		bool ReadFromFile(const FString& FilePath);
-		void ReadFromArchive(FArchive& Archive, const FString& Identifier);
+		/** Detects conflicts between loaded localization resources and logs them as warnings. */
+		void DetectAndLogConflicts() const;
 	};
 
-	struct FStringEntry
+	/** Utility class for managing the currently loaded or registered text localizations. */
+	class FDisplayStringLookupTable
 	{
-		bool bIsLocalized;
-		FString TableName;
-		uint32 SourceStringHash;
-		TSharedRef<FString, ESPMode::ThreadSafe> String;
-
-		FStringEntry(const bool InIsLocalized, const FString& InTableName, const uint32 InSourceStringHash, const TSharedRef<FString, ESPMode::ThreadSafe>& InString)
-			:	bIsLocalized(InIsLocalized), TableName(InTableName), SourceStringHash(InSourceStringHash), String(InString)
+	public:
+		/** Data struct for tracking a display string. */
+		struct FDisplayStringEntry
 		{
-		}
-	};
+			bool bIsLocalized;
+			FString LocResID;
+			uint32 SourceStringHash;
+			FTextDisplayStringRef DisplayString;
+
+			FDisplayStringEntry(const bool InIsLocalized, const FString& InLocResID, const uint32 InSourceStringHash, const FTextDisplayStringRef& InDisplayString)
+				: bIsLocalized(InIsLocalized)
+				, LocResID(InLocResID)
+				, SourceStringHash(InSourceStringHash)
+				, DisplayString(InDisplayString)
+			{
+			}
+		};
 
-	struct FTextLookupTable
-	{
-		typedef TMap<FString, FStringEntry> FKeyTable;
-		typedef TMap<FString, FKeyTable> FNamespaceTable;
-		FNamespaceTable NamespaceTable;
+		typedef TMap<FString, FDisplayStringEntry> FKeysTable;
+		typedef TMap<FString, FKeysTable> FNamespacesTable;
+		FNamespacesTable NamespacesTable;
 
-		const TSharedRef<FString, ESPMode::ThreadSafe>* GetString(const FString& Namespace, const FString& Key, const uint32 SourceStringHash) const;
+	public:
+		/** Finds the keys table for the specified namespace and the display string entry for the specified namespace and key combination. If not found, the out parameters are set to null. */
+		void Find(const FString& InNamespace, FKeysTable*& OutKeysTableForNamespace, const FString& InKey, FDisplayStringEntry*& OutDisplayStringEntry);
+		/** Finds the keys table for the specified namespace and the display string entry for the specified namespace and key combination. If not found, the out parameters are set to null. */
+		void Find(const FString& InNamespace, const FKeysTable*& OutKeysTableForNamespace, const FString& InKey, const FDisplayStringEntry*& OutDisplayStringEntry) const;
 	};
 
-	/** Stores a Namespace and Key, used in reverse lookups from a DisplayString to find the Namespace and Key */
+	/** Simple data structure containing the name of the namespace and key associated with a display string, for use in looking up namespace and key from a display string. */
 	struct FNamespaceKeyEntry
 	{
-		TSharedPtr<FString, ESPMode::ThreadSafe> Namespace;
-		TSharedPtr<FString, ESPMode::ThreadSafe> Key;
+		FString Namespace;
+		FString Key;
 
-		FNamespaceKeyEntry(TSharedPtr<FString, ESPMode::ThreadSafe> InNamespace, TSharedPtr<FString, ESPMode::ThreadSafe> InKey)
+		FNamespaceKeyEntry(const FString& InNamespace, const FString& InKey)
 			: Namespace(InNamespace)
 			, Key(InKey)
 		{}
 	};
+	typedef TMap<FTextDisplayStringPtr, FNamespaceKeyEntry> FNamespaceKeyLookupTable;
 
-public:
-	static CORE_API FTextLocalizationManager& Get();
-
-	TSharedPtr<FString, ESPMode::ThreadSafe> CORE_API FindString(const FString& Namespace, const FString& Key, const FString* const SourceString = nullptr);
-	TSharedRef<FString, ESPMode::ThreadSafe> CORE_API GetString(const FString& Namespace, const FString& Key, const FString* const SourceString);
-
-	/**
-	 * Using an FText SourceString's, returns the associated Namespace and Key
-	 *
-	 * @param InSourceString		The FText DisplayString
-	 * @param OutNamespace			Returns with the associated Namespace
-	 * @param OutKey				Returns with the associated Key
-	 */
-	void CORE_API FindKeyNamespaceFromDisplayString(TSharedRef<FString, ESPMode::ThreadSafe> InDisplayString, TSharedPtr<FString, ESPMode::ThreadSafe>& OutNamespace, TSharedPtr<FString, ESPMode::ThreadSafe>& OutKey);
-
-	/** Get table name for given Namespace/key */
-	TSharedPtr<FString, ESPMode::ThreadSafe> CORE_API GetTableName(const FString& Namespace, const FString& Key);
-
-	void CORE_API RegenerateResources( const FString& ConfigFilePath, IInternationalizationArchiveSerializer& ArchiveSerializer, IInternationalizationManifestSerializer& ManifestSerializer );
-
-	/** Returns the current culture revision index */
-	int CORE_API GetHeadCultureRevision() const { return HeadCultureRevisionIndex; }
+private:
+	bool bIsInitialized;
 
-	/** Broadcasts whenever the new translations are available */
-	DECLARE_EVENT(FTextLocalizationManager, FTranslationsChangedEvent)
-	CORE_API FTranslationsChangedEvent&  OnTranslationsChanged() { return TranslationsChangedEvent; }
+	FCriticalSection SynchronizationObject;
+	FDisplayStringLookupTable DisplayStringLookupTable;
+	FNamespaceKeyLookupTable NamespaceKeyLookupTable;
+	int32 TextRevisionCounter;
 
 private:
 	FTextLocalizationManager() 
 		: bIsInitialized(false)
 		, SynchronizationObject()
-		, HeadCultureRevisionIndex(0)
+		, TextRevisionCounter(0)
 	{}
 
-	void LoadResources(const bool ShouldLoadEditor, const bool ShouldLoadGame);
-	void UpdateLiveTable(const TArray<FLocalizationEntryTracker>& LocalizationEntryTrackers, const bool FilterUpdatesByTableName = false);
-	void OnCultureChanged();
+public:
+	/** Singleton accessor */
+	static FTextLocalizationManager& Get();
+
+	/**	Finds and returns the display string with the given namespace and key, if it exists.
+	 *	Additionally, if a source string is specified and the found localized display string was not localized from that source string, null will be returned. */
+	FTextDisplayStringPtr FindDisplayString(const FString& Namespace, const FString& Key, const FString* const SourceString = nullptr);
+
+	/**	Returns a display string with the given namespace and key.
+	 *	If no display string exists, it will be created using the source string or an empty string if no source string is provided.
+	 *	If a display string exists ...
+	 *		... but it was not localized from the specified source string, the display string will be set to the specified source and returned.
+	 *		... and it was localized from the specified source string (or none was provided), the display string will be returned.
+	*/
+	FTextDisplayStringRef GetDisplayString(const FString& Namespace, const FString& Key, const FString* const SourceString);
+
+	/** If an entry exists for the specified namespace and key, returns true and provides the localization resource identifier from which it was loaded. Otherwise, returns false. */
+	bool GetLocResID(const FString& Namespace, const FString& Key, FString& OutLocResId);
+
+	/**	Finds the namespace and key associated with the specified display string.
+	 *	Returns true if found and sets the out parameters. Otherwise, returns false.
+	 */
+	bool FindNamespaceAndKeyFromDisplayString(const FTextDisplayStringRef& InDisplayString, FString& OutNamespace, FString& OutKey);
+	
+	/**	Attempts to register the specified display string, associating it with the specified namespace and key.
+	 *	Returns true if the display string has been or was already associated with the namespace and key.
+	 *	Returns false if the display string was already associated with another namespace and key or the namespace and key are already in use by another display string.
+	 */
+	bool AddDisplayString(const FTextDisplayStringRef& DisplayString, const FString& Namespace, const FString& Key);
+
+	/**
+	 * Updates the underlying value of a display string and associates it with a specified namespace and key, then returns true.
+	 * If the namespace and key are already in use by another display string, no changes occur and false is returned.
+	 */
+	bool UpdateDisplayString(const FTextDisplayStringRef& DisplayString, const FString& Value, const FString& Namespace, const FString& Key);
+
+	/** Loads localizations for the current culture based on a configuration file specifying which manifests and archives to use. Effectively generates a localization resource in memory. */
+	void LoadFromManifestAndArchives(const FString& ConfigFilePath, IInternationalizationArchiveSerializer& ArchiveSerializer, IInternationalizationManifestSerializer& ManifestSerializer);
+
+	/**	Returns the current text revision number. This value can be cached when caching information from the text localization manager.
+	 *	If the revision does not match, cached information may be invalid and should be recached. */
+	int32 GetTextRevision() const { return TextRevisionCounter; }
+
+	/** Event type for immediately reacting to changes in display strings for text. */
+	DECLARE_EVENT(FTextLocalizationManager, FTextRevisionChangedEvent)
+	FTextRevisionChangedEvent OnTextRevisionChangedEvent;
 
 private:
-	bool bIsInitialized;
-	FTextLookupTable LiveTable;
-	
-	/** Table for looking up namespaces and keys using the DisplayString of an FText */
-	TMap< TSharedPtr< FString, ESPMode::ThreadSafe >, FNamespaceKeyEntry > ReverseLiveTable;
+	/** Callback for changes in culture. Loads the new culture's localization resources. */
+	void OnCultureChanged();
 
-	FCriticalSection SynchronizationObject;
+	/** Loads localization resources for the specified culture, optionally loading localization resources that are editor-specific or game-specific. */
+	void LoadLocalizationResourcesForCulture(const FString& CultureName, const bool ShouldLoadEditor, const bool ShouldLoadGame);
 
-	/** The current culture revision, increments every time the culture is changed and allows for FTexts to know when to rebuild under a new culture */
-	int32 HeadCultureRevisionIndex;
+	/** Updates display string entries based on specified localizations and adds new display string entries. */
+	void UpdateFromLocalizations(const TArray<FLocalizationEntryTracker>& LocalizationEntryTrackers);
 
-	FTranslationsChangedEvent TranslationsChangedEvent;
+	/** Dirties the text revision counter by incrementing it, causing a revision mismatch for any information cached before this happens.  */
+	void DirtyTextRevision() { ++TextRevisionCounter; OnTextRevisionChangedEvent.Broadcast(); }
 };
\ No newline at end of file
diff --git a/Engine/Source/Runtime/SlateCore/Private/Fonts/LegacySlateFontInfoCache.cpp b/Engine/Source/Runtime/SlateCore/Private/Fonts/LegacySlateFontInfoCache.cpp
--- a/Engine/Source/Runtime/SlateCore/Private/Fonts/LegacySlateFontInfoCache.cpp
+++ b/Engine/Source/Runtime/SlateCore/Private/Fonts/LegacySlateFontInfoCache.cpp
@@ -74,7 +74,7 @@ TSharedPtr<const FCompositeFont> FLegacySlateFontInfoCache::GetFallbackFont()
 	FScopeLock Lock(&FallbackFontCS);
 
 	// The fallback font can change if the active culture is changed
-	const int32 CurrentHistoryVersion = FTextLocalizationManager::Get().GetHeadCultureRevision();
+	const int32 CurrentHistoryVersion = FTextLocalizationManager::Get().GetTextRevision();
 
 	if (!FallbackFont.IsValid() || FallbackFontHistoryVersion != CurrentHistoryVersion)
 	{
@@ -95,7 +95,7 @@ const FFontData& FLegacySlateFontInfoCache::GetFallbackFontData()
 	FScopeLock Lock(&FallbackFontDataCS);
 
 	// The fallback font can change if the active culture is changed
-	const int32 CurrentHistoryVersion = FTextLocalizationManager::Get().GetHeadCultureRevision();
+	const int32 CurrentHistoryVersion = FTextLocalizationManager::Get().GetTextRevision();
 
 	if (!FallbackFontData.IsValid() || FallbackFontDataHistoryVersion != CurrentHistoryVersion)
 	{
