PR #1015: Allow Windows windows to disable max/min buttons (Contributed by [REDACTED_PERSON])

diff --git a/Engine/Source/Runtime/Core/Private/Windows/WindowsWindow.cpp b/Engine/Source/Runtime/Core/Private/Windows/WindowsWindow.cpp
--- a/Engine/Source/Runtime/Core/Private/Windows/WindowsWindow.cpp
+++ b/Engine/Source/Runtime/Core/Private/Windows/WindowsWindow.cpp
@@ -110,7 +110,17 @@ void FWindowsWindow::Initialize( FWindowsApplication* const Application, const T
 	{
 		// OS Window border setup
 		WindowExStyle = WS_EX_APPWINDOW;
-		WindowStyle = WS_POPUP | WS_OVERLAPPED | WS_SYSMENU | WS_MINIMIZEBOX | WS_MAXIMIZEBOX | WS_BORDER | WS_CAPTION;
+		WindowStyle = WS_POPUP | WS_OVERLAPPED | WS_SYSMENU | WS_BORDER | WS_CAPTION;
+
+		if (Definition->SupportsMaximize)
+		{
+			WindowStyle |= WS_MAXIMIZEBOX;
+		}
+
+		if (Definition->SupportsMinimize)
+		{
+			WindowStyle |= WS_MINIMIZEBOX;
+		}
 
 		// Note SizeX and SizeY should be the size of the client area.  We need to get the actual window size by adjusting the client size to account for standard windows border around the window
 		RECT WindowRect = { 0, 0, ClientWidth, ClientWidth };
