Make sure FScopedSlowTask works properly across threads.

diff --git a/Engine/Source/Runtime/Core/Public/Misc/FeedbackContext.h b/Engine/Source/Runtime/Core/Public/Misc/FeedbackContext.h
--- a/Engine/Source/Runtime/Core/Public/Misc/FeedbackContext.h
+++ b/Engine/Source/Runtime/Core/Public/Misc/FeedbackContext.h
@@ -181,7 +181,7 @@ public:
 		, CurrentFrameScope(0)
 		, bVisibleOnUI(true)
 		, StartTime(FPlatformTime::Seconds())
-		, bEnabled(bInEnabled)
+		, bEnabled(bInEnabled && IsInGameThread())
 		, bCreatedDialog(false)		// only set to true if we create a dialog
 		, Context(InContext)
 	{
