Fix HTTPChunkInstaller log spam. UE-13764

diff --git a/Engine/Source/Runtime/Core/Private/Windows/WindowsPlatformMisc.cpp b/Engine/Source/Runtime/Core/Private/Windows/WindowsPlatformMisc.cpp
--- a/Engine/Source/Runtime/Core/Private/Windows/WindowsPlatformMisc.cpp
+++ b/Engine/Source/Runtime/Core/Private/Windows/WindowsPlatformMisc.cpp
@@ -2411,10 +2411,10 @@ EConvertibleLaptopMode FWindowsPlatformMisc::GetConvertibleLaptopMode()
 IPlatformChunkInstall* FWindowsPlatformMisc::GetPlatformChunkInstall()
 {
 	static IPlatformChunkInstall* ChunkInstall = nullptr;
-	IPlatformChunkInstallModule* PlatformChunkInstallModule = FModuleManager::LoadModulePtr<IPlatformChunkInstallModule>("HTTPChunkInstaller");
 	if (!ChunkInstall)
 	{
 #if !(WITH_EDITORONLY_DATA || IS_PROGRAM)
+		static IPlatformChunkInstallModule* PlatformChunkInstallModule = FModuleManager::LoadModulePtr<IPlatformChunkInstallModule>("HTTPChunkInstaller");
 		if (PlatformChunkInstallModule != NULL)
 		{
 			// Attempt to grab the platform installer
