Stats - Fixed FStatNameAndInfo::GetGroupCategoryFrom

diff --git a/Engine/Source/Runtime/Core/Private/Stats/Stats2.cpp b/Engine/Source/Runtime/Core/Private/Stats/Stats2.cpp
--- a/Engine/Source/Runtime/Core/Private/Stats/Stats2.cpp
+++ b/Engine/Source/Runtime/Core/Private/Stats/Stats2.cpp
@@ -639,7 +639,7 @@ FName FStatNameAndInfo::GetGroupCategoryFrom(FName InLongName)
 {
 	FString Input(InLongName.ToString());
 
-	const int32 IndexStart = Input.Find(TEXT("####"), ESearchCase::CaseSensitive, ESearchDir::FromEnd);
+	const int32 IndexStart = Input.Find(TEXT("####"), ESearchCase::CaseSensitive);
 	if (IndexStart != INDEX_NONE)
 	{
 		Input = Input.RightChop(IndexStart + 4);
