Fix case-sensitive filename error (PropertyPortflags->PropertyPortFlags.h)

diff --git a/Engine/Source/Runtime/Core/Private/Internationalization/TextHistory.cpp b/Engine/Source/Runtime/Core/Private/Internationalization/TextHistory.cpp
--- a/Engine/Source/Runtime/Core/Private/Internationalization/TextHistory.cpp
+++ b/Engine/Source/Runtime/Core/Private/Internationalization/TextHistory.cpp
@@ -3,7 +3,7 @@
 #include "CorePrivatePCH.h"
 
 #include "TextHistory.h"
-#include "PropertyPortflags.h"
+#include "PropertyPortFlags.h"
 
 ///////////////////////////////////////
 // FTextHistory
