TArray::AddDefaulted added.
Some comments fixed.

diff --git a/Engine/Source/Runtime/Core/Public/Containers/Array.h b/Engine/Source/Runtime/Core/Public/Containers/Array.h
--- a/Engine/Source/Runtime/Core/Public/Containers/Array.h
+++ b/Engine/Source/Runtime/Core/Public/Containers/Array.h
@@ -1764,16 +1764,16 @@ public:
 	FORCEINLINE int32 Add(const ElementType& Item) { CheckAddress(&Item); return Emplace(Item); }
 
 	/**
-	 * Adds a new item to the end of the array, possibly reallocating the whole
-	 * array to fit. The new item will be zeroed.
+	 * Adds new items to the end of the array, possibly reallocating the whole
+	 * array to fit. The new items will be zeroed.
 	 *
 	 * Caution, AddZeroed() will create elements without calling the
 	 * constructor and this is not appropriate for element types that require
 	 * a constructor to function properly.
 	 *
-	 * @param Item The item to add.
+	 * @param  Count  The number of new items to add.
 	 *
-	 * @return Index to the new item.
+	 * @return Index to the first of the new items.
 	 */
 	int32 AddZeroed(int32 Count = 1)
 	{
@@ -1782,6 +1782,21 @@ public:
 		return Index;
 	}
 
+	/**
+	 * Adds new items to the end of the array, possibly reallocating the whole
+	 * array to fit. The new items will be default-constructed.
+	 *
+	 * @param  Count  The number of new items to add.
+	 *
+	 * @return Index to the first of the new items.
+	 */
+	int32 AddDefaulted(int32 Count = 1)
+	{
+		const int32 Index = AddUninitialized(Count);
+		DefaultConstructItems<ElementType>((uint8*)AllocatorInstance.GetAllocation() + Index * sizeof(ElementType), Count);
+		return Index;
+	}
+
 private:
 	/**
 	 * Adds unique element to array if it doesn't exist.
