Merging CL 2516447 from LauncherDev to Main

diff --git a/Engine/Source/Runtime/Core/Private/Mac/MacPlatformMisc.cpp b/Engine/Source/Runtime/Core/Private/Mac/MacPlatformMisc.cpp
--- a/Engine/Source/Runtime/Core/Private/Mac/MacPlatformMisc.cpp
+++ b/Engine/Source/Runtime/Core/Private/Mac/MacPlatformMisc.cpp
@@ -392,6 +392,7 @@ void FMacPlatformMisc::PlatformPostInit(bool ShowSplashScreen)
 	{
 		// No MacApplication means that app is a dedicated server, commandline tool or the editor running a commandlet. In these cases we don't want OS X to put our app into App Nap mode.
 		CommandletActivity = [[NSProcessInfo processInfo] beginActivityWithOptions:NSActivityUserInitiated reason:IsRunningCommandlet() ? @"Running commandlet" : @"Running dedicated server"];
+		[CommandletActivity retain];
 	}
 }
 
@@ -401,6 +402,7 @@ void FMacPlatformMisc::PlatformTearDown()
 	{
 		MainThreadCall(^{
 			[[NSProcessInfo processInfo] endActivity:CommandletActivity];
+			[CommandletActivity release];
 		}, NSDefaultRunLoopMode, false);
 		CommandletActivity = nil;
 	}
