UEPLAT-633
iOS now looks in document directory before application directory in preparation for chunked install, DLC, and content patching
COTF and iterative deploy copy data to the /Documents/cacheddata directory
#ios

diff --git a/Engine/Source/Programs/IOS/MobileDeviceInterface/MobileDeviceInstance.cs b/Engine/Source/Programs/IOS/MobileDeviceInterface/MobileDeviceInstance.cs
--- a/Engine/Source/Programs/IOS/MobileDeviceInterface/MobileDeviceInstance.cs
+++ b/Engine/Source/Programs/IOS/MobileDeviceInterface/MobileDeviceInstance.cs
@@ -611,7 +611,7 @@ namespace Manzana
 							string Trimmed = Filename.Trim();
 							string SourceFilename = BaseFolder + "\\" + Trimmed;
 							SourceFilename = SourceFilename.Replace('/', '\\');
-							string DestFilename = "/Documents/" + Trimmed.Replace("cookeddata/", "");
+							string DestFilename = "/Documents/" + Trimmed.Replace("cookeddata/", "cacheddata/");
 							DestFilename = DestFilename.Replace('\\', '/');
 							WriteProgressLine("Copying '{0}' -> '{1}' ...", 0, SourceFilename, DestFilename);
 							CopyFileToPhone(SourceFilename, DestFilename, 1024 * 1024);
diff --git a/Engine/Source/Runtime/Core/Private/IOS/IOSPlatformFile.cpp b/Engine/Source/Runtime/Core/Private/IOS/IOSPlatformFile.cpp
--- a/Engine/Source/Runtime/Core/Private/IOS/IOSPlatformFile.cpp
+++ b/Engine/Source/Runtime/Core/Private/IOS/IOSPlatformFile.cpp
@@ -365,11 +365,11 @@ bool FIOSPlatformFile::FileExists(const TCHAR* Filename)
 {
 	struct stat FileInfo;
 	FString NormalizedFilename = NormalizeFilename(Filename);
-	// check the read path
-	if (stat(TCHAR_TO_UTF8(*ConvertToIOSPath(NormalizedFilename, false)), &FileInfo) == -1)
+	// check the write path
+	if (stat(TCHAR_TO_UTF8(*ConvertToIOSPath(NormalizedFilename, true)), &FileInfo) == -1)
 	{
-		// if not in read path, check the write path
-		if (stat(TCHAR_TO_UTF8(*ConvertToIOSPath(NormalizedFilename, true)), &FileInfo) == -1)
+		// if not in write path, check the read path
+		if (stat(TCHAR_TO_UTF8(*ConvertToIOSPath(NormalizedFilename, false)), &FileInfo) == -1)
 		{
 			return false;
 		}
@@ -383,11 +383,11 @@ int64 FIOSPlatformFile::FileSize(const TCHAR* Filename)
 	struct stat FileInfo;
 	FileInfo.st_size = -1;
 	FString NormalizedFilename = NormalizeFilename(Filename);
-	// check the read path
-	if(stat(TCHAR_TO_UTF8(*ConvertToIOSPath(NormalizedFilename, false)), &FileInfo) == -1)
+	// check the write path
+	if(stat(TCHAR_TO_UTF8(*ConvertToIOSPath(NormalizedFilename, true)), &FileInfo) == -1)
 	{
-		// if not in read path, check the write path
-		stat(TCHAR_TO_UTF8(*ConvertToIOSPath(NormalizedFilename, true)), &FileInfo);
+		// if not in write path, check the read path
+		stat(TCHAR_TO_UTF8(*ConvertToIOSPath(NormalizedFilename, false)), &FileInfo);
 	}
 
 	// make sure to return -1 for directories
@@ -408,12 +408,12 @@ bool FIOSPlatformFile::DeleteFile(const TCHAR* Filename)
 bool FIOSPlatformFile::IsReadOnly(const TCHAR* Filename)
 {
 	FString NormalizedFilename = NormalizeFilename(Filename);
-	FString Filepath = ConvertToIOSPath(NormalizedFilename, false);
-	// check read path
+	FString Filepath = ConvertToIOSPath(NormalizedFilename, true);
+	// check write path
 	if (access(TCHAR_TO_UTF8(*Filepath), F_OK) == -1)
 	{
-		// if not in read path, check write path
-		Filepath = ConvertToIOSPath(NormalizedFilename, true);
+		// if not in write path, check read path
+		Filepath = ConvertToIOSPath(NormalizedFilename, false);
 		if (access(TCHAR_TO_UTF8(*Filepath), F_OK) == -1)
 		{
 			return false; // file doesn't exist
@@ -460,11 +460,11 @@ FDateTime FIOSPlatformFile::GetTimeStamp(const TCHAR* Filename)
 	// get file times
 	struct stat FileInfo;
 	FString NormalizedFilename = NormalizeFilename(Filename);
-	// check the read path
-	if(stat(TCHAR_TO_UTF8(*ConvertToIOSPath(NormalizedFilename, false)), &FileInfo) == -1)
+	// check the write path
+	if(stat(TCHAR_TO_UTF8(*ConvertToIOSPath(NormalizedFilename, true)), &FileInfo) == -1)
 	{
-		// if not in the read path, check the write path
-		if(stat(TCHAR_TO_UTF8(*ConvertToIOSPath(NormalizedFilename, true)), &FileInfo) == -1)
+		// if not in the write path, check the read path
+		if(stat(TCHAR_TO_UTF8(*ConvertToIOSPath(NormalizedFilename, false)), &FileInfo) == -1)
 		{
 			return FDateTime::MinValue();
 		}
@@ -498,11 +498,11 @@ FDateTime FIOSPlatformFile::GetAccessTimeStamp(const TCHAR* Filename)
 	// get file times
 	struct stat FileInfo;
 	FString NormalizedFilename = NormalizeFilename(Filename);
-	// check the read path
-	if(stat(TCHAR_TO_UTF8(*ConvertToIOSPath(NormalizedFilename, false)), &FileInfo) == -1)
+	// check the write path
+	if(stat(TCHAR_TO_UTF8(*ConvertToIOSPath(NormalizedFilename, true)), &FileInfo) == -1)
 	{
-		// if not in the read path, check the write path
-		if(stat(TCHAR_TO_UTF8(*ConvertToIOSPath(NormalizedFilename, true)), &FileInfo) == -1)
+		// if not in the write path, check the read path
+		if(stat(TCHAR_TO_UTF8(*ConvertToIOSPath(NormalizedFilename, false)), &FileInfo) == -1)
 		{
 			return FDateTime::MinValue();
 		}
@@ -522,13 +522,13 @@ IFileHandle* FIOSPlatformFile::OpenRead(const TCHAR* Filename, bool bAllowWrite)
 {
 	FString NormalizedFilename = NormalizeFilename(Filename);
 
-	// check the read path
-	FString FinalPath = ConvertToIOSPath(NormalizedFilename, false);
+	// check the write path
+	FString FinalPath = ConvertToIOSPath(NormalizedFilename, true);
 	int32 Handle = open(TCHAR_TO_UTF8(*FinalPath), O_RDONLY);
 	if(Handle == -1)
 	{
-		// if not in the read path, check the write path
-		FinalPath = ConvertToIOSPath(NormalizedFilename, true);
+		// if not in the write path, check the read path
+		FinalPath = ConvertToIOSPath(NormalizedFilename, false);
 		Handle = open(TCHAR_TO_UTF8(*FinalPath), O_RDONLY);
 	}
 
@@ -576,9 +576,9 @@ bool FIOSPlatformFile::DirectoryExists(const TCHAR* Directory)
 {
 	struct stat FileInfo;
 	FString NormalizedDirectory = NormalizeFilename(Directory);
-	if (stat(TCHAR_TO_UTF8(*ConvertToIOSPath(NormalizedDirectory, false)), &FileInfo)== -1)
+	if (stat(TCHAR_TO_UTF8(*ConvertToIOSPath(NormalizedDirectory, true)), &FileInfo)== -1)
 	{
-		if (stat(TCHAR_TO_UTF8(*ConvertToIOSPath(NormalizedDirectory, true)), &FileInfo)== -1)
+		if (stat(TCHAR_TO_UTF8(*ConvertToIOSPath(NormalizedDirectory, false)), &FileInfo)== -1)
 		{
 			return false;
 		}
@@ -619,11 +619,11 @@ bool FIOSPlatformFile::IterateDirectory(const TCHAR* Directory, FDirectoryVisito
 
 	FString NormalizedDirectory = NormalizeFilename(Directory);
 	// If Directory is an empty string, assume that we want to iterate Binaries/Mac (current dir), but because we're an app bundle, iterate bundle's Contents/Frameworks instead
-	DIR* Handle = opendir(Directory[0] ? TCHAR_TO_UTF8(*ConvertToIOSPath(NormalizedDirectory, false)) : FrameworksPath);
+	DIR* Handle = opendir(Directory[0] ? TCHAR_TO_UTF8(*ConvertToIOSPath(NormalizedDirectory, true)) : FrameworksPath);
 	if(!Handle)
 	{
 		// look in the write file path if it's not in the read file path
-		Handle = opendir(Directory[0] ? TCHAR_TO_UTF8(*ConvertToIOSPath(NormalizedDirectory, true)) : FrameworksPath);
+		Handle = opendir(Directory[0] ? TCHAR_TO_UTF8(*ConvertToIOSPath(NormalizedDirectory, false)) : FrameworksPath);
 	}
 	if (Handle)
 	{
@@ -651,32 +651,19 @@ FString FIOSPlatformFile::ConvertToIOSPath(const FString& Filename, bool bForWri
 	Result.ReplaceInline(FPlatformProcess::BaseDir(), TEXT(""));
 
 	if(bForWrite)
-	{
-		static FString WritePathBase = FString([NSSearchPathForDirectoriesInDomains(NSDocumentDirectory, NSUserDomainMask, YES) objectAtIndex:0]) + TEXT("/");
-		return WritePathBase + Result;
-	}
-	else
 	{
 		// if filehostip exists in the command line, cook on the fly read path should be used
 		FString Value;
 		// Cache this value as the command line doesn't change...
 		static bool bHasHostIP = FParse::Value(FCommandLine::Get(), TEXT("filehostip"), Value) || FParse::Value(FCommandLine::Get(), TEXT("streaminghostip"), Value);
 		static bool bIsIterative = FParse::Value(FCommandLine::Get(), TEXT("iterative"), Value);
-		if (bHasHostIP)
-		{
-			static FString ReadPathBase = FString([NSSearchPathForDirectoriesInDomains(NSDocumentDirectory, NSUserDomainMask, YES) objectAtIndex:0]) + TEXT("/");
-			return ReadPathBase + Result;
-		}
-		else if (bIsIterative)
-		{
-			static FString ReadPathBase = FString([NSSearchPathForDirectoriesInDomains(NSDocumentDirectory, NSUserDomainMask, YES) objectAtIndex:0]) + TEXT("/");
-			return ReadPathBase + Result.ToLower();
-		}
-		else
-		{
-			static FString ReadPathBase = FString([[NSBundle mainBundle] bundlePath]) + TEXT("/cookeddata/");
-			return ReadPathBase + Result.ToLower();
-		}
+		static FString WritePathBase = FString([NSSearchPathForDirectoriesInDomains(NSDocumentDirectory, NSUserDomainMask, YES) objectAtIndex:0]) + TEXT("/");
+		return WritePathBase + ((bHasHostIP || bIsIterative) ? TEXT("cacheddata/") : TEXT("")) + Result.ToLower();
+	}
+	else
+	{
+		static FString ReadPathBase = FString([[NSBundle mainBundle] bundlePath]) + TEXT("/cookeddata/");
+		return ReadPathBase + Result.ToLower();
 	}
 
 	return Result;
