UE-9698 - Adding file log flushing on crash (Android and iOS)
#ue4
#android
#ios

diff --git a/Engine/Source/Runtime/Core/Private/Android/AndroidOutputDevices.cpp b/Engine/Source/Runtime/Core/Private/Android/AndroidOutputDevices.cpp
--- a/Engine/Source/Runtime/Core/Private/Android/AndroidOutputDevices.cpp
+++ b/Engine/Source/Runtime/Core/Private/Android/AndroidOutputDevices.cpp
@@ -62,11 +62,14 @@ void FOutputDeviceAndroidError::HandleError()
 
 	if (NewCallCount != 1)
 	{
+		UE_LOG(LogAndroid, Error, TEXT("HandleError re-entered."));
 		return;
 	}
+	
+	GIsGuarded = 0;
+	GIsRunning = 0;
+	GIsCriticalError = 1;
+	GLogConsole = NULL;
 
-	GIsGuarded			= 0;
-	GIsRunning			= 0;
-	GIsCriticalError	= 1;
-	GLogConsole			= NULL;
+	GLog->Flush();
 }
diff --git a/Engine/Source/Runtime/Core/Private/IOS/IOSPlatformOutputDevices.cpp b/Engine/Source/Runtime/Core/Private/IOS/IOSPlatformOutputDevices.cpp
--- a/Engine/Source/Runtime/Core/Private/IOS/IOSPlatformOutputDevices.cpp
+++ b/Engine/Source/Runtime/Core/Private/IOS/IOSPlatformOutputDevices.cpp
@@ -81,4 +81,6 @@ void FOutputDeviceIOSError::HandleError()
 	GIsRunning = 0;
 	GIsCriticalError = 1;
 	GLogConsole = NULL;
+
+	GLog->Flush();
 }
