Minor formatting changes

diff --git a/Engine/Source/Runtime/Core/Private/Windows/WindowsPlatformProcess.cpp b/Engine/Source/Runtime/Core/Private/Windows/WindowsPlatformProcess.cpp
--- a/Engine/Source/Runtime/Core/Private/Windows/WindowsPlatformProcess.cpp
+++ b/Engine/Source/Runtime/Core/Private/Windows/WindowsPlatformProcess.cpp
@@ -283,34 +283,29 @@ FProcHandle FWindowsPlatformProcess::CreateProc( const TCHAR* URL, const TCHAR*
 
 	// initialize process attributes
 	SECURITY_ATTRIBUTES Attr;
-	{
 	Attr.nLength = sizeof(SECURITY_ATTRIBUTES);
 	Attr.lpSecurityDescriptor = NULL;
 	Attr.bInheritHandle = true;
-	}
 
 	// initialize process creation flags
 	uint32 CreateFlags = NORMAL_PRIORITY_CLASS;
-	{
 	if (PriorityModifier < 0)
 	{
-			CreateFlags = (PriorityModifier == -1) ? BELOW_NORMAL_PRIORITY_CLASS : IDLE_PRIORITY_CLASS;
+		CreateFlags = (PriorityModifier == -1) ? BELOW_NORMAL_PRIORITY_CLASS : IDLE_PRIORITY_CLASS;
 	}
 	else if (PriorityModifier > 0)
 	{
-			CreateFlags = (PriorityModifier == 1) ? ABOVE_NORMAL_PRIORITY_CLASS : HIGH_PRIORITY_CLASS;
+		CreateFlags = (PriorityModifier == 1) ? ABOVE_NORMAL_PRIORITY_CLASS : HIGH_PRIORITY_CLASS;
 	}
 
 	if (bLaunchDetached)
 	{
 		CreateFlags |= DETACHED_PROCESS;
 	}
-	}
 
 	// initialize window flags
 	uint32 dwFlags = 0;
 	uint16 ShowWindowFlags = SW_HIDE;
-	{
 	if (bLaunchReallyHidden)
 	{
 		dwFlags = STARTF_USESHOWWINDOW;
@@ -321,11 +316,10 @@ FProcHandle FWindowsPlatformProcess::CreateProc( const TCHAR* URL, const TCHAR*
 		ShowWindowFlags = SW_SHOWMINNOACTIVE;
 	}
 
-		if (PipeWrite != nullptr)
+	if (PipeWrite != nullptr)
 	{
 		dwFlags |= STARTF_USESTDHANDLES;
 	}
-	}
 
 	// initialize startup info
 	STARTUPINFO StartupInfo = {
