Fix Malloc incorrectly giving OOM error when using aligned memory

diff --git a/Engine/Source/Runtime/Core/Private/HAL/MallocTBB.cpp b/Engine/Source/Runtime/Core/Private/HAL/MallocTBB.cpp
--- a/Engine/Source/Runtime/Core/Private/HAL/MallocTBB.cpp
+++ b/Engine/Source/Runtime/Core/Private/HAL/MallocTBB.cpp
@@ -44,7 +44,7 @@ void* FMallocTBB::Malloc( SIZE_T Size, uint32 Alignment )
 		NewPtr = scalable_malloc( Size );
 	}
 
-	if( !NewPtr )
+	if( !NewPtr && Size )
 	{
 		OutOfMemory(Size, Alignment);
 	}
