Removed redundant function parameters

diff --git a/Engine/Source/Developer/Settings/Private/SettingsSection.cpp b/Engine/Source/Developer/Settings/Private/SettingsSection.cpp
--- a/Engine/Source/Developer/Settings/Private/SettingsSection.cpp
+++ b/Engine/Source/Developer/Settings/Private/SettingsSection.cpp
@@ -177,7 +177,7 @@ bool FSettingsSection::ResetDefaults()
 		GConfig->EmptySection(*SettingsObject->GetClass()->GetPathName(), ConfigName);
 		GConfig->Flush(false);
 
-		FConfigCacheIni::LoadGlobalIniFile(ConfigName, *FPaths::GetBaseFilename(ConfigName), nullptr, nullptr, true);
+		FConfigCacheIni::LoadGlobalIniFile(ConfigName, *FPaths::GetBaseFilename(ConfigName), nullptr, true);
 
 		SettingsObject->ReloadConfig(nullptr, nullptr, UE4::LCPF_PropagateToInstances|UE4::LCPF_PropagateToChildDefaultObjects);
 
diff --git a/Engine/Source/Developer/SourceControl/Private/SourceControlHelpers.cpp b/Engine/Source/Developer/SourceControl/Private/SourceControlHelpers.cpp
--- a/Engine/Source/Developer/SourceControl/Private/SourceControlHelpers.cpp
+++ b/Engine/Source/Developer/SourceControl/Private/SourceControlHelpers.cpp
@@ -27,7 +27,7 @@ const FString& GetSettingsIni()
 		if(SourceControlSettingsIni.Len() == 0)
 		{
 			const FString SourceControlSettingsDir = FPaths::GeneratedConfigDir();
-			FConfigCacheIni::LoadGlobalIniFile(SourceControlSettingsIni, TEXT("SourceControlSettings"), NULL, NULL, false, false, true, *SourceControlSettingsDir);
+			FConfigCacheIni::LoadGlobalIniFile(SourceControlSettingsIni, TEXT("SourceControlSettings"), NULL, false, false, true, *SourceControlSettingsDir);
 		}
 		return SourceControlSettingsIni;
 	}
@@ -39,7 +39,7 @@ const FString& GetGlobalSettingsIni()
 	if(SourceControlGlobalSettingsIni.Len() == 0)
 	{
 		const FString SourceControlSettingsDir = FPaths::EngineSavedDir() + TEXT("Config/");
-		FConfigCacheIni::LoadGlobalIniFile(SourceControlGlobalSettingsIni, TEXT("SourceControlSettings"), NULL, NULL, false, false, true, *SourceControlSettingsDir);
+		FConfigCacheIni::LoadGlobalIniFile(SourceControlGlobalSettingsIni, TEXT("SourceControlSettings"), NULL, false, false, true, *SourceControlSettingsDir);
 	}
 	return SourceControlGlobalSettingsIni;
 }
diff --git a/Engine/Source/Editor/UnrealEd/Private/Factories/PackFactory.cpp b/Engine/Source/Editor/UnrealEd/Private/Factories/PackFactory.cpp
--- a/Engine/Source/Editor/UnrealEd/Private/Factories/PackFactory.cpp
+++ b/Engine/Source/Editor/UnrealEd/Private/Factories/PackFactory.cpp
@@ -418,7 +418,7 @@ UObject* UPackFactory::FactoryCreateBinary
 						NewFile.UpdateSections(*EngineIniFilename, *RedirectsSection);
 
 						FString FinalIniFileName;
-						GConfig->LoadGlobalIniFile(FinalIniFileName, *RedirectsSection, NULL, NULL, true);
+						GConfig->LoadGlobalIniFile(FinalIniFileName, *RedirectsSection, NULL, true);
 
 						FLinkerLoad::AddGameNameRedirect(*LongOldGameName, *LongNewGameName);
 						FLinkerLoad::AddGameNameRedirect(*ConfigParameters.GameName, *LongNewGameName);
diff --git a/Engine/Source/Editor/UnrealEd/Private/StaticLightingSystem/StaticLightingSystem.cpp b/Engine/Source/Editor/UnrealEd/Private/StaticLightingSystem/StaticLightingSystem.cpp
--- a/Engine/Source/Editor/UnrealEd/Private/StaticLightingSystem/StaticLightingSystem.cpp
+++ b/Engine/Source/Editor/UnrealEd/Private/StaticLightingSystem/StaticLightingSystem.cpp
@@ -428,7 +428,7 @@ bool FStaticLightingSystem::BeginLightmassProcess()
 			GWarn->BeginSlowTask( LOCTEXT("BeginBuildingStaticLightingTaskStatus", "Building lighting"), false );
 		}
 		
-		FConfigCacheIni::LoadGlobalIniFile(GLightmassIni, TEXT("Lightmass"), NULL, NULL, true);
+		FConfigCacheIni::LoadGlobalIniFile(GLightmassIni, TEXT("Lightmass"), NULL, true);
 		verify(GConfig->GetBool(TEXT("DevOptions.StaticLighting"), TEXT("bUseBilinearFilterLightmaps"), GUseBilinearLightmaps, GLightmassIni));
 		verify(GConfig->GetBool(TEXT("DevOptions.StaticLighting"), TEXT("bAllowCropping"), GAllowLightmapCropping, GLightmassIni));
 		verify(GConfig->GetBool(TEXT("DevOptions.StaticLighting"), TEXT("bRebuildDirtyGeometryForLighting"), bRebuildDirtyGeometryForLighting, GLightmassIni));
diff --git a/Engine/Source/Runtime/Core/Private/HAL/ConsoleManager.cpp b/Engine/Source/Runtime/Core/Private/HAL/ConsoleManager.cpp
--- a/Engine/Source/Runtime/Core/Private/HAL/ConsoleManager.cpp
+++ b/Engine/Source/Runtime/Core/Private/HAL/ConsoleManager.cpp
@@ -873,7 +873,7 @@ void FConsoleManager::LoadHistoryIfNeeded()
 	FConfigFile Ini;
 
 	FString ConfigPath = FPaths::GeneratedConfigDir() + TEXT("ConsoleHistory.ini");
-	ProcessIniContents(*ConfigPath, *ConfigPath, &Ini, false, false, false);
+	ProcessIniContents(*ConfigPath, *ConfigPath, &Ini, false, false);
 
 	const FString History = TEXT("History");
 
diff --git a/Engine/Source/Runtime/Core/Private/Misc/ConfigCacheIni.cpp b/Engine/Source/Runtime/Core/Private/Misc/ConfigCacheIni.cpp
--- a/Engine/Source/Runtime/Core/Private/Misc/ConfigCacheIni.cpp
+++ b/Engine/Source/Runtime/Core/Private/Misc/ConfigCacheIni.cpp
@@ -765,9 +765,8 @@ static bool LoadIniFileHierarchy(const TArray<FIniFilename>& HierarchyToLoad, FC
 
 			bool bDoEmptyConfig = false;
 			bool bDoCombine = (IniIndex != 0);
-			bool bDoWrite = false;
 			//UE_LOG(LogConfig, Log,  TEXT( "Combining configFile: %s" ), *IniList(IniIndex) );
-			ProcessIniContents(*HierarchyToLoad.Last().Filename, *IniFileName, &ConfigFile, bDoEmptyConfig, bDoCombine, bDoWrite);
+			ProcessIniContents(*HierarchyToLoad.Last().Filename, *IniFileName, &ConfigFile, bDoEmptyConfig, bDoCombine);
 #if INI_CACHE
 			if ( bShouldCache )
 			{
@@ -1259,7 +1258,7 @@ void FConfigFile::ProcessSourceAndCheckAgainstBackup()
 		FString BackupFilename = FString::Printf( TEXT( "%s%s.ini" ), *BetweenRunsDir, *Name.ToString() );
 
 		FConfigFile BackupFile;
-		ProcessIniContents(*BackupFilename, *BackupFilename, &BackupFile, false, false, false);
+		ProcessIniContents(*BackupFilename, *BackupFilename, &BackupFile, false, false);
 
 		for( TMap<FString,FConfigSection>::TIterator SectionIterator(*SourceConfigFile); SectionIterator; ++SectionIterator )
 		{
@@ -1573,8 +1572,7 @@ void FConfigCacheIni::LoadFile( const FString& Filename, const FConfigFile* Fall
 		FConfigFile* Result = &Add(Filename, FConfigFile());
 		bool bDoEmptyConfig = false;
 		bool bDoCombine = false;
-		bool bDoWrite = true;
-		ProcessIniContents(*Filename, *Filename, Result, bDoEmptyConfig, bDoCombine, bDoWrite);
+		ProcessIniContents(*Filename, *Filename, Result, bDoEmptyConfig, bDoCombine);
 		UE_LOG(LogConfig, Log, TEXT( "GConfig::LoadFile has loaded file:  %s" ), *Filename);
 	}
 	else if( Fallback )
@@ -2553,7 +2551,7 @@ static void LoadAnIniFile(const FString& FilenameToLoad, FConfigFile& ConfigFile
 
 	if( !IsUsingLocalIniFile(*FilenameToLoad, NULL) || (IFileManager::Get().FileSize( *FilenameToLoad ) >= 0) )
 	{
-		ProcessIniContents(*FilenameToLoad, *FilenameToLoad, &ConfigFile, false, false, false);
+		ProcessIniContents(*FilenameToLoad, *FilenameToLoad, &ConfigFile, false, false);
 	}
 	else
 	{
@@ -2712,10 +2710,9 @@ static FString GetSourceIniFilename(const TCHAR* ConfigDir, const TCHAR* Prefix,
  *
  * @param InBaseIniName Ini name.
  * @param InPlatformName Platform name, NULL means to use the current platform
- * @param InGameName Game name, NULL means to use the current
  * @param OutHierarchy An array which is to receive the generated hierachy of ini filenames.
  */
-static void GetSourceIniHierarchyFilenames(const TCHAR* InBaseIniName, const TCHAR* InPlatformName, const TCHAR* InGameName, const TCHAR* EngineConfigDir, const TCHAR* SourceConfigDir, TArray<FIniFilename>& OutHierarchy, bool bRequireDefaultIni)
+static void GetSourceIniHierarchyFilenames(const TCHAR* InBaseIniName, const TCHAR* InPlatformName, const TCHAR* EngineConfigDir, const TCHAR* SourceConfigDir, TArray<FIniFilename>& OutHierarchy, bool bRequireDefaultIni)
 {
 	
 	/**************************************************
@@ -2794,15 +2791,12 @@ static void GetSourceIniHierarchyFilenames(const TCHAR* InBaseIniName, const TCH
  * 
  * @param IniBaseName Base name of the .ini (Engine, Game)
  * @param PlatformName Name of the platform to get the .ini path for (NULL means to use the current platform)
- * @param GameName The name of the game to get the .ini path for (if NULL, uses current)
  * @param GeneratedConfigDir The base folder that will contain the generated config files.
  * 
  * @return Standardized .ini filename
  */
-static FString GetDestIniFilename(const TCHAR* BaseIniName, const TCHAR* PlatformName, const TCHAR* GameName, const TCHAR* GeneratedConfigDir)
+static FString GetDestIniFilename(const TCHAR* BaseIniName, const TCHAR* PlatformName, const TCHAR* GeneratedConfigDir)
 {
-	checkf(GameName==NULL, TEXT("Specifying a GameName in GetDestIniFilename requires appGeneratedConfigDir to take a GameName parameter"));
-
 	// figure out what to look for on the commandline for an override
 	FString CommandLineSwitch = FString::Printf(TEXT("%sINI="), BaseIniName);
 	
@@ -2828,7 +2822,7 @@ void FConfigCacheIni::InitializeConfigSystem()
 	// load the main .ini files (unless we're running a program or a gameless UE4Editor.exe, DefaultEngine.ini is required).
 	const bool bIsGamelessExe = !FApp::HasGameName();
 	const bool bDefaultEngineIniRequired = !bIsGamelessExe && (GIsGameAgnosticExe || FApp::IsGameNameEmpty());
-	bool bEngineConfigCreated = FConfigCacheIni::LoadGlobalIniFile(GEngineIni, TEXT("Engine"), NULL, NULL, bDefaultEngineIniRequired);
+	bool bEngineConfigCreated = FConfigCacheIni::LoadGlobalIniFile(GEngineIni, TEXT("Engine"), NULL, bDefaultEngineIniRequired);
 
 	if ( !bIsGamelessExe )
 	{
@@ -2864,7 +2858,7 @@ void FConfigCacheIni::InitializeConfigSystem()
 	if ( !GIsBuildMachine )
 	{
 		const FString GameAgnosticEditorSettingsDir = FPaths::Combine(*FPaths::GameAgnosticSavedDir(), TEXT("Config")) + TEXT("/");
-		FConfigCacheIni::LoadGlobalIniFile(GEditorGameAgnosticIni, TEXT("EditorGameAgnostic"), NULL, NULL, false, false, true, *GameAgnosticEditorSettingsDir);
+		FConfigCacheIni::LoadGlobalIniFile(GEditorGameAgnosticIni, TEXT("EditorGameAgnostic"), NULL, false, false, true, *GameAgnosticEditorSettingsDir);
 	}
 #endif
 #if PLATFORM_DESKTOP
@@ -2883,10 +2877,10 @@ void FConfigCacheIni::InitializeConfigSystem()
 	GConfig->bIsReadyForUse = true;
 }
 
-bool FConfigCacheIni::LoadGlobalIniFile(FString& FinalIniFilename, const TCHAR* BaseIniName, const TCHAR* Platform, const TCHAR* GameName, bool bForceReload, bool bRequireDefaultIni, bool bAllowGeneratedIniWhenCooked, const TCHAR* GeneratedConfigDir)
+bool FConfigCacheIni::LoadGlobalIniFile(FString& FinalIniFilename, const TCHAR* BaseIniName, const TCHAR* Platform, bool bForceReload, bool bRequireDefaultIni, bool bAllowGeneratedIniWhenCooked, const TCHAR* GeneratedConfigDir)
 {
 	// figure out where the end ini file is
-	FinalIniFilename = GetDestIniFilename(BaseIniName, Platform, GameName, GeneratedConfigDir);
+	FinalIniFilename = GetDestIniFilename(BaseIniName, Platform, GeneratedConfigDir);
 	
 	// Start the loading process for the remote config file when appropriate
 	if (FRemoteConfig::Get()->ShouldReadRemoteFile(*FinalIniFilename))
@@ -2919,7 +2913,7 @@ bool FConfigCacheIni::LoadGlobalIniFile(FString& FinalIniFilename, const TCHAR*
 
 
 	// calculate the source ini file name,
-	GetSourceIniHierarchyFilenames( BaseIniName, Platform, GameName, *FPaths::EngineConfigDir(), *FPaths::SourceConfigDir(), NewConfigFile.SourceIniHierarchy, bRequireDefaultIni );
+	GetSourceIniHierarchyFilenames( BaseIniName, Platform, *FPaths::EngineConfigDir(), *FPaths::SourceConfigDir(), NewConfigFile.SourceIniHierarchy, bRequireDefaultIni );
 
 	// Keep a record of the original settings
 	NewConfigFile.SourceConfigFile = new FConfigFile();
@@ -2947,11 +2941,11 @@ bool FConfigCacheIni::LoadGlobalIniFile(FString& FinalIniFilename, const TCHAR*
 	return bResult;
 }
 
-void FConfigCacheIni::LoadLocalIniFile(FConfigFile& ConfigFile, const TCHAR* IniName, bool bGenerateDestIni, const TCHAR* Platform, const TCHAR* GameName, const bool bForceReload )
+void FConfigCacheIni::LoadLocalIniFile(FConfigFile& ConfigFile, const TCHAR* IniName, bool bGenerateDestIni, const TCHAR* Platform, const bool bForceReload )
 {
 	DECLARE_SCOPE_CYCLE_COUNTER( TEXT( "FConfigCacheIni::LoadLocalIniFile" ), STAT_FConfigCacheIni_LoadLocalIniFile, STATGROUP_LoadTime );
 
-	LoadExternalIniFile( ConfigFile, IniName, *FPaths::EngineConfigDir(), *FPaths::SourceConfigDir(), bGenerateDestIni, Platform, GameName, bForceReload );
+	LoadExternalIniFile( ConfigFile, IniName, *FPaths::EngineConfigDir(), *FPaths::SourceConfigDir(), bGenerateDestIni, Platform, bForceReload );
 
 
 	// if bGenerateDestIni is false, that means the .ini is a ready-to-go .ini file, and just needs to be loaded into the FConfigFile
@@ -2965,19 +2959,19 @@ void FConfigCacheIni::LoadLocalIniFile(FConfigFile& ConfigFile, const TCHAR* Ini
 	}
 	else
 	{
-		GetSourceIniHierarchyFilenames( IniName, Platform, GameName, *FPaths::EngineConfigDir(), *FPaths::SourceConfigDir(), ConfigFile.SourceIniHierarchy, false );
+		GetSourceIniHierarchyFilenames( IniName, Platform, *FPaths::EngineConfigDir(), *FPaths::SourceConfigDir(), ConfigFile.SourceIniHierarchy, false );
 		
 		// Keep a record of the original settings
 		ConfigFile.SourceConfigFile = new FConfigFile();
 
 		// now generate and make sure it's up to date (using IniName as a Base for an ini filename)
 		const bool bAllowGeneratedINIs = true;
-		GenerateDestIniFile(ConfigFile, GetDestIniFilename(IniName, Platform, GameName, *FPaths::GeneratedConfigDir()), ConfigFile.SourceIniHierarchy, bAllowGeneratedINIs, true);
+		GenerateDestIniFile(ConfigFile, GetDestIniFilename(IniName, Platform, *FPaths::GeneratedConfigDir()), ConfigFile.SourceIniHierarchy, bAllowGeneratedINIs, true);
 	}
 	ConfigFile.Name = IniName;*/
 }
 
-void FConfigCacheIni::LoadExternalIniFile(FConfigFile& ConfigFile, const TCHAR* IniName, const TCHAR* EngineConfigDir, const TCHAR* SourceConfigDir, bool bGenerateDestIni, const TCHAR* Platform, const TCHAR* GameName, const bool bForceReload)
+void FConfigCacheIni::LoadExternalIniFile(FConfigFile& ConfigFile, const TCHAR* IniName, const TCHAR* EngineConfigDir, const TCHAR* SourceConfigDir, bool bGenerateDestIni, const TCHAR* Platform, const bool bForceReload)
 {
 	// if bGenerateDestIni is false, that means the .ini is a ready-to-go .ini file, and just needs to be loaded into the FConfigFile
 	if (!bGenerateDestIni)
@@ -2990,7 +2984,7 @@ void FConfigCacheIni::LoadExternalIniFile(FConfigFile& ConfigFile, const TCHAR*
 	}
 	else
 	{
-		GetSourceIniHierarchyFilenames( IniName, Platform, GameName, EngineConfigDir, SourceConfigDir, ConfigFile.SourceIniHierarchy, false );
+		GetSourceIniHierarchyFilenames( IniName, Platform, EngineConfigDir, SourceConfigDir, ConfigFile.SourceIniHierarchy, false );
 
 		if ( bForceReload )
 		{
@@ -3003,7 +2997,7 @@ void FConfigCacheIni::LoadExternalIniFile(FConfigFile& ConfigFile, const TCHAR*
 
 		// now generate and make sure it's up to date (using IniName as a Base for an ini filename)
 		const bool bAllowGeneratedINIs = true;
-		GenerateDestIniFile(ConfigFile, GetDestIniFilename(IniName, Platform, GameName, *FPaths::GeneratedConfigDir()), ConfigFile.SourceIniHierarchy, bAllowGeneratedINIs, true);
+		GenerateDestIniFile(ConfigFile, GetDestIniFilename(IniName, Platform, *FPaths::GeneratedConfigDir()), ConfigFile.SourceIniHierarchy, bAllowGeneratedINIs, true);
 	}
 	ConfigFile.Name = IniName;
 }
@@ -3111,7 +3105,7 @@ void FConfigFile::UpdateSections(const TCHAR* DiskFilename, const TCHAR* IniRoot
 	{
 		// get the standard ini files
 		SourceIniHierarchy.Empty();
-		GetSourceIniHierarchyFilenames(IniRootName, NULL, NULL, *FPaths::EngineConfigDir(), *FPaths::SourceConfigDir(), SourceIniHierarchy, false);
+		GetSourceIniHierarchyFilenames(IniRootName, NULL, *FPaths::EngineConfigDir(), *FPaths::SourceConfigDir(), SourceIniHierarchy, false);
 
 		// now chop off this file and any after it
 		for (int32 FileIndex = 0; FileIndex < SourceIniHierarchy.Num(); FileIndex++)
diff --git a/Engine/Source/Runtime/Core/Private/Misc/RemoteConfigIni.cpp b/Engine/Source/Runtime/Core/Private/Misc/RemoteConfigIni.cpp
--- a/Engine/Source/Runtime/Core/Private/Misc/RemoteConfigIni.cpp
+++ b/Engine/Source/Runtime/Core/Private/Misc/RemoteConfigIni.cpp
@@ -443,7 +443,7 @@ bool IsUsingLocalIniFile(const TCHAR* FilenameToLoad, const TCHAR* IniFileName)
 }
 
 /** Contains the logic for processing config files, local or remote */
-void ProcessIniContents(const TCHAR* FilenameToLoad, const TCHAR* IniFileName, FConfigFile* Config, bool bDoEmptyConfig, bool bDoCombine, bool bDoWrite)
+void ProcessIniContents(const TCHAR* FilenameToLoad, const TCHAR* IniFileName, FConfigFile* Config, bool bDoEmptyConfig, bool bDoCombine)
 {
 	DECLARE_SCOPE_CYCLE_COUNTER( TEXT( "ProcessIniContents" ), STAT_ProcessIniContents, STATGROUP_LoadTime );
 
diff --git a/Engine/Source/Runtime/Core/Public/Misc/ConfigCacheIni.h b/Engine/Source/Runtime/Core/Public/Misc/ConfigCacheIni.h
--- a/Engine/Source/Runtime/Core/Public/Misc/ConfigCacheIni.h
+++ b/Engine/Source/Runtime/Core/Public/Misc/ConfigCacheIni.h
@@ -515,14 +515,13 @@ public:
 	 * @param FinalIniFilename The output name of the generated .ini file (in Game\Saved\Config)
 	 * @param BaseIniName The "base" ini name, with no extension (ie, Engine, Game, etc)
 	 * @param Platform The platform to load the .ini for (if NULL, uses current)
-	 * @param GameName The name of the game to load the .ini for (if NULL, uses current)
 	 * @param bForceReload If true, the destination .in will be regenerated from the source, otherwise this will only process if the dest isn't in GConfig
 	 * @param bRequireDefaultIni If true, the Default*.ini file is required to exist when generating the final ini file.
 	 * @param bAllowGeneratedIniWhenCooked If true, the engine will attempt to load the generated/user INI file when loading.
 	 * @param GeneratedConfigDir The location where generated config files are made.
 	 * @return true if the final ini was created successfully.
 	 */
-	static bool LoadGlobalIniFile(FString& FinalIniFilename, const TCHAR* BaseIniName, const TCHAR* Platform=NULL, const TCHAR* GameName=NULL, bool bForceReload=false, bool bRequireDefaultIni=false, bool bAllowGeneratedIniWhenCooked=true, const TCHAR* GeneratedConfigDir = *FPaths::GeneratedConfigDir());
+	static bool LoadGlobalIniFile(FString& FinalIniFilename, const TCHAR* BaseIniName, const TCHAR* Platform=NULL, bool bForceReload=false, bool bRequireDefaultIni=false, bool bAllowGeneratedIniWhenCooked=true, const TCHAR* GeneratedConfigDir = *FPaths::GeneratedConfigDir());
 
 	/**
 	 * Load an ini file directly into an FConfigFile, and nothing is written to GConfig or disk. 
@@ -533,10 +532,9 @@ public:
 	 * @param IniName Either a Base ini name (Engine) or a full ini name (WrangleContent). NO PATH OR EXTENSION SHOULD BE USED!
 	 * @param bIsBaseIniName true if IniName is a Base name, which can be overridden on commandline, etc.
 	 * @param Platform The platform to use for Base ini names, NULL means to use the current platform
-	 * @param GameName The game to use for Base ini names, NULL means to use the current
 	 * @param bForceReload force reload the ini file from disk this is required if you make changes to the ini file not using the config system as the hierarchy cache will not be updated in this case
 	 */
-	static void LoadLocalIniFile(FConfigFile& ConfigFile, const TCHAR* IniName, bool bIsBaseIniName, const TCHAR* Platform=NULL, const TCHAR* GameName=NULL, const bool bForceReload=false);
+	static void LoadLocalIniFile(FConfigFile& ConfigFile, const TCHAR* IniName, bool bIsBaseIniName, const TCHAR* Platform=NULL, const bool bForceReload=false);
 
 	/**
 	 * Load an ini file directly into an FConfigFile from the specified config folders, and nothing is written to GConfig or disk. 
@@ -549,10 +547,9 @@ public:
 	 * @param SourceConfigDir Game config directory.
 	 * @param bGenerateDestIni true if IniName is a Base name, which can be overridden on commandline, etc.
 	 * @param Platform The platform to use for Base ini names
-	 * @param GameName The game to use for Base ini names
 	 * @param bForceReload force reload the ini file from disk this is required if you make changes to the ini file not using the config system as the hierarchy cache will not be updated in this case
 	 */
-	static void LoadExternalIniFile(FConfigFile& ConfigFile, const TCHAR* IniName, const TCHAR* EngineConfigDir, const TCHAR* SourceConfigDir, bool bGenerateDestIni, const TCHAR* Platform=NULL, const TCHAR* GameName=NULL, const bool bForceReload=false);
+	static void LoadExternalIniFile(FConfigFile& ConfigFile, const TCHAR* IniName, const TCHAR* EngineConfigDir, const TCHAR* SourceConfigDir, bool bGenerateDestIni, const TCHAR* Platform=NULL, const bool bForceReload=false);
 
 	/**
 	 * Needs to be called after GConfig is set and LoadCoalescedFile was called.
diff --git a/Engine/Source/Runtime/Core/Public/Misc/RemoteConfigIni.h b/Engine/Source/Runtime/Core/Public/Misc/RemoteConfigIni.h
--- a/Engine/Source/Runtime/Core/Public/Misc/RemoteConfigIni.h
+++ b/Engine/Source/Runtime/Core/Public/Misc/RemoteConfigIni.h
@@ -211,7 +211,7 @@ private:
 bool IsUsingLocalIniFile(const TCHAR* FilenameToLoad, const TCHAR* IniFileName);
 
 /** Contains the logic for processing config files, local or remote. */
-void ProcessIniContents(const TCHAR* FilenameToLoad, const TCHAR* IniFileName, FConfigFile* Config, bool bDoEmptyConfig, bool bDoCombine, bool bDoWrite);
+void ProcessIniContents(const TCHAR* FilenameToLoad, const TCHAR* IniFileName, FConfigFile* Config, bool bDoEmptyConfig, bool bDoCombine);
 
 /** Returns the timestamp of the appropriate config file. */
 FDateTime GetIniTimeStamp(const TCHAR* FilenameToLoad, const TCHAR* IniFileName);
diff --git a/Engine/Source/Runtime/CoreUObject/Private/UObject/Obj.cpp b/Engine/Source/Runtime/CoreUObject/Private/UObject/Obj.cpp
--- a/Engine/Source/Runtime/CoreUObject/Private/UObject/Obj.cpp
+++ b/Engine/Source/Runtime/CoreUObject/Private/UObject/Obj.cpp
@@ -1851,7 +1851,7 @@ void UObject::UpdateSingleSectionOfConfigFile(const FString& ConfigIniName)
 
 	// reload the file, so that it refresh the cache internally.
 	FString FinalIniFileName;
-	GConfig->LoadGlobalIniFile(FinalIniFileName, *GetClass()->ClassConfigName.ToString(), NULL, NULL, true);
+	GConfig->LoadGlobalIniFile(FinalIniFileName, *GetClass()->ClassConfigName.ToString(), NULL, true);
 }
 
 void UObject::UpdateDefaultConfigFile()
diff --git a/Engine/Source/Runtime/Engine/Private/GameUserSettings.cpp b/Engine/Source/Runtime/Engine/Private/GameUserSettings.cpp
--- a/Engine/Source/Runtime/Engine/Private/GameUserSettings.cpp
+++ b/Engine/Source/Runtime/Engine/Private/GameUserSettings.cpp
@@ -279,7 +279,7 @@ void UGameUserSettings::SaveSettings()
 void UGameUserSettings::LoadConfigIni( bool bForceReload/*=false*/ )
 {
 	// Load .ini, allowing merging
-	FConfigCacheIni::LoadGlobalIniFile(GGameUserSettingsIni, TEXT("GameUserSettings"), NULL, NULL, bForceReload);
+	FConfigCacheIni::LoadGlobalIniFile(GGameUserSettingsIni, TEXT("GameUserSettings"), NULL, bForceReload);
 }
 
 void UGameUserSettings::PreloadResolutionSettings()
