Cleaned up Crash Reporter Rocket checks, changed some code so that it's the same for all build types.

diff --git a/Engine/Source/Programs/CrashReporter/CrashReportClient/Private/CrashDescription.cpp b/Engine/Source/Programs/CrashReporter/CrashReportClient/Private/CrashDescription.cpp
--- a/Engine/Source/Programs/CrashReporter/CrashReportClient/Private/CrashDescription.cpp
+++ b/Engine/Source/Programs/CrashReporter/CrashReportClient/Private/CrashDescription.cpp
@@ -215,13 +215,9 @@ void FCrashDescription::InitializeIDs()
 	// The Epic ID can be looked up from this ID.
 	EpicAccountId = FPlatformMisc::GetEpicAccountId();
 
-	// Get the user name only for non-UE4 releases.
-	if( !FRocketSupport::IsRocket() )
-	{
-		// Remove periods from internal user names to match AutoReporter user names
-		// The name prefix is read by CrashRepository.AddNewCrash in the website code
-		UserName = FString( FPlatformProcess::UserName() ).Replace( TEXT( "." ), TEXT( "" ) );
-	}
+	// Remove periods from user names to match AutoReporter user names
+	// The name prefix is read by CrashRepository.AddNewCrash in the website code
+	UserName = FString( FPlatformProcess::UserName() ).Replace( TEXT( "." ), TEXT( "" ) );
 }
 
 void FCrashDescription::SendAnalytics()
diff --git a/Engine/Source/Programs/CrashReporter/CrashReportClient/Private/CrashReportClient.cpp b/Engine/Source/Programs/CrashReporter/CrashReportClient/Private/CrashReportClient.cpp
--- a/Engine/Source/Programs/CrashReporter/CrashReportClient/Private/CrashReportClient.cpp
+++ b/Engine/Source/Programs/CrashReporter/CrashReportClient/Private/CrashReportClient.cpp
@@ -200,15 +200,7 @@ void FDiagnoseReportWorker::DoWork()
 
 FText FCrashReportUtil::FormatDiagnosticText( const FText& DiagnosticText, const FString MachineId, const FString EpicAccountId, const FString UserNameNoDot )
 {
-	if( FRocketSupport::IsRocket() )
-	{
-		return FText::Format( LOCTEXT( "CrashReportClientCallstackPattern", "MachineId:{0}\nEpicAccountId:{1}\n\n{2}" ), FText::FromString( MachineId ), FText::FromString( EpicAccountId ), DiagnosticText );
-	}
-	else
-	{
-		return FText::Format( LOCTEXT( "CrashReportClientCallstackPattern", "MachineId:{0}\n\n{1}" ), FText::FromString( MachineId ), DiagnosticText );
-	}
-
+	return FText::Format( LOCTEXT( "CrashReportClientCallstackPattern", "MachineId:{0}\nEpicAccountId:{1}\n\n{2}" ), FText::FromString( MachineId ), FText::FromString( EpicAccountId ), DiagnosticText );
 }
 
 #undef LOCTEXT_NAMESPACE
diff --git a/Engine/Source/Programs/CrashReporter/CrashReportClient/Private/GenericErrorReport.cpp b/Engine/Source/Programs/CrashReporter/CrashReportClient/Private/GenericErrorReport.cpp
--- a/Engine/Source/Programs/CrashReporter/CrashReportClient/Private/GenericErrorReport.cpp
+++ b/Engine/Source/Programs/CrashReporter/CrashReportClient/Private/GenericErrorReport.cpp
@@ -73,17 +73,9 @@ bool FGenericErrorReport::SetUserComment(const FText& UserComment, bool bAllowTo
 		DynamicSignaturesNode->AppendChildNode(TEXT("Parameter3"), UserComment.ToString());
 	}
 	
-	FString MachineIDandUserID;
 	// Set global user name ID: will be added to the report
 	extern FCrashDescription& GetCrashDescription();
-	if( FRocketSupport::IsRocket() )
-	{
-		MachineIDandUserID = FString::Printf( TEXT( "!MachineId:%s!EpicAccountId:%s" ), *GetCrashDescription().MachineId, *GetCrashDescription().EpicAccountId );
-	}
-	else
-	{
-		MachineIDandUserID = FString::Printf( TEXT( "!MachineId:%s!Name:%s" ), *GetCrashDescription().MachineId, *GetCrashDescription().UserName );
-	}
+	FString MachineIDandUserID = FString::Printf( TEXT( "!MachineId:%s!EpicAccountId:%s!Name:%s" ), *GetCrashDescription().MachineId, *GetCrashDescription().EpicAccountId, *GetCrashDescription().UserName );
 
 	// Add or update a user ID.
 	FXmlNode* Parameter4Node = DynamicSignaturesNode->FindChildNode(TEXT("Parameter4"));
diff --git a/Engine/Source/Runtime/Core/Private/GenericPlatform/GenericPlatformContext.cpp b/Engine/Source/Runtime/Core/Private/GenericPlatform/GenericPlatformContext.cpp
--- a/Engine/Source/Runtime/Core/Private/GenericPlatform/GenericPlatformContext.cpp
+++ b/Engine/Source/Runtime/Core/Private/GenericPlatform/GenericPlatformContext.cpp
@@ -94,17 +94,11 @@ void FGenericCrashContext::SerializeContentToBuffer()
 	AddCrashProperty( TEXT( "LanguageLCID" ), FInternationalization::Get().GetCurrentCulture()->GetLCID() );
 	AddCrashProperty( TEXT( "DefaultLocale" ), *DefaultLocale );
 
-	AddCrashProperty( TEXT( "IsUE4Release" ), (int32)FRocketSupport::IsRocket() );
-
-	// Get the user name only for non-UE4 releases.
-	if( !FRocketSupport::IsRocket() )
-	{
-		// Remove periods from internal user names to match AutoReporter user names
-		// The name prefix is read by CrashRepository.AddNewCrash in the website code
-
-		AddCrashProperty( TEXT( "UserName" ), *UserName.Replace( TEXT( "." ), TEXT( "" ) ) );
-	}
+	AddCrashProperty( TEXT( "IsUE4Release" ), (int32)FApp::IsEngineInstalled() );
 
+	// Remove periods from user names to match AutoReporter user names
+	// The name prefix is read by CrashRepository.AddNewCrash in the website code
+	AddCrashProperty( TEXT( "UserName" ), *UserName.Replace( TEXT( "." ), TEXT( "" ) ) );
 
 	AddCrashProperty( TEXT( "BaseDir" ), *BaseDir );
 	AddCrashProperty( TEXT( "RootDir" ), *RootDir );
