CIS fix for passing one preprocessor macro to another

diff --git a/Engine/Source/Runtime/Core/Private/Misc/ConfigManifest.cpp b/Engine/Source/Runtime/Core/Private/Misc/ConfigManifest.cpp
--- a/Engine/Source/Runtime/Core/Private/Misc/ConfigManifest.cpp
+++ b/Engine/Source/Runtime/Core/Private/Misc/ConfigManifest.cpp
@@ -7,6 +7,8 @@
 #include "RemoteConfigIni.h"
 #include "EngineVersion.h"
 
+#include "Runtime/Launch/Resources/Version.h"
+
 enum class EConfigManifestVersion
 {
 	/******* DO NOT REMOVE OLD VERSIONS ********/
@@ -67,7 +69,7 @@ void MigratePreviousEngineInis()
 	IPlatformFile& PlatformFile = FPlatformFileManager::Get().GetPlatformFile();
 	while(PreviousVersion.GetMinor() >= 0)
 	{
-		const FString Directory = FString(FPlatformProcess::UserSettingsDir()) / TEXT(EPIC_PRODUCT_IDENTIFIER) / PreviousVersion.ToString(EVersionComponent::Minor) / TEXT("Saved") / TEXT("Config") / ANSI_TO_TCHAR(FPlatformProperties::PlatformName());
+		const FString Directory = FString(FPlatformProcess::UserSettingsDir()) / ENGINE_VERSION_TEXT(EPIC_PRODUCT_IDENTIFIER) / PreviousVersion.ToString(EVersionComponent::Minor) / TEXT("Saved") / TEXT("Config") / ANSI_TO_TCHAR(FPlatformProperties::PlatformName());
 		if (FPaths::DirectoryExists(Directory))
 		{
 			const FString DestDir = ProjectAgnosticIniPath(TEXT(""));
