Core - IPlatformFile::CreateDirectoryTree buffer overflow (UE-13663)

diff --git a/Engine/Source/Runtime/Core/Private/GenericPlatform/GenericPlatformFile.cpp b/Engine/Source/Runtime/Core/Private/GenericPlatform/GenericPlatformFile.cpp
--- a/Engine/Source/Runtime/Core/Private/GenericPlatform/GenericPlatformFile.cpp
+++ b/Engine/Source/Runtime/Core/Private/GenericPlatform/GenericPlatformFile.cpp
@@ -207,7 +207,9 @@ bool IPlatformFile::CreateDirectoryTree(const TCHAR* Directory)
 	FPaths::NormalizeDirectoryName(LocalFilename);
 	const TCHAR* LocalPath = *LocalFilename;
 	int32 CreateCount = 0;
-	for (TCHAR Full[MAX_UNREAL_FILENAME_LENGTH]=TEXT( "" ), *Ptr=Full; ; *Ptr++=*LocalPath++ )
+	const int32 MaxCharacters = MAX_UNREAL_FILENAME_LENGTH - 1;
+	int32 Index = 0;
+	for (TCHAR Full[MaxCharacters + 1] = TEXT( "" ), *Ptr = Full; Index < MaxCharacters; *Ptr++ = *LocalPath++, Index++)
 	{
 		if (((*LocalPath) == TEXT('/')) || (*LocalPath== 0))
 		{
