#ue4 - Optimized FTransform::TransformPosition() and FTransform::TransformPositionNoScale().

- Vectorized versions are now 25% faster than previous implementation.

diff --git a/Engine/Source/Runtime/Core/Public/Math/Transform.h b/Engine/Source/Runtime/Core/Public/Math/Transform.h
--- a/Engine/Source/Runtime/Core/Public/Math/Transform.h
+++ b/Engine/Source/Runtime/Core/Public/Math/Transform.h
@@ -1142,7 +1142,7 @@ FORCEINLINE FVector4 FTransform::TransformFVector4NoScale(const FVector4& V) con
 
 	//Transform using QST is following
 	//QST(P) = Q*S*P*-Q + T where Q = quaternion, S = scale, T = translation
-	FVector4 Transform = Rotation*V;
+	FVector4 Transform = Rotation.RotateVector(V);
 	if (V.W == 1.f)
 	{
 		Transform += FVector4(Translation, 1.f);
@@ -1162,7 +1162,7 @@ FORCEINLINE FVector4 FTransform::TransformFVector4(const FVector4& V) const
 	//Transform using QST is following
 	//QST(P) = Q*S*P*-Q + T where Q = quaternion, S = scale, T = translation
 
-	FVector4 Transform = (Rotation*(Scale3D*V));
+	FVector4 Transform = (Rotation.RotateVector(Scale3D*V));
 	if (V.W == 1.f)
 	{
 		Transform += FVector4(Translation, 1.f);
@@ -1175,12 +1175,14 @@ FORCEINLINE FVector4 FTransform::TransformFVector4(const FVector4& V) const
 
 FORCEINLINE FVector FTransform::TransformPosition(const FVector& V) const
 {
-	return TransformFVector4(FVector4(V.X,V.Y,V.Z,1.0f));
+	DiagnosticCheckNaN_All();
+	return Rotation.RotateVector(Scale3D*V) + Translation;
 }
 
 FORCEINLINE FVector FTransform::TransformPositionNoScale(const FVector& V) const
 {
-	return TransformFVector4NoScale(FVector4(V.X,V.Y,V.Z,1.0f));
+	DiagnosticCheckNaN_All();
+	return Rotation.RotateVector(V) + Translation;
 }
 
 FORCEINLINE FVector FTransform::TransformVector(const FVector& V) const
diff --git a/Engine/Source/Runtime/Core/Public/Math/TransformVectorized.h b/Engine/Source/Runtime/Core/Public/Math/TransformVectorized.h
--- a/Engine/Source/Runtime/Core/Public/Math/TransformVectorized.h
+++ b/Engine/Source/Runtime/Core/Public/Math/TransformVectorized.h
@@ -1446,12 +1446,45 @@ FORCEINLINE FVector4 FTransform::TransformFVector4(const FVector4& V) const
 
 FORCEINLINE FVector FTransform::TransformPosition(const FVector& V) const
 {
-	return TransformFVector4(FVector4(V.X,V.Y,V.Z,1.0f));
+	DiagnosticCheckNaN_All();
+
+	const VectorRegister InputVectorW0 = VectorLoadFloat3_W0(&V);
+
+	//Transform using QST is following
+	//QST(P) = Q*S*P*-Q + T where Q = quaternion, S = scale, T = translation
+	const VectorRegister InverseRotation = MAKE_QUATINV_VECTORREGISTER(Rotation);
+
+	//FQuat Transform = Rotation*FQuat(Scale*V.X, Scale*V.Y, Scale*V.Z, 0.f)*Rotation.Inverse();
+	const VectorRegister ScaledVec = VectorMultiply(Scale3D, InputVectorW0);
+	const VectorRegister TempStorage = VectorQuaternionMultiply2(Rotation, ScaledVec);
+	const VectorRegister RotatedVec = VectorQuaternionMultiply2(TempStorage, InverseRotation);
+
+	const VectorRegister TranslatedVec = VectorAdd(RotatedVec, Translation);
+
+	FVector Result;
+	VectorStoreFloat3(TranslatedVec, &Result);
+	return Result;
 }
 
 FORCEINLINE FVector FTransform::TransformPositionNoScale(const FVector& V) const
 {
-	return TransformFVector4NoScale(FVector4(V.X,V.Y,V.Z,1.0f));
+	DiagnosticCheckNaN_All();
+
+	const VectorRegister InputVectorW0 = VectorLoadFloat3_W0(&V);
+
+	//Transform using QST is following
+	//QST(P) = Q*S*P*-Q + T where Q = quaternion, S = 1.0f, T = translation
+	const VectorRegister InverseRotation = MAKE_QUATINV_VECTORREGISTER(Rotation);
+
+	//FQuat Transform = Rotation*FQuat(V.X, V.Y, V.Z, 0.f)*Rotation.Inverse();
+	const VectorRegister TempStorage = VectorQuaternionMultiply2(Rotation, InputVectorW0);
+	const VectorRegister RotatedVec = VectorQuaternionMultiply2(TempStorage, InverseRotation);
+
+	const VectorRegister TranslatedVec = VectorAdd(RotatedVec, Translation);
+
+	FVector Result;
+	VectorStoreFloat3(TranslatedVec, &Result);
+	return Result;
 }
 
 FORCEINLINE FVector FTransform::TransformVector(const FVector& V) const
