#ue4 - Added more DiagnosticCheckNaN_All() checks in common FTransform methods.

- Better to catch them early than skip a NaN check for translation just because we didn't use it, only to catch it later.

diff --git a/Engine/Source/Runtime/Core/Public/Math/Transform.h b/Engine/Source/Runtime/Core/Public/Math/Transform.h
--- a/Engine/Source/Runtime/Core/Public/Math/Transform.h
+++ b/Engine/Source/Runtime/Core/Public/Math/Transform.h
@@ -1187,26 +1187,27 @@ FORCEINLINE FVector FTransform::TransformPositionNoScale(const FVector& V) const
 
 FORCEINLINE FVector FTransform::TransformVector(const FVector& V) const
 {
-	DiagnosticCheckNaN_Rotate();
-	DiagnosticCheckNaN_Scale3D();
+	DiagnosticCheckNaN_All();
 	return Rotation.RotateVector(Scale3D*V);
 }
 
 FORCEINLINE FVector FTransform::TransformVectorNoScale(const FVector& V) const
 {
-	DiagnosticCheckNaN_Rotate();
+	DiagnosticCheckNaN_All();
 	return Rotation.RotateVector(V);
 }
 
 // do backward operation when inverse, translation -> rotation -> scale
 FORCEINLINE FVector FTransform::InverseTransformPosition(const FVector &V) const
 {
+	DiagnosticCheckNaN_All();
 	return ( Rotation.UnrotateVector(V-Translation) ) * GetSafeScaleReciprocal(Scale3D);
 }
 
 // do backward operation when inverse, translation -> rotation
 FORCEINLINE FVector FTransform::InverseTransformPositionNoScale(const FVector &V) const
 {
+	DiagnosticCheckNaN_All();
 	return ( Rotation.UnrotateVector(V-Translation) );
 }
 
@@ -1214,12 +1215,14 @@ FORCEINLINE FVector FTransform::InverseTransformPositionNoScale(const FVector &V
 // do backward operation when inverse, translation -> rotation -> scale
 FORCEINLINE FVector FTransform::InverseTransformVector(const FVector &V) const
 {
+	DiagnosticCheckNaN_All();
 	return ( Rotation.UnrotateVector(V) ) * GetSafeScaleReciprocal(Scale3D);
 }
 
 // do backward operation when inverse, translation -> rotation
 FORCEINLINE FVector FTransform::InverseTransformVectorNoScale(const FVector &V) const
 {
+	DiagnosticCheckNaN_All();
 	return ( Rotation.UnrotateVector(V) );
 }
 
diff --git a/Engine/Source/Runtime/Core/Public/Math/TransformVectorized.h b/Engine/Source/Runtime/Core/Public/Math/TransformVectorized.h
--- a/Engine/Source/Runtime/Core/Public/Math/TransformVectorized.h
+++ b/Engine/Source/Runtime/Core/Public/Math/TransformVectorized.h
@@ -1387,8 +1387,7 @@ FORCEINLINE FTransform FTransform::GetScaled(FVector InScale) const
 
 FORCEINLINE FVector4 FTransform::TransformFVector4NoScale(const FVector4& V) const
 {
-	DiagnosticCheckNaN_Rotate();
-	DiagnosticCheckNaN_Translate();
+	DiagnosticCheckNaN_All();
 
 	// if not, this won't work
 	checkSlow (V.W == 0.f || V.W == 1.f);
@@ -1489,8 +1488,7 @@ FORCEINLINE FVector FTransform::TransformPositionNoScale(const FVector& V) const
 
 FORCEINLINE FVector FTransform::TransformVector(const FVector& V) const
 {
-	DiagnosticCheckNaN_Rotate();
-	DiagnosticCheckNaN_Scale3D();
+	DiagnosticCheckNaN_All();
 
 	const VectorRegister InputVectorW0 = VectorLoadFloat3_W0(&V);
 	const VectorRegister InverseRotation = MAKE_QUATINV_VECTORREGISTER(Rotation);
@@ -1507,7 +1505,7 @@ FORCEINLINE FVector FTransform::TransformVector(const FVector& V) const
 
 FORCEINLINE FVector FTransform::TransformVectorNoScale(const FVector& V) const
 {
-	DiagnosticCheckNaN_Rotate();
+	DiagnosticCheckNaN_All();
 
 	const VectorRegister InputVectorW0 = VectorLoadFloat3_W0(&V);
 	const VectorRegister InverseRotation = MAKE_QUATINV_VECTORREGISTER(Rotation);
@@ -1524,6 +1522,8 @@ FORCEINLINE FVector FTransform::TransformVectorNoScale(const FVector& V) const
 // do backward operation when inverse, translation -> rotation -> scale
 FORCEINLINE FVector FTransform::InverseTransformPosition(const FVector &V) const
 {
+	DiagnosticCheckNaN_All();
+
 	VectorRegister InputVector = VectorLoadFloat3_W0(&V);
 
 	// Rotation.Inverse()
@@ -1550,6 +1550,8 @@ FORCEINLINE FVector FTransform::InverseTransformPosition(const FVector &V) const
 // do backward operation when inverse, translation -> rotation
 FORCEINLINE FVector FTransform::InverseTransformPositionNoScale(const FVector &V) const
 {
+	DiagnosticCheckNaN_All();
+
 	VectorRegister InputVector = VectorLoadFloat3_W0(&V);
 
 	// Rotation.Inverse()
@@ -1571,6 +1573,8 @@ FORCEINLINE FVector FTransform::InverseTransformPositionNoScale(const FVector &V
 // do backward operation when inverse, translation -> rotation -> scale
 FORCEINLINE FVector FTransform::InverseTransformVector(const FVector &V) const
 {
+	DiagnosticCheckNaN_All();
+
 	VectorRegister InputVector = VectorLoadFloat3_W0(&V);
 
 	// Rotation.Inverse()
@@ -1594,6 +1598,8 @@ FORCEINLINE FVector FTransform::InverseTransformVector(const FVector &V) const
 // do backward operation when inverse, translation -> rotation
 FORCEINLINE FVector FTransform::InverseTransformVectorNoScale(const FVector &V) const
 {
+	DiagnosticCheckNaN_All();
+
 	VectorRegister InputVector = VectorLoadFloat3_W0(&V);
 
 	// Rotation.Inverse()
