Converted FText::Format to use the same code on all platforms

We have two backends for FText i18n support; ICU, and our own legacy backend.

We typically favor the ICU backend where available as it provides much more accurate i18n support than our own legacy code, however in the case of FText::Format, we weren't taking advantage of ICU at all (nor do we plan to take advantage of its advanced syntax or message formatting features), so it makes more sense to always use the legacy version instead.

This change makes the legacy version the default regardless of which i18n backend is being used. This also involved making some improvements to the legacy code to support newer features that are required by the editor, such as text history.

diff --git a/Engine/Source/Runtime/Core/Private/Internationalization/ICUText.cpp b/Engine/Source/Runtime/Core/Private/Internationalization/ICUText.cpp
--- a/Engine/Source/Runtime/Core/Private/Internationalization/ICUText.cpp
+++ b/Engine/Source/Runtime/Core/Private/Internationalization/ICUText.cpp
@@ -222,332 +222,4 @@ bool FText::FSortPredicate::operator()(const FText& A, const FText& B) const
 	return Implementation->Compare(A, B);
 }
 
-void FText::GetFormatPatternParameters(const FText& Pattern, TArray<FString>& ParameterNames)
-{
-	UErrorCode ICUStatus = U_ZERO_ERROR;
-
-	const FString& NativePatternString = Pattern.ToString();
-
-	icu::UnicodeString ICUPatternString;
-	ICUUtilities::ConvertString(NativePatternString, ICUPatternString);
-
-	UParseError ICUParseError;
-	icu::MessagePattern ICUMessagePattern(ICUPatternString, &(ICUParseError), ICUStatus);
-
-	const int32_t PartCount = ICUMessagePattern.countParts();
-	for(int32_t i = 0; i < PartCount; ++i)
-	{
-		UMessagePatternPartType PartType = ICUMessagePattern.getPartType(i);
-		if(PartType == UMSGPAT_PART_TYPE_ARG_NAME || PartType == UMSGPAT_PART_TYPE_ARG_NUMBER)
-		{
-			icu::MessagePattern::Part Part = ICUMessagePattern.getPart(i);
-
-			/** The following does a case-sensitive "TArray::AddUnique" by first checking to see if
-			the parameter is in the ParameterNames list (using a case-sensitive comparison) followed
-			by adding it to the ParameterNames */
-
-			bool bIsCaseSensitiveUnique = true;
-			FString CurrentArgument;
-			ICUUtilities::ConvertString(ICUMessagePattern.getSubstring(Part), CurrentArgument);
-
-			for(auto It = ParameterNames.CreateConstIterator(); It; ++It)
-			{
-				if(It->Equals(CurrentArgument))
-				{
-					bIsCaseSensitiveUnique = false;
-				}
-			}
-
-			if(bIsCaseSensitiveUnique)
-			{
-				ParameterNames.Add( CurrentArgument );
-			}
-		}
-	}
-}
-
-FText FText::FormatInternal(const FText& Pattern, const FFormatNamedArguments& Arguments, bool bInRebuildText, bool bInRebuildAsSource)
-{
-	checkf(FInternationalization::Get().IsInitialized() == true, TEXT("FInternationalization is not initialized. An FText formatting method was likely used in static object initialization - this is not supported."));
-	//SCOPE_CYCLE_COUNTER( STAT_TextFormat );
-
-	CA_SUPPRESS(6316)
-	const bool bEnableErrorResults = ENABLE_TEXT_ERROR_CHECKING_RESULTS && bEnableErrorCheckingResults;
-
-	TArray<icu::UnicodeString> ArgumentNames;
-	ArgumentNames.Reserve(Arguments.Num());
-
-	TArray<icu::Formattable> ArgumentValues;
-	ArgumentValues.Reserve(Arguments.Num());
-
-	for(auto It = Arguments.CreateConstIterator(); It; ++It)
-	{
-		const FString& NativeArgumentName = It.Key();
-		icu::UnicodeString ICUArgumentName;
-		ICUUtilities::ConvertString(NativeArgumentName, ICUArgumentName, false);
-		ArgumentNames.Add( ICUArgumentName );
-
-		const FFormatArgumentValue& ArgumentValue = It.Value();
-		switch(ArgumentValue.Type)
-		{
-		case EFormatArgumentType::Int:
-			{
-				ArgumentValues.Add( icu::Formattable( static_cast<int64_t>(ArgumentValue.IntValue) ));
-			}
-			break;
-		case EFormatArgumentType::UInt:
-			{
-				ArgumentValues.Add( icu::Formattable( static_cast<int64_t>(ArgumentValue.UIntValue) ));
-			}
-			break;
-		case EFormatArgumentType::Float:
-			{
-				ArgumentValues.Add( icu::Formattable( ArgumentValue.FloatValue ));
-			}
-			break;
-		case EFormatArgumentType::Double:
-			{
-				ArgumentValues.Add( icu::Formattable( ArgumentValue.DoubleValue ));
-			}
-			break;
-		case EFormatArgumentType::Text:
-			{
-				// When doing a rebuild, all FText arguments need to be rebuilt during the Format
-				if(bInRebuildText)
-				{
-					ArgumentValue.TextValue->Rebuild();
-				}
-
-				FString NativeStringValue = bInRebuildAsSource? ArgumentValue.TextValue->BuildSourceString() : ArgumentValue.TextValue->ToString();
-				icu::UnicodeString ICUStringValue;
-				ICUUtilities::ConvertString(NativeStringValue, ICUStringValue, false);
-				ArgumentValues.Add( icu::Formattable( ICUStringValue ) );
-			}
-			break;
-		}
-	}
-
-	const FString& NativePatternString = bInRebuildAsSource? Pattern.BuildSourceString() : Pattern.ToString();
-	icu::UnicodeString ICUPatternString;
-	ICUUtilities::ConvertString(NativePatternString, ICUPatternString);
-
-	UErrorCode ICUStatus = U_ZERO_ERROR;
-	UParseError ICUParseError;
-	icu::MessageFormat ICUMessageFormatter(ICUPatternString, icu::Locale::getDefault(), ICUParseError, ICUStatus);
-	if(U_FAILURE(ICUStatus))
-	{
-		return bEnableErrorResults ? FText::FromString(u_errorName(ICUStatus)) : FText();
-	}
-
-	icu::UnicodeString ICUResultString;
-	ICUMessageFormatter.format(ArgumentNames.GetData(), ArgumentValues.GetData(), Arguments.Num(), ICUResultString, ICUStatus);
-	if(U_FAILURE(ICUStatus))
-	{
-		return bEnableErrorResults ? FText::FromString(u_errorName(ICUStatus)) : FText();
-	}
-
-	FString NativeResultString;
-	ICUUtilities::ConvertString(ICUResultString, NativeResultString);
-
-	FText ResultText(NativeResultString);
-	ResultText.History = MakeShareable(new FTextHistory_NamedFormat(Pattern, Arguments));
-
-	return ResultText;
-}
-
-FText FText::FormatInternal(const FText& Pattern, const FFormatOrderedArguments& Arguments, bool bInRebuildText, bool bInRebuildAsSource)
-{
-	checkf(FInternationalization::Get().IsInitialized() == true, TEXT("FInternationalization is not initialized. An FText formatting method was likely used in static object initialization - this is not supported."));
-	//SCOPE_CYCLE_COUNTER( STAT_TextFormat );
-
-	CA_SUPPRESS(6316)
-	const bool bEnableErrorResults = ENABLE_TEXT_ERROR_CHECKING_RESULTS && bEnableErrorCheckingResults;
-
-	TArray<icu::Formattable> ArgumentValues;
-	ArgumentValues.Reserve(Arguments.Num());
-
-	for(auto It = Arguments.CreateConstIterator(); It; ++It)
-	{
-		const FFormatArgumentValue& ArgumentValue = *It;
-		switch(ArgumentValue.Type)
-		{
-		case EFormatArgumentType::Int:
-			{
-				ArgumentValues.Add( icu::Formattable( static_cast<int64_t>(ArgumentValue.IntValue) ));
-			}
-			break;
-		case EFormatArgumentType::UInt:
-			{
-				ArgumentValues.Add( icu::Formattable( static_cast<int64_t>(ArgumentValue.UIntValue) ));
-			}
-			break;
-		case EFormatArgumentType::Float:
-			{
-				ArgumentValues.Add( icu::Formattable( ArgumentValue.FloatValue ));
-			}
-			break;
-		case EFormatArgumentType::Double:
-			{
-				ArgumentValues.Add( icu::Formattable( ArgumentValue.DoubleValue ));
-			}
-			break;
-		case EFormatArgumentType::Text:
-			{
-				// When doing a rebuild, all FText arguments need to be rebuilt during the Format
-				if(bInRebuildText)
-				{
-					ArgumentValue.TextValue->Rebuild();
-				}
-
-				FString NativeStringValue = bInRebuildAsSource? ArgumentValue.TextValue->BuildSourceString() : ArgumentValue.TextValue->ToString();
-				icu::UnicodeString ICUStringValue;
-				ICUUtilities::ConvertString(NativeStringValue, ICUStringValue, false);
-				ArgumentValues.Add( icu::Formattable( ICUStringValue ) );
-			}
-			break;
-		}
-	}
-
-	const FString& NativePatternString = bInRebuildAsSource? Pattern.BuildSourceString() : Pattern.ToString();
-	icu::UnicodeString ICUPatternString;
-	ICUUtilities::ConvertString(NativePatternString, ICUPatternString);
-
-	UErrorCode ICUStatus = U_ZERO_ERROR;
-	UParseError ICUParseError;
-	icu::MessageFormat ICUMessageFormatter(ICUPatternString, icu::Locale::getDefault(), ICUParseError, ICUStatus);
-	if(U_FAILURE(ICUStatus))
-	{
-		return bEnableErrorResults ? FText::FromString(u_errorName(ICUStatus)) : FText();
-	}
-
-	icu::UnicodeString ICUResultString;
-	ICUMessageFormatter.format(NULL, ArgumentValues.GetData(), Arguments.Num(), ICUResultString, ICUStatus);
-	if(U_FAILURE(ICUStatus))
-	{
-		return bEnableErrorResults ? FText::FromString(u_errorName(ICUStatus)) : FText();
-	}
-
-	FString NativeResultString;
-	ICUUtilities::ConvertString(ICUResultString, NativeResultString);
-
-	FText ResultText = FText(NativeResultString);
-
-	if (!GIsEditor)
-	{
-		ResultText.Flags = ResultText.Flags | ETextFlag::Transient;
-	}
-
-	ResultText.History = MakeShareable(new FTextHistory_OrderedFormat(Pattern, Arguments));
-	return ResultText;
-}
-
-FText FText::FormatInternal(const FText& Pattern, const TArray< FFormatArgumentData > InArguments, bool bInRebuildText, bool bInRebuildAsSource)
-{
-	checkf(FInternationalization::Get().IsInitialized() == true, TEXT("FInternationalization is not initialized. An FText formatting method was likely used in static object initialization - this is not supported."));
-	//SCOPE_CYCLE_COUNTER( STAT_TextFormat );
-
-	CA_SUPPRESS(6316)
-	const bool bEnableErrorResults = ENABLE_TEXT_ERROR_CHECKING_RESULTS && bEnableErrorCheckingResults;
-
-	TArray<icu::UnicodeString> ArgumentNames;
-	ArgumentNames.Reserve(InArguments.Num());
-
-	TArray<icu::Formattable> ArgumentValues;
-	ArgumentValues.Reserve(InArguments.Num());
-
-	for(int32 x = 0; x < InArguments.Num(); ++x)
-	{
-		const FString& ArgumentName = InArguments[x].ArgumentName.ToString();
-		icu::UnicodeString ICUArgumentName;
-		ICUUtilities::ConvertString(ArgumentName, ICUArgumentName, false);
-		ArgumentNames.Add( ICUArgumentName );
-
-		const FFormatArgumentValue& ArgumentValue = InArguments[x].ArgumentValue;
-		switch(ArgumentValue.Type)
-		{
-		case EFormatArgumentType::Int:
-			{
-				ArgumentValues.Add( icu::Formattable( static_cast<int64_t>(ArgumentValue.IntValue) ));
-			}
-			break;
-		case EFormatArgumentType::UInt:
-			{
-				ArgumentValues.Add( icu::Formattable( static_cast<int64_t>(ArgumentValue.UIntValue) ));
-			}
-			break;
-		case EFormatArgumentType::Float:
-			{
-				ArgumentValues.Add( icu::Formattable( ArgumentValue.FloatValue ));
-			}
-			break;
-		case EFormatArgumentType::Double:
-			{
-				ArgumentValues.Add( icu::Formattable( ArgumentValue.DoubleValue ));
-			}
-			break;
-		case EFormatArgumentType::Text:
-			{
-				// When doing a rebuild, all FText arguments need to be rebuilt during the Format
-				if(bInRebuildText)
-				{
-					ArgumentValue.TextValue->Rebuild();
-				}
-
-				FString StringValue = bInRebuildAsSource? ArgumentValue.TextValue->BuildSourceString() : ArgumentValue.TextValue->ToString();
-				icu::UnicodeString ICUStringValue;
-				ICUUtilities::ConvertString(StringValue, ICUStringValue, false);
-				ArgumentValues.Add( ICUStringValue );
-			}
-			break;
-		}
-	}
-
-	const FString& PatternString = bInRebuildAsSource? Pattern.BuildSourceString() : Pattern.ToString();
-	icu::UnicodeString ICUPatternString;
-	ICUUtilities::ConvertString(PatternString, ICUPatternString, false);
-
-	UErrorCode ICUStatus = U_ZERO_ERROR;
-	UParseError ICUParseError;
-	icu::MessageFormat ICUMessageFormatter(ICUPatternString, icu::Locale::getDefault(), ICUParseError, ICUStatus);
-	if(U_FAILURE(ICUStatus))
-	{
-		return bEnableErrorResults ? FText::FromString(u_errorName(ICUStatus)) : FText();
-	}
-
-	icu::UnicodeString Result;
-	ICUMessageFormatter.format(ArgumentNames.GetData(), ArgumentValues.GetData(), InArguments.Num(), Result, ICUStatus);
-	if(U_FAILURE(ICUStatus))
-	{
-		return bEnableErrorResults ? FText::FromString(u_errorName(ICUStatus)) : FText();
-	}
-
-	FString NativeResultString;
-	ICUUtilities::ConvertString(Result, NativeResultString);
-
-	FText ResultText = FText(NativeResultString);
-
-	if (!GIsEditor)
-	{
-		ResultText.Flags = ResultText.Flags | ETextFlag::Transient;
-	}
-
-	ResultText.History = MakeShareable(new FTextHistory_ArgumentDataFormat(Pattern, InArguments));
-	return ResultText;
-}
-
-FText FText::Format(const FText& Pattern, const FFormatNamedArguments& Arguments)
-{
-	return FormatInternal(Pattern, Arguments, false, false);
-}
-
-FText FText::Format(const FText& Pattern, const FFormatOrderedArguments& Arguments)
-{
-	return FormatInternal(Pattern, Arguments, false, false);
-}
-
-FText FText::Format(const FText& Pattern, const TArray< struct FFormatArgumentData > InArguments)
-{
-	return FormatInternal(Pattern, InArguments, false, false);
-}
-
 #endif
diff --git a/Engine/Source/Runtime/Core/Private/Internationalization/LegacyText.cpp b/Engine/Source/Runtime/Core/Private/Internationalization/LegacyText.cpp
--- a/Engine/Source/Runtime/Core/Private/Internationalization/LegacyText.cpp
+++ b/Engine/Source/Runtime/Core/Private/Internationalization/LegacyText.cpp
@@ -91,411 +91,4 @@ bool FText::FSortPredicate::operator()(const FText& A, const FText& B) const
 	return A.ToString() < B.ToString();
 }
 
-namespace
-{
-	namespace EEscapeState
-	{
-		enum Type
-		{
-			None,
-			BeginEscaping,
-			Escaping,
-			EndEscaping
-		};
-	}
-
-	namespace EBlockState
-	{
-		enum Type
-		{
-			None,
-			InBlock
-		};
-	}
-}
-
-DECLARE_DELEGATE_RetVal_OneParam( const FFormatArgumentValue*, FGetArgumentValue, const FString );
-
-class FLegacyTextHelper
-{
-public:
-	static void EnumerateParameters(const FText& Pattern, TArray<FString>& ParameterNames)
-	{
-		const FString& PatternString = Pattern.ToString();
-
-		EEscapeState::Type EscapeState = EEscapeState::None;
-		EBlockState::Type BlockState = EBlockState::None;
-
-		FString ParameterName;
-
-		for(int32 i = 0; i < PatternString.Len(); ++i)
-		{
-			switch( EscapeState )
-			{
-			case EEscapeState::None:
-				{
-					switch( PatternString[i] )
-					{
-					case '`':	{ EscapeState = EEscapeState::BeginEscaping; } break;
-					}
-				}
-				break;
-			case EEscapeState::BeginEscaping:
-				{
-					switch( PatternString[i] )
-					{
-						// Only begin EEscapeState::Escaping if there's a syntax character.
-					case '{':
-					case '}':
-						{
-							EscapeState = EEscapeState::Escaping;
-						}
-						break;
-						// Cancel beginning EEscapeState::Escaping if the escape is itself escaped.
-					case '`':
-						{
-							EscapeState = EEscapeState::None;
-						}
-						break;
-						// Cancel beginning EEscapeState::Escaping if not a syntax character.
-					default:
-						{
-							EscapeState = EEscapeState::None;
-						}
-						break;
-					}
-				}
-				break;
-			case EEscapeState::Escaping:
-				{
-					switch( PatternString[i] )
-					{
-					case '`':	{ EscapeState = EEscapeState::EndEscaping; } break;
-					}
-				}
-				break;
-			case EEscapeState::EndEscaping:
-				{
-					switch( PatternString[i] )
-					{
-					case '`':	{ EscapeState = EEscapeState::Escaping; } break; // Cancel ending EEscapeState::Escaping if the escape is itself escaped, copy over escaped character.
-					default:	{ EscapeState = EEscapeState::None; } break;
-					}
-				}
-				break;
-			}
-
-			if(EscapeState == EEscapeState::None)
-			{
-				switch( BlockState )
-				{
-				case EBlockState::None:
-					{
-						switch( PatternString[i] )
-						{
-						case '{':
-							{
-								BlockState = EBlockState::InBlock;
-							}
-							break;
-						default:
-							{
-							}
-							break;
-						}
-					}
-					break;
-				case EBlockState::InBlock:
-					{
-						switch( PatternString[i] )
-						{
-						case '}':
-							{
-								/** The following does a case-sensitive "TArray::AddUnique" by first checking to see if
-								the parameter is in the ParameterNames list (using a case-sensitive comparison) followed
-								by adding it to the ParameterNames */
-
-								bool bIsCaseSensitiveUnique = true;
-								for(auto It = ParameterNames.CreateConstIterator(); It; ++It)
-								{
-									if(It->Equals(ParameterName))
-									{
-										bIsCaseSensitiveUnique = false;
-									}
-								}
-								if(bIsCaseSensitiveUnique)
-								{
-									ParameterNames.Add( ParameterName );
-								}
-								ParameterName.Empty();
-								BlockState = EBlockState::None;
-							}
-							break;
-						default:
-							{
-								ParameterName += PatternString[i];
-							}
-							break;
-						}
-					}
-					break;
-				}
-			}
-		}
-	}
-
-	static FText Format(const FText& Pattern, FGetArgumentValue GetArgumentValue)
-	{
-		checkf(FInternationalization::Get().IsInitialized() == true, TEXT("FInternationalization is not initialized. An FText formatting method was likely used in static object initialization - this is not supported."));
-		//SCOPE_CYCLE_COUNTER( STAT_TextFormat );
-
-		const FString& PatternString = Pattern.ToString();
-
-		FString ResultString;
-
-		EEscapeState::Type EscapeState = EEscapeState::None;
-		EBlockState::Type BlockState = EBlockState::None;
-
-		FString ArgumentName;
-
-		for(int32 i = 0; i < PatternString.Len(); ++i)
-		{
-			switch( EscapeState )
-			{
-			case EEscapeState::None:
-				{
-					switch( PatternString[i] )
-					{
-					case '`':	{ EscapeState = EEscapeState::BeginEscaping; } break;
-					}
-				}
-				break;
-			case EEscapeState::BeginEscaping:
-				{
-					switch( PatternString[i] )
-					{
-						// Only begin EEscapeState::Escaping if there's a syntax character.
-					case '{':
-					case '}':
-						{
-							EscapeState = EEscapeState::Escaping;
-							ResultString += PatternString[i]; // Characters are escaped, copy over.
-						}
-						break;
-						// Cancel beginning EEscapeState::Escaping if the escape is itself escaped.
-					case '`':
-						{
-							EscapeState = EEscapeState::None;
-						}
-						break;
-						// Cancel beginning EEscapeState::Escaping if not a syntax character.
-					default:
-						{
-							EscapeState = EEscapeState::None;
-							ResultString += '`'; // Insert previously ignored escape marker.
-						}
-						break;
-					}
-				}
-				break;
-			case EEscapeState::Escaping:
-				{
-					switch( PatternString[i] )
-					{
-					case '`':	{ EscapeState = EEscapeState::EndEscaping; } break;
-					default:	{ ResultString += PatternString[i]; } break; // Characters are escaped, copy over.
-					}
-				}
-				break;
-			case EEscapeState::EndEscaping:
-				{
-					switch( PatternString[i] )
-					{
-					case '`':	{ EscapeState = EEscapeState::Escaping; ResultString += PatternString[i]; } break; // Cancel ending EEscapeState::Escaping if the escape is itself escaped, copy over escaped character.
-					default:	{ EscapeState = EEscapeState::None; } break;
-					}
-				}
-				break;
-			}
-
-			if(EscapeState == EEscapeState::None)
-			{
-				switch( BlockState )
-				{
-				case EBlockState::None:
-					{
-						switch( PatternString[i] )
-						{
-						case '{':
-							{
-								BlockState = EBlockState::InBlock;
-							}
-							break;
-						default:
-							{
-								ResultString += PatternString[i]; // Copy over characters.
-							}
-							break;
-						}
-					}
-					break;
-				case EBlockState::InBlock:
-					{
-						switch( PatternString[i] )
-						{
-						case '}':
-							{
-							
-								const FFormatArgumentValue* const PossibleArgumentValue = GetArgumentValue.Execute(ArgumentName);
-								if( PossibleArgumentValue )
-								{
-									FString ArgumentValueAsString;
-									const FFormatArgumentValue& ArgumentValue = *PossibleArgumentValue;
-									switch(ArgumentValue.Type)
-									{
-									case EFormatArgumentType::Text:
-										{
-											ArgumentValueAsString = ArgumentValue.TextValue->ToString();
-										}
-										break;
-									case EFormatArgumentType::Int:
-										{
-											ArgumentValueAsString = FText::AsNumber(ArgumentValue.IntValue).ToString();
-										}
-										break;
-									case EFormatArgumentType::UInt:
-										{
-											ArgumentValueAsString = FText::AsNumber(ArgumentValue.UIntValue).ToString();
-										}
-										break;
-									case EFormatArgumentType::Float:
-										{
-											ArgumentValueAsString = FText::AsNumber(ArgumentValue.FloatValue).ToString();
-										}
-										break;
-									case EFormatArgumentType::Double:
-										{
-											ArgumentValueAsString = FText::AsNumber(ArgumentValue.DoubleValue).ToString();
-										}
-										break;
-									}
-									ResultString += ArgumentValueAsString;
-								}
-								else
-								{
-									ResultString += FString("{") + ArgumentName + FString("}");
-								}
-								ArgumentName.Empty();
-								BlockState = EBlockState::None;
-							}
-							break;
-						default:
-							{
-								ArgumentName += PatternString[i];
-							}
-							break;
-						}
-					}
-					break;
-				}
-			}
-		}
-
-		FText Result = FText(MoveTemp(ResultString));
-		if (!GIsEditor)
-		{
-			Result.Flags = Result.Flags | ETextFlag::Transient;
-		}
-		return Result;
-	}
-};
-
-void FText::GetFormatPatternParameters(const FText& Pattern, TArray<FString>& ParameterNames)
-{
-	FLegacyTextHelper::EnumerateParameters(Pattern, ParameterNames);
-}
-
-FText FText::Format(const FText& Pattern, const FFormatNamedArguments& Arguments)
-{
-	checkf(FInternationalization::Get().IsInitialized() == true, TEXT("FInternationalization is not initialized. An FText formatting method was likely used in static object initialization - this is not supported."));
-	//SCOPE_CYCLE_COUNTER( STAT_TextFormat );
-
-	struct FArgumentGetter
-	{
-		const FFormatNamedArguments& Arguments;
-		FArgumentGetter(const FFormatNamedArguments& InArguments) : Arguments(InArguments) {}
-		const FFormatArgumentValue* GetArgumentValue( const FString ArgumentName )
-		{
-			return Arguments.Find(ArgumentName);
-		}
-	};
-
-	FArgumentGetter ArgumentGetter(Arguments);
-	return FLegacyTextHelper::Format(Pattern, FGetArgumentValue::CreateRaw(&ArgumentGetter, &FArgumentGetter::GetArgumentValue));
-}
-
-FText FText::Format(const FText& Pattern, const FFormatOrderedArguments& Arguments)
-{
-	checkf(FInternationalization::Get().IsInitialized() == true, TEXT("FInternationalization is not initialized. An FText formatting method was likely used in static object initialization - this is not supported."));
-	//SCOPE_CYCLE_COUNTER( STAT_TextFormat );
-
-	struct FArgumentGetter
-	{
-		const FFormatOrderedArguments& Arguments;
-		FArgumentGetter(const FFormatOrderedArguments& InArguments) : Arguments(InArguments) {}
-		const FFormatArgumentValue* GetArgumentValue( const FString ArgumentName )
-		{
-			int32 ArgumentIndex = INDEX_NONE;
-			if( ArgumentName.IsNumeric() )
-			{
-				ArgumentIndex = FCString::Atoi(*ArgumentName);
-			}
-			return ArgumentIndex != INDEX_NONE && ArgumentIndex < Arguments.Num() ? &(Arguments[ArgumentIndex]) : NULL;
-		}
-	};
-
-	FArgumentGetter ArgumentGetter(Arguments);
-	return FLegacyTextHelper::Format(Pattern, FGetArgumentValue::CreateRaw(&ArgumentGetter, &FArgumentGetter::GetArgumentValue));
-}
-
-FText FText::Format(const FText& Pattern, const TArray< FFormatArgumentData > InArguments)
-{
-	checkf(FInternationalization::Get().IsInitialized() == true, TEXT("FInternationalization is not initialized. An FText formatting method was likely used in static object initialization - this is not supported."));
-	//SCOPE_CYCLE_COUNTER( STAT_TextFormat );
-
-	struct FArgumentGetter
-	{
-		const FFormatNamedArguments& Arguments;
-		FArgumentGetter(const FFormatNamedArguments& InArguments) : Arguments(InArguments) {}
-		const FFormatArgumentValue* GetArgumentValue( const FString ArgumentName )
-		{
-			return Arguments.Find(ArgumentName);
-		}
-	};
-
-	FFormatNamedArguments FormatNamedArguments;
-	for(auto It = InArguments.CreateConstIterator(); It; ++It)
-	{
-		FormatNamedArguments.Add(It->ArgumentName.ToString(), FFormatArgumentValue(It->ArgumentValue));
-	}
-
-	FArgumentGetter ArgumentGetter(FormatNamedArguments);
-	return FLegacyTextHelper::Format(Pattern, FGetArgumentValue::CreateRaw(&ArgumentGetter, &FArgumentGetter::GetArgumentValue));
-}
-
-FText FText::FormatInternal(const FText& Pattern, const FFormatNamedArguments& Arguments, bool bInRebuildText, bool bInRebuildAsSource)
-{
-	return Format(Pattern, Arguments);
-}
-
-FText FText::FormatInternal(const FText& Pattern, const FFormatOrderedArguments& Arguments, bool bInRebuildText, bool bInRebuildAsSource)
-{
-	return Format(Pattern, Arguments);
-}
-
-FText FText::FormatInternal(const FText& Pattern, const TArray< struct FFormatArgumentData > InArguments, bool bInRebuildText, bool bInRebuildAsSource)
-{
-	return Format(Pattern, InArguments);
-}
-
 #endif
diff --git a/Engine/Source/Runtime/Core/Private/Internationalization/Text.cpp b/Engine/Source/Runtime/Core/Private/Internationalization/Text.cpp
--- a/Engine/Source/Runtime/Core/Private/Internationalization/Text.cpp
+++ b/Engine/Source/Runtime/Core/Private/Internationalization/Text.cpp
@@ -472,6 +472,429 @@ FText FText::Format(const FText& Fmt,const FText& v1,const FText& v2,const FText
 	return FText::Format(Fmt, Arguments);
 }
 
+class FTextFormatHelper
+{
+private:
+	enum class EEscapeState
+	{
+		None,
+		BeginEscaping,
+		Escaping,
+		EndEscaping,
+	};
+
+	enum class EBlockState
+	{
+		None,
+		InBlock,
+	};
+
+public:
+	DECLARE_DELEGATE_RetVal_OneParam( const FFormatArgumentValue*, FGetArgumentValue, const FString& );
+
+	static int32 EstimateArgumentValueLength(const FFormatArgumentValue& ArgumentValue)
+	{
+		switch(ArgumentValue.Type)
+		{
+		case EFormatArgumentType::Text:
+			return ArgumentValue.TextValue->ToString().Len();
+		case EFormatArgumentType::Int:
+		case EFormatArgumentType::UInt:
+		case EFormatArgumentType::Float:
+		case EFormatArgumentType::Double:
+			return 16;
+		default:
+			break;
+		}
+
+		return 0;
+	}
+
+	static void EnumerateParameters(const FText& Pattern, TArray<FString>& ParameterNames)
+	{
+		const FString& PatternString = Pattern.ToString();
+
+		EEscapeState EscapeState = EEscapeState::None;
+		EBlockState BlockState = EBlockState::None;
+
+		FString ParameterName;
+
+		for(int32 i = 0; i < PatternString.Len(); ++i)
+		{
+			switch( EscapeState )
+			{
+			case EEscapeState::None:
+				{
+					switch( PatternString[i] )
+					{
+					case '`':	{ EscapeState = EEscapeState::BeginEscaping; } break;
+					}
+				}
+				break;
+			case EEscapeState::BeginEscaping:
+				{
+					switch( PatternString[i] )
+					{
+						// Only begin EEscapeState::Escaping if there's a syntax character.
+					case '{':
+					case '}':
+						{
+							EscapeState = EEscapeState::Escaping;
+						}
+						break;
+						// Cancel beginning EEscapeState::Escaping if the escape is itself escaped.
+					case '`':
+						{
+							EscapeState = EEscapeState::None;
+						}
+						break;
+						// Cancel beginning EEscapeState::Escaping if not a syntax character.
+					default:
+						{
+							EscapeState = EEscapeState::None;
+						}
+						break;
+					}
+				}
+				break;
+			case EEscapeState::Escaping:
+				{
+					switch( PatternString[i] )
+					{
+					case '`':	{ EscapeState = EEscapeState::EndEscaping; } break;
+					}
+				}
+				break;
+			case EEscapeState::EndEscaping:
+				{
+					switch( PatternString[i] )
+					{
+					case '`':	{ EscapeState = EEscapeState::Escaping; } break; // Cancel ending EEscapeState::Escaping if the escape is itself escaped, copy over escaped character.
+					default:	{ EscapeState = EEscapeState::None; } break;
+					}
+				}
+				break;
+			}
+
+			if(EscapeState == EEscapeState::None)
+			{
+				switch( BlockState )
+				{
+				case EBlockState::None:
+					{
+						switch( PatternString[i] )
+						{
+						case '{':
+							{
+								BlockState = EBlockState::InBlock;
+							}
+							break;
+						default:
+							{
+							}
+							break;
+						}
+					}
+					break;
+				case EBlockState::InBlock:
+					{
+						switch( PatternString[i] )
+						{
+						case '}':
+							{
+								/** The following does a case-sensitive "TArray::AddUnique" by first checking to see if
+								the parameter is in the ParameterNames list (using a case-sensitive comparison) followed
+								by adding it to the ParameterNames */
+
+								bool bIsCaseSensitiveUnique = true;
+								for(auto It = ParameterNames.CreateConstIterator(); It; ++It)
+								{
+									if(It->Equals(ParameterName))
+									{
+										bIsCaseSensitiveUnique = false;
+									}
+								}
+								if(bIsCaseSensitiveUnique)
+								{
+									ParameterNames.Add( ParameterName );
+								}
+								ParameterName.Empty();
+								BlockState = EBlockState::None;
+							}
+							break;
+						default:
+							{
+								ParameterName += PatternString[i];
+							}
+							break;
+						}
+					}
+					break;
+				}
+			}
+		}
+	}
+
+	static FText Format(const FText& Pattern, const int32 EstimatedArgumentValuesLength, FGetArgumentValue GetArgumentValue, const TSharedPtr<FTextHistory, ESPMode::ThreadSafe>& InHistory, bool bInRebuildText, bool bInRebuildAsSource)
+	{
+		checkf(FInternationalization::Get().IsInitialized() == true, TEXT("FInternationalization is not initialized. An FText formatting method was likely used in static object initialization - this is not supported."));
+		//SCOPE_CYCLE_COUNTER( STAT_TextFormat );
+
+		const FString& PatternString = bInRebuildAsSource ? Pattern.BuildSourceString() : Pattern.ToString();
+
+		// Guesstimate the result string size to minimize reallocations (it's okay to be a little bit over)
+		// We'll assume that each argument will be replaced once, plus a bit of slack
+		FString ResultString;
+		ResultString.Reserve(PatternString.Len() + EstimatedArgumentValuesLength + 8);
+
+		EEscapeState EscapeState = EEscapeState::None;
+		EBlockState BlockState = EBlockState::None;
+
+		FString ArgumentName;
+
+		for(int32 i = 0; i < PatternString.Len(); ++i)
+		{
+			switch( EscapeState )
+			{
+			case EEscapeState::None:
+				{
+					switch( PatternString[i] )
+					{
+					case '`':	{ EscapeState = EEscapeState::BeginEscaping; } break;
+					}
+				}
+				break;
+			case EEscapeState::BeginEscaping:
+				{
+					switch( PatternString[i] )
+					{
+						// Only begin EEscapeState::Escaping if there's a syntax character.
+					case '{':
+					case '}':
+						{
+							EscapeState = EEscapeState::Escaping;
+							ResultString += PatternString[i]; // Characters are escaped, copy over.
+						}
+						break;
+						// Cancel beginning EEscapeState::Escaping if the escape is itself escaped.
+					case '`':
+						{
+							EscapeState = EEscapeState::None;
+						}
+						break;
+						// Cancel beginning EEscapeState::Escaping if not a syntax character.
+					default:
+						{
+							EscapeState = EEscapeState::None;
+							ResultString += '`'; // Insert previously ignored escape marker.
+						}
+						break;
+					}
+				}
+				break;
+			case EEscapeState::Escaping:
+				{
+					switch( PatternString[i] )
+					{
+					case '`':	{ EscapeState = EEscapeState::EndEscaping; } break;
+					default:	{ ResultString += PatternString[i]; } break; // Characters are escaped, copy over.
+					}
+				}
+				break;
+			case EEscapeState::EndEscaping:
+				{
+					switch( PatternString[i] )
+					{
+					case '`':	{ EscapeState = EEscapeState::Escaping; ResultString += PatternString[i]; } break; // Cancel ending EEscapeState::Escaping if the escape is itself escaped, copy over escaped character.
+					default:	{ EscapeState = EEscapeState::None; } break;
+					}
+				}
+				break;
+			}
+
+			if(EscapeState == EEscapeState::None)
+			{
+				switch( BlockState )
+				{
+				case EBlockState::None:
+					{
+						switch( PatternString[i] )
+						{
+						case '{':
+							{
+								BlockState = EBlockState::InBlock;
+							}
+							break;
+						default:
+							{
+								ResultString += PatternString[i]; // Copy over characters.
+							}
+							break;
+						}
+					}
+					break;
+				case EBlockState::InBlock:
+					{
+						switch( PatternString[i] )
+						{
+						case '}':
+							{
+							
+								const FFormatArgumentValue* const PossibleArgumentValue = GetArgumentValue.Execute(ArgumentName);
+								if( PossibleArgumentValue )
+								{
+									const FFormatArgumentValue& ArgumentValue = *PossibleArgumentValue;
+									switch(ArgumentValue.Type)
+									{
+									case EFormatArgumentType::Text:
+										{
+											// When doing a rebuild, all FText arguments need to be rebuilt during the Format
+											if( bInRebuildText )
+											{
+												ArgumentValue.TextValue->Rebuild();
+											}
+
+											ResultString += bInRebuildAsSource ? ArgumentValue.TextValue->BuildSourceString() : ArgumentValue.TextValue->ToString();
+										}
+										break;
+									case EFormatArgumentType::Int:
+										{
+											ResultString += FText::AsNumber(ArgumentValue.IntValue).ToString();
+										}
+										break;
+									case EFormatArgumentType::UInt:
+										{
+											ResultString += FText::AsNumber(ArgumentValue.UIntValue).ToString();
+										}
+										break;
+									case EFormatArgumentType::Float:
+										{
+											ResultString += FText::AsNumber(ArgumentValue.FloatValue).ToString();
+										}
+										break;
+									case EFormatArgumentType::Double:
+										{
+											ResultString += FText::AsNumber(ArgumentValue.DoubleValue).ToString();
+										}
+										break;
+									}
+								}
+								else
+								{
+									ResultString += TEXT("{");
+									ResultString += ArgumentName;
+									ResultString += TEXT("}");
+								}
+								ArgumentName.Reset();
+								BlockState = EBlockState::None;
+							}
+							break;
+						default:
+							{
+								ArgumentName += PatternString[i];
+							}
+							break;
+						}
+					}
+					break;
+				}
+			}
+		}
+
+		FText Result = FText(MoveTemp(ResultString));
+		Result.History = InHistory;
+		if (!GIsEditor)
+		{
+			Result.Flags = Result.Flags | ETextFlag::Transient;
+		}
+		return Result;
+	}
+};
+
+void FText::GetFormatPatternParameters(const FText& Pattern, TArray<FString>& ParameterNames)
+{
+	FTextFormatHelper::EnumerateParameters(Pattern, ParameterNames);
+}
+
+FText FText::Format(const FText& Pattern, const FFormatNamedArguments& Arguments)
+{
+	return FormatInternal(Pattern, Arguments, false, false);
+}
+
+FText FText::Format(const FText& Pattern, const FFormatOrderedArguments& Arguments)
+{
+	return FormatInternal(Pattern, Arguments, false, false);
+}
+
+FText FText::Format(const FText& Pattern, const TArray< FFormatArgumentData > InArguments)
+{
+	return FormatInternal(Pattern, InArguments, false, false);
+}
+
+FText FText::FormatInternal(const FText& Pattern, const FFormatNamedArguments& Arguments, bool bInRebuildText, bool bInRebuildAsSource)
+{
+	checkf(FInternationalization::Get().IsInitialized() == true, TEXT("FInternationalization is not initialized. An FText formatting method was likely used in static object initialization - this is not supported."));
+	//SCOPE_CYCLE_COUNTER( STAT_TextFormat );
+
+	int32 EstimatedArgumentValuesLength = 0;
+	for (const auto& Arg : Arguments)
+	{
+		EstimatedArgumentValuesLength += FTextFormatHelper::EstimateArgumentValueLength(Arg.Value);
+	}
+
+	auto GetArgumentValue = [&Arguments](const FString& ArgumentName) -> const FFormatArgumentValue*
+	{
+		return Arguments.Find(ArgumentName);
+	};
+
+	return FTextFormatHelper::Format(Pattern, EstimatedArgumentValuesLength, FTextFormatHelper::FGetArgumentValue::CreateLambda(GetArgumentValue), MakeShareable(new FTextHistory_NamedFormat(Pattern, Arguments)), bInRebuildText, bInRebuildAsSource);
+}
+
+FText FText::FormatInternal(const FText& Pattern, const FFormatOrderedArguments& Arguments, bool bInRebuildText, bool bInRebuildAsSource)
+{
+	checkf(FInternationalization::Get().IsInitialized() == true, TEXT("FInternationalization is not initialized. An FText formatting method was likely used in static object initialization - this is not supported."));
+	//SCOPE_CYCLE_COUNTER( STAT_TextFormat );
+
+	int32 EstimatedArgumentValuesLength = 0;
+	for (const auto& Arg : Arguments)
+	{
+		EstimatedArgumentValuesLength += FTextFormatHelper::EstimateArgumentValueLength(Arg);
+	}
+
+	auto GetArgumentValue = [&Arguments](const FString& ArgumentName) -> const FFormatArgumentValue*
+	{
+		int32 ArgumentIndex = INDEX_NONE;
+		if( ArgumentName.IsNumeric() )
+		{
+			ArgumentIndex = FCString::Atoi(*ArgumentName);
+		}
+		return ArgumentIndex != INDEX_NONE && ArgumentIndex < Arguments.Num() ? &(Arguments[ArgumentIndex]) : nullptr;
+	};
+
+	return FTextFormatHelper::Format(Pattern, EstimatedArgumentValuesLength, FTextFormatHelper::FGetArgumentValue::CreateLambda(GetArgumentValue), MakeShareable(new FTextHistory_OrderedFormat(Pattern, Arguments)), bInRebuildText, bInRebuildAsSource);
+}
+
+FText FText::FormatInternal(const FText& Pattern, const TArray< struct FFormatArgumentData > Arguments, bool bInRebuildText, bool bInRebuildAsSource)
+{
+	checkf(FInternationalization::Get().IsInitialized() == true, TEXT("FInternationalization is not initialized. An FText formatting method was likely used in static object initialization - this is not supported."));
+	//SCOPE_CYCLE_COUNTER( STAT_TextFormat );
+
+	int32 EstimatedArgumentValuesLength = 0;
+	FFormatNamedArguments FormatNamedArguments;
+	for (const auto& Arg : Arguments)
+	{
+		EstimatedArgumentValuesLength += FTextFormatHelper::EstimateArgumentValueLength(Arg.ArgumentValue);
+		FormatNamedArguments.Add(Arg.ArgumentName.ToString(), FFormatArgumentValue(Arg.ArgumentValue));
+	}
+
+	auto GetArgumentValue = [&FormatNamedArguments](const FString& ArgumentName) -> const FFormatArgumentValue*
+	{
+		return FormatNamedArguments.Find(ArgumentName);
+	};
+
+	return FTextFormatHelper::Format(Pattern, EstimatedArgumentValuesLength, FTextFormatHelper::FGetArgumentValue::CreateLambda(GetArgumentValue), MakeShareable(new FTextHistory_ArgumentDataFormat(Pattern, Arguments)), bInRebuildText, bInRebuildAsSource);
+}
+
 /**
 * Generate an FText that represents the passed number in the passed culture
 */
diff --git a/Engine/Source/Runtime/Core/Public/Internationalization/Text.h b/Engine/Source/Runtime/Core/Public/Internationalization/Text.h
--- a/Engine/Source/Runtime/Core/Public/Internationalization/Text.h
+++ b/Engine/Source/Runtime/Core/Public/Internationalization/Text.h
@@ -383,6 +383,7 @@ public:
 	static const FText UnEscapedCloseBraceOutsideOfArgumentBlock;
 	static const FText SerializationFailureError;
 
+	friend class FTextFormatHelper;
 	friend class FTextSnapshot;
 	friend class FTextInspector;
 	friend class FInternationalization;
@@ -394,10 +395,6 @@ public:
 	friend class FTextHistory_ArgumentDataFormat;
 	friend class FTextHistory_OrderedFormat;
 	friend class FScopedTextIdentityPreserver;
-
-#if !UE_ENABLE_ICU
-	friend class FLegacyTextHelper;
-#endif
 };
 
 /** A snapshot of an FText at a point in time that can be used to detect changes in the FText, including live-culture changes */
