Changed editor project unit settings to be more flexible

Users can now choose a set of relevant units to display applicable properties in.
Setting a single unit for a given type will result in only that unit ever being used. When multiple exist, the most appropriate unit will be shown when a property editor is opened, and will remain in this unit until something else is selected.

diff --git a/Engine/Config/BaseEditor.ini b/Engine/Config/BaseEditor.ini
--- a/Engine/Config/BaseEditor.ini
+++ b/Engine/Config/BaseEditor.ini
@@ -365,4 +365,13 @@ MaxMemoryAllowance=8192
 +EngineTargetsSettings=(Name="ToolTips",Guid=0F116534468918AEA432DD8C77703BA8,TargetDependencies=(33482D004789784C9DA695A682ACCA1B,AC8BFD2A41A2FB2893BB8EA0AF903E6D),AdditionalManifestDependencies=,RequiredModuleNames=,GatherFromTextFiles=(SearchDirectories=,ExcludePathWildcards=,FileExtensions=),GatherFromPackages=(IncludePathWildcards=,ExcludePathWildcards=,FileExtensions=),GatherFromMetaData=(IncludePathWildcards=("Source/Editor/*","Source/Runtime/*","Source/Developer/*"),ExcludePathWildcards=("*/Source/Developer/NoRedist/*"),KeySpecifications=((MetaDataKey="ToolTip",TextNamespace="UObjectToolTips",TextKeyPattern="{FieldPath}"),(MetaDataKey="ShortToolTip",TextNamespace="UObjectShortToolTips",TextKeyPattern="{FieldPath}"))),NativeCultureStatistics=(CultureName="en"),SupportedCulturesStatistics=((CultureName="es"),(CultureName="ja"),(CultureName="ko"),(CultureName="pt-BR"),(CultureName="zh-CN")))
 +GameTargetsSettings=(Name="Game",Guid=AE0EA34A45461A25BA65A391026F19F8,TargetDependencies=(33482D004789784C9DA695A682ACCA1B,AC8BFD2A41A2FB2893BB8EA0AF903E6D),AdditionalManifestDependencies=,RequiredModuleNames=,GatherFromTextFiles=(SearchDirectories=,ExcludePathWildcards=,FileExtensions=),GatherFromPackages=(IncludePathWildcards=,ExcludePathWildcards=,FileExtensions=),GatherFromMetaData=(IncludePathWildcards=,ExcludePathWildcards=,KeySpecifications=),NativeCultureStatistics=(CultureName="en"),SupportedCulturesStatistics=)
 
-
+[/Script/EditorProjectAppearanceSettings.EditorProjectAppearanceSettings]
+bDisplayUnits=True
++DistanceUnits=EUnit::Centimeters
++MassUnits=EUnit::Kilograms
++TimeUnits=EUnit::Seconds
++TimeUnits=EUnit::Minutes
+AngleUnits=EUnit::Degrees
+SpeedUnits=EUnit::MetersPerSecond
+TemperatureUnits=EUnit::Celsius
+ForceUnits=EUnit::Newtons
diff --git a/Engine/Source/Editor/DetailCustomizations/Private/ComponentTransformDetails.cpp b/Engine/Source/Editor/DetailCustomizations/Private/ComponentTransformDetails.cpp
--- a/Engine/Source/Editor/DetailCustomizations/Private/ComponentTransformDetails.cpp
+++ b/Engine/Source/Editor/DetailCustomizations/Private/ComponentTransformDetails.cpp
@@ -354,15 +354,12 @@ void FComponentTransformDetails::GenerateChildContent( IDetailChildrenBuilder& C
 	const bool bHideRotationField = ( HiddenFieldMask & ( 1 << ETransformField::Rotation ) ) != 0;
 	const bool bHideScaleField = ( HiddenFieldMask & ( 1 << ETransformField::Scale ) ) != 0;
 
-	VectorUnits = EUnit::Centimeters;
-
 	// Location
 	if(!bHideLocationField)
 	{
 		TSharedPtr<INumericTypeInterface<float>> TypeInterface;
 		if( FUnitConversion::Settings().ShouldDisplayUnits() )
 		{
-			UseDefaultInputUnits();
 			TypeInterface = SharedThis(this);
 		}
 
@@ -569,53 +566,29 @@ END_SLATE_FUNCTION_BUILD_OPTIMIZATION
 void FComponentTransformDetails::Tick( float DeltaTime ) 
 {
 	CacheTransform();
-	CacheCommonLocationUnits();
+	if (!FixedDisplayUnits.IsSet())
+	{
+		CacheCommonLocationUnits();
+	}
 }
 
 void FComponentTransformDetails::CacheCommonLocationUnits()
 {
-	// Calculate the largest units
-	TOptional<EUnit> LargestUnit;
-	auto SetLargestUnit = [&](EUnit NewUnit)
+	float LargestValue = 0.f;
+	if (CachedLocation.X.IsSet() && CachedLocation.X.GetValue() > LargestValue)
 	{
-		if (!LargestUnit.IsSet() || (uint8)NewUnit > (uint8)LargestUnit.GetValue())
-		{
-			LargestUnit = NewUnit;
-		}
-	};
-
-	if (CachedLocation.X.IsSet() && CachedLocation.X.GetValue() != 0)
-	{
-		SetLargestUnit(FUnitConversion::QuantizeUnitsToBestFit(CachedLocation.X.GetValue(), Units).Units);
+		LargestValue = CachedLocation.X.GetValue();
 	}
-	if (CachedLocation.Y.IsSet() && CachedLocation.Y.GetValue() != 0)
+	if (CachedLocation.X.IsSet() && CachedLocation.Y.GetValue() > LargestValue)
 	{
-		SetLargestUnit(FUnitConversion::QuantizeUnitsToBestFit(CachedLocation.Y.GetValue(), Units).Units);
+		LargestValue = CachedLocation.Y.GetValue();
 	}
-	if (CachedLocation.Z.IsSet() && CachedLocation.Z.GetValue() != 0)
+	if (CachedLocation.X.IsSet() && CachedLocation.Z.GetValue() > LargestValue)
 	{
-		SetLargestUnit(FUnitConversion::QuantizeUnitsToBestFit(CachedLocation.Z.GetValue(), Units).Units);
+		LargestValue = CachedLocation.Z.GetValue();
 	}
 
-	if (LargestUnit.IsSet())
-	{
-		VectorUnits = LargestUnit.GetValue();
-	}
-	else if (DefaultInputUnits.IsSet())
-	{
-		VectorUnits	= DefaultInputUnits.GetValue();
-	}
-	else
-	{
-		VectorUnits = EUnit::Unspecified;
-	}
-
-	VectorUnits = FUnitConversion::ConvertToGlobalDisplayRange(VectorUnits);
-}
-
-FNumericUnit<float> FComponentTransformDetails::QuantizeUnitsToBestFit(const float& InValue, EUnit InUnits) const
-{
-	return FNumericUnit<float>(FUnitConversion::Convert(InValue, InUnits, VectorUnits), VectorUnits);
+	SetupFixedDisplay(LargestValue);
 }
 
 bool FComponentTransformDetails::GetIsEnabled() const
@@ -1122,7 +1095,6 @@ void FComponentTransformDetails::OnSetLocation( float NewValue, ETextCommit::Typ
 	}
 
 	CacheTransform();
-	CacheCommonLocationUnits();
 
 	GUnrealEd->RedrawLevelEditingViewports();
 }
diff --git a/Engine/Source/Editor/DetailCustomizations/Private/ComponentTransformDetails.h b/Engine/Source/Editor/DetailCustomizations/Private/ComponentTransformDetails.h
--- a/Engine/Source/Editor/DetailCustomizations/Private/ComponentTransformDetails.h
+++ b/Engine/Source/Editor/DetailCustomizations/Private/ComponentTransformDetails.h
@@ -34,8 +34,6 @@ public:
 	virtual void Tick( float DeltaTime ) override;
 	virtual void SetOnRebuildChildren( FSimpleDelegate OnRebuildChildren ) override{}
 
-	virtual FNumericUnit<float> QuantizeUnitsToBestFit(const float& InValue, EUnit InUnits) const override;
-
 	void HideTransformField(const ETransformField::Type InTransformField)
 	{
 		HiddenFieldMask |= (1 << InTransformField);
@@ -330,6 +328,4 @@ private:
 	bool bEditingRotationInUI;
 	/** Bitmask to indicate which fields should be hidden (if any) */
 	uint8 HiddenFieldMask;
-	/** Enum specifying the most appropriate units for the whole vector to be displayed in */
-	EUnit VectorUnits;
 };
diff --git a/Engine/Source/Editor/PropertyEditor/Private/PropertyHandleImpl.cpp b/Engine/Source/Editor/PropertyEditor/Private/PropertyHandleImpl.cpp
--- a/Engine/Source/Editor/PropertyEditor/Private/PropertyHandleImpl.cpp
+++ b/Engine/Source/Editor/PropertyEditor/Private/PropertyHandleImpl.cpp
@@ -1853,6 +1853,26 @@ FPropertyAccess::Result FPropertyHandleBase::GetPerObjectValues( TArray<FString>
 	return Result;
 }
 
+TArray<FName> GetValidEnumEntries(UProperty* Property, const UEnum* InEnum)
+{
+	TArray<FName> ValidEnumValues;
+
+	static const FName ValidEnumValuesName("ValidEnumValues");
+	if (Property->HasMetaData(ValidEnumValuesName))
+	{
+		TArray<FString> ValidEnumValuesAsString;
+
+		Property->GetMetaData(ValidEnumValuesName).ParseIntoArray(ValidEnumValuesAsString, TEXT(","));
+		for (auto& Value : ValidEnumValuesAsString)
+		{
+			Value.Trim();
+			ValidEnumValues.Add(*UEnum::GenerateFullEnumName(InEnum, *Value));
+		}
+	}
+
+	return ValidEnumValues;
+}
+
 bool FPropertyHandleBase::GeneratePossibleValues(TArray< TSharedPtr<FString> >& OutOptionStrings, TArray< FText >& OutToolTips, TArray<bool>& OutRestrictedItems)
 {
 	UProperty* Property = GetProperty();
@@ -1876,6 +1896,8 @@ bool FPropertyHandleBase::GeneratePossibleValues(TArray< TSharedPtr<FString> >&
 
 		check( Enum );
 
+		const TArray<FName> ValidEnumValues = GetValidEnumEntries(Property, Enum);
+
 		//NumEnums() - 1, because the last item in an enum is the _MAX item
 		for( int32 EnumIndex = 0; EnumIndex < Enum->NumEnums() - 1; ++EnumIndex )
 		{
@@ -1885,6 +1907,10 @@ bool FPropertyHandleBase::GeneratePossibleValues(TArray< TSharedPtr<FString> >&
 
 			// Ignore hidden enums
 			bool bShouldBeHidden = Enum->HasMetaData(TEXT("Hidden"), EnumIndex ) || Enum->HasMetaData(TEXT("Spacer"), EnumIndex );
+			if (!bShouldBeHidden && ValidEnumValues.Num() != 0)
+			{
+				bShouldBeHidden = ValidEnumValues.Find(Enum->GetEnum(EnumIndex)) == INDEX_NONE;
+			}
 
 			if( !bShouldBeHidden )
 			{
diff --git a/Engine/Source/Editor/PropertyEditor/Private/UserInterface/PropertyEditor/SPropertyEditorNumeric.h b/Engine/Source/Editor/PropertyEditor/Private/UserInterface/PropertyEditor/SPropertyEditorNumeric.h
--- a/Engine/Source/Editor/PropertyEditor/Private/UserInterface/PropertyEditor/SPropertyEditorNumeric.h
+++ b/Engine/Source/Editor/PropertyEditor/Private/UserInterface/PropertyEditor/SPropertyEditorNumeric.h
@@ -93,7 +93,6 @@ public:
 			&& !PropertyNode->GetProperty()->GetBoolMetaData("NoSpinbox");
 
 		// Set up the correct type interface if we want to display units on the property editor
-		TSharedPtr<INumericTypeInterface<NumericType>> TypeInterface;
 		if (FUnitConversion::Settings().ShouldDisplayUnits())
 		{
 			const FString& Units = InPropertyEditor->GetProperty()->GetMetaData(TEXT("Units"));
@@ -101,17 +100,14 @@ public:
 
 			if (PropertyUnits.IsSet())
 			{
-				bool bAllowUnitRangeAdaption = true;
-				const auto& AllowRageAdaptionString = Property->GetMetaData(TEXT("AllowUnitRangeAdaption"));
-				if (!AllowRageAdaptionString.IsEmpty())
+				// Create the type interface and set up the default input units if they are compatible
+				TypeInterface = MakeShareable(new TNumericUnitTypeInterface<NumericType>(PropertyUnits.GetValue()));
+
+				auto Value = OnGetValue();
+				if (Value.IsSet())
 				{
-					LexicalConversion::FromString(bAllowUnitRangeAdaption, *AllowRageAdaptionString);
+					TypeInterface->SetupFixedDisplay(Value.GetValue());
 				}
-
-				// Create the type interface and set up the default input units if they are compatible
-				auto* Impl = new TNumericUnitTypeInterface<NumericType>(PropertyUnits.GetValue(), bAllowUnitRangeAdaption);
-				Impl->UseDefaultInputUnits();
-				TypeInterface = MakeShareable(Impl);
 			}
 		}
 
@@ -242,6 +238,11 @@ private:
 			// We don't create a transaction for each property change when using the slider.  Only once when the slider first is moved
 			EPropertyValueSetFlags::Type Flags = (EPropertyValueSetFlags::InteractiveChange | EPropertyValueSetFlags::NotTransactable);
 			PropertyHandle->SetValue( NewValue, Flags );
+
+			if (TypeInterface.IsValid() && !TypeInterface->FixedDisplayUnits.IsSet())
+			{
+				TypeInterface->SetupFixedDisplay(NewValue);
+			}
 		}
 	}
 
@@ -249,6 +250,11 @@ private:
 	{
 		const TSharedRef< IPropertyHandle > PropertyHandle = PropertyEditor->GetPropertyHandle();
 		PropertyHandle->SetValue( NewValue );
+
+		if (TypeInterface.IsValid() && !TypeInterface->FixedDisplayUnits.IsSet())
+		{
+			TypeInterface->SetupFixedDisplay(NewValue);
+		}
 	}
 	
 	/**
@@ -279,6 +285,9 @@ private:
 	}
 
 private:
+
+	TSharedPtr<TNumericUnitTypeInterface<NumericType>> TypeInterface;	
+
 	TSharedPtr< class FPropertyEditor > PropertyEditor;
 
 	TSharedPtr< class SWidget > PrimaryWidget;
diff --git a/Engine/Source/Editor/UnrealEd/Private/Settings/EditorProjectSettings.cpp b/Engine/Source/Editor/UnrealEd/Private/Settings/EditorProjectSettings.cpp
--- a/Engine/Source/Editor/UnrealEd/Private/Settings/EditorProjectSettings.cpp
+++ b/Engine/Source/Editor/UnrealEd/Private/Settings/EditorProjectSettings.cpp
@@ -23,22 +23,31 @@ EUnit ConvertDefaultInputUnits(UEditorProjectAppearanceSettings::EDefaultLocatio
 	}
 }
 
-EGlobalUnitDisplay ConvertGlobalUnitDisplay(EUnitDisplay In)
+// EGlobalUnitDisplay ConvertGlobalUnitDisplay(EUnitDisplay In)
+// {
+// 	switch(In)
+// 	{
+// 	case EUnitDisplay::None:		return EGlobalUnitDisplay::None;
+// 	case EUnitDisplay::Metric:		return EGlobalUnitDisplay::Metric;
+// 	case EUnitDisplay::Imperial:	return EGlobalUnitDisplay::Imperial;
+// 	default:				 		return EGlobalUnitDisplay::Metric;
+// 	}
+// }
+
+TArray<EUnit> ToRawUnits(const TArray<TEnumAsByte<EUnit>>& In)
 {
-	switch(In)
+	TArray<EUnit> Units;
+	for (const auto AsByte : In)
 	{
-	case EUnitDisplay::None:		return EGlobalUnitDisplay::None;
-	case EUnitDisplay::Metric:		return EGlobalUnitDisplay::Metric;
-	case EUnitDisplay::Imperial:	return EGlobalUnitDisplay::Imperial;
-	default:				 		return EGlobalUnitDisplay::Metric;
+		Units.Add(AsByte);
 	}
+	return Units;	
 }
 
 UEditorProjectAppearanceSettings::UEditorProjectAppearanceSettings(const FObjectInitializer&)
-	: UnitDisplay(EUnitDisplay::Metric)
-	, DefaultInputUnits(EDefaultLocationUnit::Centimeters)
+	: UnitDisplay_DEPRECATED(EUnitDisplay::Invalid)
+	, DefaultInputUnits_DEPRECATED(EDefaultLocationUnit::Invalid)
 {
-
 }
 
 void UEditorProjectAppearanceSettings::PostEditChangeProperty( struct FPropertyChangedEvent& PropertyChangedEvent )
@@ -47,26 +56,96 @@ void UEditorProjectAppearanceSettings::PostEditChangeProperty( struct FPropertyC
 
 	const FName Name = PropertyChangedEvent.Property ? PropertyChangedEvent.Property->GetFName() : NAME_None;
 
-	if (Name == GET_MEMBER_NAME_CHECKED(UEditorProjectAppearanceSettings, UnitDisplay))
+	auto& Settings = FUnitConversion::Settings();
+	if (Name == GET_MEMBER_NAME_CHECKED(UEditorProjectAppearanceSettings, DistanceUnits))
 	{
-		FUnitConversion::Settings().SetGlobalUnitDisplay(ConvertGlobalUnitDisplay(UnitDisplay));
+		TArray<EUnit> TempArray = ToRawUnits(DistanceUnits);
+		Settings.SetDisplayUnits(EUnitType::Distance, TempArray);
 	}
-	else if (Name == GET_MEMBER_NAME_CHECKED(UEditorProjectAppearanceSettings, DefaultInputUnits))
+	else if (Name == GET_MEMBER_NAME_CHECKED(UEditorProjectAppearanceSettings, MassUnits))
 	{
-		FUnitConversion::Settings().SetDefaultInputUnit(ConvertDefaultInputUnits(DefaultInputUnits));
+		TArray<EUnit> TempArray = ToRawUnits(MassUnits);
+		Settings.SetDisplayUnits(EUnitType::Mass, TempArray);
 	}
+	else if (Name == GET_MEMBER_NAME_CHECKED(UEditorProjectAppearanceSettings, TimeUnits))
+	{
+		TArray<EUnit> TempArray = ToRawUnits(TimeUnits);
+		Settings.SetDisplayUnits(EUnitType::Time, TempArray);
+	}
+	else if (Name == GET_MEMBER_NAME_CHECKED(UEditorProjectAppearanceSettings, AngleUnits))
+	{
+		Settings.SetDisplayUnits(EUnitType::Angle, AngleUnits);
+	}
+	else if (Name == GET_MEMBER_NAME_CHECKED(UEditorProjectAppearanceSettings, SpeedUnits))
+	{
+		Settings.SetDisplayUnits(EUnitType::Speed, SpeedUnits);
+	}
+	else if (Name == GET_MEMBER_NAME_CHECKED(UEditorProjectAppearanceSettings, TemperatureUnits))
+	{
+		Settings.SetDisplayUnits(EUnitType::Temperature, TemperatureUnits);
+	}
+	else if (Name == GET_MEMBER_NAME_CHECKED(UEditorProjectAppearanceSettings, ForceUnits))
+	{
+		Settings.SetDisplayUnits(EUnitType::Force, ForceUnits);
+	}
+	else if (Name == GET_MEMBER_NAME_CHECKED(UEditorProjectAppearanceSettings, bDisplayUnits))
+	{
+		Settings.SetShouldDisplayUnits(bDisplayUnits);
+	}
+
+	DefaultInputUnits_DEPRECATED = EDefaultLocationUnit::Invalid;
+	UnitDisplay_DEPRECATED = EUnitDisplay::Invalid;
 
 	SaveConfig();
 }
 
+void SetupEnumMetaData(UClass* Class, const FName& MemberName, const TCHAR* Values)
+{
+	UArrayProperty* Array = Cast<UArrayProperty>(Class->FindPropertyByName(MemberName));
+	if (Array && Array->Inner)
+	{
+		Array->Inner->SetMetaData(TEXT("ValidEnumValues"), Values);
+	}
+}
+
 void UEditorProjectAppearanceSettings::PostInitProperties()
 {
 	Super::PostInitProperties();
 
-	FUnitConversion::Settings().SetGlobalUnitDisplay(ConvertGlobalUnitDisplay(UnitDisplay));
-	FUnitConversion::Settings().SetDefaultInputUnit(ConvertDefaultInputUnits(DefaultInputUnits));
-}
+	/** Setup the meta data for the array properties */
+	SetupEnumMetaData(GetClass(), GET_MEMBER_NAME_CHECKED(UEditorProjectAppearanceSettings, DistanceUnits), TEXT("Micrometers, Millimeters, Centimeters, Meters, Kilometers, Inches, Feet, Yards, Miles"));
+	SetupEnumMetaData(GetClass(), GET_MEMBER_NAME_CHECKED(UEditorProjectAppearanceSettings, MassUnits), TEXT("Micrograms, Milligrams, Grams, Kilograms, MetricTons,	Ounces, Pounds, Stones"));
+	SetupEnumMetaData(GetClass(), GET_MEMBER_NAME_CHECKED(UEditorProjectAppearanceSettings, TimeUnits), TEXT("Milliseconds, Seconds, Minutes, Hours, Days, Months, Years"));
+
+	if (UnitDisplay_DEPRECATED != EUnitDisplay::Invalid)
+	{
+		bDisplayUnits = UnitDisplay_DEPRECATED != EUnitDisplay::None;
+	}
+
+	if (DefaultInputUnits_DEPRECATED != EDefaultLocationUnit::Invalid)
+	{
+		DistanceUnits.Empty();
+		DistanceUnits.Add(ConvertDefaultInputUnits(DefaultInputUnits_DEPRECATED));
+	}
+
+	auto& Settings = FUnitConversion::Settings();
+
+	TArray<EUnit> TempArray = ToRawUnits(DistanceUnits);
+	Settings.SetDisplayUnits(EUnitType::Distance, TempArray);
 
+	TempArray = ToRawUnits(MassUnits);
+	Settings.SetDisplayUnits(EUnitType::Mass, TempArray);
+
+	TempArray = ToRawUnits(TimeUnits);
+	Settings.SetDisplayUnits(EUnitType::Time, TempArray);
+
+	Settings.SetDisplayUnits(EUnitType::Angle, AngleUnits);
+	Settings.SetDisplayUnits(EUnitType::Speed, SpeedUnits);
+	Settings.SetDisplayUnits(EUnitType::Temperature, TemperatureUnits);
+	Settings.SetDisplayUnits(EUnitType::Force, ForceUnits);
+
+	Settings.SetShouldDisplayUnits(bDisplayUnits);
+}
 
 
 /* ULevelEditor2DSettings
diff --git a/Engine/Source/Editor/UnrealEd/Public/Settings/EditorProjectSettings.h b/Engine/Source/Editor/UnrealEd/Public/Settings/EditorProjectSettings.h
--- a/Engine/Source/Editor/UnrealEd/Public/Settings/EditorProjectSettings.h
+++ b/Engine/Source/Editor/UnrealEd/Public/Settings/EditorProjectSettings.h
@@ -6,13 +6,6 @@
 
 #include "EditorProjectSettings.generated.h"
 
-/** UENUM to define the specific set of allowable unit types */
-UENUM()
-enum class EUnitDisplay : uint8
-{
-	None, Metric, Imperial
-};
-
 /** Editor project appearance settings. Stored in default config, per-project */
 UCLASS(config=Editor, defaultconfig)
 class UNREALED_API UEditorProjectAppearanceSettings : public UObject
@@ -21,28 +14,61 @@ public:
 	GENERATED_BODY()
 	UEditorProjectAppearanceSettings(const FObjectInitializer&);
 
-	/** UENUM to define the specific set of allowable default units */
-	UENUM()
-	enum class EDefaultLocationUnit : uint8
-	{
-		Micrometers, Millimeters, Centimeters, Meters, Kilometers,
-		Inches, Feet, Yards, Miles
-	};
-
 protected:
 	/** Called when a property on this object is changed */
 	virtual void PostEditChangeProperty( struct FPropertyChangedEvent& PropertyChangedEvent ) override;
 	virtual void PostInitProperties() override;
 
 private:
+
+	UPROPERTY(EditAnywhere, config, Category=Units, meta=(DisplayName="Display Units on Applicable Properties", Tooltip="Whether to display units on editor properties where the property has units set."))
+	bool bDisplayUnits;
+
+	UPROPERTY(EditAnywhere, config, Category=Units, AdvancedDisplay, meta=(DisplayName="Distance/Length", Tooltip="Choose a set of units in which to display distance/length values."))
+	TArray<TEnumAsByte<EUnit>> DistanceUnits;
+
+	UPROPERTY(EditAnywhere, config, Category=Units, AdvancedDisplay, meta=(DisplayName="Mass", Tooltip="Choose a set of units in which to display masses."))
+	TArray<TEnumAsByte<EUnit>> MassUnits;
+
+	UPROPERTY(EditAnywhere, config, Category=Units, AdvancedDisplay, meta=(DisplayName="Time", Tooltip="Choose the units in which to display time."))
+	TArray<TEnumAsByte<EUnit>> TimeUnits;
 	
-	UPROPERTY(EditAnywhere, config, Category=Units, meta=(Tooltip="Specifies how to present applicable property units on editor interfaces."))
-	EUnitDisplay UnitDisplay;
+	UPROPERTY(EditAnywhere, config, Category=Units, AdvancedDisplay, meta=(DisplayName="Angles", Tooltip="Choose the units in which to display angles.", ValidEnumValues="Degrees, Radians"))
+	TEnumAsByte<EUnit> AngleUnits;
 
-	UPROPERTY(EditAnywhere, config, Category=Units, meta=(Tooltip="The default units to use for location and distance measurements on editor user interfaces when not explicitly specified by the user."))
-	EDefaultLocationUnit DefaultInputUnits;
-};
+	UPROPERTY(EditAnywhere, config, Category=Units, AdvancedDisplay, meta=(DisplayName="Speed/Velocity", Tooltip="Choose the units in which to display speeds and velocities.", ValidEnumValues="MetersPerSecond, KilometersPerHour, MilesPerHour"))
+	TEnumAsByte<EUnit> SpeedUnits;
+
+	UPROPERTY(EditAnywhere, config, Category=Units, AdvancedDisplay, meta=(DisplayName="Temperature", Tooltip="Choose the units in which to display temperatures.", ValidEnumValues="Celsius, Farenheit, Kelvin"))
+	TEnumAsByte<EUnit> TemperatureUnits;
 
+	UPROPERTY(EditAnywhere, config, Category=Units, AdvancedDisplay, meta=(DisplayName="Force", Tooltip="Choose the units in which to display forces.", ValidEnumValues="Newtons, PoundsForce, KilogramsForce"))
+	TEnumAsByte<EUnit> ForceUnits;
+
+public:
+	/** Deprecated properties that didn't live very long */
+
+	/** UENUM to define the specific set of allowable unit types */
+	UENUM()
+	enum class EUnitDisplay : uint8
+	{
+		None, Metric, Imperial, Invalid
+	};
+
+	UPROPERTY(config)
+	EUnitDisplay UnitDisplay_DEPRECATED;
+
+	/** UENUM to define the specific set of allowable default units */
+	UENUM()
+	enum class EDefaultLocationUnit : uint8
+	{
+		Micrometers, Millimeters, Centimeters, Meters, Kilometers,
+		Inches, Feet, Yards, Miles, Invalid
+	};
+
+	UPROPERTY(config)
+	EDefaultLocationUnit DefaultInputUnits_DEPRECATED;
+};
 /**
 * 2D layer settings
 */
diff --git a/Engine/Source/Runtime/Core/Private/Math/UnitConversion.cpp b/Engine/Source/Runtime/Core/Private/Math/UnitConversion.cpp
--- a/Engine/Source/Runtime/Core/Private/Math/UnitConversion.cpp
+++ b/Engine/Source/Runtime/Core/Private/Math/UnitConversion.cpp
@@ -88,7 +88,7 @@ const TCHAR* const FUnitConversion::DisplayStrings[] = {
 	TEXT("ms"), TEXT("s"), TEXT("min"), TEXT("hr"), TEXT("dy"), TEXT("mth"), TEXT("yr"),
 };
 
-const FUnitConversion::EUnitType FUnitConversion::Types[] = {
+const EUnitType FUnitConversion::Types[] = {
 	EUnitType::Distance,	EUnitType::Distance,	EUnitType::Distance,	EUnitType::Distance,	EUnitType::Distance,
 	EUnitType::Distance,	EUnitType::Distance,	EUnitType::Distance,	EUnitType::Distance,
 	EUnitType::Distance,
@@ -113,79 +113,55 @@ const FUnitConversion::EUnitType FUnitConversion::Types[] = {
 	EUnitType::Time,		EUnitType::Time,		EUnitType::Time,		EUnitType::Time,		EUnitType::Time,		EUnitType::Time,		EUnitType::Time,
 };
 
-const EUnit MetricToImperial[] = {
-	/** Scalar distance/length units */
-	EUnit::Inches/*Micrometers*/, EUnit::Inches/*Millimeters*/, EUnit::Inches/*Centimeters*/, EUnit::Yards/*Meters*/, EUnit::Miles/*Kilometers*/,
-	EUnit::Inches, EUnit::Feet, EUnit::Yards, EUnit::Miles,
-	EUnit::Lightyears,
-
-	/** Angular units */
-	EUnit::Degrees, EUnit::Radians,
-
-	/** Speed units */
-	EUnit::MilesPerHour/*MetersPerSecond*/, EUnit::MilesPerHour/*KilometersPerHour*/, EUnit::MilesPerHour,
-
-	/** Temperature units */
-	EUnit::Farenheit/*Celsius*/, EUnit::Farenheit, EUnit::Kelvin,
-
-	/** Mass units */
-	EUnit::Ounces/*Micrograms*/, EUnit::Ounces/*Milligrams*/, EUnit::Ounces/*Grams*/, EUnit::Pounds/*Kilograms*/, EUnit::Stones/*MetricTons*/,
-	EUnit::Ounces, EUnit::Pounds, EUnit::Stones,
-
-	/** Force units */
-	EUnit::PoundsForce/*Newtons*/, EUnit::PoundsForce, EUnit::PoundsForce/*KilogramsForce*/,
-
-	/** Frequency units */
-	EUnit::Hertz, EUnit::Kilohertz, EUnit::Megahertz, EUnit::Gigahertz, EUnit::RevolutionsPerMinute,
-
-	/** Data Size units */
-	EUnit::Bytes, EUnit::Kilobytes, EUnit::Megabytes, EUnit::Gigabytes, EUnit::Terabytes,
-
-	/** Luminous flux units */
-	EUnit::Lumens,
-
-	/** Time units */
-	EUnit::Milliseconds, EUnit::Seconds, EUnit::Minutes, EUnit::Hours, EUnit::Days, EUnit::Months, EUnit::Years
-};
-EUnit ImperialToMetric[] = {
-	/** Scalar distance/length units */
-	EUnit::Micrometers, EUnit::Millimeters, EUnit::Centimeters, EUnit::Meters, EUnit::Kilometers,
-	EUnit::Centimeters/*Inches*/, EUnit::Meters/*Feet*/, EUnit::Meters/*Yards*/, EUnit::Kilometers/*Miles*/,
-	EUnit::Lightyears,
-
-	/** Angular units */
-	EUnit::Degrees, EUnit::Radians,
-
-	/** Speed units */
-	EUnit::MetersPerSecond, EUnit::KilometersPerHour, EUnit::KilometersPerHour/*MilesPerHour*/,
-
-	/** Temperature units */
-	EUnit::Celsius, EUnit::Celsius/*Farenheit*/, EUnit::Kelvin,
-
-	/** Mass units */
-	EUnit::Micrograms, EUnit::Milligrams, EUnit::Grams, EUnit::Kilograms, EUnit::MetricTons,
-	EUnit::Grams/*Ounces*/, EUnit::Kilograms/*Pounds*/, EUnit::MetricTons/*Stones*/,
-
-	/** Force units */
-	EUnit::Newtons, EUnit::Newtons/*PoundsForce*/, EUnit::KilogramsForce,
-
-	/** Frequency units */
-	EUnit::Hertz, EUnit::Kilohertz, EUnit::Megahertz, EUnit::Gigahertz, EUnit::RevolutionsPerMinute,
+FUnitSettings::FUnitSettings()
+	: bGlobalUnitDisplay(true)
+{
+	DisplayUnits[(uint8)EUnitType::Distance].Add(EUnit::Centimeters);
+	DisplayUnits[(uint8)EUnitType::Angle].Add(EUnit::Degrees);
+	DisplayUnits[(uint8)EUnitType::Speed].Add(EUnit::MetersPerSecond);
+	DisplayUnits[(uint8)EUnitType::Temperature].Add(EUnit::Celsius);
+	DisplayUnits[(uint8)EUnitType::Mass].Add(EUnit::Kilograms);
+	DisplayUnits[(uint8)EUnitType::Force].Add(EUnit::Newtons);
+	DisplayUnits[(uint8)EUnitType::Frequency].Add(EUnit::Hertz);
+	DisplayUnits[(uint8)EUnitType::DataSize].Add(EUnit::Megabytes);
+	DisplayUnits[(uint8)EUnitType::LuminousFlux].Add(EUnit::Lumens);
+	DisplayUnits[(uint8)EUnitType::Time].Add(EUnit::Seconds);
+}
 
-	/** Data Size units */
-	EUnit::Bytes, EUnit::Kilobytes, EUnit::Megabytes, EUnit::Gigabytes, EUnit::Terabytes,
+const TArray<EUnit>& FUnitSettings::GetDisplayUnits(EUnitType InType)
+{
+	return DisplayUnits[(uint8)InType];
+}
 
-	/** Luminous flux units */
-	EUnit::Lumens,
+void FUnitSettings::SetDisplayUnits(EUnitType InType, EUnit Unit)
+{
+	if (InType != EUnitType::NumberOf)
+	{
+		DisplayUnits[(uint8)InType].Empty();
+		if (FUnitConversion::IsUnitOfType(Unit, InType))
+		{
+			DisplayUnits[(uint8)InType].Add(Unit);
+		}
 
-	/** Time units */
-	EUnit::Milliseconds, EUnit::Seconds, EUnit::Minutes, EUnit::Hours, EUnit::Days, EUnit::Months, EUnit::Years,
-};
+		SettingChangedEvent.Broadcast();
+	}	
+}
 
-FUnitSettings::FUnitSettings()
-	: GlobalUnitDisplay(EGlobalUnitDisplay::Metric)
-	, DefaultInputUnit(EUnit::Centimeters)
+void FUnitSettings::SetDisplayUnits(EUnitType InType, const TArray<EUnit>& Units)
 {
+	if (InType != EUnitType::NumberOf)
+	{
+		DisplayUnits[(uint8)InType].Reset();
+		for (EUnit Unit : Units)
+		{
+			if (FUnitConversion::IsUnitOfType(Unit, InType))
+			{
+				DisplayUnits[(uint8)InType].Add(Unit);
+			}
+		}
+
+		SettingChangedEvent.Broadcast();
+	}
 }
 
 /** Get the global settings for unit conversion/display */
@@ -200,8 +176,6 @@ const TCHAR* FUnitConversion::GetUnitDisplayString(EUnit Unit)
 {
 	static_assert(ARRAY_COUNT(FUnitConversion::Types) == (uint8)EUnit::Unspecified, "Type array does not match size of unit enum");
 	static_assert(ARRAY_COUNT(FUnitConversion::DisplayStrings) == (uint8)EUnit::Unspecified, "Display String array does not match size of unit enum");
-	static_assert(ARRAY_COUNT(MetricToImperial) == (uint8)EUnit::Unspecified, "MetricToImperial array does not match size of unit enum");
-	static_assert(ARRAY_COUNT(ImperialToMetric) == (uint8)EUnit::Unspecified, "ImperialToMetric array does not match size of unit enum");
 	
 	if (Unit != EUnit::Unspecified)
 	{
@@ -229,21 +203,6 @@ TOptional<EUnit> FUnitConversion::UnitFromString(const TCHAR* UnitString)
 	return TOptional<EUnit>();
 }
 
-EUnit FUnitConversion::ConvertToGlobalDisplayRange(EUnit InUnit)
-{
-	if (InUnit == EUnit::Unspecified)
-	{
-		return EUnit::Unspecified;
-	}
-
-	switch(Settings().GetGlobalUnitDisplay())
-	{
-	case EGlobalUnitDisplay::Metric:	return ImperialToMetric[(int8)InUnit];
-	case EGlobalUnitDisplay::Imperial:	return MetricToImperial[(int8)InUnit];
-	default:							return InUnit;
-	}
-}
-
 double FUnitConversion::DistanceUnificationFactor(EUnit From)
 {
 	// Convert to meters
diff --git a/Engine/Source/Runtime/Core/Public/Math/UnitConversion.h b/Engine/Source/Runtime/Core/Public/Math/UnitConversion.h
--- a/Engine/Source/Runtime/Core/Public/Math/UnitConversion.h
+++ b/Engine/Source/Runtime/Core/Public/Math/UnitConversion.h
@@ -5,7 +5,8 @@
 #include "Optional.h"
 #include "Function.h"
 
-/** Enum *must* be zero-indexed and sequential */
+/** Enum *must* be zero-indexed and sequential. Must be grouped by relevance and ordered by magnitude. */
+/** Enum *must* match the mirrored enum that exists in CoreUObject/Classes/Object.h for the purposes of UObject reflection */
 enum class EUnit
 {
 	/** Scalar distance/length units */
@@ -45,9 +46,13 @@ enum class EUnit
 	Unspecified
 };
 
-enum class EGlobalUnitDisplay
+/** Enumeration that specifies particular classes of unit */
+enum class EUnitType
 {
-	None, Metric, Imperial
+	Distance, Angle, Speed, Temperature, Mass, Force, Frequency, DataSize, LuminousFlux, Time,
+
+	// Symbolic entry - do not use directly
+	NumberOf,
 };
 
 template<typename NumericType> struct FNumericUnit;
@@ -56,19 +61,17 @@ template<typename NumericType> struct FNumericUnit;
 class CORE_API FUnitSettings
 {
 public:
-	
+
 	FUnitSettings();
 
 	/** Check whether unit display is globally enabled or disabled */
-	bool ShouldDisplayUnits() { return GlobalUnitDisplay != EGlobalUnitDisplay::None; }
-
-	/** Get/set the global unit display to either mertic, imperial, or none */
-	EGlobalUnitDisplay GetGlobalUnitDisplay() { return GlobalUnitDisplay; }
-	void SetGlobalUnitDisplay(EGlobalUnitDisplay InGlobalUnitDisplay) { GlobalUnitDisplay = InGlobalUnitDisplay; SettingChangedEvent.Broadcast(); }
+	bool ShouldDisplayUnits() { return bGlobalUnitDisplay; }
+	void SetShouldDisplayUnits(bool bInGlobalUnitDisplay) { bGlobalUnitDisplay = bInGlobalUnitDisplay; SettingChangedEvent.Broadcast(); }
 	
-	/** Default unit to use for text input when no unit was specified. Can be EUnit::Unspecified. Check for compatibility with property units using AreUnitsCompatible */
-	EUnit GetDefaultInputUnit() { return DefaultInputUnit; }
-	void SetDefaultInputUnit(EUnit InDefaultInputUnit) { DefaultInputUnit = InDefaultInputUnit; SettingChangedEvent.Broadcast(); }
+	/** Get the specific valid units to display the specified type of unit in */
+	const TArray<EUnit>& GetDisplayUnits(EUnitType InType);
+	void SetDisplayUnits(EUnitType InType, const TArray<EUnit>& Units);
+	void SetDisplayUnits(EUnitType InType, EUnit Units);
 
 	/** Returns an event delegate that is executed when a display setting has changed. (GlobalUnitDisplay or DefaultInputUnits) */
 	DECLARE_EVENT(FUnitSettings, FDisplaySettingChanged);
@@ -76,11 +79,11 @@ public:
 
 private:
 
-	/** Globally check how we should display units on user interfaces */
-	EGlobalUnitDisplay GlobalUnitDisplay;
+	/** Global toggle controlling whether we should display units or not */
+	bool bGlobalUnitDisplay;
 
-	/** Default unit to use for text input when no unit was specified. Can be EUnit::Unspecified. Check for compatibility with property units using AreUnitsCompatible */
-	EUnit DefaultInputUnit;
+	/** Arrays of units that are valid to display on interfaces */
+	TArray<EUnit> DisplayUnits[(uint8)EUnitType::NumberOf + 1];
 	
 	/** Holds an event delegate that is executed when a display setting has changed. */
 	FDisplaySettingChanged SettingChangedEvent;
@@ -97,6 +100,12 @@ struct CORE_API FUnitConversion
 		return From == EUnit::Unspecified || To == EUnit::Unspecified || Types[(uint8)From] == Types[(uint8)To];
 	}
 
+	/** Check whether a unit is of the specified type */
+	static bool IsUnitOfType(EUnit Unit, EUnitType Type)
+	{
+		return Unit != EUnit::Unspecified && Types[(uint8)Unit] == Type;
+	}
+
 	/** Convert the specified number from one unit to another. Does nothing if the units are incompatible. */
 	template<typename T>
 	static T Convert(T InValue, EUnit From, EUnit To)
@@ -151,9 +160,6 @@ struct CORE_API FUnitConversion
 	/** Helper function to find a unit from a string (name or display string) */
 	static TOptional<EUnit> UnitFromString(const TCHAR* UnitString);
 
-	/** Perform metric <-> imperial conversion on the specified unit according to the GlobalUnitDisplay setting */
-	static EUnit ConvertToGlobalDisplayRange(EUnit InUnit);
-
 private:
 
 	/** Find the common quantization factor for the specified distance unit. Quantizes to Meters. */
@@ -248,13 +254,51 @@ public:
 		return FNumericUnit<T>(NewValue, NewUnits);
 	}
 
-private:
-
-	/** Private enumeration that specifies particular classes of unit */
-	enum class EUnitType
+	/** Quantizes this number to the most appropriate unit for user friendly presentation (e.g. 1000m returns 1km), adhereing to global display settings. */
+	template<typename T>
+	static EUnit CalculateDisplayUnit(T Value, EUnit InUnits)
 	{
-		Distance, Angle, Speed, Temperature, Mass, Force, Frequency, DataSize, LuminousFlux, Time
-	};
+		if (InUnits == EUnit::Unspecified)
+		{
+			return EUnit::Unspecified;
+		}
+
+		const TArray<EUnit>& DisplayUnits = Settings().GetDisplayUnits(Types[(uint8)InUnits]);
+		if (DisplayUnits.Num() == 0)
+		{
+			return QuantizeUnitsToBestFit(Value, InUnits).Units;
+		}
+		else if (DisplayUnits.Num() == 1)
+		{
+			return DisplayUnits[0];
+		}
+
+		// If the value we were given was 0, change it to something we can actually work with
+		if (Value == 0)
+		{
+			Value = 1;
+		}
+
+		int32 BestIndex = 0;
+		for (int32 Index = 0; Index < DisplayUnits.Num() - 1; ++Index)
+		{
+			double This = Convert(Value, InUnits, DisplayUnits[Index]);
+			double Next = Convert(Value, InUnits, DisplayUnits[Index + 1]);
+
+			if (FMath::Abs(FMath::LogX(10.0f, This)) < FMath::Abs(FMath::LogX(10.0f, Next)))
+			{
+				BestIndex = Index;
+			}
+			else
+			{
+				BestIndex = Index + 1;
+			}
+		}
+
+		return DisplayUnits[BestIndex];
+	}
+
+private:
 
 	/** Structure used to match units when parsing */
 	struct FParseCandidate
diff --git a/Engine/Source/Runtime/CoreUObject/Classes/Object.h b/Engine/Source/Runtime/CoreUObject/Classes/Object.h
--- a/Engine/Source/Runtime/CoreUObject/Classes/Object.h
+++ b/Engine/Source/Runtime/CoreUObject/Classes/Object.h
@@ -206,6 +206,47 @@ namespace EMouseCursor
 	};
 }
 
+/** A set of numerical unit types supported by the engine. Mirrored from UnitConversion.h */
+UENUM(BlueprintType)
+enum class EUnit
+{
+	/** Scalar distance/length units */
+	Micrometers, Millimeters, Centimeters, Meters, Kilometers,
+	Inches, Feet, Yards, Miles,
+	Lightyears,
+
+	/** Angular units */
+	Degrees, Radians,
+
+	/** Speed units */
+	MetersPerSecond, KilometersPerHour, MilesPerHour,
+
+	/** Temperature units */
+	Celsius, Farenheit, Kelvin,
+
+	/** Mass units */
+	Micrograms, Milligrams, Grams, Kilograms, MetricTons,
+	Ounces, Pounds, Stones,
+
+	/** Force units */
+	Newtons, PoundsForce, KilogramsForce,
+
+	/** Frequency units */
+	Hertz, Kilohertz, Megahertz, Gigahertz, RevolutionsPerMinute,
+
+	/** Data Size units */
+	Bytes, Kilobytes, Megabytes, Gigabytes, Terabytes,
+
+	/** Luminous flux units */
+	Lumens,
+
+	/** Time units */
+	Milliseconds, Seconds, Minutes, Hours, Days, Months, Years,
+
+	/** Symbolic entry, not specifyable on meta data */
+	Unspecified
+};
+
 // A globally unique identifier.
 USTRUCT(immutable, noexport, BlueprintType)
 struct FGuid
diff --git a/Engine/Source/Runtime/Slate/Public/Widgets/Input/NumericTypeInterface.h b/Engine/Source/Runtime/Slate/Public/Widgets/Input/NumericTypeInterface.h
--- a/Engine/Source/Runtime/Slate/Public/Widgets/Input/NumericTypeInterface.h
+++ b/Engine/Source/Runtime/Slate/Public/Widgets/Input/NumericTypeInterface.h
@@ -71,21 +71,14 @@ template<typename> struct FNumericUnit;
 template<typename NumericType>
 struct TNumericUnitTypeInterface : TDefaultNumericTypeInterface<NumericType>
 {
-	/** The units that we should treat this value as */
-	EUnit Units;
+	/** The underlying units which the numeric type are specfied in. */
+	const EUnit UnderlyingUnits;
 
-	/** Optional units to use for input when no units where specified */
-	TOptional<EUnit> DefaultInputUnits;
-
-	/** True if we are allowed to display this unit to the user as a more relevant related unit type (e.g. 10000 m shows as 10 km) */
-	bool bAllowUnitRangeAdaption;
+	/** Optional units that this type interface will be fixed on */
+	TOptional<EUnit> FixedDisplayUnits;
 
 	/** Constructor */
-	TNumericUnitTypeInterface(EUnit InUnits, bool bInAllowUnitRangeAdaption = true);
-	~TNumericUnitTypeInterface();
-
-	/** Set up this interface to use the global default input units */
-	void UseDefaultInputUnits();
+	TNumericUnitTypeInterface(EUnit InUnits);
 
 	/** Convert this type to a string */
 	virtual FString ToString(const NumericType& Value) const override;
@@ -93,12 +86,12 @@ struct TNumericUnitTypeInterface : TDefaultNumericTypeInterface<NumericType>
 	/** Attempt to parse a numeral with our units from the specified string. */
 	virtual TOptional<NumericType> FromString(const FString& ValueString) override;
 
-	/** Called by this class to quantize the specified value to the best unit */
-	virtual FNumericUnit<NumericType> QuantizeUnitsToBestFit(const NumericType& InValue, EUnit InUnits) const;
-
 	/** Check whether the specified typed character is valid */
 	virtual bool IsCharacterValid(TCHAR InChar) const override;
 
+	/** Set up this interface to use a fixed display unit based on the specified value */
+	void SetupFixedDisplay(const NumericType& InValue);
+
 private:
 	/** Called when the global unit settings have changed, if this type interface is using the default input units */
 	void OnGlobalUnitSettingChanged();
diff --git a/Engine/Source/Runtime/Slate/Public/Widgets/Input/NumericUnitTypeInterface.inl b/Engine/Source/Runtime/Slate/Public/Widgets/Input/NumericUnitTypeInterface.inl
--- a/Engine/Source/Runtime/Slate/Public/Widgets/Input/NumericUnitTypeInterface.inl
+++ b/Engine/Source/Runtime/Slate/Public/Widgets/Input/NumericUnitTypeInterface.inl
@@ -5,63 +5,27 @@
 #include "UnitConversion.h"
 
 template<typename NumericType>
-TNumericUnitTypeInterface<NumericType>::TNumericUnitTypeInterface(EUnit InUnits, bool bInAllowUnitRangeAdaption)
-	: Units(InUnits), bAllowUnitRangeAdaption(bInAllowUnitRangeAdaption)
+TNumericUnitTypeInterface<NumericType>::TNumericUnitTypeInterface(EUnit InUnits)
+	: UnderlyingUnits(InUnits)
 {}
 
-template<typename NumericType>
-TNumericUnitTypeInterface<NumericType>::~TNumericUnitTypeInterface()
-{
-	FUnitConversion::Settings().OnDisplaySettingsChanged().RemoveAll(this);
-}
-
-
-template<typename NumericType>
-void TNumericUnitTypeInterface<NumericType>::UseDefaultInputUnits()
-{
-	FUnitConversion::Settings().OnDisplaySettingsChanged().AddRaw(this, &TNumericUnitTypeInterface<NumericType>::OnGlobalUnitSettingChanged);
-	OnGlobalUnitSettingChanged();
-}
-
-template<typename NumericType>
-void TNumericUnitTypeInterface<NumericType>::OnGlobalUnitSettingChanged()
-{
-	const EUnit Default = FUnitConversion::Settings().GetDefaultInputUnit();
-	if (FUnitConversion::AreUnitsCompatible(Default, Units))
-	{
-		DefaultInputUnits = Default;
-	}
-	else
-	{
-		DefaultInputUnits.Reset();
-	}
-}
-
-template<typename NumericType>
-FNumericUnit<NumericType> TNumericUnitTypeInterface<NumericType>::QuantizeUnitsToBestFit(const NumericType& InValue, EUnit InUnits) const
-{
-	// Use the default input units for 0
-	if (InValue == 0 && DefaultInputUnits.IsSet() && FUnitConversion::AreUnitsCompatible(DefaultInputUnits.GetValue(), InUnits))
-	{
-		return FNumericUnit<NumericType>(0, FUnitConversion::ConvertToGlobalDisplayRange(DefaultInputUnits.GetValue()));
-	}
-	return FUnitConversion::QuantizeUnitsToBestFit(InValue, InUnits);
-}
-
 template<typename NumericType>
 FString TNumericUnitTypeInterface<NumericType>::ToString(const NumericType& Value) const
 {
 	using namespace LexicalConversion;
 
-	NumericType ValueToUse = Value;
-	EUnit DisplayUnits = FUnitConversion::ConvertToGlobalDisplayRange(Units);
+	FNumericUnit<NumericType> FinalValue(Value, UnderlyingUnits);
 
-	if (DisplayUnits != Units)
+	if (FixedDisplayUnits.IsSet())
 	{
-		ValueToUse = FUnitConversion::Convert(ValueToUse, Units, DisplayUnits);
+		auto Converted = FinalValue.ConvertTo(FixedDisplayUnits.GetValue());
+		if (Converted.IsSet())
+		{
+			return ToSanitizedString(Converted.GetValue());
+		}
 	}
 	
-	return ToSanitizedString(bAllowUnitRangeAdaption ? QuantizeUnitsToBestFit(ValueToUse, DisplayUnits) : FNumericUnit<NumericType>(ValueToUse, DisplayUnits));
+	return ToSanitizedString(FinalValue);
 }
 
 template<typename NumericType>
@@ -71,17 +35,17 @@ TOptional<NumericType> TNumericUnitTypeInterface<NumericType>::FromString(const
 
 	// Always parse in as a double, to allow for input of higher-order units with decimal numerals into integral types (eg, inputting 0.5km as 500m)
 	FNumericUnit<double> NewValue;
-	bool bEvalResult = TryParseString( NewValue, *InString ) && FUnitConversion::AreUnitsCompatible( NewValue.Units, Units );
+	bool bEvalResult = TryParseString( NewValue, *InString ) && FUnitConversion::AreUnitsCompatible( NewValue.Units, UnderlyingUnits );
 	if (bEvalResult)
 	{
 		// Convert the number into the correct units
 		EUnit SourceUnits = NewValue.Units;
-		if (SourceUnits == EUnit::Unspecified && DefaultInputUnits.IsSet())
+		if (SourceUnits == EUnit::Unspecified && FixedDisplayUnits.IsSet())
 		{
 			// Use the default supplied input units
-			SourceUnits = DefaultInputUnits.GetValue();
+			SourceUnits = FixedDisplayUnits.GetValue();
 		}
-		return FUnitConversion::Convert(NewValue.Value, SourceUnits, Units);
+		return FUnitConversion::Convert(NewValue.Value, SourceUnits, UnderlyingUnits);
 	}
 	else
 	{
@@ -99,4 +63,14 @@ template<typename NumericType>
 bool TNumericUnitTypeInterface<NumericType>::IsCharacterValid(TCHAR InChar) const
 {
 	return true;
+}
+
+template<typename NumericType>
+void TNumericUnitTypeInterface<NumericType>::SetupFixedDisplay(const NumericType& InValue)
+{
+	EUnit DisplayUnit = FUnitConversion::CalculateDisplayUnit(InValue, UnderlyingUnits);
+	if (DisplayUnit != EUnit::Unspecified)
+	{
+		FixedDisplayUnits = DisplayUnit;
+	}
 }
\ No newline at end of file
