FText::AsNumber usability and performance improvements

FNumberFormattingOptions now provides a convenient way to create a default number formatter that doesn't use grouping (FNumberFormattingOptions::DefaultNoGrouping()), as well as one that does (FNumberFormattingOptions::DefaultWithGrouping()).

This also improves the performance when using the default formatting options, by caching the formatters used by those options. This helps greatly when using Blueprint formatting (see KismetTextLibrary.cpp) as these were always creating custom formatters for every operations, which is a *very* expensive thing to do.

To improve the cases where you are frequently re-using a custom set of formatting options, we also have a LRU cache of the most recently used formatting options which we check before creating a new custom formatter.

diff --git a/Engine/Plugins/2D/Paper2D/Source/Paper2D/Private/PaperTileMap.cpp b/Engine/Plugins/2D/Paper2D/Source/Paper2D/Private/PaperTileMap.cpp
--- a/Engine/Plugins/2D/Paper2D/Source/Paper2D/Private/PaperTileMap.cpp
+++ b/Engine/Plugins/2D/Paper2D/Source/Paper2D/Private/PaperTileMap.cpp
@@ -373,11 +373,7 @@ FText UPaperTileMap::GenerateNewLayerName(UPaperTileMap* TileMap)
 	do
 	{
 		TileMap->LayerNameIndex++;
-
-		FNumberFormattingOptions NoGroupingFormat;
-		NoGroupingFormat.SetUseGrouping(false);
-
-		TestLayerName = FText::Format(LOCTEXT("NewLayerNameFormatString", "Layer {0}"), FText::AsNumber(TileMap->LayerNameIndex, &NoGroupingFormat));
+		TestLayerName = FText::Format(LOCTEXT("NewLayerNameFormatString", "Layer {0}"), FText::AsNumber(TileMap->LayerNameIndex, &FNumberFormattingOptions::DefaultNoGrouping()));
 	} while (ExistingNames.Contains(TestLayerName.ToString()));
 
 	return TestLayerName;
diff --git a/Engine/Source/Editor/BehaviorTreeEditor/Private/SBehaviorTreeDiff.cpp b/Engine/Source/Editor/BehaviorTreeEditor/Private/SBehaviorTreeDiff.cpp
--- a/Engine/Source/Editor/BehaviorTreeEditor/Private/SBehaviorTreeDiff.cpp
+++ b/Engine/Source/Editor/BehaviorTreeEditor/Private/SBehaviorTreeDiff.cpp
@@ -466,12 +466,9 @@ FText SBehaviorTreeDiff::FBehaviorTreeDiffPanel::GetTitle() const
 	if (!RevisionInfo.Revision.IsEmpty())
 	{
 		// Don't use grouping on the revision or CL numbers to match how Perforce displays them
-		static const FNumberFormattingOptions RevisionFormatOptions = FNumberFormattingOptions()
-			.SetUseGrouping(false);
-
 		const FText DateText = FText::AsDate(RevisionInfo.Date, EDateTimeStyle::Short);
 		const FText RevisionText = FText::FromString(RevisionInfo.Revision);
-		const FText ChangelistText = FText::AsNumber(RevisionInfo.Changelist, &RevisionFormatOptions);
+		const FText ChangelistText = FText::AsNumber(RevisionInfo.Changelist, &FNumberFormattingOptions::DefaultNoGrouping());
 
 		if (bShowAssetName)
 		{
diff --git a/Engine/Source/Editor/Kismet/Private/DiffUtils.cpp b/Engine/Source/Editor/Kismet/Private/DiffUtils.cpp
--- a/Engine/Source/Editor/Kismet/Private/DiffUtils.cpp
+++ b/Engine/Source/Editor/Kismet/Private/DiffUtils.cpp
@@ -546,11 +546,9 @@ FText DiffViewUtils::GetPanelLabel(const UBlueprint* Blueprint, const FRevisionI
 		
 		if(ISourceControlModule::Get().GetProvider().UsesChangelists())
 		{
-			FNumberFormattingOptions NumberOptions;
-			NumberOptions.SetUseGrouping(false);
 			RevisionData = FText::Format(NSLOCTEXT("DiffViewUtils", "RevisionData", "Revision {0} - CL {1} - {2}")
 				, FText::FromString(Revision.Revision)
-				, FText::AsNumber(Revision.Changelist, &NumberOptions)
+				, FText::AsNumber(Revision.Changelist, &FNumberFormattingOptions::DefaultNoGrouping())
 				, FText::FromString(Revision.Date.ToString(TEXT("%m/%d/%Y"))));
 		}
 		else
diff --git a/Engine/Source/Editor/SourceControlWindows/Private/SSourceControlHistory.cpp b/Engine/Source/Editor/SourceControlWindows/Private/SSourceControlHistory.cpp
--- a/Engine/Source/Editor/SourceControlWindows/Private/SSourceControlHistory.cpp
+++ b/Engine/Source/Editor/SourceControlWindows/Private/SSourceControlHistory.cpp
@@ -917,9 +917,8 @@ private:
 	{
 		if(LastSelectedRevisionItem.IsValid())
 		{
-			static const FNumberFormattingOptions ChangelistFormatOptions = FNumberFormattingOptions()
-				.SetUseGrouping(false); // don't group the CL# as Perforce doesn't display it that way
-			return FText::AsNumber(LastSelectedRevisionItem.Pin()->ChangelistNumber, &ChangelistFormatOptions);
+			// don't group the CL# as Perforce doesn't display it that way
+			return FText::AsNumber(LastSelectedRevisionItem.Pin()->ChangelistNumber, &FNumberFormattingOptions::DefaultNoGrouping());
 		}
 		return FText::GetEmpty();
 	}
diff --git a/Engine/Source/Runtime/Core/Private/Internationalization/ICUCulture.cpp b/Engine/Source/Runtime/Core/Private/Internationalization/ICUCulture.cpp
--- a/Engine/Source/Runtime/Core/Private/Internationalization/ICUCulture.cpp
+++ b/Engine/Source/Runtime/Core/Private/Internationalization/ICUCulture.cpp
@@ -44,10 +44,10 @@ namespace
 		return Ptr.ToSharedRef();
 	}
 
-	TSharedRef<const icu::DecimalFormat> CreateDecimalFormat( const icu::Locale& ICULocale )
+	TSharedRef<const icu::DecimalFormat, ESPMode::ThreadSafe> CreateDecimalFormat( const icu::Locale& ICULocale )
 	{
 		UErrorCode ICUStatus = U_ZERO_ERROR;
-		TSharedPtr<const icu::DecimalFormat> Ptr = MakeShareable( static_cast<icu::DecimalFormat*>(icu::NumberFormat::createInstance( ICULocale, ICUStatus )) );
+		TSharedPtr<const icu::DecimalFormat, ESPMode::ThreadSafe> Ptr = MakeShareable( static_cast<icu::DecimalFormat*>(icu::NumberFormat::createInstance( ICULocale, ICUStatus )) );
 		checkf(Ptr.IsValid(), TEXT("Creating a decimal format object failed using locale %s. Perhaps this locale has no data."), StringCast<TCHAR>(ICULocale.getName()).Get());
 		return Ptr.ToSharedRef();
 	}
@@ -98,6 +98,7 @@ namespace
 
 FCulture::FICUCultureImplementation::FICUCultureImplementation(const FString& LocaleName)
 	: ICULocale( TCHAR_TO_ANSI( *LocaleName ) )
+	, ICUDecimalFormatLRUCache( 10 )
 {
 
 }
@@ -289,31 +290,70 @@ TSharedRef<const icu::Collator, ESPMode::ThreadSafe> FCulture::FICUCultureImplem
 	}
 }
 
-TSharedRef<const icu::DecimalFormat> FCulture::FICUCultureImplementation::GetDecimalFormatter(const FNumberFormattingOptions* const Options)
+TSharedRef<const icu::DecimalFormat, ESPMode::ThreadSafe> FCulture::FICUCultureImplementation::GetDecimalFormatter(const FNumberFormattingOptions* const Options)
 {
-	if (!ICUDecimalFormat.IsValid())
+	if (!ICUDecimalFormat_DefaultForCulture.IsValid())
 	{
-		ICUDecimalFormat = CreateDecimalFormat( ICULocale );
+		ICUDecimalFormat_DefaultForCulture = CreateDecimalFormat( ICULocale );
 	}
 
-	const bool bIsDefault = Options == NULL;
-	const TSharedRef<const icu::DecimalFormat> DefaultFormatter( ICUDecimalFormat.ToSharedRef() );
-	if(bIsDefault)
+	const bool bIsCultureDefault = Options == nullptr;
+	const TSharedRef<const icu::DecimalFormat, ESPMode::ThreadSafe> DefaultFormatter( ICUDecimalFormat_DefaultForCulture.ToSharedRef() );
+	if (bIsCultureDefault)
 	{
 		return DefaultFormatter;
 	}
-	else
+	else if (FNumberFormattingOptions::DefaultWithGrouping().IsIdentical(*Options))
 	{
-		const TSharedRef<icu::DecimalFormat> Formatter( static_cast<icu::DecimalFormat*>(DefaultFormatter->clone()) );
-		if(Options)
+		if (!ICUDecimalFormat_DefaultWithGrouping.IsValid())
 		{
+			const TSharedRef<icu::DecimalFormat, ESPMode::ThreadSafe> Formatter( static_cast<icu::DecimalFormat*>(DefaultFormatter->clone()) );
 			Formatter->setGroupingUsed(Options->UseGrouping);
 			Formatter->setRoundingMode(UEToICU(Options->RoundingMode));
 			Formatter->setMinimumIntegerDigits(Options->MinimumIntegralDigits);
 			Formatter->setMaximumIntegerDigits(Options->MaximumIntegralDigits);
 			Formatter->setMinimumFractionDigits(Options->MinimumFractionalDigits);
 			Formatter->setMaximumFractionDigits(Options->MaximumFractionalDigits);
+			ICUDecimalFormat_DefaultWithGrouping = Formatter;
 		}
+		return ICUDecimalFormat_DefaultWithGrouping.ToSharedRef();
+	}
+	else if (FNumberFormattingOptions::DefaultNoGrouping().IsIdentical(*Options))
+	{
+		if (!ICUDecimalFormat_DefaultNoGrouping.IsValid())
+		{
+			const TSharedRef<icu::DecimalFormat, ESPMode::ThreadSafe> Formatter( static_cast<icu::DecimalFormat*>(DefaultFormatter->clone()) );
+			Formatter->setGroupingUsed(Options->UseGrouping);
+			Formatter->setRoundingMode(UEToICU(Options->RoundingMode));
+			Formatter->setMinimumIntegerDigits(Options->MinimumIntegralDigits);
+			Formatter->setMaximumIntegerDigits(Options->MaximumIntegralDigits);
+			Formatter->setMinimumFractionDigits(Options->MinimumFractionalDigits);
+			Formatter->setMaximumFractionDigits(Options->MaximumFractionalDigits);
+			ICUDecimalFormat_DefaultNoGrouping = Formatter;
+		}
+		return ICUDecimalFormat_DefaultNoGrouping.ToSharedRef();
+	}
+	else
+	{
+		/*
+		FScopeLock ScopeLock(&ICUDecimalFormatLRUCacheCS);
+
+		const TSharedPtr<const icu::DecimalFormat, ESPMode::ThreadSafe> CachedFormatter = ICUDecimalFormatLRUCache.AccessItem(*Options);
+		if (CachedFormatter.IsValid())
+		{
+			return CachedFormatter.ToSharedRef();
+		}
+		*/
+		const TSharedRef<icu::DecimalFormat, ESPMode::ThreadSafe> Formatter( static_cast<icu::DecimalFormat*>(DefaultFormatter->clone()) );
+		Formatter->setGroupingUsed(Options->UseGrouping);
+		Formatter->setRoundingMode(UEToICU(Options->RoundingMode));
+		Formatter->setMinimumIntegerDigits(Options->MinimumIntegralDigits);
+		Formatter->setMaximumIntegerDigits(Options->MaximumIntegralDigits);
+		Formatter->setMinimumFractionDigits(Options->MinimumFractionalDigits);
+		Formatter->setMaximumFractionDigits(Options->MaximumFractionalDigits);
+
+		//ICUDecimalFormatLRUCache.Add(*Options, Formatter);
+
 		return Formatter;
 	}
 }
@@ -522,4 +562,4 @@ TSharedRef<const icu::DateFormat> FCulture::FICUCultureImplementation::GetDateTi
 	}
 }
 
-#endif
\ No newline at end of file
+#endif
diff --git a/Engine/Source/Runtime/Core/Private/Internationalization/ICUCulture.h b/Engine/Source/Runtime/Core/Private/Internationalization/ICUCulture.h
--- a/Engine/Source/Runtime/Core/Private/Internationalization/ICUCulture.h
+++ b/Engine/Source/Runtime/Core/Private/Internationalization/ICUCulture.h
@@ -109,6 +109,215 @@ enum class EBreakIteratorType
 	Title
 };
 
+/**
+ * Basic Least Recently Used (LRU) cache to potentially keep ICU formatters alive for multiple format calls
+ */
+template <typename TFormatterType, ESPMode Mode>
+class FLRUFormatterCache
+{
+public:
+	typedef TSharedPtr<const TFormatterType, Mode> ValueType;
+
+	FLRUFormatterCache( int32 InMaxNumElements )
+		: LookupSet()
+		, MostRecent(nullptr)
+		, LeastRecent(nullptr)
+		, MaxNumElements(InMaxNumElements)
+	{
+	}
+
+	~FLRUFormatterCache()
+	{
+		Empty();
+	}
+
+	/**
+	 * Accesses an item in the cache.  
+	 */
+	FORCEINLINE const ValueType AccessItem(const FNumberFormattingOptions& Key)
+	{
+		CacheEntry** Entry = LookupSet.Find( Key );
+		if( Entry )
+		{
+			MarkAsRecent( *Entry );
+			return ((*Entry)->Value);
+		}
+		
+		return nullptr;
+	}
+
+	void Add( const FNumberFormattingOptions& Key, const ValueType& Value )
+	{
+		CacheEntry** Entry = LookupSet.Find( Key );
+	
+		// Make a new link
+		if( !Entry )
+		{
+			if( LookupSet.Num() == MaxNumElements )
+			{
+				Eject();
+				checkf( LookupSet.Num() < MaxNumElements, TEXT("Could not eject item from the LRU: (%d of %d)"), LookupSet.Num(), MaxNumElements );
+			}
+
+			CacheEntry* NewEntry = new CacheEntry( Key, Value );
+
+			// Link before the most recent so that we become the most recent
+			NewEntry->Link(MostRecent);
+			MostRecent = NewEntry;
+
+			if( LeastRecent == nullptr )
+			{
+				LeastRecent = NewEntry;
+			}
+
+			LookupSet.Add( NewEntry );
+		}
+		else
+		{
+			// Trying to add an existing value 
+			CacheEntry* EntryPtr = *Entry;
+			
+			// Update the value
+			EntryPtr->Value = Value;
+			checkSlow( EntryPtr->Key == Key );
+			// Mark as the most recent
+			MarkAsRecent( EntryPtr );
+		}
+	}
+
+	void Empty()
+	{
+		for( auto& Entry : LookupSet )
+		{
+			// Note no need to unlink anything here. we are emptying the entire list
+			delete Entry;
+		}
+		LookupSet.Empty();
+
+		MostRecent = LeastRecent = nullptr;
+	}
+
+private:
+	struct CacheEntry
+	{
+		FNumberFormattingOptions Key;
+		ValueType Value;
+		CacheEntry* Next;
+		CacheEntry* Prev;
+
+		CacheEntry( const FNumberFormattingOptions& InKey, const ValueType& InValue )
+			: Key( InKey )
+			, Value( InValue )
+			, Next( nullptr )
+			, Prev( nullptr )
+		{
+		}
+
+		~CacheEntry()
+		{
+		}
+
+		FORCEINLINE void Link( CacheEntry* Before )
+		{
+			Next = Before;
+
+			if( Before )
+			{
+				Before->Prev = this;
+			}
+		}
+
+		FORCEINLINE void Unlink()
+		{
+			if( Prev )
+			{
+				Prev->Next = Next;
+			}
+
+			if( Next )
+			{
+				Next->Prev = Prev;
+			}
+
+			Prev = nullptr;
+			Next = nullptr;
+		}
+	};
+
+	struct FNumberFormattingOptionsKeyFuncs : BaseKeyFuncs<CacheEntry*, FNumberFormattingOptions>
+	{
+		FORCEINLINE static const FNumberFormattingOptions& GetSetKey( const CacheEntry* Entry )
+		{
+			return Entry->Key;
+		}
+
+		FORCEINLINE static bool Matches(const FNumberFormattingOptions& A,const FNumberFormattingOptions& B)
+		{
+			return A.IsIdentical( B );
+		}
+
+		FORCEINLINE static uint32 GetKeyHash(const FNumberFormattingOptions& Identifier)
+		{
+			return GetTypeHash( Identifier );
+		}
+	};
+
+	/**
+	 * Marks the link as the the most recent
+	 *
+	 * @param Link	The link to mark as most recent
+	 */
+	FORCEINLINE void MarkAsRecent( CacheEntry* Entry )
+	{
+		checkSlow( LeastRecent && MostRecent );
+
+		// If we are the least recent entry we are no longer the least recent
+		// The previous least recent item is now the most recent.  If it is nullptr then this entry is the only item in the list
+		if( Entry == LeastRecent && LeastRecent->Prev != nullptr )
+		{
+			LeastRecent = LeastRecent->Prev;
+		}
+
+		// No need to relink if we happen to already be accessing the most recent value
+		if( Entry != MostRecent )
+		{
+			// Unlink from its current spot
+			Entry->Unlink();
+
+			// Relink before the most recent so that we become the most recent
+			Entry->Link(MostRecent);
+			MostRecent = Entry;
+		}
+	}
+
+	/**
+	 * Removes the least recent item from the cache
+	 */
+	FORCEINLINE void Eject()
+	{
+		CacheEntry* EntryToRemove = LeastRecent;
+		// Eject the least recent, no more space
+		check( EntryToRemove );
+		LookupSet.Remove( EntryToRemove->Key );
+
+		LeastRecent = LeastRecent->Prev;
+
+		// Unlink the LRU
+		EntryToRemove->Unlink();
+
+		delete EntryToRemove;
+	}
+
+private:
+	TSet< CacheEntry*, FNumberFormattingOptionsKeyFuncs > LookupSet;
+	/** Most recent item in the cache */
+	CacheEntry* MostRecent;
+	/** Least recent item in the cache */
+	CacheEntry* LeastRecent;
+	/** The maximum number of elements in the cache */
+	int32 MaxNumElements;
+};
+
 class FCulture::FICUCultureImplementation
 {
 	friend FCulture;
@@ -149,7 +358,7 @@ class FCulture::FICUCultureImplementation
 
 	TSharedRef<const icu::BreakIterator> GetBreakIterator(const EBreakIteratorType Type);
 	TSharedRef<const icu::Collator, ESPMode::ThreadSafe> GetCollator(const ETextComparisonLevel::Type ComparisonLevel);
-	TSharedRef<const icu::DecimalFormat> GetDecimalFormatter(const FNumberFormattingOptions* const Options = NULL);
+	TSharedRef<const icu::DecimalFormat, ESPMode::ThreadSafe> GetDecimalFormatter(const FNumberFormattingOptions* const Options = NULL);
 	TSharedRef<const icu::DecimalFormat> GetCurrencyFormatter(const FString& CurrencyCode = FString(), const FNumberFormattingOptions* const Options = NULL);
 	TSharedRef<const icu::DecimalFormat> GetPercentFormatter(const FNumberFormattingOptions* const Options = NULL);
 	TSharedRef<const icu::DateFormat> GetDateFormatter(const EDateTimeStyle::Type DateStyle, const FString& TimeZone);
@@ -163,11 +372,17 @@ class FCulture::FICUCultureImplementation
 	TSharedPtr<const icu::BreakIterator> ICUSentenceBreakIterator;
 	TSharedPtr<const icu::BreakIterator> ICUTitleBreakIterator;
 	TSharedPtr<const icu::Collator, ESPMode::ThreadSafe> ICUCollator;
-	TSharedPtr<const icu::DecimalFormat> ICUDecimalFormat;
+
+	TSharedPtr<const icu::DecimalFormat, ESPMode::ThreadSafe> ICUDecimalFormat_DefaultForCulture;
+	TSharedPtr<const icu::DecimalFormat, ESPMode::ThreadSafe> ICUDecimalFormat_DefaultWithGrouping;
+	TSharedPtr<const icu::DecimalFormat, ESPMode::ThreadSafe> ICUDecimalFormat_DefaultNoGrouping;
+	FLRUFormatterCache<icu::DecimalFormat, ESPMode::ThreadSafe> ICUDecimalFormatLRUCache;
+	FCriticalSection ICUDecimalFormatLRUCacheCS;
+
 	TSharedPtr<const icu::DecimalFormat> ICUCurrencyFormat;
 	TSharedPtr<const icu::DecimalFormat> ICUPercentFormat;
 	TSharedPtr<const icu::DateFormat> ICUDateFormat;
 	TSharedPtr<const icu::DateFormat> ICUTimeFormat;
 	TSharedPtr<const icu::DateFormat> ICUDateTimeFormat;
 };
-#endif
\ No newline at end of file
+#endif
diff --git a/Engine/Source/Runtime/Core/Private/Internationalization/ICUText.h b/Engine/Source/Runtime/Core/Private/Internationalization/ICUText.h
--- a/Engine/Source/Runtime/Core/Private/Internationalization/ICUText.h
+++ b/Engine/Source/Runtime/Core/Private/Internationalization/ICUText.h
@@ -17,7 +17,7 @@ FText FText::AsNumberTemplate(T1 Val, const FNumberFormattingOptions* const Opti
 	checkf(I18N.IsInitialized() == true, TEXT("FInternationalization is not initialized. An FText formatting method was likely used in static object initialization - this is not supported."));
 	const FCultureRef Culture = TargetCulture.IsValid() ? TargetCulture.ToSharedRef() : I18N.GetCurrentCulture();
 	UErrorCode ICUStatus = U_ZERO_ERROR;
-	const TSharedRef<const icu::DecimalFormat> ICUDecimalFormat( Culture->Implementation->GetDecimalFormatter(Options) );
+	const TSharedRef<const icu::DecimalFormat, ESPMode::ThreadSafe> ICUDecimalFormat( Culture->Implementation->GetDecimalFormatter(Options) );
 	icu::Formattable FormattableVal(static_cast<T2>(Val));
 	icu::UnicodeString FormattedString;
 	ICUDecimalFormat->format(FormattableVal, FormattedString, ICUStatus);
@@ -68,4 +68,4 @@ FText FText::AsPercentTemplate(T1 Val, const FNumberFormattingOptions* const Opt
 	FText ReturnText = FText::CreateNumericalText(MoveTemp(NativeString));
 	ReturnText.History = MakeShareable(new FTextHistory_AsPercent(Val, Options, TargetCulture));
 	return ReturnText;
-}
\ No newline at end of file
+}
diff --git a/Engine/Source/Runtime/Core/Private/Internationalization/Text.cpp b/Engine/Source/Runtime/Core/Private/Internationalization/Text.cpp
--- a/Engine/Source/Runtime/Core/Private/Internationalization/Text.cpp
+++ b/Engine/Source/Runtime/Core/Private/Internationalization/Text.cpp
@@ -223,6 +223,40 @@ FArchive& operator<<(FArchive& Ar, FNumberFormattingOptions& Value)
 	return Ar;
 }
 
+uint32 GetTypeHash( const FNumberFormattingOptions& Key )
+{
+	uint32 Hash = 0;
+	Hash = HashCombine(Hash, GetTypeHash(Key.UseGrouping));
+	Hash = HashCombine(Hash, GetTypeHash(Key.RoundingMode));
+	Hash = HashCombine(Hash, GetTypeHash(Key.MinimumIntegralDigits));
+	Hash = HashCombine(Hash, GetTypeHash(Key.MaximumIntegralDigits));
+	Hash = HashCombine(Hash, GetTypeHash(Key.MinimumFractionalDigits));
+	Hash = HashCombine(Hash, GetTypeHash(Key.MaximumFractionalDigits));
+	return Hash;
+}
+
+bool FNumberFormattingOptions::IsIdentical( const FNumberFormattingOptions& Other ) const
+{
+	return UseGrouping == Other.UseGrouping
+		&& RoundingMode == Other.RoundingMode
+		&& MinimumIntegralDigits == Other.MinimumIntegralDigits
+		&& MaximumIntegralDigits == Other.MaximumIntegralDigits
+		&& MinimumFractionalDigits == Other.MinimumFractionalDigits
+		&& MaximumFractionalDigits == Other.MaximumFractionalDigits;
+}
+
+const FNumberFormattingOptions& FNumberFormattingOptions::DefaultWithGrouping()
+{
+	static const FNumberFormattingOptions Options = FNumberFormattingOptions().SetUseGrouping(true);
+	return Options;
+}
+
+const FNumberFormattingOptions& FNumberFormattingOptions::DefaultNoGrouping()
+{
+	static const FNumberFormattingOptions Options = FNumberFormattingOptions().SetUseGrouping(false);
+	return Options;
+}
+
 bool FText::bEnableErrorCheckingResults = ENABLE_TEXT_ERROR_CHECKING_RESULTS;
 bool FText::bSuppressWarnings = false;
 
diff --git a/Engine/Source/Runtime/Core/Public/Internationalization/Text.h b/Engine/Source/Runtime/Core/Public/Internationalization/Text.h
--- a/Engine/Source/Runtime/Core/Public/Internationalization/Text.h
+++ b/Engine/Source/Runtime/Core/Public/Internationalization/Text.h
@@ -116,6 +116,18 @@ struct CORE_API FNumberFormattingOptions
 	FNumberFormattingOptions& SetMaximumFractionalDigits( int32 InValue ){ MaximumFractionalDigits = InValue; return *this; }
 
 	friend FArchive& operator<<(FArchive& Ar, FNumberFormattingOptions& Value);
+
+	/** Get the hash code to use for the given formatting options */
+	friend uint32 GetTypeHash( const FNumberFormattingOptions& Key );
+
+	/** Check to see if our formatting options match the other formatting options */
+	bool IsIdentical( const FNumberFormattingOptions& Other ) const;
+
+	/** Get the default number formatting options with grouping enabled */
+	static const FNumberFormattingOptions& DefaultWithGrouping();
+
+	/** Get the default number formatting options with grouping disabled */
+	static const FNumberFormattingOptions& DefaultNoGrouping();
 };
 
 class FCulture;
diff --git a/Engine/Source/Runtime/Engine/Private/KismetTextLibrary.cpp b/Engine/Source/Runtime/Engine/Private/KismetTextLibrary.cpp
--- a/Engine/Source/Runtime/Engine/Private/KismetTextLibrary.cpp
+++ b/Engine/Source/Runtime/Engine/Private/KismetTextLibrary.cpp
@@ -115,27 +115,18 @@ FText UKismetTextLibrary::Conv_BoolToText(bool InBool)
 
 FText UKismetTextLibrary::Conv_ByteToText(uint8 Value)
 {
-	FNumberFormattingOptions NumberFormatOptions;
-	NumberFormatOptions.UseGrouping = true;
-	NumberFormatOptions.RoundingMode = ERoundingMode::ToZero;
-	NumberFormatOptions.MinimumIntegralDigits = 1;
-	NumberFormatOptions.MaximumIntegralDigits = 3;
-	NumberFormatOptions.MinimumFractionalDigits = 0;
-	NumberFormatOptions.MaximumFractionalDigits = 0;
-
-	return FText::AsNumber(Value, &NumberFormatOptions);
+	return FText::AsNumber(Value, &FNumberFormattingOptions::DefaultNoGrouping());
 }
 
 
 FText UKismetTextLibrary::Conv_IntToText(int32 Value, bool bUseGrouping/* = true*/, int32 MinimumIntegralDigits/* = 1*/, int32 MaximumIntegralDigits/* = 324*/)
 {
+	// Only update the values that need to be changed from the default FNumberFormattingOptions, 
+	// as this lets us use the default formatter if possible (which is a performance win!)
 	FNumberFormattingOptions NumberFormatOptions;
 	NumberFormatOptions.UseGrouping = bUseGrouping;
-	NumberFormatOptions.RoundingMode = ERoundingMode::FromZero;
 	NumberFormatOptions.MinimumIntegralDigits = MinimumIntegralDigits;
 	NumberFormatOptions.MaximumIntegralDigits = MaximumIntegralDigits;
-	NumberFormatOptions.MinimumFractionalDigits = 0;
-	NumberFormatOptions.MaximumFractionalDigits = 0;
 
 	return FText::AsNumber(Value, &NumberFormatOptions);
 }
