FText::AsNumber usability and performance improvements

Re-added some code which had been commented out while testing.

diff --git a/Engine/Source/Runtime/Core/Private/Internationalization/ICUCulture.cpp b/Engine/Source/Runtime/Core/Private/Internationalization/ICUCulture.cpp
--- a/Engine/Source/Runtime/Core/Private/Internationalization/ICUCulture.cpp
+++ b/Engine/Source/Runtime/Core/Private/Internationalization/ICUCulture.cpp
@@ -335,7 +335,6 @@ TSharedRef<const icu::DecimalFormat, ESPMode::ThreadSafe> FCulture::FICUCultureI
 	}
 	else
 	{
-		/*
 		FScopeLock ScopeLock(&ICUDecimalFormatLRUCacheCS);
 
 		const TSharedPtr<const icu::DecimalFormat, ESPMode::ThreadSafe> CachedFormatter = ICUDecimalFormatLRUCache.AccessItem(*Options);
@@ -343,7 +342,7 @@ TSharedRef<const icu::DecimalFormat, ESPMode::ThreadSafe> FCulture::FICUCultureI
 		{
 			return CachedFormatter.ToSharedRef();
 		}
-		*/
+
 		const TSharedRef<icu::DecimalFormat, ESPMode::ThreadSafe> Formatter( static_cast<icu::DecimalFormat*>(DefaultFormatter->clone()) );
 		Formatter->setGroupingUsed(Options->UseGrouping);
 		Formatter->setRoundingMode(UEToICU(Options->RoundingMode));
@@ -352,7 +351,7 @@ TSharedRef<const icu::DecimalFormat, ESPMode::ThreadSafe> FCulture::FICUCultureI
 		Formatter->setMinimumFractionDigits(Options->MinimumFractionalDigits);
 		Formatter->setMaximumFractionDigits(Options->MaximumFractionalDigits);
 
-		//ICUDecimalFormatLRUCache.Add(*Options, Formatter);
+		ICUDecimalFormatLRUCache.Add(*Options, Formatter);
 
 		return Formatter;
 	}
