Core - Stats for module manager

diff --git a/Engine/Source/Runtime/Core/Private/Modules/ModuleManager.cpp b/Engine/Source/Runtime/Core/Private/Modules/ModuleManager.cpp
--- a/Engine/Source/Runtime/Core/Private/Modules/ModuleManager.cpp
+++ b/Engine/Source/Runtime/Core/Private/Modules/ModuleManager.cpp
@@ -306,6 +306,14 @@ TSharedPtr<IModuleInterface> FModuleManager::LoadModuleWithFailureReason(const F
 {
 	DECLARE_SCOPE_CYCLE_COUNTER(TEXT("Module Load"), STAT_ModuleLoad, STATGROUP_LoadTime);
 
+#if	STATS
+	// This is fine here, we only load a handful of modules.
+	static FString Module = TEXT( "Module" );
+	const FString LongName = Module / InModuleName.GetPlainNameString();
+	const TStatId StatId = FDynamicStats::CreateStatId<FStatGroup_STATGROUP_UObjects>( LongName );
+	FScopeCycleCounter CycleCounter( StatId );
+#endif // STATS
+
 	TSharedPtr<IModuleInterface> LoadedModule;
 	OutFailureReason = EModuleLoadResult::Success;
 
