Fix curly apostrophe in comment; breaks doc tool.

diff --git a/Engine/Source/Runtime/Core/Public/Misc/ILifeCycle.h b/Engine/Source/Runtime/Core/Public/Misc/ILifeCycle.h
--- a/Engine/Source/Runtime/Core/Public/Misc/ILifeCycle.h
+++ b/Engine/Source/Runtime/Core/Public/Misc/ILifeCycle.h
@@ -4,7 +4,7 @@
 
 
 /**
- * Interface for controlling an object’s start/stop life cycle.
+ * Interface for controlling an object's start/stop life cycle.
  */
 class ILifeCycle
 {
