Comment out checkSlow() which doesn't compile in debug.

diff --git a/Engine/Source/Runtime/Core/Private/Internationalization/ICUCulture.h b/Engine/Source/Runtime/Core/Private/Internationalization/ICUCulture.h
--- a/Engine/Source/Runtime/Core/Private/Internationalization/ICUCulture.h
+++ b/Engine/Source/Runtime/Core/Private/Internationalization/ICUCulture.h
@@ -179,7 +179,7 @@ public:
 			
 			// Update the value
 			EntryPtr->Value = Value;
-			checkSlow( EntryPtr->Key == Key );
+//			checkSlow( EntryPtr->Key == Key );
 			// Mark as the most recent
 			MarkAsRecent( EntryPtr );
 		}
