Fix for incorrect character input in HTML5 UMG widgets. Also (hopefully) fixed keyboard mappings not working correctly.
Fixes UE-9555

diff --git a/Engine/Source/Runtime/Core/Private/HTML5/HTML5InputInterface.cpp b/Engine/Source/Runtime/Core/Private/HTML5/HTML5InputInterface.cpp
--- a/Engine/Source/Runtime/Core/Private/HTML5/HTML5InputInterface.cpp
+++ b/Engine/Source/Runtime/Core/Private/HTML5/HTML5InputInterface.cpp
@@ -87,122 +87,6 @@ void FHTML5InputInterface::SetMessageHandler( const TSharedRef< FGenericApplicat
 	MessageHandler = InMessageHandler;
 }
 
-static TCHAR ConvertChar( SDL_Keysym Keysym )
-{
-	if( Keysym.sym >= 128 )
-	{
-		return 0;
-	}
-
-	TCHAR Char = Keysym.sym;
-
-	if (Keysym.mod & (KMOD_LSHIFT | KMOD_RSHIFT))
-	{
-		// Convert to uppercase (FIXME: what about CAPS?)
-		if( Keysym.sym >= 97 && Keysym.sym <= 122)
-		{
-			return Keysym.sym - 32;
-		}
-		else if( Keysym.sym >= 91 && Keysym.sym <= 93)
-		{
-			return Keysym.sym + 32; // [ \ ] -> { | }
-		}
-		else
-		{
-			switch(Keysym.sym)
-			{
-			case '`': // ` -> ~
-				Char = TEXT('`');
-				break;
-
-			case '-': // - -> _
-				Char = TEXT('_');
-				break;
-
-			case '=': // - -> _
-				Char = TEXT('+');
-				break;
-
-			case ',':
-				Char = TEXT('<');
-				break;
-
-			case '.':
-				Char = TEXT('>');
-				break;
-
-			case ';':
-				Char = TEXT(':');
-				break;
-
-			case '\'':
-				Char = TEXT('\"');
-				break;
-
-			case '/':
-				Char = TEXT('?');
-				break;
-
-			case '0':
-				Char = TEXT(')');
-				break;
-
-			case '9':
-				Char = TEXT('(');
-				break;
-
-			case '8':
-				Char = TEXT('*');
-				break;
-
-			case '7':
-				Char = TEXT('&');
-				break;
-
-			case '6':
-				Char = TEXT('^');
-				break;
-
-			case '5':
-				Char = TEXT('%');
-				break;
-
-			case '4':
-				Char = TEXT('$');
-				break;
-
-			case '3':
-				Char = TEXT('#');
-				break;
-
-			case '2':
-				Char = TEXT('@');
-				break;
-
-			case '1':
-				Char = TEXT('!');
-				break;
-
-			default:
-				break;
-			}
-		}
-	}
-
-	return Char;
-}
-
-
-static bool GeneratesKeyCharMessage(const SDL_KeyboardEvent & KeyDownEvent)
-{
-	bool bCmdKeyPressed = (KeyDownEvent.keysym.mod & (KMOD_LCTRL | KMOD_RCTRL)) != 0;
-	const SDL_Keycode Sym = KeyDownEvent.keysym.sym;
-
-	// filter out command keys, non-ASCI and arrow keycodes that don't generate WM_CHAR under Windows (TODO: find a table?)
-	return !bCmdKeyPressed && Sym < 128 &&
-		(Sym != SDLK_DOWN && Sym != SDLK_LEFT && Sym != SDLK_RIGHT && Sym != SDLK_UP && Sym != SDLK_DELETE);
-}
-
 void FHTML5InputInterface::Tick(float DeltaTime, const SDL_Event& Event,TSharedRef < FGenericWindow>& ApplicationWindow )
 {
 
@@ -217,9 +101,12 @@ void FHTML5InputInterface::Tick(float DeltaTime, const SDL_Event& Event,TSharedR
 				// First KeyDown, then KeyChar. This is important, as in-game console ignores first character otherwise
 				MessageHandler->OnKeyDown(KeyCode, KeyEvent.keysym.sym, bIsRepeated);
 
-				if (GeneratesKeyCharMessage(KeyEvent))
+				// Backspace/Return input caught here. 
+				// Note that TextInput still seems to get characters messages too but slate seems to not process them.
+				if (KeyCode == SDL_SCANCODE_BACKSPACE || KeyCode == SDL_SCANCODE_RETURN)
 				{
-					const TCHAR Character = ConvertChar(KeyEvent.keysym);
+					const TCHAR Character = SDL_GetKeyFromScancode(Event.key.keysym.scancode);
+					UE_LOG(LogHTML5Input, Verbose, TEXT("TextInput: Text:%c bIsRepeated:%s"), Character, bIsRepeated ? TEXT("TRUE") : TEXT("FALSE"));
 					MessageHandler->OnKeyChar(Character, bIsRepeated);
 				}
 				UE_LOG(LogHTML5Input, Verbose, TEXT("KeyDown: Code:%d bIsRepeated:%s"), KeyCode, bIsRepeated ? TEXT("TRUE") : TEXT("FALSE"));
@@ -229,13 +116,21 @@ void FHTML5InputInterface::Tick(float DeltaTime, const SDL_Event& Event,TSharedR
 			{
 				SDL_KeyboardEvent keyEvent = Event.key;
 				const SDL_Keycode KeyCode = keyEvent.keysym.scancode;
-				const TCHAR Character = ConvertChar( keyEvent.keysym );
 				const bool IsRepeat = keyEvent.repeat != 0;
 
 				MessageHandler->OnKeyUp( KeyCode, keyEvent.keysym.sym, IsRepeat );
 				UE_LOG(LogHTML5Input, Verbose, TEXT("KeyUp Code:%d"), KeyCode);
 			}
 			break;
+		case SDL_TEXTINPUT:
+			{
+				const bool bIsRepeated = Event.key.repeat != 0;
+				const TCHAR Character = *ANSI_TO_TCHAR(Event.text.text);
+
+				MessageHandler->OnKeyChar(Character, bIsRepeated);
+				UE_LOG(LogHTML5Input, Verbose, TEXT("TextInput: Text:%c bIsRepeated:%s"), Character, bIsRepeated ? TEXT("TRUE") : TEXT("FALSE"));
+			}
+			break;
 		case SDL_MOUSEBUTTONDOWN:
 			{
 				EMouseButtons::Type MouseButton = Event.button.button == 1 ? EMouseButtons::Left : Event.button.button == 2 ? EMouseButtons::Middle : EMouseButtons::Right;
