Stats - Cleaned compression stats

diff --git a/Engine/Source/Runtime/Core/Private/Misc/Compression.cpp b/Engine/Source/Runtime/Core/Private/Misc/Compression.cpp
--- a/Engine/Source/Runtime/Core/Private/Misc/Compression.cpp
+++ b/Engine/Source/Runtime/Core/Private/Misc/Compression.cpp
@@ -7,6 +7,9 @@
 
 DEFINE_LOG_CATEGORY_STATIC(LogCompression, Log, All);
 
+DECLARE_STATS_GROUP( TEXT( "Compression" ), STATGROUP_Compression, STATCAT_Advanced );
+
+
 /**
  * Thread-safe abstract compression routine. Compresses memory from uncompressed buffer and writes it to compressed
  * buffer. Updates CompressedSize with size of compressed data.
@@ -17,11 +20,9 @@ DEFINE_LOG_CATEGORY_STATIC(LogCompression, Log, All);
  * @param	UncompressedSize			Size of uncompressed data in bytes
  * @return true if compression succeeds, false if it fails because CompressedBuffer was too small or other reasons
  */
-DECLARE_CYCLE_STAT(TEXT("Compress Memory ZLIB"),Stat_appCompressMemoryZLIB,STATGROUP_Engine);
-
 static bool appCompressMemoryZLIB( void* CompressedBuffer, int32& CompressedSize, const void* UncompressedBuffer, int32 UncompressedSize )
 {
-	SCOPE_CYCLE_COUNTER( Stat_appCompressMemoryZLIB );
+	DECLARE_SCOPE_CYCLE_COUNTER( TEXT( "Compress Memory ZLIB" ), STAT_appCompressMemoryZLIB, STATGROUP_Compression );
 
 	// Zlib wants to use unsigned long.
 	unsigned long ZCompressedSize	= CompressedSize;
@@ -35,7 +36,7 @@ static bool appCompressMemoryZLIB( void* CompressedBuffer, int32& CompressedSize
 
 static bool appCompressMemoryGZIP(void* CompressedBuffer, int32& CompressedSize, const void* UncompressedBuffer, int32 UncompressedSize)
 {
-	SCOPE_CYCLE_COUNTER(Stat_appCompressMemoryZLIB);
+	DECLARE_SCOPE_CYCLE_COUNTER( TEXT( "Compress Memory GZIP" ), STAT_appCompressMemoryGZIP, STATGROUP_Compression );
 
 	z_stream gzipstream;
 	gzipstream.zalloc = Z_NULL;
@@ -88,11 +89,9 @@ static bool appCompressMemoryGZIP(void* CompressedBuffer, int32& CompressedSize,
  * @param	CompressedSize				Size of CompressedBuffer data in bytes
  * @return true if compression succeeds, false if it fails because CompressedBuffer was too small or other reasons
  */
-DECLARE_CYCLE_STAT(TEXT("Uncompress Memory ZLIB"),Stat_appUncompressMemoryZLIB,STATGROUP_Engine);
-
 bool appUncompressMemoryZLIB( void* UncompressedBuffer, int32 UncompressedSize, const void* CompressedBuffer, int32 CompressedSize )
 {
-	SCOPE_CYCLE_COUNTER( Stat_appUncompressMemoryZLIB );
+	DECLARE_SCOPE_CYCLE_COUNTER( TEXT( "Uncompress Memory ZLIB" ), STAT_appUncompressMemoryZLIB, STATGROUP_Compression );
 
 	// Zlib wants to use unsigned long.
 	unsigned long ZCompressedSize	= CompressedSize;
@@ -224,7 +223,7 @@ bool FCompression::CompressMemory( ECompressionFlags Flags, void* CompressedBuff
  * @param	bIsSourcePadded				Whether the source memory is padded with a full cache line at the end
  * @return true if compression succeeds, false if it fails because CompressedBuffer was too small or other reasons
  */
-DECLARE_FLOAT_ACCUMULATOR_STAT(TEXT("Uncompressor total time"),STAT_UncompressorTime,STATGROUP_AsyncIO);
+DECLARE_FLOAT_ACCUMULATOR_STAT(TEXT("Uncompressor total time"),STAT_UncompressorTime,STATGROUP_Compression);
 
 bool FCompression::UncompressMemory( ECompressionFlags Flags, void* UncompressedBuffer, int32 UncompressedSize, const void* CompressedBuffer, int32 CompressedSize, bool bIsSourcePadded /*= false*/ )
 {
@@ -245,7 +244,13 @@ bool FCompression::UncompressMemory( ECompressionFlags Flags, void* Uncompressed
 			UE_LOG(LogCompression, Warning, TEXT("FCompression::UncompressMemory - This compression type not supported"));
 			bUncompressSucceeded = false;
 	}
-	STAT(if (FThreadStats::IsThreadingReady()) { INC_FLOAT_STAT_BY(STAT_UncompressorTime,(float)(FPlatformTime::Seconds()-UncompressorStartTime))} );
+
+#if	STATS
+	if (FThreadStats::IsThreadingReady())
+	{
+		INC_FLOAT_STAT_BY( STAT_UncompressorTime, (float)(FPlatformTime::Seconds() - UncompressorStartTime) )
+	}
+#endif // STATS
 	
 	return bUncompressSucceeded;
 }
