Stats - Cleaned AsyncIo stats

diff --git a/Engine/Source/Runtime/Core/Private/Misc/Core.cpp b/Engine/Source/Runtime/Core/Private/Misc/Core.cpp
--- a/Engine/Source/Runtime/Core/Private/Misc/Core.cpp
+++ b/Engine/Source/Runtime/Core/Private/Misc/Core.cpp
@@ -304,18 +304,6 @@ DEFINE_STAT(STAT_PumpMessages);
 DEFINE_STAT(STAT_CPUTimePct);
 DEFINE_STAT(STAT_CPUTimePctRelative);
 
-DEFINE_STAT(STAT_AsyncIO_FulfilledReadCount);
-DEFINE_STAT(STAT_AsyncIO_FulfilledReadSize);
-DEFINE_STAT(STAT_AsyncIO_CanceledReadCount);
-DEFINE_STAT(STAT_AsyncIO_CanceledReadSize);
-DEFINE_STAT(STAT_AsyncIO_OutstandingReadCount);
-DEFINE_STAT(STAT_AsyncIO_OutstandingReadSize);
-DEFINE_STAT(STAT_AsyncIO_UncompressorWaitTime);
-DEFINE_STAT(STAT_AsyncIO_AsyncLoadingBlockingTime);
-DEFINE_STAT(STAT_AsyncIO_AsyncPackagePrecacheWaitTime);
-DEFINE_STAT(STAT_AsyncIO_Bandwidth);
-DEFINE_STAT(STAT_AsyncIO_PlatformReadTime);
-
 DEFINE_LOG_CATEGORY(LogHAL);
 DEFINE_LOG_CATEGORY(LogMac);
 DEFINE_LOG_CATEGORY(LogLinux);
diff --git a/Engine/Source/Runtime/Core/Private/Serialization/AsyncIOSystemBase.cpp b/Engine/Source/Runtime/Core/Private/Serialization/AsyncIOSystemBase.cpp
--- a/Engine/Source/Runtime/Core/Private/Serialization/AsyncIOSystemBase.cpp
+++ b/Engine/Source/Runtime/Core/Private/Serialization/AsyncIOSystemBase.cpp
@@ -7,7 +7,21 @@
 #include "CorePrivatePCH.h"
 #include "Serialization/AsyncIOSystemBase.h"
 
-DECLARE_STATS_GROUP_VERBOSE(TEXT("AsyncIOSystem"),STATGROUP_AsyncIOSystem, STATCAT_Advanced);
+DECLARE_STATS_GROUP_VERBOSE(TEXT("AsyncIOSystem"),STATGROUP_AsyncIO_Verbose, STATCAT_Advanced);
+
+DECLARE_CYCLE_STAT( TEXT( "Platform read time" ), STAT_AsyncIO_PlatformReadTime, STATGROUP_AsyncIO );
+
+DECLARE_DWORD_ACCUMULATOR_STAT( TEXT( "Fulfilled read count" ), STAT_AsyncIO_FulfilledReadCount, STATGROUP_AsyncIO );
+DECLARE_DWORD_ACCUMULATOR_STAT( TEXT( "Canceled read count" ), STAT_AsyncIO_CanceledReadCount, STATGROUP_AsyncIO );
+
+DECLARE_MEMORY_STAT( TEXT( "Fulfilled read size" ), STAT_AsyncIO_FulfilledReadSize, STATGROUP_AsyncIO );
+DECLARE_MEMORY_STAT( TEXT( "Canceled read size" ), STAT_AsyncIO_CanceledReadSize, STATGROUP_AsyncIO );
+DECLARE_MEMORY_STAT( TEXT( "Outstanding read size" ), STAT_AsyncIO_OutstandingReadSize, STATGROUP_AsyncIO );
+
+DECLARE_DWORD_ACCUMULATOR_STAT( TEXT( "Outstanding read count" ), STAT_AsyncIO_OutstandingReadCount, STATGROUP_AsyncIO );
+DECLARE_FLOAT_ACCUMULATOR_STAT( TEXT( "Uncompressor wait time" ), STAT_AsyncIO_UncompressorWaitTime, STATGROUP_AsyncIO );
+
+DECLARE_FLOAT_COUNTER_STAT( TEXT( "Bandwidth (MByte/ sec)" ), STAT_AsyncIO_Bandwidth, STATGROUP_AsyncIO );
 
 /*-----------------------------------------------------------------------------
 	FAsyncIOSystemBase implementation.
@@ -108,14 +122,12 @@ void FAsyncIOSystemBase::LogIORequest(const FString& Message, const FAsyncIORequ
 	FPlatformMisc::LowLevelOutputDebugString(*OutputStr);
 }
 
-bool FAsyncIOSystemBase::InternalRead( IFileHandle* FileHandle, int64 Offset, int64 Size, void* Dest )
+void FAsyncIOSystemBase::InternalRead( IFileHandle* FileHandle, int64 Offset, int64 Size, void* Dest )
 {
-	DECLARE_SCOPE_CYCLE_COUNTER(TEXT("FAsyncIOSystemBase::InternalRead"), STAT_AsyncIOSystemBase_InternalRead, STATGROUP_AsyncIOSystem);
+	DECLARE_SCOPE_CYCLE_COUNTER(TEXT("FAsyncIOSystemBase::InternalRead"), STAT_AsyncIOSystemBase_InternalRead, STATGROUP_AsyncIO_Verbose);
 
 	FScopeLock ScopeLock( ExclusiveReadCriticalSection );
 
-	bool bRetVal = false;
-	
 	STAT(double ReadTime = 0);
 	{	
 		SCOPE_SECONDS_COUNTER(ReadTime);
@@ -127,8 +139,6 @@ bool FAsyncIOSystemBase::InternalRead( IFileHandle* FileHandle, int64 Offset, in
 	// though we only really care about throttling requested bandwidth as it's not very accurate
 	// to begin with.
 	STAT(ConstrainBandwidth(Size, ReadTime));
-
-	return bRetVal;
 }
 
 bool FAsyncIOSystemBase::PlatformReadDoNotCallDirectly( IFileHandle* FileHandle, int64 Offset, int64 Size, void* Dest )
@@ -188,7 +198,7 @@ int64 FAsyncIOSystemBase::PlatformMinimumReadSize()
 
 void FAsyncIOSystemBase::FulfillCompressedRead( const FAsyncIORequest& IORequest, IFileHandle* FileHandle )
 {
-	DECLARE_SCOPE_CYCLE_COUNTER(TEXT("FAsyncIOSystemBase::FulfillCompressedRead"), STAT_AsyncIOSystemBase_FulfillCompressedRead, STATGROUP_AsyncIOSystem);
+	DECLARE_SCOPE_CYCLE_COUNTER(TEXT("FAsyncIOSystemBase::FulfillCompressedRead"), STAT_AsyncIOSystemBase_FulfillCompressedRead, STATGROUP_AsyncIO_Verbose);
 
 	if (GbLogAsyncLoading == true)
 	{
@@ -720,7 +730,7 @@ void FAsyncIOSystemBase::TickSingleThreaded()
 
 void FAsyncIOSystemBase::BlockTillAllRequestsFinished()
 {
-	DECLARE_SCOPE_CYCLE_COUNTER( TEXT( "FAsyncIOSystemBase::BlockTillAllRequestsFinished" ), STAT_FAsyncIOSystemBase_BlockTillAllRequestsFinished, STATGROUP_AsyncIOSystem );
+	DECLARE_SCOPE_CYCLE_COUNTER(TEXT("FAsyncIOSystemBase::BlockTillAllRequestsFinished"), STAT_FAsyncIOSystemBase_BlockTillAllRequestsFinished, STATGROUP_AsyncIO_Verbose);
 
 	// Block till all requests are fulfilled.
 	while( true ) 
diff --git a/Engine/Source/Runtime/Core/Private/Serialization/AsyncIOSystemBase.h b/Engine/Source/Runtime/Core/Private/Serialization/AsyncIOSystemBase.h
--- a/Engine/Source/Runtime/Core/Private/Serialization/AsyncIOSystemBase.h
+++ b/Engine/Source/Runtime/Core/Private/Serialization/AsyncIOSystemBase.h
@@ -233,9 +233,8 @@ protected:
 	 * @param	Size		Size in bytes to read at current position from passed in file handle
 	 * @param	Dest		Pointer to data to read into
 	 *
-	 * @return	true if read was successful, false otherwise
 	 */	
-	bool InternalRead( IFileHandle* FileHandle, int64 Offset, int64 Size, void* Dest );
+	void InternalRead( IFileHandle* FileHandle, int64 Offset, int64 Size, void* Dest );
 
 
 	/** 
diff --git a/Engine/Source/Runtime/Core/Public/Misc/CoreStats.h b/Engine/Source/Runtime/Core/Public/Misc/CoreStats.h
--- a/Engine/Source/Runtime/Core/Public/Misc/CoreStats.h
+++ b/Engine/Source/Runtime/Core/Public/Misc/CoreStats.h
@@ -71,18 +71,6 @@ DECLARE_CYCLE_STAT_EXTERN(TEXT("Pump Messages"),STAT_PumpMessages,STATGROUP_Thre
 DECLARE_FLOAT_COUNTER_STAT_EXTERN(TEXT("Percentage CPU utilization"),STAT_CPUTimePct,STATGROUP_Threading, CORE_API);
 DECLARE_FLOAT_COUNTER_STAT_EXTERN(TEXT("Percentage CPU utilization (relative to one core)"),STAT_CPUTimePctRelative,STATGROUP_Threading, CORE_API);
 
-DECLARE_DWORD_ACCUMULATOR_STAT_EXTERN(TEXT("Fulfilled read count"),STAT_AsyncIO_FulfilledReadCount,STATGROUP_AsyncIO, CORE_API);
-DECLARE_MEMORY_STAT_EXTERN(TEXT("Fulfilled read size"),STAT_AsyncIO_FulfilledReadSize,STATGROUP_AsyncIO, CORE_API);
-DECLARE_DWORD_ACCUMULATOR_STAT_EXTERN(TEXT("Canceled read count"),STAT_AsyncIO_CanceledReadCount,STATGROUP_AsyncIO, CORE_API);
-DECLARE_MEMORY_STAT_EXTERN(TEXT("Canceled read size"),STAT_AsyncIO_CanceledReadSize,STATGROUP_AsyncIO, CORE_API);
-DECLARE_DWORD_ACCUMULATOR_STAT_EXTERN(TEXT("Outstanding read count"),STAT_AsyncIO_OutstandingReadCount,STATGROUP_AsyncIO, CORE_API);
-DECLARE_MEMORY_STAT_EXTERN(TEXT("Outstanding read size"),STAT_AsyncIO_OutstandingReadSize,STATGROUP_AsyncIO, CORE_API);
-DECLARE_FLOAT_ACCUMULATOR_STAT_EXTERN(TEXT("Platform read time"),STAT_AsyncIO_PlatformReadTime,STATGROUP_AsyncIO, CORE_API);
-DECLARE_FLOAT_ACCUMULATOR_STAT_EXTERN(TEXT("Uncompressor wait time"),STAT_AsyncIO_UncompressorWaitTime,STATGROUP_AsyncIO, CORE_API);
-DECLARE_FLOAT_ACCUMULATOR_STAT_EXTERN(TEXT("Async loading block time"),STAT_AsyncIO_AsyncLoadingBlockingTime,STATGROUP_AsyncIO, CORE_API);
-DECLARE_FLOAT_ACCUMULATOR_STAT_EXTERN(TEXT("Async package precache wait time"),STAT_AsyncIO_AsyncPackagePrecacheWaitTime,STATGROUP_AsyncIO, CORE_API);
-DECLARE_FLOAT_COUNTER_STAT_EXTERN(TEXT("Bandwidth (MByte/ sec)"),STAT_AsyncIO_Bandwidth,STATGROUP_AsyncIO, CORE_API);
-
 /**
  * Math stats
  */ 
diff --git a/Engine/Source/Runtime/CoreUObject/Private/Serialization/AsyncLoading.cpp b/Engine/Source/Runtime/CoreUObject/Private/Serialization/AsyncLoading.cpp
--- a/Engine/Source/Runtime/CoreUObject/Private/Serialization/AsyncLoading.cpp
+++ b/Engine/Source/Runtime/CoreUObject/Private/Serialization/AsyncLoading.cpp
@@ -21,6 +21,7 @@ DECLARE_STATS_GROUP_VERBOSE(TEXT("Async Load"), STATGROUP_AsyncLoad, STATCAT_Adv
 
 DECLARE_CYCLE_STAT(TEXT("Tick AsyncPackage"),STAT_FAsyncPackage_Tick,STATGROUP_AsyncLoad);
 DECLARE_FLOAT_ACCUMULATOR_STAT(TEXT("Tick AsyncPackage Time"), STAT_FAsyncPackage_TickTime, STATGROUP_AsyncLoad);
+
 DECLARE_CYCLE_STAT(TEXT("CreateLinker AsyncPackage"),STAT_FAsyncPackage_CreateLinker,STATGROUP_AsyncLoad);
 DECLARE_CYCLE_STAT(TEXT("FinishLinker AsyncPackage"),STAT_FAsyncPackage_FinishLinker,STATGROUP_AsyncLoad);
 DECLARE_CYCLE_STAT(TEXT("LoadImports AsyncPackage"),STAT_FAsyncPackage_LoadImports,STATGROUP_AsyncLoad);
@@ -35,13 +36,14 @@ DECLARE_CYCLE_STAT(TEXT("ProcessAsyncLoading AsyncLoadingThread"), STAT_FAsyncLo
 DECLARE_CYCLE_STAT(TEXT("Async Loading Time"),STAT_AsyncLoadingTime,STATGROUP_AsyncLoad);
 DECLARE_CYCLE_STAT(TEXT("Async Loading Time Detailed"), STAT_AsyncLoadingTimeDetailed, STATGROUP_AsyncLoad);
 
-
 DECLARE_STATS_GROUP(TEXT("Async Load Game Thread"), STATGROUP_AsyncLoadGameThread, STATCAT_Advanced);
 
-DECLARE_CYCLE_STAT(TEXT("PostLoadObjects GT"), STAT_FAsyncPackage_PostLoadObjectsGameThread, STATGROUP_AsyncLoad);
+DECLARE_CYCLE_STAT(TEXT("PostLoadObjects GT"), STAT_FAsyncPackage_PostLoadObjectsGameThread, STATGROUP_AsyncLoadGameThread);
 DECLARE_CYCLE_STAT(TEXT("TickAsyncLoading GT"), STAT_FAsyncPackage_TickAsyncLoadingGameThread, STATGROUP_AsyncLoadGameThread);
-DECLARE_CYCLE_STAT(TEXT("Flush Async Loading GT"), STAT_FAsyncPackage_FlushAsyncLoadingGameThread, STATGROUP_AsyncLoad);
+DECLARE_CYCLE_STAT(TEXT("Flush Async Loading GT"), STAT_FAsyncPackage_FlushAsyncLoadingGameThread, STATGROUP_AsyncLoadGameThread);
 
+DECLARE_FLOAT_ACCUMULATOR_STAT( TEXT( "Async loading block time" ), STAT_AsyncIO_AsyncLoadingBlockingTime, STATGROUP_AsyncIO );
+DECLARE_FLOAT_ACCUMULATOR_STAT( TEXT( "Async package precache wait time" ), STAT_AsyncIO_AsyncPackagePrecacheWaitTime, STATGROUP_AsyncIO );
 
 /** Keeps a reference to all objects created during async load until streaming has finished */
 class FAsyncObjectsReferencer : FGCObject
