Stats - Better location for stats log

diff --git a/Engine/Source/Runtime/Core/Private/Stats/Stats2.cpp b/Engine/Source/Runtime/Core/Private/Stats/Stats2.cpp
--- a/Engine/Source/Runtime/Core/Private/Stats/Stats2.cpp
+++ b/Engine/Source/Runtime/Core/Private/Stats/Stats2.cpp
@@ -1119,6 +1119,8 @@ void FThreadStats::StartThread()
 	FStartupMessages::Get().AddThreadMetadata( NAME_GameThread, FPlatformTLS::GetCurrentThreadId() );
 
 	CheckForCollectingStartupStats();
+
+	UE_LOG( LogStats, Log, TEXT( "Stats thread started" ) );
 }
 
 static FGraphEventRef LastFramesEvents[MAX_STAT_LAG];
diff --git a/Engine/Source/Runtime/Launch/Private/LaunchEngineLoop.cpp b/Engine/Source/Runtime/Launch/Private/LaunchEngineLoop.cpp
--- a/Engine/Source/Runtime/Launch/Private/LaunchEngineLoop.cpp
+++ b/Engine/Source/Runtime/Launch/Private/LaunchEngineLoop.cpp
@@ -1063,7 +1063,6 @@ int32 FEngineLoop::PreInit( const TCHAR* CmdLine )
 
 #if STATS
 	FThreadStats::StartThread();
-	UE_LOG( LogStats, Log, TEXT( "Stats thread started" ) );
 #endif
 
 	if (FPlatformProcess::SupportsMultithreading())
