Added fallback logic to FText::Format when using ordered arguments with a format string using named arguments

UE-14170 - REGRESSION: FText::Format is failing in many places

We have some existing code that relies on this suspect behavior, so we need to emulate it to avoid breaking things.

diff --git a/Engine/Source/Runtime/Core/Private/Internationalization/Text.cpp b/Engine/Source/Runtime/Core/Private/Internationalization/Text.cpp
--- a/Engine/Source/Runtime/Core/Private/Internationalization/Text.cpp
+++ b/Engine/Source/Runtime/Core/Private/Internationalization/Text.cpp
@@ -524,7 +524,7 @@ private:
 	};
 
 public:
-	DECLARE_DELEGATE_RetVal_OneParam( const FFormatArgumentValue*, FGetArgumentValue, const FString& );
+	DECLARE_DELEGATE_RetVal_TwoParams( const FFormatArgumentValue*, FGetArgumentValue, const FString&, int32 );
 
 	static int32 EstimateArgumentValueLength(const FFormatArgumentValue& ArgumentValue)
 	{
@@ -685,6 +685,7 @@ public:
 		EBlockState BlockState = EBlockState::None;
 
 		FString ArgumentName;
+		int32 ArgumentNumber = 0;
 
 		for(int32 i = 0; i < PatternString.Len(); ++i)
 		{
@@ -774,7 +775,7 @@ public:
 						case '}':
 							{
 							
-								const FFormatArgumentValue* const PossibleArgumentValue = GetArgumentValue.Execute(ArgumentName);
+								const FFormatArgumentValue* const PossibleArgumentValue = GetArgumentValue.Execute(ArgumentName, ArgumentNumber++);
 								if( PossibleArgumentValue )
 								{
 									const FFormatArgumentValue& ArgumentValue = *PossibleArgumentValue;
@@ -876,7 +877,7 @@ FText FText::FormatInternal(const FText& Pattern, const FFormatNamedArguments& A
 		EstimatedArgumentValuesLength += FTextFormatHelper::EstimateArgumentValueLength(Arg.Value);
 	}
 
-	auto GetArgumentValue = [&Arguments](const FString& ArgumentName) -> const FFormatArgumentValue*
+	auto GetArgumentValue = [&Arguments](const FString& ArgumentName, int32 ArgumentNumber) -> const FFormatArgumentValue*
 	{
 		return Arguments.Find(ArgumentName);
 	};
@@ -895,12 +896,16 @@ FText FText::FormatInternal(const FText& Pattern, const FFormatOrderedArguments&
 		EstimatedArgumentValuesLength += FTextFormatHelper::EstimateArgumentValueLength(Arg);
 	}
 
-	auto GetArgumentValue = [&Arguments](const FString& ArgumentName) -> const FFormatArgumentValue*
+	auto GetArgumentValue = [&Arguments](const FString& ArgumentName, int32 ArgumentNumber) -> const FFormatArgumentValue*
 	{
 		int32 ArgumentIndex = INDEX_NONE;
-		if( ArgumentName.IsNumeric() )
+		if (!LexicalConversion::TryParseString(ArgumentIndex, *ArgumentName))
 		{
-			ArgumentIndex = FCString::Atoi(*ArgumentName);
+			// We failed to parse the argument name into a number...
+			// We have existing code that is incorrectly using names in the format string when providing ordered arguments
+			// ICU used to fallback to treating the index of the argument within the string as if it were the index specified 
+			// by the argument name, so we need to emulate that behavior to avoid breaking some format operations
+			ArgumentIndex = ArgumentNumber;
 		}
 		return ArgumentIndex != INDEX_NONE && ArgumentIndex < Arguments.Num() ? &(Arguments[ArgumentIndex]) : nullptr;
 	};
@@ -921,7 +926,7 @@ FText FText::FormatInternal(const FText& Pattern, const TArray< struct FFormatAr
 		FormatNamedArguments.Add(Arg.ArgumentName.ToString(), FFormatArgumentValue(Arg.ArgumentValue));
 	}
 
-	auto GetArgumentValue = [&FormatNamedArguments](const FString& ArgumentName) -> const FFormatArgumentValue*
+	auto GetArgumentValue = [&FormatNamedArguments](const FString& ArgumentName, int32 ArgumentNumber) -> const FFormatArgumentValue*
 	{
 		return FormatNamedArguments.Find(ArgumentName);
 	};
