fixed UE-12134 Post Process Effect Screen Percentage Not Functioning

diff --git a/Engine/Source/Runtime/Core/Private/HAL/ConsoleManager.cpp b/Engine/Source/Runtime/Core/Private/HAL/ConsoleManager.cpp
--- a/Engine/Source/Runtime/Core/Private/HAL/ConsoleManager.cpp
+++ b/Engine/Source/Runtime/Core/Private/HAL/ConsoleManager.cpp
@@ -1730,10 +1730,11 @@ static TAutoConsoleVariable<int32> CVarDepthOfFieldQuality(
 
 static TAutoConsoleVariable<float> CVarScreenPercentage(
 	TEXT("r.ScreenPercentage"),
-	-1.0f,
-	TEXT("To render in lower resolution and upscale for better performance.\n")
+	100.0f,
+	TEXT("To render in lower resolution and upscale for better performance (combined up with the blenable post process setting).\n")
 	TEXT("70 is a good value for low aliasing and performance, can be verified with 'show TestImage'\n")
-	TEXT("in percent, >0 and <=100, <0 means the post process volume settings are used"),
+	TEXT("in percent, >0 and <=100, larger numbers are possible (supersampling) but the downsampling quality is improvable.")
+	TEXT("<0 is treated like 100."),
 	ECVF_Scalability | ECVF_Default);
 
 static TAutoConsoleVariable<int32> CVarMaterialQualityLevel(
diff --git a/Engine/Source/Runtime/Engine/Private/SceneView.cpp b/Engine/Source/Runtime/Engine/Private/SceneView.cpp
--- a/Engine/Source/Runtime/Engine/Private/SceneView.cpp
+++ b/Engine/Source/Runtime/Engine/Private/SceneView.cpp
@@ -1165,7 +1165,7 @@ void FSceneView::EndFinalPostprocessSettings(const FSceneViewInitOptions& ViewIn
 
 		if(Value >= 0.0)
 		{
-			FinalPostProcessSettings.ScreenPercentage = Value;
+			FinalPostProcessSettings.ScreenPercentage *= Value / 100.0f;
 		}
 	}
 
