Correctly report available desktop work area on Linux (UEPLAT-693).

diff --git a/Engine/Source/Runtime/Core/Private/Linux/LinuxApplication.cpp b/Engine/Source/Runtime/Core/Private/Linux/LinuxApplication.cpp
--- a/Engine/Source/Runtime/Core/Private/Linux/LinuxApplication.cpp
+++ b/Engine/Source/Runtime/Core/Private/Linux/LinuxApplication.cpp
@@ -1257,8 +1257,9 @@ void FDisplayMetrics::GetDisplayMetrics(FDisplayMetrics& OutDisplayMetrics)
 	OutDisplayMetrics.MonitorInfo.Empty();
 	
 	FMonitorInfo Primary;
-	SDL_Rect PrimaryBounds;
+	SDL_Rect PrimaryBounds, PrimaryUsableBounds;
 	SDL_GetDisplayBounds(0, &PrimaryBounds);
+	SDL_GetDisplayUsableBounds(0, &PrimaryUsableBounds);
 
 	Primary.Name = UTF8_TO_TCHAR(SDL_GetDisplayName(0));
 	Primary.ID = TEXT("display0");
@@ -1267,14 +1268,13 @@ void FDisplayMetrics::GetDisplayMetrics(FDisplayMetrics& OutDisplayMetrics)
 	Primary.bIsPrimary = true;
 	OutDisplayMetrics.MonitorInfo.Add(Primary);
 
-	// @TODO [RCL] 2014-09-30 - try to account for real work area and not just display size.
-	OutDisplayMetrics.PrimaryDisplayWorkAreaRect.Left = PrimaryBounds.x;
-	OutDisplayMetrics.PrimaryDisplayWorkAreaRect.Top = PrimaryBounds.y;
-	OutDisplayMetrics.PrimaryDisplayWorkAreaRect.Right = PrimaryBounds.x + PrimaryBounds.w;
-	OutDisplayMetrics.PrimaryDisplayWorkAreaRect.Bottom = PrimaryBounds.y + PrimaryBounds.h;
+	OutDisplayMetrics.PrimaryDisplayWorkAreaRect.Left = PrimaryUsableBounds.x;
+	OutDisplayMetrics.PrimaryDisplayWorkAreaRect.Top = PrimaryUsableBounds.y;
+	OutDisplayMetrics.PrimaryDisplayWorkAreaRect.Right = PrimaryUsableBounds.x + PrimaryUsableBounds.w;
+	OutDisplayMetrics.PrimaryDisplayWorkAreaRect.Bottom = PrimaryUsableBounds.y + PrimaryUsableBounds.h;
 
 	OutDisplayMetrics.PrimaryDisplayWidth = PrimaryBounds.w;
- 	OutDisplayMetrics.PrimaryDisplayHeight = PrimaryBounds.h;
+	OutDisplayMetrics.PrimaryDisplayHeight = PrimaryBounds.h;
 
 	// accumulate the total bound rect
 	OutDisplayMetrics.VirtualDisplayRect = OutDisplayMetrics.PrimaryDisplayWorkAreaRect;
