Stats - Cleaned Math stats

diff --git a/Engine/Source/Runtime/Core/Private/Math/UnrealMath.cpp b/Engine/Source/Runtime/Core/Private/Math/UnrealMath.cpp
--- a/Engine/Source/Runtime/Core/Private/Math/UnrealMath.cpp
+++ b/Engine/Source/Runtime/Core/Private/Math/UnrealMath.cpp
@@ -8,8 +8,12 @@
 
 DEFINE_LOG_CATEGORY(LogUnrealMath);
 
-DEFINE_STAT(STAT_MathConvertQuatToRotator);
-DEFINE_STAT(STAT_MathConvertRotatorToQuat);
+/**
+* Math stats
+*/
+
+DECLARE_CYCLE_STAT( TEXT( "Convert Rotator to Quat" ), STAT_MathConvertRotatorToQuat, STATGROUP_Math );
+DECLARE_CYCLE_STAT( TEXT( "Convert Quat to Rotator" ), STAT_MathConvertQuatToRotator, STATGROUP_Math );
 
 /*-----------------------------------------------------------------------------
 	Globals
diff --git a/Engine/Source/Runtime/Core/Public/Misc/CoreStats.h b/Engine/Source/Runtime/Core/Public/Misc/CoreStats.h
--- a/Engine/Source/Runtime/Core/Public/Misc/CoreStats.h
+++ b/Engine/Source/Runtime/Core/Public/Misc/CoreStats.h
@@ -71,11 +71,4 @@ DECLARE_CYCLE_STAT_EXTERN(TEXT("Pump Messages"),STAT_PumpMessages,STATGROUP_Thre
 DECLARE_FLOAT_COUNTER_STAT_EXTERN(TEXT("Percentage CPU utilization"),STAT_CPUTimePct,STATGROUP_Threading, CORE_API);
 DECLARE_FLOAT_COUNTER_STAT_EXTERN(TEXT("Percentage CPU utilization (relative to one core)"),STAT_CPUTimePctRelative,STATGROUP_Threading, CORE_API);
 
-/**
- * Math stats
- */ 
-
-DECLARE_CYCLE_STAT_EXTERN(TEXT("Convert Rotator to Quat"), STAT_MathConvertRotatorToQuat, STATGROUP_Math, CORE_API);
-DECLARE_CYCLE_STAT_EXTERN(TEXT("Convert Quat to Rotator"), STAT_MathConvertQuatToRotator, STATGROUP_Math, CORE_API);
-
 DECLARE_CYCLE_STAT_EXTERN(TEXT("CPU Stall - Sleep"), STAT_Sleep, STATGROUP_CPUStalls, CORE_API);
