Add warning log when process fails to read environment variable.

diff --git a/Engine/Source/Runtime/Core/Private/Windows/WindowsPlatformProcess.cpp b/Engine/Source/Runtime/Core/Private/Windows/WindowsPlatformProcess.cpp
--- a/Engine/Source/Runtime/Core/Private/Windows/WindowsPlatformProcess.cpp
+++ b/Engine/Source/Runtime/Core/Private/Windows/WindowsPlatformProcess.cpp
@@ -37,10 +37,17 @@ void FWindowsPlatformProcess::AddDllDirectory(const TCHAR* Directory)
 	FPaths::NormalizeDirectoryName(NormalizedDirectory);
 	FPaths::MakePlatformFilename(NormalizedDirectory);
 
-	// Get the current value of the PATH variable
+	// Get the size of the PATH variable
 	TArray<TCHAR> PathVariable;
-	PathVariable.AddUninitialized(GetEnvironmentVariable(TEXT("PATH"), NULL, 0));
-	verify(::GetEnvironmentVariable(TEXT("PATH"), PathVariable.GetData(), PathVariable.Num()) == PathVariable.Num() - 1);
+	PathVariable.AddUninitialized(::GetEnvironmentVariable(TEXT("PATH"), NULL, 0));
+
+	// Get the actual value of variable.
+	if (::GetEnvironmentVariable(TEXT("PATH"), PathVariable.GetData(), PathVariable.Num()) == 0)
+	{
+		// Log a warning if reading value fails, but continue anyway.
+		UE_LOG(LogWindows, Warning, TEXT("Failed to load PATH environment variable. It either doesn't exist, or is too long."));
+		PathVariable.Add(TEXT(';'));
+	}
 
 	// Set the new path variable with the input directory at the start. Skip over any existing instances of the input directory.
 	FString NewPathVariable = NormalizedDirectory;
