Add critical section guarding adding modules to module manager.

diff --git a/Engine/Source/Runtime/Core/Private/Modules/ModuleManager.cpp b/Engine/Source/Runtime/Core/Private/Modules/ModuleManager.cpp
--- a/Engine/Source/Runtime/Core/Private/Modules/ModuleManager.cpp
+++ b/Engine/Source/Runtime/Core/Private/Modules/ModuleManager.cpp
@@ -177,6 +177,19 @@ bool FindNewestModuleFile(TArray<FString>& FilesToSearch, const FDateTime& Newer
 	return bFound;
 }
 
+void FModuleManager::AddModuleToModulesList(const FName InModuleName, TSharedRef<FModuleManager::FModuleInfo>& InModuleInfo)
+{
+	{
+		FWriteScopeLock Lock(&ModulesCriticalSection);
+
+		// Update hash table
+		Modules.Add(InModuleName, InModuleInfo);
+	}
+
+	// List of known modules has changed.  Fire callbacks.
+	FModuleManager::Get().ModulesChangedEvent.Broadcast(InModuleName, EModuleChangeReason::PluginDirectoryChanged);
+}
+
 void FModuleManager::AddModule(const FName InModuleName)
 {
 	// Do we already know about this module?  If not, we'll create information for this module now.
@@ -198,11 +211,7 @@ void FModuleManager::AddModule(const FName InModuleName)
 
 		~FAtExit()
 		{
-			// Update hash table
-			FModuleManager::Get().Modules.Add(ModuleName, ModuleInfo);
-
-			// List of known modules has changed.  Fire callbacks.
-			FModuleManager::Get().ModulesChangedEvent.Broadcast(ModuleName, EModuleChangeReason::PluginDirectoryChanged);
+			FModuleManager::Get().AddModuleToModulesList(ModuleName, ModuleInfo);
 		}
 
 		FName ModuleName;
diff --git a/Engine/Source/Runtime/Core/Public/Modules/ModuleManager.h b/Engine/Source/Runtime/Core/Public/Modules/ModuleManager.h
--- a/Engine/Source/Runtime/Core/Public/Modules/ModuleManager.h
+++ b/Engine/Source/Runtime/Core/Public/Modules/ModuleManager.h
@@ -496,6 +496,8 @@ public:
 	 */
 	void MakeUniqueModuleFilename( const FName InModuleName, FString& UniqueSuffix, FString& UniqueModuleFileName );
 
+	void AddModuleToModulesList(const FName InModuleName, TSharedRef<FModuleInfo>& ModuleInfo);
+
 private:
 	/** Thread safe module finding routine. */
 	TSharedRef<FModuleInfo>* FindModule(FName InModuleName);
