Stats - Cleared MemoryStaticMesh stats

diff --git a/Engine/Source/Runtime/Core/Private/Misc/Core.cpp b/Engine/Source/Runtime/Core/Private/Misc/Core.cpp
--- a/Engine/Source/Runtime/Core/Private/Misc/Core.cpp
+++ b/Engine/Source/Runtime/Core/Private/Misc/Core.cpp
@@ -256,7 +256,6 @@ DEFINE_STAT(STAT_AnimationMemory);
 DEFINE_STAT(STAT_PrecomputedVisibilityMemory);
 DEFINE_STAT(STAT_PrecomputedShadowDepthMapMemory);
 DEFINE_STAT(STAT_PrecomputedLightVolumeMemory);
-DEFINE_STAT(STAT_StaticMeshTotalMemory);
 DEFINE_STAT(STAT_SkeletalMeshVertexMemory);
 DEFINE_STAT(STAT_SkeletalMeshIndexMemory);
 DEFINE_STAT(STAT_SkeletalMeshMotionBlurSkinningMemory);
@@ -269,14 +268,6 @@ DEFINE_STAT(STAT_PhysSceneWriteLock);
 DEFINE_STAT(STAT_ReflectionCaptureTextureMemory);
 DEFINE_STAT(STAT_ReflectionCaptureMemory);
 
-DEFINE_STAT(STAT_StaticMeshTotalMemory2);
-DEFINE_STAT(STAT_StaticMeshVertexMemory);
-DEFINE_STAT(STAT_ResourceVertexColorMemory);
-DEFINE_STAT(STAT_InstVertexColorMemory);
-DEFINE_STAT(STAT_StaticMeshIndexMemory);
-DEFINE_STAT(STAT_StaticMeshDistanceFieldMemory);
-
-
 /** Threading stats objects */
 
 DEFINE_STAT(STAT_RenderingIdleTime_WaitingForGPUQuery);
diff --git a/Engine/Source/Runtime/Core/Public/Misc/CoreStats.h b/Engine/Source/Runtime/Core/Public/Misc/CoreStats.h
--- a/Engine/Source/Runtime/Core/Public/Misc/CoreStats.h
+++ b/Engine/Source/Runtime/Core/Public/Misc/CoreStats.h
@@ -24,7 +24,6 @@ DECLARE_MEMORY_STAT_EXTERN(TEXT("Animation Memory"),STAT_AnimationMemory,STATGRO
 DECLARE_MEMORY_STAT_EXTERN(TEXT("Precomputed Visibility Memory"),STAT_PrecomputedVisibilityMemory,STATGROUP_Memory, CORE_API);
 DECLARE_MEMORY_STAT_EXTERN(TEXT("Precomputed Shadow Depth Map Memory"),STAT_PrecomputedShadowDepthMapMemory,STATGROUP_Memory, CORE_API);
 DECLARE_MEMORY_STAT_EXTERN(TEXT("Precomputed Light Volume Memory"),STAT_PrecomputedLightVolumeMemory,STATGROUP_Memory, CORE_API);
-DECLARE_MEMORY_STAT_EXTERN(TEXT("StaticMesh Total Memory"),STAT_StaticMeshTotalMemory,STATGROUP_Memory, CORE_API);
 DECLARE_MEMORY_STAT_EXTERN(TEXT("SkeletalMesh Vertex Memory"),STAT_SkeletalMeshVertexMemory,STATGROUP_Memory, CORE_API);
 DECLARE_MEMORY_STAT_EXTERN(TEXT("SkeletalMesh Index Memory"),STAT_SkeletalMeshIndexMemory,STATGROUP_Memory, CORE_API);
 DECLARE_MEMORY_STAT_EXTERN(TEXT("SkeletalMesh M.BlurSkinning Memory"),STAT_SkeletalMeshMotionBlurSkinningMemory,STATGROUP_Memory, CORE_API);
@@ -38,13 +37,6 @@ DECLARE_MEMORY_STAT_EXTERN(TEXT("ICU Data File Memory Used"),STAT_MemoryICUDataF
 DECLARE_CYCLE_STAT_EXTERN(TEXT("PhysX Scene ReadLock"), STAT_PhysSceneReadLock, STATGROUP_Physics, CORE_API);
 DECLARE_CYCLE_STAT_EXTERN(TEXT("PhysX Scene WriteLock"), STAT_PhysSceneWriteLock, STATGROUP_Physics, CORE_API);
 
-DECLARE_MEMORY_STAT_EXTERN(TEXT("StaticMesh Total Memory"),STAT_StaticMeshTotalMemory2,STATGROUP_MemoryStaticMesh, CORE_API);
-DECLARE_MEMORY_STAT_EXTERN(TEXT("StaticMesh Vertex Memory"),STAT_StaticMeshVertexMemory,STATGROUP_MemoryStaticMesh, CORE_API);
-DECLARE_MEMORY_STAT_EXTERN(TEXT("StaticMesh VxColor Resource Mem"),STAT_ResourceVertexColorMemory,STATGROUP_MemoryStaticMesh, CORE_API);
-DECLARE_MEMORY_STAT_EXTERN(TEXT("StaticMesh VxColor Inst Mem"),STAT_InstVertexColorMemory,STATGROUP_MemoryStaticMesh, CORE_API);
-DECLARE_MEMORY_STAT_EXTERN(TEXT("StaticMesh Index Memory"),STAT_StaticMeshIndexMemory,STATGROUP_MemoryStaticMesh, CORE_API);
-DECLARE_MEMORY_STAT_EXTERN(TEXT("StaticMesh Distance Field Memory"),STAT_StaticMeshDistanceFieldMemory,STATGROUP_MemoryStaticMesh, CORE_API);
-
 DECLARE_MEMORY_STAT_EXTERN(TEXT("Texture Memory Used"),STAT_TextureMemory,STATGROUP_Memory, CORE_API);
 DECLARE_MEMORY_STAT_EXTERN(TEXT("Reflection Capture Texture Memory"),STAT_ReflectionCaptureTextureMemory,STATGROUP_Memory, CORE_API);
 DECLARE_MEMORY_STAT_EXTERN(TEXT("Reflection Capture Memory"),STAT_ReflectionCaptureMemory,STATGROUP_Memory, CORE_API);
diff --git a/Engine/Source/Runtime/Engine/Private/Components/StaticMeshComponent.cpp b/Engine/Source/Runtime/Engine/Private/Components/StaticMeshComponent.cpp
--- a/Engine/Source/Runtime/Engine/Private/Components/StaticMeshComponent.cpp
+++ b/Engine/Source/Runtime/Engine/Private/Components/StaticMeshComponent.cpp
@@ -23,6 +23,8 @@
 
 #define LOCTEXT_NAMESPACE "StaticMeshComponent"
 
+DECLARE_MEMORY_STAT( TEXT( "StaticMesh VxColor Inst Mem" ), STAT_InstVertexColorMemory, STATGROUP_MemoryStaticMesh );
+
 class FStaticMeshComponentInstanceData : public FSceneComponentInstanceData
 {
 public:
diff --git a/Engine/Source/Runtime/Engine/Private/StaticMesh.cpp b/Engine/Source/Runtime/Engine/Private/StaticMesh.cpp
--- a/Engine/Source/Runtime/Engine/Private/StaticMesh.cpp
+++ b/Engine/Source/Runtime/Engine/Private/StaticMesh.cpp
@@ -27,6 +27,14 @@
 
 DEFINE_LOG_CATEGORY(LogStaticMesh);	
 
+DECLARE_MEMORY_STAT( TEXT( "StaticMesh Total Memory" ), STAT_StaticMeshTotalMemory2, STATGROUP_MemoryStaticMesh );
+DECLARE_MEMORY_STAT( TEXT( "StaticMesh Vertex Memory" ), STAT_StaticMeshVertexMemory, STATGROUP_MemoryStaticMesh );
+DECLARE_MEMORY_STAT( TEXT( "StaticMesh VxColor Resource Mem" ), STAT_ResourceVertexColorMemory, STATGROUP_MemoryStaticMesh );
+DECLARE_MEMORY_STAT( TEXT( "StaticMesh Index Memory" ), STAT_StaticMeshIndexMemory, STATGROUP_MemoryStaticMesh );
+DECLARE_MEMORY_STAT( TEXT( "StaticMesh Distance Field Memory" ), STAT_StaticMeshDistanceFieldMemory, STATGROUP_MemoryStaticMesh );
+
+DECLARE_MEMORY_STAT( TEXT( "StaticMesh Total Memory" ), STAT_StaticMeshTotalMemory, STATGROUP_Memory );
+
 /** Package name, that if set will cause only static meshes in that package to be rebuilt based on SM version. */
 ENGINE_API FName GStaticMeshPackageNameToRebuild = NAME_None;
 
