Prefer __debugbreak() instead of ::DebugBreak() on Windows platform
- After an assertion triggers, this prevents the code pointer being left inside the DebugBreak() Windows system library (which we usually won't have symbols for), and avoids us having to "step out" to get back to Unreal code.

diff --git a/Engine/Source/Runtime/Core/Public/Windows/WindowsPlatformMisc.h b/Engine/Source/Runtime/Core/Public/Windows/WindowsPlatformMisc.h
--- a/Engine/Source/Runtime/Core/Public/Windows/WindowsPlatformMisc.h
+++ b/Engine/Source/Runtime/Core/Public/Windows/WindowsPlatformMisc.h
@@ -46,7 +46,10 @@ struct CORE_API FWindowsPlatformMisc
 	{
 		if (IsDebuggerPresent())
 		{
-			::DebugBreak();
+			// Prefer __debugbreak() instead of ::DebugBreak() on Windows platform, so the code pointer isn't left
+			// inside the DebugBreak() Windows system library (which we usually won't have symbols for), and avoids
+			// us having to "step out" to get back to Unreal code.
+			__debugbreak();
 		}
 	}
 #endif
