Linux: suppress '-W-inconsistent-missing-override' warning for 4.8 for reals

diff --git a/Engine/Source/Runtime/Core/Public/Linux/LinuxPlatformCompilerPreSetup.h b/Engine/Source/Runtime/Core/Public/Linux/LinuxPlatformCompilerPreSetup.h
--- a/Engine/Source/Runtime/Core/Public/Linux/LinuxPlatformCompilerPreSetup.h
+++ b/Engine/Source/Runtime/Core/Public/Linux/LinuxPlatformCompilerPreSetup.h
@@ -61,6 +61,6 @@
 // Make certain warnings always be warnings, even despite -Werror.
 // Rationale: we don't want to suppress those as there are plans to address them (e.g. UE-12341), but breaking builds due to these warnings is very expensive
 // since they cannot be caught by all compilers that we support. They are deemed to be relatively safe to be ignored, at least until all SDKs/toolchains start supporting them.
-#pragma clang diagnostic warning "-Winconsistent-missing-override"
+//#pragma clang diagnostic warning "-Winconsistent-missing-override"
 #pragma clang diagnostic warning "-Wreorder"
 #pragma clang diagnostic warning "-Wparentheses-equality"
