Linux: fix shadowing issues.

diff --git a/Engine/Source/Runtime/Core/Private/Linux/LinuxPlatformFile.cpp b/Engine/Source/Runtime/Core/Private/Linux/LinuxPlatformFile.cpp
--- a/Engine/Source/Runtime/Core/Private/Linux/LinuxPlatformFile.cpp
+++ b/Engine/Source/Runtime/Core/Private/Linux/LinuxPlatformFile.cpp
@@ -168,7 +168,7 @@ public:
 			return FileSize;
 		}
 		else
-			#endif
+		#endif
 		{
 			struct stat FileInfo;
 			fstat(FileHandle, &FileInfo);
@@ -942,8 +942,6 @@ bool FLinuxPlatformFile::CreateDirectoriesFromPath(const TCHAR* Path)
 
 	for (int32 i=0; i<Len; ++i)
 	{
-		struct stat FileInfo;
-
 		SubPath[i] = DirPath[i];
 
 		if (SubPath[i] == '/')
@@ -951,7 +949,8 @@ bool FLinuxPlatformFile::CreateDirectoriesFromPath(const TCHAR* Path)
 			SubPath[i+1] = 0;
 
 			// directory exists?
-			if (stat(SubPath, &FileInfo) == -1)
+			struct stat SubPathFileInfo;
+			if (stat(SubPath, &SubPathFileInfo) == -1)
 			{
 				// nope. create it.
 				if (mkdir(SubPath, 0755) == -1)
diff --git a/Engine/Source/Runtime/Core/Private/Linux/LinuxPlatformProcess.cpp b/Engine/Source/Runtime/Core/Private/Linux/LinuxPlatformProcess.cpp
--- a/Engine/Source/Runtime/Core/Private/Linux/LinuxPlatformProcess.cpp
+++ b/Engine/Source/Runtime/Core/Private/Linux/LinuxPlatformProcess.cpp
@@ -626,11 +626,11 @@ FProcHandle FLinuxPlatformProcess::CreateProc(const TCHAR* URL, const TCHAR* Par
 		posix_spawn_file_actions_adddup2(&FileActions, PipeWriteHandle->GetHandle(), STDOUT_FILENO);
 	}
 
-	int ErrNo = posix_spawn(&ChildPid, TCHAR_TO_ANSI(*ProcessPath), &FileActions, nullptr, Argv, environ);
+	int PosixSpawnErrNo = posix_spawn(&ChildPid, TCHAR_TO_ANSI(*ProcessPath), &FileActions, nullptr, Argv, environ);
 	posix_spawn_file_actions_destroy(&FileActions);
-	if (ErrNo != 0)
+	if (PosixSpawnErrNo != 0)
 	{
-		UE_LOG(LogHAL, Fatal, TEXT("FLinuxPlatformProcess::CreateProc: posix_spawn() failed (%d, %s)"), ErrNo, ANSI_TO_TCHAR(strerror(ErrNo)));
+		UE_LOG(LogHAL, Fatal, TEXT("FLinuxPlatformProcess::CreateProc: posix_spawn() failed (%d, %s)"), PosixSpawnErrNo, ANSI_TO_TCHAR(strerror(PosixSpawnErrNo)));
 		return FProcHandle();	// produce knowingly invalid handle if for some reason Fatal log (above) returns
 	}
 
