Allow volume buttons to be handled by game or system on Android
Adds 2 BP nodes (SetVolumeButtonsHandledBySystem and GetVolumeButtonsHandledBySystem)
#ue4
#android

diff --git a/Engine/Source/Runtime/Core/Private/Android/AndroidMisc.cpp b/Engine/Source/Runtime/Core/Private/Android/AndroidMisc.cpp
--- a/Engine/Source/Runtime/Core/Private/Android/AndroidMisc.cpp
+++ b/Engine/Source/Runtime/Core/Private/Android/AndroidMisc.cpp
@@ -25,6 +25,9 @@ FString FAndroidMisc::OSLanguage; // language code the device is set to eg "deu"
 // Build/API level we are running.
 int32 FAndroidMisc::AndroidBuildVersion = 0;
 
+// Whether or not the system handles the volume buttons (event will still be generated either way)
+bool FAndroidMisc::VolumeButtonsHandledBySystem = true;
+
 GenericApplication* FAndroidMisc::CreateApplication()
 {
 	return FAndroidApplication::CreateAndroidApplication();
@@ -717,6 +720,17 @@ uint32 FAndroidMisc::GetKeyMap( uint16* KeyCodes, FString* KeyNames, uint32 MaxM
 #undef ADDKEYMAP
 }
 
+bool FAndroidMisc::GetVolumeButtonsHandledBySystem()
+{
+	return VolumeButtonsHandledBySystem;
+}
+
+void FAndroidMisc::SetVolumeButtonsHandledBySystem(bool enabled)
+{
+	VolumeButtonsHandledBySystem = enabled;
+}
+
+
 int32 FAndroidMisc::GetAndroidBuildVersion()
 {
 	if (AndroidBuildVersion > 0)
diff --git a/Engine/Source/Runtime/Core/Public/Android/AndroidMisc.h b/Engine/Source/Runtime/Core/Public/Android/AndroidMisc.h
--- a/Engine/Source/Runtime/Core/Public/Android/AndroidMisc.h
+++ b/Engine/Source/Runtime/Core/Public/Android/AndroidMisc.h
@@ -43,6 +43,8 @@ struct CORE_API FAndroidMisc : public FGenericPlatformMisc
 	static uint32 GetCharKeyMap(uint16* KeyCodes, FString* KeyNames, uint32 MaxMappings);
 	static uint32 GetKeyMap( uint16* KeyCodes, FString* KeyNames, uint32 MaxMappings );
 	static const TCHAR* GetDefaultDeviceProfileName() { return TEXT("Android"); }
+	static bool GetVolumeButtonsHandledBySystem();
+	static void SetVolumeButtonsHandledBySystem(bool enabled);
 
 	/** @return Memory representing a true type or open type font provided by the platform as a default font for unreal to consume; empty array if the default font failed to load. */
 	static TArray<uint8> GetSystemFontBytes();
@@ -59,7 +61,6 @@ struct CORE_API FAndroidMisc : public FGenericPlatformMisc
 	static bool SupportsFloatingPointRenderTargets();
 	static int GetAndroidBuildVersion();
 
-
 #if !UE_BUILD_SHIPPING
 	static bool IsDebuggerPresent();
 
@@ -112,6 +113,8 @@ struct CORE_API FAndroidMisc : public FGenericPlatformMisc
 
 	// Build version of Android, i.e. API level.
 	static int32 AndroidBuildVersion;
+
+	static bool VolumeButtonsHandledBySystem;
 };
 
 typedef FAndroidMisc FPlatformMisc;
diff --git a/Engine/Source/Runtime/Core/Public/GenericPlatform/GenericPlatformMisc.h b/Engine/Source/Runtime/Core/Public/GenericPlatform/GenericPlatformMisc.h
--- a/Engine/Source/Runtime/Core/Public/GenericPlatform/GenericPlatformMisc.h
+++ b/Engine/Source/Runtime/Core/Public/GenericPlatform/GenericPlatformMisc.h
@@ -755,6 +755,20 @@ public:
 		return PLATFORM_HAS_TOUCH_MAIN_SCREEN;
 	}
 
+	/*
+	 * Returns whether the volume buttons are handled by the system
+	 */
+	static bool GetVolumeButtonsHandledBySystem()
+	{
+		return true;
+	}
+
+	/*
+	 * Set whether the volume buttons are handled by the system
+	 */
+	static void SetVolumeButtonsHandledBySystem(bool enabled)
+	{}
+
 	/** @return Memory representing a true type or open type font provided by the platform as a default font for unreal to consume; empty array if the default font failed to load. */
 	static TArray<uint8> GetSystemFontBytes();
 	/**
diff --git a/Engine/Source/Runtime/Engine/Classes/Kismet/KismetSystemLibrary.h b/Engine/Source/Runtime/Engine/Classes/Kismet/KismetSystemLibrary.h
--- a/Engine/Source/Runtime/Engine/Classes/Kismet/KismetSystemLibrary.h
+++ b/Engine/Source/Runtime/Engine/Classes/Kismet/KismetSystemLibrary.h
@@ -1272,10 +1272,23 @@ class ENGINE_API UKismetSystemLibrary : public UBlueprintFunctionLibrary
 	/**
 	 * Allows or inhibits screensaver
 	 * @param	bAllowScreenSaver		If false, don't allow screensaver if possible, otherwise allow default behavior
-	*/
+	 */
 	UFUNCTION(BlueprintCallable, Category = "Utilities|Platform")
 	static void ControlScreensaver(bool bAllowScreenSaver);
 
+	/**
+	 * Allows or inhibits system default handling of volume up and volume down buttons (Android only)
+	 * @param	bEnabled				If true, allow Android to handle volume up and down events
+	 */
+	UFUNCTION(BlueprintCallable, Category = "Utilities|Platform")
+	static void SetVolumeButtonsHandledBySystem(bool bEnabled);
+
+	/**
+	 * Returns true if system default handling of volume up and volume down buttons enabled (Android only)
+	 */
+	UFUNCTION(BlueprintPure, Category = "Utilities|Platform")
+	static bool GetVolumeButtonsHandledBySystem();
+
 	/**
 	 * Sets the state of the transition message rendered by the viewport. (The blue text displayed when the game is paused and so forth.)
 	 *
diff --git a/Engine/Source/Runtime/Engine/Private/KismetSystemLibrary.cpp b/Engine/Source/Runtime/Engine/Private/KismetSystemLibrary.cpp
--- a/Engine/Source/Runtime/Engine/Private/KismetSystemLibrary.cpp
+++ b/Engine/Source/Runtime/Engine/Private/KismetSystemLibrary.cpp
@@ -2919,6 +2919,16 @@ void UKismetSystemLibrary::ControlScreensaver(bool bAllowScreenSaver)
 	FPlatformMisc::ControlScreensaver(bAllowScreenSaver ? FPlatformMisc::EScreenSaverAction::Enable : FPlatformMisc::EScreenSaverAction::Disable);
 }
 
+void UKismetSystemLibrary::SetVolumeButtonsHandledBySystem(bool bEnabled)
+{
+	FPlatformMisc::SetVolumeButtonsHandledBySystem(bEnabled);
+}
+
+bool UKismetSystemLibrary::GetVolumeButtonsHandledBySystem()
+{
+	return FPlatformMisc::GetVolumeButtonsHandledBySystem();
+}
+
 void UKismetSystemLibrary::SetSupressViewportTransitionMessage(UObject* WorldContextObject, bool bState)
 {
 	UWorld* World = GEngine->GetWorldFromContextObject(WorldContextObject);
diff --git a/Engine/Source/Runtime/Launch/Private/Android/LaunchAndroid.cpp b/Engine/Source/Runtime/Launch/Private/Android/LaunchAndroid.cpp
--- a/Engine/Source/Runtime/Launch/Private/Android/LaunchAndroid.cpp
+++ b/Engine/Source/Runtime/Launch/Private/Android/LaunchAndroid.cpp
@@ -626,6 +626,15 @@ static int32_t HandleInputCB(struct android_app* app, AInputEvent* event)
 			Message.KeyEventData.modifier = AKeyEvent_getMetaState(event);
 			Message.KeyEventData.isRepeat = AKeyEvent_getAction(event) == AKEY_EVENT_ACTION_MULTIPLE;
 			FAndroidInputInterface::DeferMessage(Message);
+
+			// allow event to be generated for volume up and down, but conditionally allow system to handle it, too
+			if (keyCode == AKEYCODE_VOLUME_UP || keyCode == AKEYCODE_VOLUME_DOWN)
+			{
+				if (FPlatformMisc::GetVolumeButtonsHandledBySystem())
+				{
+					return 0;
+				}
+			}
 		}
 
 		return 1;
