Fix for a shadowing issue in Linux splash.

diff --git a/Engine/Source/Runtime/Core/Private/Linux/LinuxPlatformSplash.cpp b/Engine/Source/Runtime/Core/Private/Linux/LinuxPlatformSplash.cpp
--- a/Engine/Source/Runtime/Core/Private/Linux/LinuxPlatformSplash.cpp
+++ b/Engine/Source/Runtime/Core/Private/Linux/LinuxPlatformSplash.cpp
@@ -906,7 +906,6 @@ void FLinuxPlatformSplash::Show( )
 
 		// Set version info
 		{
-			const FText GameName = FText::FromString( FApp::GetGameName() );
 			const FText Version = FText::FromString( GEngineVersion.ToString( FEngineBuildSettings::IsPerforceBuild() ? EVersionComponent::Branch : EVersionComponent::Patch ) );
 
 			FText VersionInfo;
