Fix orientation determination when _IPHONE_OS_VERSION_MIN_REQUIRED < __IPHONE_8_0 (shadowed variable)

diff --git a/Engine/Source/Runtime/Core/Private/IOS/IOSInputInterface.cpp b/Engine/Source/Runtime/Core/Private/IOS/IOSInputInterface.cpp
--- a/Engine/Source/Runtime/Core/Private/IOS/IOSInputInterface.cpp
+++ b/Engine/Source/Runtime/Core/Private/IOS/IOSInputInterface.cpp
@@ -43,7 +43,7 @@ void ModifyVectorByOrientation(FVector& Vec, bool bIsRotation)
 #endif
     {
 #if __IPHONE_OS_VERSION_MIN_REQUIRED < __IPHONE_8_0
-        UIInterfaceOrientation Orientation = [[IOSAppDelegate GetDelegate].IOSController interfaceOrientation];
+        Orientation = [[IOSAppDelegate GetDelegate].IOSController interfaceOrientation];
 #endif
     }
 
