Make it check slow.

diff --git a/Engine/Source/Runtime/Core/Private/Math/Sphere.cpp b/Engine/Source/Runtime/Core/Private/Math/Sphere.cpp
--- a/Engine/Source/Runtime/Core/Private/Math/Sphere.cpp
+++ b/Engine/Source/Runtime/Core/Private/Math/Sphere.cpp
@@ -102,8 +102,8 @@ FSphere& FSphere::operator+=(const FSphere &Other)
 		NewSphere.W = NewRadius;
 
 		// make sure both are inside afterwards
-		check (Other.IsInside(NewSphere, 1.f));
-		check (IsInside(NewSphere, 1.f));
+		checkSlow (Other.IsInside(NewSphere, 1.f));
+		checkSlow (IsInside(NewSphere, 1.f));
 
 		*this = NewSphere;
 	}
