Project agnostic ini files are now always loaded on build machines, but are set to be read-only.
This should fix problems caused by not loading the ini files at all.

diff --git a/Engine/Source/Runtime/Core/Private/Misc/ConfigCacheIni.cpp b/Engine/Source/Runtime/Core/Private/Misc/ConfigCacheIni.cpp
--- a/Engine/Source/Runtime/Core/Private/Misc/ConfigCacheIni.cpp
+++ b/Engine/Source/Runtime/Core/Private/Misc/ConfigCacheIni.cpp
@@ -2831,6 +2831,24 @@ static FString GetDestIniFilename(const TCHAR* BaseIniName, const TCHAR* Platfor
 	return IniFilename;
 }
 
+/** Load a project agnostic ini file. Will be read-only for build machines. */
+void LoadProjectAgnosticIni(FString& InOutFilename, const TCHAR* InName)
+{
+	static const FString EditorSettingsDir = FPaths::Combine(*FPaths::GameAgnosticSavedDir(), TEXT("Config")) + TEXT("/");
+
+	if (FConfigCacheIni::LoadGlobalIniFile(InOutFilename, InName, NULL, false, false, true, *EditorSettingsDir))
+	{
+		// Only save these configs if we are not the build machine as it would introduce an indeterminate state
+		if (GIsBuildMachine)
+		{
+			if (FConfigFile* Config = GConfig->FindConfigFile(*InOutFilename))
+			{
+				Config->NoSave = true;
+			}
+		}
+	}
+}
+
 void FConfigCacheIni::InitializeConfigSystem()
 {
 	// Perform any upgrade we need before we load any configuration files
@@ -2874,17 +2892,12 @@ void FConfigCacheIni::InitializeConfigSystem()
 	FConfigManifest::MigrateEditorUserSettings();
 	FConfigCacheIni::LoadGlobalIniFile(GEditorPerProjectIni, TEXT("EditorPerProjectUserSettings"));
 
+	// Project agnostic editor ini files
+	LoadProjectAgnosticIni(GEditorIni, TEXT("Editor"));
+	LoadProjectAgnosticIni(GEditorSettingsIni, TEXT("EditorSettings"));
+	LoadProjectAgnosticIni(GEditorLayoutIni, TEXT("EditorLayout"));
+	LoadProjectAgnosticIni(GEditorKeyBindingsIni, TEXT("EditorKeyBindings"));
 
-	// Project agnostic editor ini files. Only load these configs if we are not the build machine as it would introduce a non-deterministic state
-	if ( !GIsBuildMachine )
-	{
-		const FString EditorSettingsDir = FPaths::Combine(*FPaths::GameAgnosticSavedDir(), TEXT("Config")) + TEXT("/");
-
-		FConfigCacheIni::LoadGlobalIniFile(GEditorIni, TEXT("Editor"), NULL, false, false, true, *EditorSettingsDir);
-		FConfigCacheIni::LoadGlobalIniFile(GEditorSettingsIni, TEXT("EditorSettings"), NULL, false, false, true, *EditorSettingsDir);
-		FConfigCacheIni::LoadGlobalIniFile(GEditorLayoutIni, TEXT("EditorLayout"), NULL, false, false, true, *EditorSettingsDir);
-		FConfigCacheIni::LoadGlobalIniFile(GEditorKeyBindingsIni, TEXT("EditorKeyBindings"), NULL, false, false, true, *EditorSettingsDir);
-	}
 #endif
 #if PLATFORM_DESKTOP
 	// load some desktop only .ini files
diff --git a/Engine/Source/Runtime/CoreUObject/Private/UObject/Obj.cpp b/Engine/Source/Runtime/CoreUObject/Private/UObject/Obj.cpp
--- a/Engine/Source/Runtime/CoreUObject/Private/UObject/Obj.cpp
+++ b/Engine/Source/Runtime/CoreUObject/Private/UObject/Obj.cpp
@@ -1477,13 +1477,6 @@ void UObject::LoadConfig( UClass* ConfigClass/*=NULL*/, const TCHAR* InFilename/
 		? InFilename
 		: GetConfigFilename(this);
 
-	// If the file name is empty on build machines, this means the config system never initialized them
-	// as this would potentially create indeterminate on the builder. As such we don't load these configs.
-	if (Filename.IsEmpty() && GIsBuildMachine)
-	{
-		return;
-	}
-
 	const bool bPerObject = UsesPerObjectConfig(this);
 
 	FString ClassSection;
@@ -1682,13 +1675,6 @@ void UObject::SaveConfig( uint64 Flags, const TCHAR* InFilename, FConfigCacheIni
 		? InFilename
 		: GetConfigFilename(this);
 
-	// If the file name is empty on build machines, this means the config system never initialized them
-	// as this would potentially create indeterminate on the builder. As such we don't save these configs.
-	if (Filename.IsEmpty() && GIsBuildMachine)
-	{
-		return;
-	}
-
 	// Determine whether the file we are writing is a default file config.
 	const bool bIsADefaultIniWrite = !Filename.Contains(FPaths::GameSavedDir())
 		&& !Filename.Contains(FPaths::EngineSavedDir())
