PR #797 - Contributed by [REDACTED_PERSON]
UE4 - Added -AllowCommandletRendering to allow running a commandlet to initialize an RHI
- Added some fixes for null pointer dereferences

diff --git a/Engine/Source/Runtime/Core/Private/Misc/Core.cpp b/Engine/Source/Runtime/Core/Private/Misc/Core.cpp
--- a/Engine/Source/Runtime/Core/Private/Misc/Core.cpp
+++ b/Engine/Source/Runtime/Core/Private/Misc/Core.cpp
@@ -25,13 +25,13 @@ IMPLEMENT_MODULE( FCoreModule, Core );
 	Global variables.
 -----------------------------------------------------------------------------*/
 
-CORE_API FFeedbackContext*	GWarn						= NULL;						/* User interaction and non critical warnings */
-FConfigCacheIni*		GConfig							= NULL;						/* Configuration database cache */
-ITransaction*		GUndo							= NULL;						/* Transaction tracker, non-NULL when a transaction is in progress */
-FOutputDeviceConsole*	GLogConsole						= NULL;						/* Console log hook */
-CORE_API FMalloc*		GMalloc							= NULL;						/* Memory allocator */
+CORE_API FFeedbackContext*	GWarn						= nullptr;		/* User interaction and non critical warnings */
+FConfigCacheIni*				GConfig						= nullptr;		/* Configuration database cache */
+ITransaction*				GUndo						= nullptr;		/* Transaction tracker, non-NULL when a transaction is in progress */
+FOutputDeviceConsole*		GLogConsole					= nullptr;		/* Console log hook */
+CORE_API FMalloc*			GMalloc						= nullptr;		/* Memory allocator */
 
-class UPropertyWindowManager*	GPropertyWindowManager	= NULL;						/* Manages and tracks property editing windows */
+class UPropertyWindowManager*	GPropertyWindowManager	= nullptr;		/* Manages and tracks property editing windows */
 
 /** For building call stack text dump in guard/unguard mechanism. */
 TCHAR GErrorHist[16384]	= TEXT("");
@@ -91,13 +91,14 @@ bool GIsReconstructingBlueprintInstances = false;
 bool GForceDisableBlueprintCompileOnLoad = false;
 
 #if WITH_ENGINE
-bool					PRIVATE_GIsRunningCommandlet	= false;					/* Whether this executable is running a commandlet (custom command-line processing code) */
-#endif
+bool					PRIVATE_GIsRunningCommandlet			= false;				/* Whether this executable is running a commandlet (custom command-line processing code) */
+bool					PRIVATE_GAllowCommandletRendering	= false;				/** If true, initialise RHI and set up scene for rendering even when running a commandlet. */
+#endif	// WITH_ENGINE
 
 #if WITH_EDITORONLY_DATA
 bool					GIsEditor						= false;					/* Whether engine was launched for editing */
 bool					GIsImportingT3D					= false;					/* Whether editor is importing T3D */
-bool					GIsUCCMakeStandaloneHeaderGenerator = false;					/* Are we rebuilding script via the standalone header generator? */
+bool					GIsUCCMakeStandaloneHeaderGenerator = false;				/* Are we rebuilding script via the standalone header generator? */
 bool					GIsTransacting					= false;					/* true if there is an undo/redo operation in progress. */
 bool					GIntraFrameDebuggingGameThread	= false;					/* Indicates that the game thread is currently paused deep in a call stack; do not process any game thread tasks */
 bool					GFirstFrameIntraFrameDebugging	= false;					/* Indicates that we're currently processing the first frame of intra-frame debugging */
diff --git a/Engine/Source/Runtime/Core/Public/CoreGlobals.h b/Engine/Source/Runtime/Core/Public/CoreGlobals.h
--- a/Engine/Source/Runtime/Core/Public/CoreGlobals.h
+++ b/Engine/Source/Runtime/Core/Public/CoreGlobals.h
@@ -58,6 +58,9 @@ extern CORE_API void(*GFlushStreamingFunc)(void);
 
 #if WITH_ENGINE
 extern CORE_API bool PRIVATE_GIsRunningCommandlet;
+
+/** If true, initialize RHI and set up scene for rendering even when running a commandlet. */
+extern CORE_API bool PRIVATE_GAllowCommandletRendering;
 #endif
 
 #if WITH_EDITORONLY_DATA
@@ -100,6 +103,18 @@ FORCEINLINE bool IsRunningCommandlet()
 #endif
 }
 
+/**
+ * Check to see if we should initialise RHI and set up scene for rendering even when running a commandlet.
+ */
+FORCEINLINE bool IsAllowCommandletRendering()
+{
+#if WITH_ENGINE
+	return PRIVATE_GAllowCommandletRendering;
+#else
+	return false;
+#endif
+}
+
 extern CORE_API bool GEdSelectionLock;
 extern CORE_API bool GIsClient;
 extern CORE_API bool GIsServer;
diff --git a/Engine/Source/Runtime/Core/Public/Misc/App.h b/Engine/Source/Runtime/Core/Public/Misc/App.h
--- a/Engine/Source/Runtime/Core/Public/Misc/App.h
+++ b/Engine/Source/Runtime/Core/Public/Misc/App.h
@@ -21,7 +21,7 @@ public:
 	 */
 	FORCEINLINE static bool CanEverRender()
 	{
-		return !IsRunningCommandlet() && !IsRunningDedicatedServer();
+		return (!IsRunningCommandlet() || IsAllowCommandletRendering()) && !IsRunningDedicatedServer();
 	}
 
 	/**
diff --git a/Engine/Source/Runtime/Engine/Private/ModelRender.cpp b/Engine/Source/Runtime/Engine/Private/ModelRender.cpp
--- a/Engine/Source/Runtime/Engine/Private/ModelRender.cpp
+++ b/Engine/Source/Runtime/Engine/Private/ModelRender.cpp
@@ -153,7 +153,10 @@ public:
 		Component(InComponent),
 		CollisionResponse(InComponent->GetCollisionResponseToChannels())
 #if WITH_EDITOR
-		, CollisionMaterialInstance(GEngine->ShadedLevelColorationUnlitMaterial->GetRenderProxy(false, false),FColor(157,149,223,255)	)
+		, CollisionMaterialInstance(GEngine->ShadedLevelColorationUnlitMaterial
+			? GEngine->ShadedLevelColorationUnlitMaterial->GetRenderProxy(false, false)
+			: nullptr,
+			FColor(157,149,223,255))
 #endif
 	{
 		OverrideOwnerName(NAME_BSP);
diff --git a/Engine/Source/Runtime/Launch/Private/LaunchEngineLoop.cpp b/Engine/Source/Runtime/Launch/Private/LaunchEngineLoop.cpp
--- a/Engine/Source/Runtime/Launch/Private/LaunchEngineLoop.cpp
+++ b/Engine/Source/Runtime/Launch/Private/LaunchEngineLoop.cpp
@@ -1177,6 +1177,9 @@ int32 FEngineLoop::PreInit( const TCHAR* CmdLine )
 #endif	//WITH_EDITOR
 		PRIVATE_GIsRunningCommandlet = true;
 
+		// Allow commandlet rendering based on command line switch (too early to let the commandlet itself override this).
+		PRIVATE_GAllowCommandletRendering = FParse::Param(FCommandLine::Get(), TEXT("AllowCommandletRendering"));
+
 		// We need to disregard the empty token as we try finding Token + "Commandlet" which would result in finding the
 		// UCommandlet class if Token is empty.
 		bool bDefinitelyCommandlet = (bTokenDoesNotHaveDash && Token.EndsWith(TEXT("Commandlet")));
diff --git a/Engine/Source/Runtime/RHI/Private/DynamicRHI.cpp b/Engine/Source/Runtime/RHI/Private/DynamicRHI.cpp
--- a/Engine/Source/Runtime/RHI/Private/DynamicRHI.cpp
+++ b/Engine/Source/Runtime/RHI/Private/DynamicRHI.cpp
@@ -40,7 +40,7 @@ void RHIInit(bool bHasEditorToken)
 	{
 		GRHICommandList.LatchBypass(); // read commandline for bypass flag
 
-		if(USE_NULL_RHI || FParse::Param(FCommandLine::Get(),TEXT("nullrhi")) || IsRunningCommandlet() || IsRunningDedicatedServer())
+		if (USE_NULL_RHI || FParse::Param(FCommandLine::Get(),TEXT("nullrhi")) || !FApp::CanEverRender())
 		{
 			InitNullRHI();
 		}
diff --git a/Engine/Source/Runtime/Renderer/Private/RendererScene.cpp b/Engine/Source/Runtime/Renderer/Private/RendererScene.cpp
--- a/Engine/Source/Runtime/Renderer/Private/RendererScene.cpp
+++ b/Engine/Source/Runtime/Renderer/Private/RendererScene.cpp
@@ -2453,7 +2453,7 @@ FSceneInterface* FRendererModule::AllocateScene(UWorld* World, bool bInRequiresH
 	check(IsInGameThread());
 
 	// Create a full fledged scene if we have something to render.
-	if( GIsClient && !IsRunningCommandlet() && !GUsingNullRHI )
+	if (GIsClient && FApp::CanEverRender() && !GUsingNullRHI)
 	{
 		FScene* NewScene = new FScene(World, bInRequiresHitProxies, GIsEditor && !World->IsGameWorld(), bCreateFXSystem, InFeatureLevel);
 		AllocatedScenes.Add(NewScene);
