UE-14132: GitHub 1077 : [ IFileManager / Generic ] Added new function, FindFilesExt, provides simplified and Easy-To-Use version of FindFiles!

GitHub PR #1077

diff --git a/Engine/Source/Runtime/Core/Private/HAL/FileManagerGeneric.cpp b/Engine/Source/Runtime/Core/Private/HAL/FileManagerGeneric.cpp
--- a/Engine/Source/Runtime/Core/Private/HAL/FileManagerGeneric.cpp
+++ b/Engine/Source/Runtime/Core/Private/HAL/FileManagerGeneric.cpp
@@ -390,6 +390,43 @@ void FFileManagerGeneric::FindFiles( TArray<FString>& Result, const TCHAR* InFil
 	GetLowLevel().IterateDirectory( *FPaths::GetPath(Filename), FileMatch );
 }
 
+void FFileManagerGeneric::FindFiles(TArray<FString>& FoundFiles, const TCHAR* Directory, const TCHAR* FileExtension)
+{
+	if (!Directory)
+	{
+		return;
+	}
+
+	FString RootDir(Directory);
+	FString ExtStr = (FileExtension != nullptr) ? FString(FileExtension) : "";
+
+	// No Directory?
+	if (RootDir.Len() < 1)
+	{
+		return;
+	}
+
+	FPaths::NormalizeDirectoryName(RootDir);
+
+	// Don't modify the ExtStr if the user supplied the form "*.EXT" or "*" or "*.*" or "Name.*"
+	if (!ExtStr.Contains(TEXT("*")))
+	{
+		if (ExtStr == "")
+		{
+			ExtStr = "*.*";
+		}
+		else
+		{
+			//Complete the supplied extension with * or *. to yield "*.EXT"
+			ExtStr = (ExtStr.Left(1) == ".") ? "*" + ExtStr : "*." + ExtStr;
+		}
+	}
+
+	// Create the full filter, which is "Directory/*.EXT".
+	FString FinalPath = RootDir + "/" + ExtStr;
+	FindFiles(FoundFiles, *FinalPath, true, false);
+}
+
 bool FFileManagerGeneric::IterateDirectory(const TCHAR* Directory, IPlatformFile::FDirectoryVisitor& Visitor)
 {
 	return GetLowLevel().IterateDirectory( Directory, Visitor );
diff --git a/Engine/Source/Runtime/Core/Public/HAL/FileManager.h b/Engine/Source/Runtime/Core/Public/HAL/FileManager.h
--- a/Engine/Source/Runtime/Core/Public/HAL/FileManager.h
+++ b/Engine/Source/Runtime/Core/Public/HAL/FileManager.h
@@ -105,6 +105,18 @@ public:
 	/** Finds file or directories. */
 	virtual void FindFiles( TArray<FString>& FileNames, const TCHAR* Filename, bool Files, bool Directories)=0;
 
+	/**
+	 * Finds all the files within the given directory, with optional file extension filter.
+	 *
+	 * @param Directory, the absolute path to the directory to search. Ex: "C:\UE4\Pictures"
+	 *
+	 * @param FileExtension, If FileExtension is NULL, or an empty string "" then all files are found.
+	 * 			Otherwise FileExtension can be of the form .EXT or just EXT and only files with that extension will be returned.
+	 *
+	 * @return FoundFiles, All the files that matched the optional FileExtension filter, or all files if none was specified.
+	 */
+	virtual void FindFiles(TArray<FString>& FoundFiles, const TCHAR* Directory, const TCHAR* FileExtension = nullptr) = 0;
+
 	/** Finds file or directories recursively. */
 	virtual void FindFilesRecursive( TArray<FString>& FileNames, const TCHAR* StartDirectory, const TCHAR* Filename, bool Files, bool Directories, bool bClearFileNames=true) = 0; // utility
 
diff --git a/Engine/Source/Runtime/Core/Public/HAL/FileManagerGeneric.h b/Engine/Source/Runtime/Core/Public/HAL/FileManagerGeneric.h
--- a/Engine/Source/Runtime/Core/Public/HAL/FileManagerGeneric.h
+++ b/Engine/Source/Runtime/Core/Public/HAL/FileManagerGeneric.h
@@ -73,6 +73,19 @@ public:
 	virtual bool	MakeDirectory( const TCHAR* Path, bool Tree=0 ) override;
 	virtual bool	DeleteDirectory( const TCHAR* Path, bool RequireExists=0, bool Tree=0 ) override;
 
+
+	/**
+	 * Finds all the files within the given directory, with optional file extension filter.
+	 *
+	 * @param Directory, the absolute path to the directory to search. Ex: "C:\UE4\Pictures"
+	 *
+	 * @param FileExtension, If FileExtension is NULL, or an empty string "" then all files are found.
+	 * 			Otherwise FileExtension can be of the form .EXT or just EXT and only files with that extension will be returned.
+	 *
+	 * @return FoundFiles, All the files that matched the optional FileExtension filter, or all files if none was specified.
+	 */
+	virtual void FindFiles(TArray<FString>& FoundFiles, const TCHAR* Directory, const TCHAR* FileExtension = nullptr) override;
+
 	/** 
 	 * Call the Visit function of the visitor once for each file or directory in a single directory. This function does not explore subdirectories.
 	 * @param Directory		The directory to iterate the contents of.
