Make FPlatformProcess::ConditionalSleep use TFunctionRef instead of TFunction.

diff --git a/Engine/Source/Runtime/Core/Private/GenericPlatform/GenericPlatformProcess.cpp b/Engine/Source/Runtime/Core/Private/GenericPlatform/GenericPlatformProcess.cpp
--- a/Engine/Source/Runtime/Core/Private/GenericPlatform/GenericPlatformProcess.cpp
+++ b/Engine/Source/Runtime/Core/Private/GenericPlatform/GenericPlatformProcess.cpp
@@ -278,7 +278,7 @@ void FGenericPlatformProcess::SleepInfinite()
 
 #endif // PLATFORM_HAS_BSD_TIME 
 
-void FGenericPlatformProcess::ConditionalSleep(const TFunction<bool()>& Condition)
+void FGenericPlatformProcess::ConditionalSleep(const TFunctionRef<bool()>& Condition)
 {
 	if (Condition())
 	{
diff --git a/Engine/Source/Runtime/Core/Public/GenericPlatform/GenericPlatformProcess.h b/Engine/Source/Runtime/Core/Public/GenericPlatform/GenericPlatformProcess.h
--- a/Engine/Source/Runtime/Core/Public/GenericPlatform/GenericPlatformProcess.h
+++ b/Engine/Source/Runtime/Core/Public/GenericPlatform/GenericPlatformProcess.h
@@ -8,7 +8,7 @@
 
 
 class FString;
-template <typename FuncType> class TFunction;
+template <typename FuncType> class TFunctionRef;
 
 namespace EProcessResource
 {
@@ -424,7 +424,7 @@ struct CORE_API FGenericPlatformProcess
 	*
 	* @param	Condition	Condition to evaluate.
 	*/
-	static void ConditionalSleep(const TFunction<bool()>& Condition);
+	static void ConditionalSleep(const TFunctionRef<bool()>& Condition);
 
 	/**
 	 * Creates a new event.
