Modify message on missing config files on startup.

diff --git a/Engine/Source/Runtime/Core/Private/Misc/ConfigCacheIni.cpp b/Engine/Source/Runtime/Core/Private/Misc/ConfigCacheIni.cpp
--- a/Engine/Source/Runtime/Core/Private/Misc/ConfigCacheIni.cpp
+++ b/Engine/Source/Runtime/Core/Private/Misc/ConfigCacheIni.cpp
@@ -2869,7 +2869,7 @@ void FConfigCacheIni::InitializeConfigSystem()
 		{
 			const FText AbsolutePath = FText::FromString( IFileManager::Get().ConvertToAbsolutePathForExternalAppForRead(*FPaths::GetPath(GEngineIni)) );
 			//@todo this is too early to localize
-			const FText Message = FText::Format( NSLOCTEXT("Core", "FirstCmdArgMustBeGameName", "The first command line argument must be a game name. ('{0}' must exist and contain a DefaultEngine.ini)"), AbsolutePath );
+			const FText Message = FText::Format( NSLOCTEXT("Core", "FirstCmdArgMustBeGameName", "'{0}' must exist and contain a DefaultEngine.ini."), AbsolutePath );
 			if (!GIsBuildMachine)
 			{
 				FMessageDialog::Open(EAppMsgType::Ok, Message);
