Getting rid of "using namespace" statement that caused shadowing problem with one of my function local variables.

From the coding standard "Do not put 'using' declarations in the global scope, even in a .cpp file (it will cause problems with our 'unity' build system.)".

diff --git a/Engine/Source/Runtime/Core/Private/GenericPlatform/GenericPlatformContext.cpp b/Engine/Source/Runtime/Core/Private/GenericPlatform/GenericPlatformContext.cpp
--- a/Engine/Source/Runtime/Core/Private/GenericPlatform/GenericPlatformContext.cpp
+++ b/Engine/Source/Runtime/Core/Private/GenericPlatform/GenericPlatformContext.cpp
@@ -31,29 +31,28 @@ namespace NCachedCrashContextProperties
 
 	static FString CrashGUID;
 }
-using namespace NCachedCrashContextProperties;
 
 void FGenericCrashContext::Initialize()
 {
-	ExecutableName = FPlatformProcess::ExecutableName();
-	PlatformName = FPlatformProperties::PlatformName();
-	PlatformNameIni = FPlatformProperties::IniPlatformName();
-	BaseDir = FPlatformProcess::BaseDir();
-	RootDir = FPlatformMisc::RootDir();
-	EpicAccountId = FPlatformMisc::GetEpicAccountId();
-	MachineIdStr = FPlatformMisc::GetMachineId().ToString( EGuidFormats::Digits );
-	FPlatformMisc::GetOSVersions(OsVersion,OsSubVersion);
-	NumberOfCores = FPlatformMisc::NumberOfCores();
-	NumberOfCoresIncludingHyperthreads = FPlatformMisc::NumberOfCoresIncludingHyperthreads();
-
-	CPUVendor = FPlatformMisc::GetCPUVendor();
-	CPUBrand = FPlatformMisc::GetCPUBrand();
-	PrimaryGPUBrand = FPlatformMisc::GetPrimaryGPUBrand();
-	UserName = FPlatformProcess::UserName();
-	DefaultLocale = FPlatformMisc::GetDefaultLocale();
+	NCachedCrashContextProperties::ExecutableName = FPlatformProcess::ExecutableName();
+	NCachedCrashContextProperties::PlatformName = FPlatformProperties::PlatformName();
+	NCachedCrashContextProperties::PlatformNameIni = FPlatformProperties::IniPlatformName();
+	NCachedCrashContextProperties::BaseDir = FPlatformProcess::BaseDir();
+	NCachedCrashContextProperties::RootDir = FPlatformMisc::RootDir();
+	NCachedCrashContextProperties::EpicAccountId = FPlatformMisc::GetEpicAccountId();
+	NCachedCrashContextProperties::MachineIdStr = FPlatformMisc::GetMachineId().ToString(EGuidFormats::Digits);
+	FPlatformMisc::GetOSVersions(NCachedCrashContextProperties::OsVersion, NCachedCrashContextProperties::OsSubVersion);
+	NCachedCrashContextProperties::NumberOfCores = FPlatformMisc::NumberOfCores();
+	NCachedCrashContextProperties::NumberOfCoresIncludingHyperthreads = FPlatformMisc::NumberOfCoresIncludingHyperthreads();
+
+	NCachedCrashContextProperties::CPUVendor = FPlatformMisc::GetCPUVendor();
+	NCachedCrashContextProperties::CPUBrand = FPlatformMisc::GetCPUBrand();
+	NCachedCrashContextProperties::PrimaryGPUBrand = FPlatformMisc::GetPrimaryGPUBrand();
+	NCachedCrashContextProperties::UserName = FPlatformProcess::UserName();
+	NCachedCrashContextProperties::DefaultLocale = FPlatformMisc::GetDefaultLocale();
 
 	const FGuid Guid = FGuid::NewGuid();
-	CrashGUID = FString::Printf( TEXT( "UE4CC-%s-%s" ), *PlatformNameIni, *Guid.ToString( EGuidFormats::Digits ) );
+	NCachedCrashContextProperties::CrashGUID = FString::Printf(TEXT("UE4CC-%s-%s"), *NCachedCrashContextProperties::PlatformNameIni, *Guid.ToString(EGuidFormats::Digits));
 
 	bIsInitialized = true;
 }
@@ -83,27 +82,27 @@ void FGenericCrashContext::SerializeContentToBuffer()
 
 	// Add common crash properties.
 	AddCrashProperty( TEXT( "GameName" ), FApp::GetGameName() );
-	AddCrashProperty( TEXT( "ExecutableName" ), *ExecutableName );
+	AddCrashProperty(TEXT("ExecutableName"), *NCachedCrashContextProperties::ExecutableName);
 	AddCrashProperty( TEXT( "BuildConfiguration" ), EBuildConfigurations::ToString( FApp::GetBuildConfiguration() ) );
 
-	AddCrashProperty( TEXT( "PlatformName" ), *PlatformName );
-	AddCrashProperty( TEXT( "PlatformNameIni" ), *PlatformNameIni );
+	AddCrashProperty(TEXT("PlatformName"), *NCachedCrashContextProperties::PlatformName);
+	AddCrashProperty(TEXT("PlatformNameIni"), *NCachedCrashContextProperties::PlatformNameIni);
 	AddCrashProperty( TEXT( "EngineMode" ), FPlatformMisc::GetEngineMode() );
 	AddCrashProperty( TEXT( "EngineVersion" ), ENGINE_VERSION_STRING );
 	AddCrashProperty( TEXT( "CommandLine" ), FCommandLine::IsInitialized() ? FCommandLine::Get() : TEXT("") );
 	AddCrashProperty( TEXT( "LanguageLCID" ), FInternationalization::Get().GetCurrentCulture()->GetLCID() );
-	AddCrashProperty( TEXT( "DefaultLocale" ), *DefaultLocale );
+	AddCrashProperty(TEXT("DefaultLocale"), *NCachedCrashContextProperties::DefaultLocale);
 
 	AddCrashProperty( TEXT( "IsUE4Release" ), (int32)FApp::IsEngineInstalled() );
 
 	// Remove periods from user names to match AutoReporter user names
 	// The name prefix is read by CrashRepository.AddNewCrash in the website code
-	AddCrashProperty( TEXT( "UserName" ), *UserName.Replace( TEXT( "." ), TEXT( "" ) ) );
+	AddCrashProperty(TEXT("UserName"), *NCachedCrashContextProperties::UserName.Replace(TEXT("."), TEXT("")));
 
-	AddCrashProperty( TEXT( "BaseDir" ), *BaseDir );
-	AddCrashProperty( TEXT( "RootDir" ), *RootDir );
-	AddCrashProperty( TEXT( "MachineId" ), *MachineIdStr );
-	AddCrashProperty( TEXT( "EpicAccountId" ), *EpicAccountId );
+	AddCrashProperty(TEXT("BaseDir"), *NCachedCrashContextProperties::BaseDir);
+	AddCrashProperty(TEXT("RootDir"), *NCachedCrashContextProperties::RootDir);
+	AddCrashProperty(TEXT("MachineId"), *NCachedCrashContextProperties::MachineIdStr);
+	AddCrashProperty(TEXT("EpicAccountId"), *NCachedCrashContextProperties::EpicAccountId);
 
 	AddCrashProperty( TEXT( "CallStack" ), TEXT( "" ) );
 	AddCrashProperty( TEXT( "SourceContext" ), TEXT( "" ) );
@@ -118,15 +117,15 @@ void FGenericCrashContext::SerializeContentToBuffer()
 	*/
 
 	// Add misc stats.
-	AddCrashProperty( TEXT( "MiscNumberOfCores" ), NumberOfCores );
-	AddCrashProperty( TEXT( "MiscNumberOfCoresIncludingHyperthreads" ), NumberOfCoresIncludingHyperthreads );
+	AddCrashProperty(TEXT("MiscNumberOfCores"), NCachedCrashContextProperties::NumberOfCores);
+	AddCrashProperty(TEXT("MiscNumberOfCoresIncludingHyperthreads"), NCachedCrashContextProperties::NumberOfCoresIncludingHyperthreads);
 	AddCrashProperty( TEXT( "MiscIs64bitOperatingSystem" ), (int32)FPlatformMisc::Is64bitOperatingSystem() );
 
-	AddCrashProperty( TEXT( "MiscCPUVendor" ), *CPUVendor );
-	AddCrashProperty( TEXT( "MiscCPUBrand" ), *CPUBrand );
-	AddCrashProperty( TEXT( "MiscPrimaryGPUBrand" ), *PrimaryGPUBrand );
-	AddCrashProperty( TEXT( "MiscOSVersionMajor" ), *OsVersion );
-	AddCrashProperty( TEXT( "MiscOSVersionMinor" ), *OsSubVersion );
+	AddCrashProperty(TEXT("MiscCPUVendor"), *NCachedCrashContextProperties::CPUVendor);
+	AddCrashProperty(TEXT("MiscCPUBrand"), *NCachedCrashContextProperties::CPUBrand);
+	AddCrashProperty(TEXT("MiscPrimaryGPUBrand"), *NCachedCrashContextProperties::PrimaryGPUBrand);
+	AddCrashProperty(TEXT("MiscOSVersionMajor"), *NCachedCrashContextProperties::OsVersion);
+	AddCrashProperty(TEXT("MiscOSVersionMinor"), *NCachedCrashContextProperties::OsSubVersion);
 
 
 	// @TODO yrx 2014-10-08 Move to the crash report client.
@@ -178,7 +177,7 @@ void FGenericCrashContext::SerializeContentToBuffer()
 
 const FString& FGenericCrashContext::GetUniqueCrashName()
 {
-	return CrashGUID;
+	return NCachedCrashContextProperties::CrashGUID;
 }
 
 void FGenericCrashContext::SerializeAsXML( const TCHAR* Filename )
