Core - Added SleepTime to ConditionalSleep

diff --git a/Engine/Source/Runtime/Core/Private/GenericPlatform/GenericPlatformProcess.cpp b/Engine/Source/Runtime/Core/Private/GenericPlatform/GenericPlatformProcess.cpp
--- a/Engine/Source/Runtime/Core/Private/GenericPlatform/GenericPlatformProcess.cpp
+++ b/Engine/Source/Runtime/Core/Private/GenericPlatform/GenericPlatformProcess.cpp
@@ -278,7 +278,7 @@ void FGenericPlatformProcess::SleepInfinite()
 
 #endif // PLATFORM_HAS_BSD_TIME 
 
-void FGenericPlatformProcess::ConditionalSleep(const TFunctionRef<bool()>& Condition)
+void FGenericPlatformProcess::ConditionalSleep(const TFunctionRef<bool()>& Condition, float SleepTime /*= 0.0f*/)
 {
 	if (Condition())
 	{
@@ -289,7 +289,7 @@ void FGenericPlatformProcess::ConditionalSleep(const TFunctionRef<bool()>& Condi
 	FThreadIdleStats::FScopeIdle Scope;
 	do
 	{
-		FPlatformProcess::SleepNoStats(0.0f);
+		FPlatformProcess::SleepNoStats(SleepTime);
 	} while (!Condition());
 }
 
diff --git a/Engine/Source/Runtime/Core/Public/GenericPlatform/GenericPlatformProcess.h b/Engine/Source/Runtime/Core/Public/GenericPlatform/GenericPlatformProcess.h
--- a/Engine/Source/Runtime/Core/Public/GenericPlatform/GenericPlatformProcess.h
+++ b/Engine/Source/Runtime/Core/Public/GenericPlatform/GenericPlatformProcess.h
@@ -423,8 +423,9 @@ struct CORE_API FGenericPlatformProcess
 	* Sleep thread until condition is satisfied.
 	*
 	* @param	Condition	Condition to evaluate.
+	* @param	SleepTime	Time to sleep
 	*/
-	static void ConditionalSleep(const TFunctionRef<bool()>& Condition);
+	static void ConditionalSleep(const TFunctionRef<bool()>& Condition, float SleepTime = 0.0f);
 
 	/**
 	 * Creates a new event.
