Exposed the low-level FileExists function to IFileManager and updated FPaths to use it

diff --git a/Engine/Source/Runtime/Core/Private/HAL/FileManagerGeneric.cpp b/Engine/Source/Runtime/Core/Private/HAL/FileManagerGeneric.cpp
--- a/Engine/Source/Runtime/Core/Private/HAL/FileManagerGeneric.cpp
+++ b/Engine/Source/Runtime/Core/Private/HAL/FileManagerGeneric.cpp
@@ -330,6 +330,11 @@ bool FFileManagerGeneric::Move( const TCHAR* Dest, const TCHAR* Src, bool Replac
 	return true;
 }
 
+bool FFileManagerGeneric::FileExists( const TCHAR* Filename )
+{
+	return GetLowLevel().FileExists( Filename );
+}
+
 bool FFileManagerGeneric::DirectoryExists( const TCHAR* InDirectory )
 {
 	return GetLowLevel().DirectoryExists( InDirectory );
diff --git a/Engine/Source/Runtime/Core/Private/Misc/Paths.cpp b/Engine/Source/Runtime/Core/Private/Misc/Paths.cpp
--- a/Engine/Source/Runtime/Core/Private/Misc/Paths.cpp
+++ b/Engine/Source/Runtime/Core/Private/Misc/Paths.cpp
@@ -458,7 +458,7 @@ FString FPaths::GetPath(const FString& InPath)
 
 bool FPaths::FileExists(const FString& InPath)
 {
-	return IFileManager::Get().GetTimeStamp(*InPath) > FDateTime::MinValue();
+	return IFileManager::Get().FileExists(*InPath);
 }
 
 bool FPaths::DirectoryExists(const FString& InPath)
diff --git a/Engine/Source/Runtime/Core/Public/HAL/FileManager.h b/Engine/Source/Runtime/Core/Public/HAL/FileManager.h
--- a/Engine/Source/Runtime/Core/Public/HAL/FileManager.h
+++ b/Engine/Source/Runtime/Core/Public/HAL/FileManager.h
@@ -93,6 +93,9 @@ public:
 	/** Moves/renames a file. */
 	virtual bool Move( const TCHAR* Dest, const TCHAR* Src, bool Replace=1, bool EvenIfReadOnly=0, bool Attributes=0, bool bDoNotRetryOrError=0 )=0;
 
+	/** Checks if a file exists */
+	virtual bool FileExists( const TCHAR* Filename )=0;
+
 	/** Checks if a directory exists. */
 	virtual bool DirectoryExists( const TCHAR* InDirectory )=0;
 
diff --git a/Engine/Source/Runtime/Core/Public/HAL/FileManagerGeneric.h b/Engine/Source/Runtime/Core/Public/HAL/FileManagerGeneric.h
--- a/Engine/Source/Runtime/Core/Public/HAL/FileManagerGeneric.h
+++ b/Engine/Source/Runtime/Core/Public/HAL/FileManagerGeneric.h
@@ -60,6 +60,7 @@ public:
 	bool Delete( const TCHAR* Filename, bool RequireExists=0, bool EvenReadOnly=0, bool Quiet=0 ) override;
 	bool IsReadOnly( const TCHAR* Filename ) override;
 	bool Move( const TCHAR* Dest, const TCHAR* Src, bool Replace=1, bool EvenIfReadOnly=0, bool Attributes=0, bool bDoNotRetryOrError=0 ) override;
+	bool FileExists( const TCHAR* Filename ) override;
 	bool DirectoryExists(const TCHAR* InDirectory) override;
 	void FindFiles( TArray<FString>& Result, const TCHAR* Filename, bool Files, bool Directories ) override;
 	void FindFilesRecursive( TArray<FString>& FileNames, const TCHAR* StartDirectory, const TCHAR* Filename, bool Files, bool Directories, bool bClearFileNames=true) override;
