Add -Wshadow on HTML5. Fixed errors relating to it. UE-14250
Fixed some potential issues from incoming SDK release. Removed unnecessary js lib.

diff --git a/Engine/Source/Programs/UnrealBuildTool/HTML5/HTML5ToolChain.cs b/Engine/Source/Programs/UnrealBuildTool/HTML5/HTML5ToolChain.cs
--- a/Engine/Source/Programs/UnrealBuildTool/HTML5/HTML5ToolChain.cs
+++ b/Engine/Source/Programs/UnrealBuildTool/HTML5/HTML5ToolChain.cs
@@ -64,6 +64,11 @@ namespace UnrealBuildTool
             Result += " -Wno-array-bounds"; // some VectorLoads go past the end of the array, but it's okay in that case
             Result += " -Wno-invalid-offsetof"; // too many warnings kills windows clang. 
 
+			if (BuildConfiguration.bEnableShadowVariableWarning)
+			{
+				Result += " -Wshadow";
+			}
+
             // JavsScript option overrides (see src/settings.js)
 
             // we have to specify the full amount of memory with Asm.JS (1.5 G)
@@ -203,8 +208,6 @@ namespace UnrealBuildTool
 				Result += " -s CASE_INSENSITIVE_FS=1 ";
 
 
-                string BaseSDKPath = HTML5SDKInfo.EmscriptenSDKPath();
-				Result += " --js-library \"" + BaseSDKPath + "/Src/library_openal.js\" ";
 			}
 
 			return Result;
diff --git a/Engine/Source/Runtime/Core/Public/HTML5/HTML5AssertionMacros.h b/Engine/Source/Runtime/Core/Public/HTML5/HTML5AssertionMacros.h
--- a/Engine/Source/Runtime/Core/Public/HTML5/HTML5AssertionMacros.h
+++ b/Engine/Source/Runtime/Core/Public/HTML5/HTML5AssertionMacros.h
@@ -28,30 +28,37 @@ void emscripten_log(int flags, ...);
 #define checkNoReentry(...)
 #define checkNoRecursion(...)
 
-#define html5_break_msg(msg, file, line) EM_ASM(alert('Expression ('+msg+') failed in '+file+':'+line+'!\nCheck console for details.\n'))
+inline void html5_break_msg(const char* msg, const char* file, int line) {
+	EM_ASM_ARGS(
+	{
+		var InMsg = Pointer_stringify($0);
+		var InFile = Pointer_stringify($1);
+		alert('Expression ('+InMsg+') failed in '+InFile+':'+$2+'!\nCheck console for details.\n'); 
+	}, msg, file, line);
+}
 
 #define check(expr)			{ if (!(expr)) { \
 		emscripten_log(255, "Expression '" #expr "' failed in " __FILE__ ":" PREPROCESSOR_TO_STRING(__LINE__) "!\n"); \
-		html5_break_msg(#expr, PREPROCESSOR_TO_STRING(__FILE__), PREPROCESSOR_TO_STRING(__LINE__)); \
+		html5_break_msg(#expr, __FILE__, __LINE__); \
 	} \
 }
 #define checkf(expr, ...)	{ if (!(expr)) { \
 		emscripten_log(255, "Expression '" #expr "' failed in " __FILE__ ":" PREPROCESSOR_TO_STRING(__LINE__) "!\n"); \
 		emscripten_log(255, ##__VA_ARGS__); \
 		FDebug::AssertFailed( #expr, __FILE__, __LINE__, ##__VA_ARGS__ ); \
-		html5_break_msg(#expr, PREPROCESSOR_TO_STRING(__FILE__), PREPROCESSOR_TO_STRING(__LINE__)); \
+		html5_break_msg(#expr, __FILE__, __LINE__); \
 	} \
 	CA_ASSUME(expr); \
 }
 #define verify(expr)		{ if(!(expr)) {\
 		emscripten_log(255, "Expression '" #expr "' failed in " __FILE__ ":" PREPROCESSOR_TO_STRING(__LINE__) "!\n"); \
-		html5_break_msg(#expr, PREPROCESSOR_TO_STRING(__FILE__), PREPROCESSOR_TO_STRING(__LINE__)); \
+		html5_break_msg(#expr, __FILE__, __LINE__); \
 	} \
 }
 #define verifyf(expr, ...)	{ if(!(expr)) { \
 		emscripten_log(255, "Expression '" #expr "' failed in " __FILE__ ":" PREPROCESSOR_TO_STRING(__LINE__) "!\n"); \
 		emscripten_log(255, ##__VA_ARGS__); \
-		html5_break_msg(#expr, PREPROCESSOR_TO_STRING(__FILE__), PREPROCESSOR_TO_STRING(__LINE__)); \
+		html5_break_msg(#expr, __FILE__, __LINE__); \
 	} \
 }
 
@@ -66,21 +73,21 @@ void emscripten_log(int flags, ...);
 #define checkSlow(expr, ...)   {if(!(expr)) { \
 		emscripten_log(255, "Expression '" #expr "' failed in " __FILE__ ":" PREPROCESSOR_TO_STRING(__LINE__) "!\n"); \
 		emscripten_log(255, ##__VA_ARGS__); FDebug::AssertFailed( #expr, __FILE__, __LINE__ ); \
-		html5_break_msg(#expr, PREPROCESSOR_TO_STRING(__FILE__), PREPROCESSOR_TO_STRING(__LINE__)); \
+		html5_break_msg(#expr, __FILE__, __LINE__); \
 	} \
 	CA_ASSUME(expr); \
 }
 #define checkfSlow(expr, ...)	{ if(!(expr)) { \
 		emscripten_log(255, "Expression '" #expr "' failed in " __FILE__ ":" PREPROCESSOR_TO_STRING(__LINE__) "!\n"); \
 		emscripten_log(255, ##__VA_ARGS__); FDebug::AssertFailed( #expr, __FILE__, __LINE__, __VA_ARGS__ ); \
-		html5_break_msg(#expr, PREPROCESSOR_TO_STRING(__FILE__), PREPROCESSOR_TO_STRING(__LINE__)); \
+		html5_break_msg(#expr, __FILE__, __LINE__); \
 	} \
 	CA_ASSUME(expr); \
 }
 #define verifySlow(expr)  {if(!(expr)) { \
 		emscripten_log(255, "Expression '" #expr "' failed in " __FILE__ ":" PREPROCESSOR_TO_STRING(__LINE__) "!\n"); \
 		FDebug::AssertFailed( #expr, __FILE__, __LINE__ ); \
-		html5_break_msg(#expr, PREPROCESSOR_TO_STRING(__FILE__), PREPROCESSOR_TO_STRING(__LINE__)); \
+		html5_break_msg(#expr, __FILE__, __LINE__); \
 	} \
 }
 
diff --git a/Engine/Source/Runtime/HTML5/HTML5Audio/Private/HTML5AudioBuffer.cpp b/Engine/Source/Runtime/HTML5/HTML5Audio/Private/HTML5AudioBuffer.cpp
--- a/Engine/Source/Runtime/HTML5/HTML5Audio/Private/HTML5AudioBuffer.cpp
+++ b/Engine/Source/Runtime/HTML5/HTML5Audio/Private/HTML5AudioBuffer.cpp
@@ -107,12 +107,12 @@ FALSoundBuffer* FALSoundBuffer::CreateNativeBuffer( FALAudioDevice* AudioDevice,
 	{
 		return( NULL );
 	}
-	FWaveModInfo WaveInfo;
 
 	Wave->InitAudioResource(AudioDevice->GetRuntimeFormat(Wave));
 
 	FALSoundBuffer* Buffer = NULL;
 	FAudioDeviceManager* AudioDeviceManager = GEngine->GetAudioDeviceManager();
+	check(AudioDeviceManager != nullptr);
 
 	// Find the existing buffer if any
 	if( Wave->ResourceID )
@@ -129,9 +129,6 @@ FALSoundBuffer* FALSoundBuffer::CreateNativeBuffer( FALAudioDevice* AudioDevice,
 
 		AudioDevice->alError( TEXT( "RegisterSound" ) );
 
-		FAudioDeviceManager* AudioDeviceManager = GEngine->GetAudioDeviceManager();
-		check(AudioDeviceManager != nullptr);
-
 		AudioDeviceManager->TrackResource(Wave, Buffer);
 
 		Buffer->InternalFormat = AudioDevice->GetInternalFormat( Wave->NumChannels );
diff --git a/Engine/Source/Runtime/Online/HTTP/Private/HTML5/HTML5HTTP.cpp b/Engine/Source/Runtime/Online/HTTP/Private/HTML5/HTML5HTTP.cpp
--- a/Engine/Source/Runtime/Online/HTTP/Private/HTML5/HTML5HTTP.cpp
+++ b/Engine/Source/Runtime/Online/HTTP/Private/HTML5/HTML5HTTP.cpp
@@ -306,7 +306,7 @@ bool FHTML5HttpRequest::StartRequest()
 	TArray<FString> AllHeaders = GetAllHeaders();
 
 	// Create a String which emscripten can understand. 
-	FString Headers = FString::Join(AllHeaders, TEXT("%"));
+	FString RequestHeaders = FString::Join(AllHeaders, TEXT("%"));
 	
 	// set up verb (note that Verb is expected to be uppercase only)
 	if (Verb == TEXT("POST"))
@@ -316,7 +316,7 @@ bool FHTML5HttpRequest::StartRequest()
 		check(!GetHeader("Content-Type").IsEmpty() || IsURLEncoded(RequestPayload));
 
 #if PLATFORM_HTML5_BROWSER
-		UE_MakeHTTPDataRequest(this, TCHAR_TO_ANSI(*URL), "POST", (char*)RequestPayload.GetData(), RequestPayload.Num(),TCHAR_TO_ANSI(*Headers), 0, StaticReceiveCallback, StaticErrorCallback, StaticProgressCallback);
+		UE_MakeHTTPDataRequest(this, TCHAR_TO_ANSI(*URL), "POST", (char*)RequestPayload.GetData(), RequestPayload.Num(),TCHAR_TO_ANSI(*RequestHeaders), 0, StaticReceiveCallback, StaticErrorCallback, StaticProgressCallback);
 #else
 		return false;
 #endif
@@ -335,7 +335,7 @@ bool FHTML5HttpRequest::StartRequest()
 	else if (Verb == TEXT("GET"))
 	{
 #if PLATFORM_HTML5_BROWSER
-		UE_MakeHTTPDataRequest(this, TCHAR_TO_ANSI(*URL), "GET", NULL, 0,TCHAR_TO_ANSI(*Headers), 1, StaticReceiveCallback, StaticErrorCallback, StaticProgressCallback);
+		UE_MakeHTTPDataRequest(this, TCHAR_TO_ANSI(*URL), "GET", NULL, 0,TCHAR_TO_ANSI(*RequestHeaders), 1, StaticReceiveCallback, StaticErrorCallback, StaticProgressCallback);
 #else
 		return false;
 #endif
