#ue4 - Add FQuat::GetNormalized() that returns a normalized copy.

diff --git a/Engine/Source/Runtime/Core/Public/Math/Quat.h b/Engine/Source/Runtime/Core/Public/Math/Quat.h
--- a/Engine/Source/Runtime/Core/Public/Math/Quat.h
+++ b/Engine/Source/Runtime/Core/Public/Math/Quat.h
@@ -260,11 +260,20 @@ public:
 
 	/**
 	 * Normalize this quaternion if it is large enough.
+	 * If it is too small, returns an identity quaternion.
 	 *
-	 * @param Tolerance Minimum squared length of vector for normalization.
+	 * @param Tolerance Minimum squared length of quaternion for normalization.
 	 */
 	FORCEINLINE void Normalize( float Tolerance=SMALL_NUMBER );
 
+	/**
+	 * Get a normalized copy of this quaternion.
+	 * If it is too small, returns an identity quaternion.
+	 *
+	 * @param Tolerance Minimum squared length of quaternion for normalization.
+	 */
+	FORCEINLINE FQuat GetNormalized( float Tolerance=SMALL_NUMBER ) const;
+
 	// Return true if this quaternion is normalized
 	bool IsNormalized() const;
 
@@ -796,6 +805,15 @@ FORCEINLINE void FQuat::Normalize(float Tolerance)
 }
 
 
+FORCEINLINE FQuat FQuat::GetNormalized(float Tolerance) const
+{
+	FQuat Result(*this);
+	Result.Normalize(Tolerance);
+	return Result;
+}
+
+
+
 FORCEINLINE bool FQuat::IsNormalized() const
 {
 	return (FMath::Abs(1.f - SizeSquared()) < THRESH_QUAT_NORMALIZED);
