Linux: game window shouldn't be resizeable (UE-8612).

diff --git a/Engine/Source/Runtime/Core/Private/Linux/LinuxWindow.cpp b/Engine/Source/Runtime/Core/Private/Linux/LinuxWindow.cpp
--- a/Engine/Source/Runtime/Core/Private/Linux/LinuxWindow.cpp
+++ b/Engine/Source/Runtime/Core/Private/Linux/LinuxWindow.cpp
@@ -98,11 +98,11 @@ void FLinuxWindow::Initialize( FLinuxApplication* const Application, const TShar
 		{
 			WindowStyle |= SDL_WINDOW_SKIP_TASKBAR;
 		}
-	}
 
-	if ( Definition->HasSizingFrame )
-	{
-		WindowStyle |= SDL_WINDOW_RESIZABLE;
+		if (Definition->IsRegularWindow && Definition->HasSizingFrame)
+		{
+			WindowStyle |= SDL_WINDOW_RESIZABLE;
+		}
 	}
 
 	// Does this combination tell us that this is a tooltip window? Need to pass information to
