Work around for compilation problems with clang 3.6.0

diff --git a/Engine/Source/Runtime/Core/Public/UObject/WeakObjectPtrTemplates.h b/Engine/Source/Runtime/Core/Public/UObject/WeakObjectPtrTemplates.h
--- a/Engine/Source/Runtime/Core/Public/UObject/WeakObjectPtrTemplates.h
+++ b/Engine/Source/Runtime/Core/Public/UObject/WeakObjectPtrTemplates.h
@@ -15,7 +15,9 @@
  * Most often it is used when you explicitly do NOT want to prevent something from being garbage collected.
  **/
 struct FWeakObjectPtr;
-struct FIndexToObject;
+
+template<class T=UObject, class TWeakObjectPtrBase=FWeakObjectPtr>
+struct TWeakObjectPtr;
 
 
 /***
@@ -23,13 +25,12 @@ struct FIndexToObject;
 * TWeakObjectPtr is templatized version of the generic FWeakObjectPtr
 * 
 **/
-template<class T=UObject, class TWeakObjectPtrBase=FWeakObjectPtr, class TUObjectArray=FIndexToObject>
+template<class T, class TWeakObjectPtrBase>
 struct TWeakObjectPtr : private TWeakObjectPtrBase
 {
 	// Although templated, these parameters are not intended to be anything other than the default,
 	// and are only templates for module organization reasons.
 	static_assert(TAreTypesEqual<TWeakObjectPtrBase, FWeakObjectPtr>::Value, "TWeakObjectPtrBase should not be overridden");
-	static_assert(TAreTypesEqual<TUObjectArray,      FIndexToObject>::Value, "TUObjectArray should not be overridden");
 
 public:
 
@@ -40,7 +41,7 @@ public:
 	 * Construct from an object pointer
 	 * @param Object object to create a weak pointer to
 	**/
-	FORCEINLINE TWeakObjectPtr(const T *Object) :
+	FORCEINLINE TWeakObjectPtr(const T* Object) :
 		TWeakObjectPtrBase((UObject*)Object)
 	{
 		// This static assert is in here rather than in the body of the class because we want
@@ -52,7 +53,7 @@ public:
 	 * Construct from another weak pointer
 	 * @param Other weak pointer to copy from
 	**/
-	FORCEINLINE TWeakObjectPtr(const TWeakObjectPtr<T> &Other) :
+	FORCEINLINE TWeakObjectPtr(const TWeakObjectPtr& Other) :
 		TWeakObjectPtrBase(Other)
 	{ }
 
@@ -61,8 +62,8 @@ public:
 	 * @param Other weak pointer to copy from
 	**/
 	template <typename OtherT>
-	FORCEINLINE TWeakObjectPtr(const TWeakObjectPtr<OtherT, TWeakObjectPtrBase, TUObjectArray> &Other) :
-		TWeakObjectPtrBase(Other)
+	FORCEINLINE TWeakObjectPtr(const TWeakObjectPtr<OtherT, TWeakObjectPtrBase>& Other) :
+		TWeakObjectPtrBase(*(TWeakObjectPtrBase*)&Other) // we do a C-style cast to private base here to avoid clang 3.6.0 compilation problems with friend declarations
 	{
 		// It's also possible that this static_assert may fail for valid conversions because
 		// one or both of the types have only been forward-declared.
@@ -82,7 +83,7 @@ public:
 	 * @param Object object to create a weak pointer to
 	**/
 	template<class U>
-	FORCEINLINE void operator=(const U *Object)
+	FORCEINLINE void operator=(const U* Object)
 	{
 		const T* TempObject = Object;
 		TWeakObjectPtrBase::operator=(TempObject);
@@ -92,7 +93,7 @@ public:
 	 * Construct from another weak pointer
 	 * @param Other weak pointer to copy from
 	**/
-	FORCEINLINE void operator=(const TWeakObjectPtr<T> &Other)
+	FORCEINLINE void operator=(const TWeakObjectPtr& Other)
 	{
 		TWeakObjectPtrBase::operator=(Other);
 	}
@@ -102,13 +103,13 @@ public:
 	 * @param Other weak pointer to copy from
 	**/
 	template <typename OtherT>
-	FORCEINLINE void operator=(const TWeakObjectPtr<OtherT, TWeakObjectPtrBase, TUObjectArray> &Other)
+	FORCEINLINE void operator=(const TWeakObjectPtr<OtherT, TWeakObjectPtrBase>& Other)
 	{
 		// It's also possible that this static_assert may fail for valid conversions because
 		// one or both of the types have only been forward-declared.
 		static_assert(TPointerIsConvertibleFromTo<OtherT, T>::Value, "Unable to convert TWeakObjectPtr - types are incompatible");
 
-		TWeakObjectPtrBase::operator=(Other);
+		*(TWeakObjectPtrBase*)this = *(TWeakObjectPtrBase*)&Other; // we do a C-style cast to private base here to avoid clang 3.6.0 compilation problems with friend declarations
 	}
 
 	/**  
@@ -167,24 +168,20 @@ public:
 	}
 
 	/** Hash function. */
-	FORCEINLINE friend uint32 GetTypeHash(const TWeakObjectPtr<T>& WeakObjectPtr)
+	FORCEINLINE friend uint32 GetTypeHash(const TWeakObjectPtr& WeakObjectPtr)
 	{
 		return GetTypeHash(static_cast<const TWeakObjectPtrBase&>(WeakObjectPtr));
 	}
 
-	friend FArchive& operator<<( FArchive& Ar, TWeakObjectPtr<T>& WeakObjectPtr )
+	friend FArchive& operator<<( FArchive& Ar, TWeakObjectPtr& WeakObjectPtr )
 	{
 		Ar << static_cast<TWeakObjectPtrBase&>(WeakObjectPtr);
 		return Ar;
 	}
-
-	// We declare ourselves as a friend (templated using OtherType) so we can access members as needed
-	template<class OtherT, class OtherTWeakObjectPtrBase, class OtherTUObjectArray>
-	friend struct TWeakObjectPtr;
 };
 
-template <typename LhsT, typename RhsT, typename OtherTWeakObjectPtrBase, typename OtherTUObjectArray>
-FORCENOINLINE bool operator==(const TWeakObjectPtr<LhsT, OtherTWeakObjectPtrBase, OtherTUObjectArray> &Lhs, const TWeakObjectPtr<RhsT, OtherTWeakObjectPtrBase, OtherTUObjectArray> &Rhs)
+template <typename LhsT, typename RhsT, typename OtherTWeakObjectPtrBase>
+FORCENOINLINE bool operator==(const TWeakObjectPtr<LhsT, OtherTWeakObjectPtrBase>& Lhs, const TWeakObjectPtr<RhsT, OtherTWeakObjectPtrBase>& Rhs)
 {
 	// It's also possible that this static_assert may fail for valid conversions because
 	// one or both of the types have only been forward-declared.
@@ -193,8 +190,8 @@ FORCENOINLINE bool operator==(const TWeakObjectPtr<LhsT, OtherTWeakObjectPtrBase
 	return (const OtherTWeakObjectPtrBase&)Lhs == (const OtherTWeakObjectPtrBase&)Rhs;
 }
 
-template <typename LhsT, typename RhsT, typename OtherTWeakObjectPtrBase, typename OtherTUObjectArray>
-FORCENOINLINE bool operator==(const TWeakObjectPtr<LhsT, OtherTWeakObjectPtrBase, OtherTUObjectArray> &Lhs, const RhsT* Rhs)
+template <typename LhsT, typename RhsT, typename OtherTWeakObjectPtrBase>
+FORCENOINLINE bool operator==(const TWeakObjectPtr<LhsT, OtherTWeakObjectPtrBase>& Lhs, const RhsT* Rhs)
 {
 	// It's also possible that these static_asserts may fail for valid conversions because
 	// one or both of the types have only been forward-declared.
@@ -205,8 +202,8 @@ FORCENOINLINE bool operator==(const TWeakObjectPtr<LhsT, OtherTWeakObjectPtrBase
 	return (const OtherTWeakObjectPtrBase&)Lhs == OtherTWeakObjectPtrBase(Rhs);
 }
 
-template <typename LhsT, typename RhsT, typename OtherTWeakObjectPtrBase, typename OtherTUObjectArray>
-FORCENOINLINE bool operator==(const LhsT* Lhs, const TWeakObjectPtr<RhsT, OtherTWeakObjectPtrBase, OtherTUObjectArray> &Rhs)
+template <typename LhsT, typename RhsT, typename OtherTWeakObjectPtrBase>
+FORCENOINLINE bool operator==(const LhsT* Lhs, const TWeakObjectPtr<RhsT, OtherTWeakObjectPtrBase>& Rhs)
 {
 	// It's also possible that these static_asserts may fail for valid conversions because
 	// one or both of the types have only been forward-declared.
@@ -217,20 +214,20 @@ FORCENOINLINE bool operator==(const LhsT* Lhs, const TWeakObjectPtr<RhsT, OtherT
 	return OtherTWeakObjectPtrBase(Lhs) == (const OtherTWeakObjectPtrBase&)Rhs;
 }
 
-template <typename LhsT, typename OtherTWeakObjectPtrBase, typename OtherTUObjectArray>
-FORCENOINLINE bool operator==(const TWeakObjectPtr<LhsT, OtherTWeakObjectPtrBase, OtherTUObjectArray> &Lhs, TYPE_OF_NULLPTR)
+template <typename LhsT, typename OtherTWeakObjectPtrBase>
+FORCENOINLINE bool operator==(const TWeakObjectPtr<LhsT, OtherTWeakObjectPtrBase>& Lhs, TYPE_OF_NULLPTR)
 {
 	return !Lhs.IsValid();
 }
 
-template <typename RhsT, typename OtherTWeakObjectPtrBase, typename OtherTUObjectArray>
-FORCENOINLINE bool operator==(TYPE_OF_NULLPTR, const TWeakObjectPtr<RhsT, OtherTWeakObjectPtrBase, OtherTUObjectArray> &Rhs)
+template <typename RhsT, typename OtherTWeakObjectPtrBase>
+FORCENOINLINE bool operator==(TYPE_OF_NULLPTR, const TWeakObjectPtr<RhsT, OtherTWeakObjectPtrBase>& Rhs)
 {
 	return !Rhs.IsValid();
 }
 
-template <typename LhsT, typename RhsT, typename OtherTWeakObjectPtrBase, typename OtherTUObjectArray>
-FORCENOINLINE bool operator!=(const TWeakObjectPtr<LhsT, OtherTWeakObjectPtrBase, OtherTUObjectArray> &Lhs, const TWeakObjectPtr<RhsT, OtherTWeakObjectPtrBase, OtherTUObjectArray> &Rhs)
+template <typename LhsT, typename RhsT, typename OtherTWeakObjectPtrBase>
+FORCENOINLINE bool operator!=(const TWeakObjectPtr<LhsT, OtherTWeakObjectPtrBase>& Lhs, const TWeakObjectPtr<RhsT, OtherTWeakObjectPtrBase>& Rhs)
 {
 	// It's also possible that this static_assert may fail for valid conversions because
 	// one or both of the types have only been forward-declared.
@@ -239,8 +236,8 @@ FORCENOINLINE bool operator!=(const TWeakObjectPtr<LhsT, OtherTWeakObjectPtrBase
 	return (const OtherTWeakObjectPtrBase&)Lhs != (const OtherTWeakObjectPtrBase&)Rhs;
 }
 
-template <typename LhsT, typename RhsT, typename OtherTWeakObjectPtrBase, typename OtherTUObjectArray>
-FORCENOINLINE bool operator!=(const TWeakObjectPtr<LhsT, OtherTWeakObjectPtrBase, OtherTUObjectArray> &Lhs, const RhsT* Rhs)
+template <typename LhsT, typename RhsT, typename OtherTWeakObjectPtrBase>
+FORCENOINLINE bool operator!=(const TWeakObjectPtr<LhsT, OtherTWeakObjectPtrBase>& Lhs, const RhsT* Rhs)
 {
 	// It's also possible that these static_asserts may fail for valid conversions because
 	// one or both of the types have only been forward-declared.
@@ -251,8 +248,8 @@ FORCENOINLINE bool operator!=(const TWeakObjectPtr<LhsT, OtherTWeakObjectPtrBase
 	return (const OtherTWeakObjectPtrBase&)Lhs != OtherTWeakObjectPtrBase(Rhs);
 }
 
-template <typename LhsT, typename RhsT, typename OtherTWeakObjectPtrBase, typename OtherTUObjectArray>
-FORCENOINLINE bool operator!=(const LhsT* Lhs, const TWeakObjectPtr<RhsT, OtherTWeakObjectPtrBase, OtherTUObjectArray> &Rhs)
+template <typename LhsT, typename RhsT, typename OtherTWeakObjectPtrBase>
+FORCENOINLINE bool operator!=(const LhsT* Lhs, const TWeakObjectPtr<RhsT, OtherTWeakObjectPtrBase>& Rhs)
 {
 	// It's also possible that these static_asserts may fail for valid conversions because
 	// one or both of the types have only been forward-declared.
@@ -263,14 +260,14 @@ FORCENOINLINE bool operator!=(const LhsT* Lhs, const TWeakObjectPtr<RhsT, OtherT
 	return OtherTWeakObjectPtrBase(Lhs) != (const OtherTWeakObjectPtrBase&)Rhs;
 }
 
-template <typename LhsT, typename OtherTWeakObjectPtrBase, typename OtherTUObjectArray>
-FORCENOINLINE bool operator!=(const TWeakObjectPtr<LhsT, OtherTWeakObjectPtrBase, OtherTUObjectArray> &Lhs, TYPE_OF_NULLPTR)
+template <typename LhsT, typename OtherTWeakObjectPtrBase>
+FORCENOINLINE bool operator!=(const TWeakObjectPtr<LhsT, OtherTWeakObjectPtrBase>& Lhs, TYPE_OF_NULLPTR)
 {
 	return Lhs.IsValid();
 }
 
-template <typename RhsT, typename OtherTWeakObjectPtrBase, typename OtherTUObjectArray>
-FORCENOINLINE bool operator!=(TYPE_OF_NULLPTR, const TWeakObjectPtr<RhsT, OtherTWeakObjectPtrBase, OtherTUObjectArray> &Rhs)
+template <typename RhsT, typename OtherTWeakObjectPtrBase>
+FORCENOINLINE bool operator!=(TYPE_OF_NULLPTR, const TWeakObjectPtr<RhsT, OtherTWeakObjectPtrBase>& Rhs)
 {
 	return Rhs.IsValid();
 }
@@ -292,17 +289,17 @@ public:
 	{
 	}
 	/** Construct from a raw pointer **/
-	FORCEINLINE TAutoWeakObjectPtr(const T *Target) 
+	FORCEINLINE TAutoWeakObjectPtr(const T* Target) 
 		: TAutoPointer<T, TWeakObjectPtr<T> >(Target)
 	{
 	}
 	/**  Construct from the base type **/
-	FORCEINLINE TAutoWeakObjectPtr(const TWeakObjectPtr<T> &Other) 
+	FORCEINLINE TAutoWeakObjectPtr(const TWeakObjectPtr<T>& Other) 
 		: TAutoPointer<T, TWeakObjectPtr<T> >(Other)
 	{
 	}
 	/**  Construct from another auto pointer **/
-	FORCEINLINE TAutoWeakObjectPtr(const TAutoWeakObjectPtr &Other) 
+	FORCEINLINE TAutoWeakObjectPtr(const TAutoWeakObjectPtr& Other) 
 		: TAutoPointer<T, TWeakObjectPtr<T> >(Other)
 	{
 	}
