Made Editor.ini per-project once again

It was erroneously being migrated to a project-agnostic location, but project specific settings are stored in here. GIsBuildMachine checks have also been removed as they were deemed no longer necessary since build machines serialise project agnostic config in Engine/Saved (not elsewhere on disk)

diff --git a/Engine/Source/Runtime/Core/Private/Misc/ConfigCacheIni.cpp b/Engine/Source/Runtime/Core/Private/Misc/ConfigCacheIni.cpp
--- a/Engine/Source/Runtime/Core/Private/Misc/ConfigCacheIni.cpp
+++ b/Engine/Source/Runtime/Core/Private/Misc/ConfigCacheIni.cpp
@@ -2831,24 +2831,6 @@ static FString GetDestIniFilename(const TCHAR* BaseIniName, const TCHAR* Platfor
 	return IniFilename;
 }
 
-/** Load a project agnostic ini file. Will be read-only for build machines. */
-void LoadProjectAgnosticIni(FString& InOutFilename, const TCHAR* InName)
-{
-	static const FString EditorSettingsDir = FPaths::Combine(*FPaths::GameAgnosticSavedDir(), TEXT("Config")) + TEXT("/");
-
-	if (FConfigCacheIni::LoadGlobalIniFile(InOutFilename, InName, NULL, false, false, true, *EditorSettingsDir))
-	{
-		// Only save these configs if we are not the build machine as it would introduce an indeterminate state
-		if (GIsBuildMachine)
-		{
-			if (FConfigFile* Config = GConfig->FindConfigFile(*InOutFilename))
-			{
-				Config->NoSave = true;
-			}
-		}
-	}
-}
-
 void FConfigCacheIni::InitializeConfigSystem()
 {
 	// Perform any upgrade we need before we load any configuration files
@@ -2888,15 +2870,17 @@ void FConfigCacheIni::InitializeConfigSystem()
 #if WITH_EDITOR
 	// load some editor specific .ini files
 
+	FConfigCacheIni::LoadGlobalIniFile(GEditorIni, TEXT("Editor"));
+
 	// Upgrade editor user settings before loading the editor per project user settings
 	FConfigManifest::MigrateEditorUserSettings();
 	FConfigCacheIni::LoadGlobalIniFile(GEditorPerProjectIni, TEXT("EditorPerProjectUserSettings"));
 
 	// Project agnostic editor ini files
-	LoadProjectAgnosticIni(GEditorIni, TEXT("Editor"));
-	LoadProjectAgnosticIni(GEditorSettingsIni, TEXT("EditorSettings"));
-	LoadProjectAgnosticIni(GEditorLayoutIni, TEXT("EditorLayout"));
-	LoadProjectAgnosticIni(GEditorKeyBindingsIni, TEXT("EditorKeyBindings"));
+	static const FString EditorSettingsDir = FPaths::Combine(*FPaths::GameAgnosticSavedDir(), TEXT("Config")) + TEXT("/");
+	FConfigCacheIni::LoadGlobalIniFile(GEditorSettingsIni, TEXT("EditorSettings"), NULL, false, false, true, *EditorSettingsDir);
+	FConfigCacheIni::LoadGlobalIniFile(GEditorLayoutIni, TEXT("EditorLayout"), NULL, false, false, true, *EditorSettingsDir);
+	FConfigCacheIni::LoadGlobalIniFile(GEditorKeyBindingsIni, TEXT("EditorKeyBindings"), NULL, false, false, true, *EditorSettingsDir);
 
 #endif
 #if PLATFORM_DESKTOP
diff --git a/Engine/Source/Runtime/Core/Private/Misc/ConfigManifest.cpp b/Engine/Source/Runtime/Core/Private/Misc/ConfigManifest.cpp
--- a/Engine/Source/Runtime/Core/Private/Misc/ConfigManifest.cpp
+++ b/Engine/Source/Runtime/Core/Private/Misc/ConfigManifest.cpp
@@ -222,11 +222,9 @@ EConfigManifestVersion FConfigManifest::UpgradeFromVersion(EConfigManifestVersio
 
 		// The initial versioning made the following changes:
 
-		// 1. Move Editor.ini from Game/Saved/Config to Engine/Saved/Config, thus making it project-agnostic
-		// 2. Move EditorLayout.ini from Game/Saved/Config to Engine/Saved/Config, thus making it project-agnostic
-		// 3. Move EditorKeyBindings.ini from Game/Saved/Config to Engine/Saved/Config, thus making it project-agnostic
+		// 1. Move EditorLayout.ini from Game/Saved/Config to Engine/Saved/Config, thus making it project-agnostic
+		// 2. Move EditorKeyBindings.ini from Game/Saved/Config to Engine/Saved/Config, thus making it project-agnostic
 
-		MigrateToAgnosticIni(TEXT("Editor.ini"));
 		MigrateToAgnosticIni(TEXT("EditorLayout.ini"));
 		MigrateToAgnosticIni(TEXT("EditorKeyBindings.ini"));
 		
