Linux: enable symbol serialization for memory profiler.

diff --git a/Engine/Source/Runtime/Core/Public/GenericPlatform/GenericPlatformStackWalk.h b/Engine/Source/Runtime/Core/Public/GenericPlatform/GenericPlatformStackWalk.h
--- a/Engine/Source/Runtime/Core/Public/GenericPlatform/GenericPlatformStackWalk.h
+++ b/Engine/Source/Runtime/Core/Public/GenericPlatform/GenericPlatformStackWalk.h
@@ -178,7 +178,7 @@ struct CORE_API FGenericPlatformStackWalk
 	 *
 	 * @return	The number of modules copied into ModuleSignatures
 	 */
-	FORCEINLINE int32 GetProcessModuleSignatures(FStackWalkModuleInfo *ModuleSignatures, const int32 ModuleSignaturesSize)
+	FORCEINLINE static int32 GetProcessModuleSignatures(FStackWalkModuleInfo *ModuleSignatures, const int32 ModuleSignaturesSize)
 	{
 		return 0;
 	}
diff --git a/Engine/Source/Runtime/Core/Public/Linux/LinuxPlatformStackWalk.h b/Engine/Source/Runtime/Core/Public/Linux/LinuxPlatformStackWalk.h
--- a/Engine/Source/Runtime/Core/Public/Linux/LinuxPlatformStackWalk.h
+++ b/Engine/Source/Runtime/Core/Public/Linux/LinuxPlatformStackWalk.h
@@ -8,6 +8,9 @@
 #pragma once
 #include "GenericPlatform/GenericPlatformStackWalk.h"
 
+#undef PLATFORM_SUPPORTS_STACK_SYMBOLS
+#define PLATFORM_SUPPORTS_STACK_SYMBOLS 1
+
 struct CORE_API FLinuxPlatformStackWalk : public FGenericPlatformStackWalk
 {
 	static void ProgramCounterToSymbolInfo( uint64 ProgramCounter, FProgramCounterSymbolInfo& out_SymbolInfo );
