Added TArray::Emplace implementations for 5 and 6 parameters, for compilers which don't yet support variadic templates.

diff --git a/Engine/Source/Runtime/Core/Public/Containers/Array.h b/Engine/Source/Runtime/Core/Public/Containers/Array.h
--- a/Engine/Source/Runtime/Core/Public/Containers/Array.h
+++ b/Engine/Source/Runtime/Core/Public/Containers/Array.h
@@ -1779,7 +1779,23 @@ public:
 			return Index;
 		}
 
-	#endif
+		template <typename Arg0Type, typename Arg1Type, typename Arg2Type, typename Arg3Type, typename Arg4Type>
+		int32 Emplace(Arg0Type&& Arg0, Arg1Type&& Arg1, Arg2Type&& Arg2, Arg3Type&& Arg3, Arg4Type&& Arg4)
+		{
+			const int32 Index = AddUninitialized(1);
+			new(GetData() + Index) ElementType(Forward<Arg0Type>(Arg0), Forward<Arg1Type>(Arg1), Forward<Arg2Type>(Arg2), Forward<Arg3Type>(Arg3), Forward<Arg4Type>(Arg4));
+			return Index;
+		}
+
+		template <typename Arg0Type, typename Arg1Type, typename Arg2Type, typename Arg3Type, typename Arg4Type, typename Arg5Type>
+		int32 Emplace(Arg0Type&& Arg0, Arg1Type&& Arg1, Arg2Type&& Arg2, Arg3Type&& Arg3, Arg4Type&& Arg4, Arg5Type&& Arg5)
+		{
+			const int32 Index = AddUninitialized(1);
+			new(GetData() + Index) ElementType(Forward<Arg0Type>(Arg0), Forward<Arg1Type>(Arg1), Forward<Arg2Type>(Arg2), Forward<Arg3Type>(Arg3), Forward<Arg4Type>(Arg4), Forward<Arg5Type>(Arg5));
+			return Index;
+		}
+
+#endif
 
 	/**
 	 * Adds a new item to the end of the array, possibly reallocating the whole array to fit.
