Fixed some editor preferences not being upgraded correctly
This addresses UE-14374.

diff --git a/Engine/Config/BaseEngine.ini b/Engine/Config/BaseEngine.ini
--- a/Engine/Config/BaseEngine.ini
+++ b/Engine/Config/BaseEngine.ini
@@ -472,6 +472,7 @@ bUseStreamingPause=false
 
 ; We renamed EditorGameAgnosticSettings to be less imposing
 +ActiveClassRedirects=(OldClassName="EditorGameAgnosticSettings",NewClassName="/Script/UnrealEd.EditorSettings")
++ActiveClassRedirects=(OldClassName="EditorUserSettings",NewClassName="/Script/UnrealEd.EditorPerProjectUserSettings")
 
 +K2FieldRedirects=(OldFieldName="BTTask_MoveDirectlyToward.bForceMoveToLocation",NewFieldName="bDisablePathUpdateOnGoalLocationChange")
 +K2FieldRedirects=(OldFieldName="BTFunctionLibrary.GetBlackboard", NewFieldName="BTFunctionLibrary.GetOwnersBlackboard")
diff --git a/Engine/Source/Runtime/Core/Private/Misc/ConfigManifest.cpp b/Engine/Source/Runtime/Core/Private/Misc/ConfigManifest.cpp
--- a/Engine/Source/Runtime/Core/Private/Misc/ConfigManifest.cpp
+++ b/Engine/Source/Runtime/Core/Private/Misc/ConfigManifest.cpp
@@ -190,8 +190,11 @@ void FConfigManifest::MigrateEditorUserSettings()
 	EditorUserSettingsConfig.NoSave = true;
 	FConfigCacheIni::LoadLocalIniFile(EditorUserSettingsConfig, TEXT("EditorUserSettings"), true);
 	
+	// Rename the config section
+	MigrateConfigSection(EditorUserSettingsConfig, TEXT("/Script/UnrealEd.EditorUserSettings"), TEXT("/Script/UnrealEd.EditorPerProjectUserSettings"));
+
 	FConfigFile EditorPerProjectUserSettingsConfig;
-	FConfigCacheIni::LoadLocalIniFile(EditorPerProjectUserSettingsConfig, TEXT("EditorPerProjectUserSettings"), true);
+	FConfigCacheIni::LoadLocalIniFile(EditorPerProjectUserSettingsConfig, TEXT("EditorPerProjectUserSettings"), false);
 
 	EditorPerProjectUserSettingsConfig.AddMissingProperties(EditorUserSettingsConfig);
 	if (EditorPerProjectUserSettingsConfig.Write(ProjectSpecificIniPath(TEXT("EditorPerProjectUserSettings.ini")), false))
@@ -207,7 +210,13 @@ EConfigManifestVersion FConfigManifest::UpgradeFromVersion(EConfigManifestVersio
 	if (FromVersion < EConfigManifestVersion::RenameEditorAgnosticSettings)
 	{
 		// First off, rename the Editor game agnostic ini config to EditorSettings
-		RenameIni(*ProjectAgnosticIniPath(TEXT("EditorGameAgnostic.ini")), *ProjectAgnosticIniPath(TEXT("EditorSettings.ini")));
+		auto Path = ProjectAgnosticIniPath(TEXT("EditorSettings.ini"));
+		RenameIni(*ProjectAgnosticIniPath(TEXT("EditorGameAgnostic.ini")), *Path);
+
+		FConfigFile EditorSettings;
+		EditorSettings.Read(Path);
+		MigrateConfigSection(EditorSettings, TEXT("/Script/UnrealEd.EditorGameAgnosticSettings"), TEXT("/Script/UnrealEd.EditorSettings"));
+		EditorSettings.Write(Path, false /*bDoRemoteWrite*/);
 
 		FromVersion = EConfigManifestVersion::RenameEditorAgnosticSettings;
 	}
@@ -232,4 +241,31 @@ EConfigManifestVersion FConfigManifest::UpgradeFromVersion(EConfigManifestVersio
 	}
 
 	return FromVersion;
+}
+
+void FConfigManifest::MigrateConfigSection(FConfigFile& ConfigFile, const TCHAR* OldSectionName, const TCHAR* NewSectionName)
+{
+	const FConfigSection* OldSection = ConfigFile.Find(OldSectionName);
+	if (OldSection)
+	{
+		FConfigSection* NewSection = ConfigFile.Find(NewSectionName);
+		if (NewSection)
+		{
+			for (auto& Setting : *OldSection)
+			{
+				if (!NewSection->Contains(Setting.Key))
+				{
+					NewSection->Add(Setting.Key, Setting.Value);
+				}
+			}
+		}
+		else
+		{
+			// Add the new section and remove the old
+			FConfigSection SectionCopy = *OldSection;
+			ConfigFile.Add(NewSectionName, MoveTemp(SectionCopy));
+			ConfigFile.Remove(OldSectionName);
+		}
+		ConfigFile.Dirty = true;
+	}
 }
\ No newline at end of file
diff --git a/Engine/Source/Runtime/Core/Private/Misc/ConfigManifest.h b/Engine/Source/Runtime/Core/Private/Misc/ConfigManifest.h
--- a/Engine/Source/Runtime/Core/Private/Misc/ConfigManifest.h
+++ b/Engine/Source/Runtime/Core/Private/Misc/ConfigManifest.h
@@ -16,6 +16,9 @@ public:
 	/** Migrate what was previously EditorUserSettings.ini to EditorPerProjectUserSettings.ini, if the former exists. */
 	static void MigrateEditorUserSettings();
 
+	/** Migrate a config section to a new section, only overwriting entries that don't exist in the new section. */
+	static void MigrateConfigSection(FConfigFile& ConfigFile, const TCHAR* OldSection, const TCHAR* NewSection);
+
 private:
 
 	/** Upgrade the config from the specified version, returning the current version after the upgrade (may be the same as FromVersion if upgrade was not possible). */
