Fix to FInterpCurve serialization, and fixed overly strict check in Eval, EvalDerivative and EvalSecondDerivative.

diff --git a/Engine/Source/Runtime/Core/Public/Math/InterpCurve.h b/Engine/Source/Runtime/Core/Public/Math/InterpCurve.h
--- a/Engine/Source/Runtime/Core/Public/Math/InterpCurve.h
+++ b/Engine/Source/Runtime/Core/Public/Math/InterpCurve.h
@@ -93,7 +93,7 @@ public:
 	float InaccurateFindNearestOnSegment( const T &PointInSpace, int32 PtIdx, float& OutSquaredDistance ) const;
 
 	/** Automatically set the tangents on the curve based on surrounding points */
-	void AutoSetTangents(float Tension = 0.f, bool bStationaryEndpoints = true);
+	void AutoSetTangents(float Tension = 0.0f, bool bStationaryEndpoints = true);
 
 	/** Calculate the min/max out value that can be returned by this InterpCurve. */
 	void CalcBounds(T& OutMin, T& OutMax, const T& Default = T(ForceInit)) const;
@@ -114,8 +114,11 @@ public:
 		//   as inline struct properties in UnClass.cpp!
 
 		Ar << Curve.Points;
-		Ar << Curve.bIsLooped;
-		Ar << Curve.LoopKeyOffset;
+		if (Ar.UE4Ver() >= VER_UE4_INTERPCURVE_SUPPORTS_LOOPING)
+		{
+			Ar << Curve.bIsLooped;
+			Ar << Curve.LoopKeyOffset;
+		}
 
 		return Ar;
 	}
@@ -271,17 +274,17 @@ T FInterpCurve<T>::Eval(const float InVal, const T& Default) const
 	// Somewhere within curve range - interpolate.
 	check(Index >= 0 && ((bIsLooped && Index < NumPoints) || (!bIsLooped && Index < LastPoint)));
 	const bool bLoopSegment = (bIsLooped && Index == LastPoint);
-	const int32 NextIndex = bLoopSegment ? 0 : Index + 1;
+	const int32 NextIndex = bLoopSegment ? 0 : (Index + 1);
 
 	const auto& PrevPoint = Points[Index];
 	const auto& NextPoint = Points[NextIndex];
 
-	const float Diff = bLoopSegment ? LoopKeyOffset : NextPoint.InVal - PrevPoint.InVal;
+	const float Diff = bLoopSegment ? LoopKeyOffset : (NextPoint.InVal - PrevPoint.InVal);
 
 	if (Diff > 0.0f && PrevPoint.InterpMode != CIM_Constant)
 	{
 		const float Alpha = (InVal - PrevPoint.InVal) / Diff;
-		check(Alpha >= 0.0f && Alpha < 1.0f);
+		check(Alpha >= 0.0f && Alpha <= 1.0f);
 
 		if (PrevPoint.InterpMode == CIM_Linear)
 		{
@@ -337,12 +340,12 @@ T FInterpCurve<T>::EvalDerivative(const float InVal, const T& Default) const
 	// Somewhere within curve range - interpolate.
 	check(Index >= 0 && ((bIsLooped && Index < NumPoints) || (!bIsLooped && Index < LastPoint)));
 	const bool bLoopSegment = (bIsLooped && Index == LastPoint);
-	const int32 NextIndex = bLoopSegment ? 0 : Index + 1;
+	const int32 NextIndex = bLoopSegment ? 0 : (Index + 1);
 
 	const auto& PrevPoint = Points[Index];
 	const auto& NextPoint = Points[NextIndex];
 
-	const float Diff = bLoopSegment ? LoopKeyOffset : NextPoint.InVal - PrevPoint.InVal;
+	const float Diff = bLoopSegment ? LoopKeyOffset : (NextPoint.InVal - PrevPoint.InVal);
 
 	if (Diff > 0.0f && PrevPoint.InterpMode != CIM_Constant)
 	{
@@ -353,7 +356,7 @@ T FInterpCurve<T>::EvalDerivative(const float InVal, const T& Default) const
 		else
 		{
 			const float Alpha = (InVal - PrevPoint.InVal) / Diff;
-			check(Alpha >= 0.0f && Alpha < 1.0f);
+			check(Alpha >= 0.0f && Alpha <= 1.0f);
 
 			return FMath::CubicInterpDerivative(PrevPoint.OutVal, PrevPoint.LeaveTangent * Diff, NextPoint.OutVal, NextPoint.ArriveTangent * Diff, Alpha) / Diff;
 		}
@@ -399,12 +402,12 @@ T FInterpCurve<T>::EvalSecondDerivative(const float InVal, const T& Default) con
 	// Somewhere within curve range - interpolate.
 	check(Index >= 0 && ((bIsLooped && Index < NumPoints) || (!bIsLooped && Index < LastPoint)));
 	const bool bLoopSegment = (bIsLooped && Index == LastPoint);
-	const int32 NextIndex = bLoopSegment ? 0 : Index + 1;
+	const int32 NextIndex = bLoopSegment ? 0 : (Index + 1);
 
 	const auto& PrevPoint = Points[Index];
 	const auto& NextPoint = Points[NextIndex];
 
-	const float Diff = bLoopSegment ? LoopKeyOffset : NextPoint.InVal - PrevPoint.InVal;
+	const float Diff = bLoopSegment ? LoopKeyOffset : (NextPoint.InVal - PrevPoint.InVal);
 
 	if (Diff > 0.0f && PrevPoint.InterpMode != CIM_Constant)
 	{
@@ -416,7 +419,7 @@ T FInterpCurve<T>::EvalSecondDerivative(const float InVal, const T& Default) con
 		else
 		{
 			const float Alpha = (InVal - PrevPoint.InVal) / Diff;
-			check(Alpha >= 0.0f && Alpha < 1.0f);
+			check(Alpha >= 0.0f && Alpha <= 1.0f);
 
 			return FMath::CubicInterpSecondDerivative(PrevPoint.OutVal, PrevPoint.LeaveTangent * Diff, NextPoint.OutVal, NextPoint.ArriveTangent * Diff, Alpha) / (Diff * Diff);
 		}
@@ -468,7 +471,7 @@ float FInterpCurve<T>::InaccurateFindNearestOnSegment(const T& PointInSpace, int
 {
 	const int32 NumPoints = Points.Num();
 	const int32 LastPoint = NumPoints - 1;
-	const int32 NextPtIdx = (bIsLooped && PtIdx == LastPoint) ? 0 : PtIdx + 1;
+	const int32 NextPtIdx = (bIsLooped && PtIdx == LastPoint) ? 0 : (PtIdx + 1);
 	check(PtIdx >= 0 && ((bIsLooped && PtIdx < NumPoints) || (!bIsLooped && PtIdx < LastPoint)));
 
 	if (CIM_Constant == Points[PtIdx].InterpMode)
@@ -558,8 +561,8 @@ void FInterpCurve<T>::AutoSetTangents(float Tension, bool bStationaryEndpoints)
 	// Iterate over all points in this InterpCurve
 	for (int32 PointIndex = 0; PointIndex < NumPoints; PointIndex++)
 	{
-		const int32 PrevIndex = (PointIndex == 0) ? (bIsLooped ? LastPoint : 0) : PointIndex - 1;
-		const int32 NextIndex = (PointIndex == LastPoint) ? (bIsLooped ? 0 : LastPoint) : PointIndex + 1;
+		const int32 PrevIndex = (PointIndex == 0) ? (bIsLooped ? LastPoint : 0) : (PointIndex - 1);
+		const int32 NextIndex = (PointIndex == LastPoint) ? (bIsLooped ? 0 : LastPoint) : (PointIndex + 1);
 
 		auto& ThisPoint = Points[PointIndex];
 		const auto& PrevPoint = Points[PrevIndex];
@@ -578,8 +581,8 @@ void FInterpCurve<T>::AutoSetTangents(float Tension, bool bStationaryEndpoints)
 				const bool bWantClamping = (ThisPoint.InterpMode == CIM_CurveAutoClamped);
 				T Tangent;
 
-				const float PrevTime = (bIsLooped && PointIndex == 0) ? ThisPoint.InVal - LoopKeyOffset : PrevPoint.InVal;
-				const float NextTime = (bIsLooped && PointIndex == LastPoint) ? ThisPoint.InVal + LoopKeyOffset : NextPoint.InVal;
+				const float PrevTime = (bIsLooped && PointIndex == 0) ? (ThisPoint.InVal - LoopKeyOffset) : PrevPoint.InVal;
+				const float NextTime = (bIsLooped && PointIndex == LastPoint) ? (ThisPoint.InVal + LoopKeyOffset) : NextPoint.InVal;
 
 				ComputeCurveTangent(
 					PrevTime,			// Previous time
@@ -637,11 +640,11 @@ void FInterpCurve<T>::CalcBounds(T& OutMin, T& OutMax, const T& Default) const
 		OutMin = Points[0].OutVal;
 		OutMax = Points[0].OutVal;
 
-		const int32 NumSegments = bIsLooped ? NumPoints : NumPoints - 1;
+		const int32 NumSegments = bIsLooped ? NumPoints : (NumPoints - 1);
 
 		for (int32 Index = 0; Index < NumSegments; Index++)
 		{
-			const int32 NextIndex = (Index == NumPoints - 1) ? 0 : Index + 1;
+			const int32 NextIndex = (Index == NumPoints - 1) ? 0 : (Index + 1);
 			CurveFindIntervalBounds(Points[Index], Points[NextIndex], OutMin, OutMax, 0.0f);
 		}
 	}
diff --git a/Engine/Source/Runtime/Core/Public/UObject/ObjectVersion.h b/Engine/Source/Runtime/Core/Public/UObject/ObjectVersion.h
--- a/Engine/Source/Runtime/Core/Public/UObject/ObjectVersion.h
+++ b/Engine/Source/Runtime/Core/Public/UObject/ObjectVersion.h
@@ -499,7 +499,7 @@ enum EUnrealEngineObjectUE4Version
 	// Added fastcall data serialization directly in UFunction
 	VER_UE4_SERIALIZE_BLUEPRINT_EVENTGRAPH_FASTCALLS_IN_UFUNCTION,
 	// Changes to USplineComponent and FInterpCurve
-	VER_UE4_SPLINE_COMPONENT_WITH_ROTATION_AND_SCALE,
+	VER_UE4_INTERPCURVE_SUPPORTS_LOOPING,
 	// Material Instances overriding base material LOD transitions
 	VER_UE4_MATERIAL_INSTANCE_BASE_PROPERTY_OVERRIDES_DITHERED_LOD_TRANSITION,
 
diff --git a/Engine/Source/Runtime/Engine/Private/Components/SplineComponent.cpp b/Engine/Source/Runtime/Engine/Private/Components/SplineComponent.cpp
--- a/Engine/Source/Runtime/Engine/Private/Components/SplineComponent.cpp
+++ b/Engine/Source/Runtime/Engine/Private/Components/SplineComponent.cpp
@@ -77,7 +77,7 @@ void USplineComponent::Serialize(FArchive& Ar)
 
 	// Support old resources which don't have the rotation and scale splines present
 	const int32 ArchiveUE4Version = Ar.UE4Ver();
-	if (ArchiveUE4Version < VER_UE4_SPLINE_COMPONENT_WITH_ROTATION_AND_SCALE)
+	if (ArchiveUE4Version < VER_UE4_INTERPCURVE_SUPPORTS_LOOPING)
 	{
 		int32 NumPoints = SplineInfo.Points.Num();
 
