Fixed FMath::RandRange() returning values > Max

diff --git a/Engine/Source/Runtime/Core/Public/Math/UnrealMathUtility.h b/Engine/Source/Runtime/Core/Public/Math/UnrealMathUtility.h
--- a/Engine/Source/Runtime/Core/Public/Math/UnrealMathUtility.h
+++ b/Engine/Source/Runtime/Core/Public/Math/UnrealMathUtility.h
@@ -100,7 +100,7 @@ struct FMath : public FPlatformMath
 	/** Helper function for rand implementations. Returns a random number >= Min and <= Max */
 	static FORCEINLINE int32 RandRange(int32 Min, int32 Max)
 	{
-		const int32 Range = (Max - Min) + 1;
+		const int32 Range = Max - Min;
 		return Min + RandHelper(Range);
 	}
 
