Core - Potential fix for a race condition, commented a check, will investigate later

diff --git a/Engine/Source/Runtime/Core/Private/HAL/ThreadingBase.cpp b/Engine/Source/Runtime/Core/Private/HAL/ThreadingBase.cpp
--- a/Engine/Source/Runtime/Core/Private/HAL/ThreadingBase.cpp
+++ b/Engine/Source/Runtime/Core/Private/HAL/ThreadingBase.cpp
@@ -263,7 +263,7 @@ FRunnableThread* FRunnableThread::Create(
 void FRunnableThread::SetTls()
 {
 	// Make sure it's called from the owning thread.
-	check( ThreadID == FPlatformTLS::GetCurrentThreadId() );
+	//check( ThreadID == FPlatformTLS::GetCurrentThreadId() );
 	check( RunnableTlsSlot );
 	FPlatformTLS::SetTlsValue( RunnableTlsSlot, this );
 }
