Core - Removed X1 related code from Windows

diff --git a/Engine/Source/Runtime/Core/Private/Windows/WindowsRunnableThread.cpp b/Engine/Source/Runtime/Core/Private/Windows/WindowsRunnableThread.cpp
--- a/Engine/Source/Runtime/Core/Private/Windows/WindowsRunnableThread.cpp
+++ b/Engine/Source/Runtime/Core/Private/Windows/WindowsRunnableThread.cpp
@@ -14,10 +14,6 @@ uint32 FRunnableThreadWin::GuardedRun()
 
 	FPlatformProcess::SetThreadAffinityMask(ThreadAffinityMask);
 
-#if PLATFORM_XBOXONE
-	UE_LOG(LogThreadingWindows, Log, TEXT("Runnable thread %s is on Process %d."), *ThreadName  , static_cast<uint32>(::GetCurrentProcessorNumber()) );
-#endif
-
 #if !PLATFORM_SEH_EXCEPTIONS_DISABLED
 	if( !FPlatformMisc::IsDebuggerPresent() || GAlwaysReportCrash )
 	{
