#ue4 - Fix "nearly equal" comparisons not returning true when values are exactly equal and tolerance is zero.

- Comparisons using a tolerance should use "Delta <= Tolerance".

diff --git a/Engine/Source/Runtime/Core/Public/Math/Quat.h b/Engine/Source/Runtime/Core/Public/Math/Quat.h
--- a/Engine/Source/Runtime/Core/Public/Math/Quat.h
+++ b/Engine/Source/Runtime/Core/Public/Math/Quat.h
@@ -586,9 +586,9 @@ inline FMatrix FQuat::operator*( const FMatrix& M ) const
 FORCEINLINE bool DebugRotatorEquals(const FRotator& R1, const FRotator& R2, float Tolerance=1.f) 
 {
 	// also 0 and 360 should be considered same 
-	return ( FMath::Abs(R1.Pitch-R2.Pitch) < Tolerance || FMath::Abs(R1.Pitch+R2.Pitch-360) < Tolerance )
-		&& ( FMath::Abs(R1.Yaw-R2.Yaw) < Tolerance || FMath::Abs(R1.Yaw+R2.Yaw-360) < Tolerance )
-		&& ( FMath::Abs(R1.Roll-R2.Roll) < Tolerance || FMath::Abs(R1.Roll+R2.Roll-360) < Tolerance );
+	return ( FMath::Abs(R1.Pitch-R2.Pitch) <= Tolerance || FMath::Abs(R1.Pitch+R2.Pitch-360) <= Tolerance )
+		&& ( FMath::Abs(R1.Yaw-R2.Yaw) <= Tolerance || FMath::Abs(R1.Yaw+R2.Yaw-360) <= Tolerance )
+		&& ( FMath::Abs(R1.Roll-R2.Roll) <= Tolerance || FMath::Abs(R1.Roll+R2.Roll-360) <= Tolerance );
 }
 #endif
 
@@ -681,8 +681,8 @@ FORCEINLINE FQuat FQuat::operator-( const FQuat& Q ) const
 
 FORCEINLINE bool FQuat::Equals(const FQuat& Q, float Tolerance) const
 {
-	return (FMath::Abs(X - Q.X) < Tolerance && FMath::Abs(Y - Q.Y) < Tolerance && FMath::Abs(Z - Q.Z) < Tolerance && FMath::Abs(W - Q.W) < Tolerance)
-		|| (FMath::Abs(X + Q.X) < Tolerance && FMath::Abs(Y + Q.Y) < Tolerance && FMath::Abs(Z + Q.Z) < Tolerance && FMath::Abs(W + Q.W) < Tolerance);
+	return (FMath::Abs(X - Q.X) <= Tolerance && FMath::Abs(Y - Q.Y) <= Tolerance && FMath::Abs(Z - Q.Z) <= Tolerance && FMath::Abs(W - Q.W) <= Tolerance)
+		|| (FMath::Abs(X + Q.X) <= Tolerance && FMath::Abs(Y + Q.Y) <= Tolerance && FMath::Abs(Z + Q.Z) <= Tolerance && FMath::Abs(W + Q.W) <= Tolerance);
 }
 
 
diff --git a/Engine/Source/Runtime/Core/Public/Math/Rotator.h b/Engine/Source/Runtime/Core/Public/Math/Rotator.h
--- a/Engine/Source/Runtime/Core/Public/Math/Rotator.h
+++ b/Engine/Source/Runtime/Core/Public/Math/Rotator.h
@@ -481,9 +481,9 @@ FORCEINLINE FRotator FRotator::operator-=( const FRotator &R )
 FORCEINLINE bool FRotator::IsNearlyZero(float Tolerance) const
 {
 	return
-		FMath::Abs(NormalizeAxis(Pitch))<Tolerance
-		&&	FMath::Abs(NormalizeAxis(Yaw))<Tolerance
-		&&	FMath::Abs(NormalizeAxis(Roll))<Tolerance;
+		FMath::Abs(NormalizeAxis(Pitch))<=Tolerance
+		&&	FMath::Abs(NormalizeAxis(Yaw))<=Tolerance
+		&&	FMath::Abs(NormalizeAxis(Roll))<=Tolerance;
 }
 
 
diff --git a/Engine/Source/Runtime/Core/Public/Math/Sphere.h b/Engine/Source/Runtime/Core/Public/Math/Sphere.h
--- a/Engine/Source/Runtime/Core/Public/Math/Sphere.h
+++ b/Engine/Source/Runtime/Core/Public/Math/Sphere.h
@@ -71,7 +71,7 @@ public:
 	 */
 	bool Equals( const FSphere& Sphere, float Tolerance = KINDA_SMALL_NUMBER ) const
 	{
-		return Center.Equals(Sphere.Center, Tolerance) && FMath::Abs(W - Sphere.W) < Tolerance;
+		return Center.Equals(Sphere.Center, Tolerance) && FMath::Abs(W - Sphere.W) <= Tolerance;
 	}
 
 	/**
diff --git a/Engine/Source/Runtime/Core/Public/Math/Transform.h b/Engine/Source/Runtime/Core/Public/Math/Transform.h
--- a/Engine/Source/Runtime/Core/Public/Math/Transform.h
+++ b/Engine/Source/Runtime/Core/Public/Math/Transform.h
@@ -612,7 +612,7 @@ public:
 
 	inline bool Equals(const FTransform& Other, float Tolerance=KINDA_SMALL_NUMBER) const
 	{
-		return Rotation.Equals(Other.Rotation, Tolerance) && Translation.Equals(Other.Translation, Tolerance) && ((Scale3D-Other.Scale3D).Size() < Tolerance);
+		return Rotation.Equals(Other.Rotation, Tolerance) && Translation.Equals(Other.Translation, Tolerance) && ((Scale3D-Other.Scale3D).SizeSquared() <= FMath::Square(Tolerance));
 	}
 
 	inline bool EqualsNoScale(const FTransform& Other, float Tolerance=KINDA_SMALL_NUMBER) const
diff --git a/Engine/Source/Runtime/Core/Public/Math/UnrealMathUtility.h b/Engine/Source/Runtime/Core/Public/Math/UnrealMathUtility.h
--- a/Engine/Source/Runtime/Core/Public/Math/UnrealMathUtility.h
+++ b/Engine/Source/Runtime/Core/Public/Math/UnrealMathUtility.h
@@ -170,7 +170,7 @@ struct FMath : public FPlatformMath
 	 */
 	static FORCEINLINE bool IsNearlyEqual(float A, float B, float ErrorTolerance = SMALL_NUMBER)
 	{
-		return Abs<float>( A - B ) < ErrorTolerance;
+		return Abs<float>( A - B ) <= ErrorTolerance;
 	}
 
 	/**
@@ -182,7 +182,7 @@ struct FMath : public FPlatformMath
 	 */
 	static FORCEINLINE bool IsNearlyEqual(double A, double B, double ErrorTolerance = SMALL_NUMBER)
 	{
-		return Abs<double>( A - B ) < ErrorTolerance;
+		return Abs<double>( A - B ) <= ErrorTolerance;
 	}
 
 	/**
@@ -193,7 +193,7 @@ struct FMath : public FPlatformMath
 	 */
 	static FORCEINLINE bool IsNearlyZero(float Value, float ErrorTolerance = SMALL_NUMBER)
 	{
-		return Abs<float>( Value ) < ErrorTolerance;
+		return Abs<float>( Value ) <= ErrorTolerance;
 	}
 
 	/**
@@ -204,7 +204,7 @@ struct FMath : public FPlatformMath
 	 */
 	static FORCEINLINE bool IsNearlyZero(double Value, double ErrorTolerance = SMALL_NUMBER)
 	{
-		return Abs<double>( Value ) < ErrorTolerance;
+		return Abs<double>( Value ) <= ErrorTolerance;
 	}
 
 	/**
diff --git a/Engine/Source/Runtime/Core/Public/Math/Vector.h b/Engine/Source/Runtime/Core/Public/Math/Vector.h
--- a/Engine/Source/Runtime/Core/Public/Math/Vector.h
+++ b/Engine/Source/Runtime/Core/Public/Math/Vector.h
@@ -1276,7 +1276,7 @@ FORCEINLINE bool FVector::Equals(const FVector& V, float Tolerance) const
 
 FORCEINLINE bool FVector::AllComponentsEqual(float Tolerance) const
 {
-	return FMath::Abs( X - Y ) < Tolerance && FMath::Abs( X - Z ) < Tolerance && FMath::Abs( Y - Z ) < Tolerance;
+	return FMath::Abs( X - Y ) <= Tolerance && FMath::Abs( X - Z ) <= Tolerance && FMath::Abs( Y - Z ) <= Tolerance;
 }
 
 
@@ -1429,9 +1429,9 @@ FORCEINLINE float FVector::SizeSquared2D() const
 FORCEINLINE bool FVector::IsNearlyZero(float Tolerance) const
 {
 	return
-		FMath::Abs(X)<Tolerance
-		&&	FMath::Abs(Y)<Tolerance
-		&&	FMath::Abs(Z)<Tolerance;
+		FMath::Abs(X)<=Tolerance
+		&&	FMath::Abs(Y)<=Tolerance
+		&&	FMath::Abs(Z)<=Tolerance;
 }
 
 FORCEINLINE bool FVector::IsZero() const
@@ -1641,7 +1641,7 @@ FORCEINLINE FVector FVector::Reciprocal() const
 
 FORCEINLINE bool FVector::IsUniform(float Tolerance) const
 {
-	return (FMath::Abs(X-Y) < Tolerance) && (FMath::Abs(Y-Z) < Tolerance);
+	return (FMath::Abs(X-Y) <= Tolerance) && (FMath::Abs(Y-Z) <= Tolerance);
 }
 
 FORCEINLINE FVector FVector::MirrorByVector( const FVector& MirrorNormal ) const
diff --git a/Engine/Source/Runtime/Core/Public/Math/Vector2D.h b/Engine/Source/Runtime/Core/Public/Math/Vector2D.h
--- a/Engine/Source/Runtime/Core/Public/Math/Vector2D.h
+++ b/Engine/Source/Runtime/Core/Public/Math/Vector2D.h
@@ -645,7 +645,7 @@ FORCEINLINE bool FVector2D::operator>=( const FVector2D& Other ) const
 
 FORCEINLINE bool FVector2D::Equals(const FVector2D& V, float Tolerance) const
 {
-	return FMath::Abs(X-V.X) < Tolerance && FMath::Abs(Y-V.Y) < Tolerance;
+	return FMath::Abs(X-V.X) <= Tolerance && FMath::Abs(Y-V.Y) <= Tolerance;
 }
 
 
@@ -797,8 +797,8 @@ FORCEINLINE void FVector2D::Normalize(float Tolerance)
 
 FORCEINLINE bool FVector2D::IsNearlyZero(float Tolerance) const
 {
-	return	FMath::Abs(X)<Tolerance
-		&&	FMath::Abs(Y)<Tolerance;
+	return	FMath::Abs(X)<=Tolerance
+		&&	FMath::Abs(Y)<=Tolerance;
 }
 
 
diff --git a/Engine/Source/Runtime/Core/Public/Math/Vector4.h b/Engine/Source/Runtime/Core/Public/Math/Vector4.h
--- a/Engine/Source/Runtime/Core/Public/Math/Vector4.h
+++ b/Engine/Source/Runtime/Core/Public/Math/Vector4.h
@@ -537,7 +537,7 @@ FORCEINLINE bool FVector4::operator!=( const FVector4& V ) const
 
 FORCEINLINE bool FVector4::Equals( const FVector4& V, float Tolerance ) const
 {
-	return FMath::Abs(X-V.X) < Tolerance && FMath::Abs(Y-V.Y) < Tolerance && FMath::Abs(Z-V.Z) < Tolerance && FMath::Abs(W-V.W) < Tolerance;
+	return FMath::Abs(X-V.X) <= Tolerance && FMath::Abs(Y-V.Y) <= Tolerance && FMath::Abs(Z-V.Z) <= Tolerance && FMath::Abs(W-V.W) <= Tolerance;
 }
 
 
@@ -623,9 +623,9 @@ FORCEINLINE bool FVector4::ContainsNaN() const
 FORCEINLINE bool FVector4::IsNearlyZero3( float Tolerance ) const
 {
 	return
-			FMath::Abs(X)<Tolerance
-		&&	FMath::Abs(Y)<Tolerance
-		&&	FMath::Abs(Z)<Tolerance;
+			FMath::Abs(X)<=Tolerance
+		&&	FMath::Abs(Y)<=Tolerance
+		&&	FMath::Abs(Z)<=Tolerance;
 }
 
 
