Changed the Windows file reader to use non-blocking i/o to speed up reading of file data (~1.6x faster)

diff --git a/Engine/Source/Runtime/Core/Private/Windows/WindowsPlatformFile.cpp b/Engine/Source/Runtime/Core/Private/Windows/WindowsPlatformFile.cpp
--- a/Engine/Source/Runtime/Core/Private/Windows/WindowsPlatformFile.cpp
+++ b/Engine/Source/Runtime/Core/Private/Windows/WindowsPlatformFile.cpp
@@ -14,6 +14,323 @@ const FDateTime WindowsEpoch(1970, 1, 1);
 	}
 #include "HideWindowsPlatformTypes.h"
 
+/**
+ * This file reader uses overlapped i/o and double buffering to asynchronously read from files
+ */
+class CORE_API FAsyncBufferedFileReaderWindows : public IFileHandle
+{
+protected:
+	enum {DEFAULT_BUFFER_SIZE = 64 * 1024};
+	/**
+	 * The file handle to operate on
+	 */
+	HANDLE Handle;
+	/**
+	 * The size of the file that is being read
+	 */
+	int64 FileSize;
+	/**
+	 * Overall position in the file and buffers combined
+	 */
+	int64 FilePos;
+	/**
+	 * Overall position in the file as the OverlappedIO struct understands it
+	 */
+	uint64 OverlappedFilePos;
+	/**
+	 * These are the two buffers used for reading the file asynchronously
+	 */
+	int8* Buffers[2];
+	/**
+	 * The size of the buffers in bytes
+	 */
+	const int32 BufferSize;
+	/**
+	 * The current index of the buffer that we are serializing from
+	 */
+	int32 SerializeBuffer;
+	/**
+	 * The current index of the streaming buffer for async reading into
+	 */
+	int32 StreamBuffer;
+	/**
+	 * Where we are in the serialize buffer
+	 */
+	int32 SerializePos;
+	/**
+	 * Tracks which buffer has the async read outstanding (0 = first read after create/seek, 1 = streaming buffer)
+	 */
+	int32 CurrentAsyncReadBuffer;
+	/**
+	 * The overlapped IO struct to use for determining async state
+	 */
+	OVERLAPPED OverlappedIO;
+	/**
+	 * Used to track whether the last read reached the end of the file or not. Reset when a Seek happens
+	 */
+	bool bIsAtEOF;
+	/**
+	 * Whether there's a read outstanding or not
+	 */
+	bool bHasReadOutstanding;
+
+	/**
+	 * Closes the file handle
+	 */
+	bool Close(void)
+	{
+		if (Handle != nullptr)
+		{
+			// Close the file handle
+			CloseHandle(Handle);
+			Handle = nullptr;
+		}
+		return true;
+	}
+
+	/**
+	 * This toggles the buffers we read into & serialize out of between buffer indices 0 & 1
+	 */
+	FORCEINLINE void SwapBuffers()
+	{
+		StreamBuffer ^= 1;
+		SerializeBuffer ^= 1;
+		// We are now at the beginning of the serialize buffer
+		SerializePos = 0;
+	}
+
+	FORCEINLINE void CopyOverlappedPosition()
+	{
+		ULARGE_INTEGER LI;
+		LI.QuadPart = OverlappedFilePos;
+		OverlappedIO.Offset = LI.LowPart;
+		OverlappedIO.OffsetHigh = LI.HighPart;
+	}
+
+	bool WaitForAsyncRead()
+	{
+		// Check for already being at EOF because we won't issue a read
+		if (bIsAtEOF || !bHasReadOutstanding)
+		{
+			return true;
+		}
+		uint32 NumRead = 0;
+		if (GetOverlappedResult(Handle, &OverlappedIO, (::DWORD*)&NumRead, true) != false)
+		{
+			bHasReadOutstanding = false;
+			OverlappedFilePos += NumRead;
+			// Update the overlapped structure or the API will get confused
+			CopyOverlappedPosition();
+			if (OverlappedFilePos >= uint64(FileSize))
+			{
+				bIsAtEOF = true;
+			}
+			return true;
+		}
+		else if (GetLastError() == ERROR_HANDLE_EOF)
+		{
+			bIsAtEOF = true;
+			return true;
+		}
+		return false;
+	}
+
+	void StartAsyncRead(int32 BufferToReadInto)
+	{
+		if (!bIsAtEOF)
+		{
+			bHasReadOutstanding = true;
+			CurrentAsyncReadBuffer = BufferToReadInto;
+			uint32 NumRead = 0;
+			// Now kick off an async read
+			ReadFile(Handle, Buffers[BufferToReadInto], BufferSize, (::DWORD*)&NumRead, &OverlappedIO);
+		}
+	}
+
+	FORCEINLINE void StartStreamBufferRead()
+	{
+		StartAsyncRead(StreamBuffer);
+	}
+
+	FORCEINLINE void StartSerializeBufferRead()
+	{
+		StartAsyncRead(SerializeBuffer);
+	}
+
+	FORCEINLINE bool IsValid()
+	{
+		return Handle != nullptr && Handle != INVALID_HANDLE_VALUE;
+	}
+
+public:
+	FAsyncBufferedFileReaderWindows(HANDLE InHandle, int32 InBufferSize = DEFAULT_BUFFER_SIZE) :
+		Handle(InHandle),
+		FilePos(0),
+		OverlappedFilePos(0),
+		BufferSize(InBufferSize),
+		SerializeBuffer(0),
+		StreamBuffer(1),
+		SerializePos(0),
+		CurrentAsyncReadBuffer(0),
+		bIsAtEOF(false),
+		bHasReadOutstanding(false)
+	{
+		LARGE_INTEGER LI;
+		GetFileSizeEx(Handle, &LI);
+		FileSize = LI.QuadPart;
+		// Allocate our two buffers
+		Buffers[0] = (int8*)FMemory::Malloc(BufferSize);
+		Buffers[1] = (int8*)FMemory::Malloc(BufferSize);
+
+		// Zero the overlapped structure
+		FMemory::Memzero(&OverlappedIO, sizeof(OVERLAPPED));
+
+		// Kick off the first async read
+		StartSerializeBufferRead();
+	}
+
+	virtual ~FAsyncBufferedFileReaderWindows(void)
+	{
+		// Can't free the buffers or close the file if a read is outstanding
+		WaitForAsyncRead();
+		Close();
+		FMemory::Free(Buffers[0]);
+		FMemory::Free(Buffers[1]);
+	}
+
+	virtual bool Seek(int64 InPos) override
+	{
+		check(IsValid());
+		check(InPos >= 0);
+		check(InPos <= FileSize);
+
+		// Determine the change in locations
+		int64 PosDelta = InPos - FilePos;
+		if (PosDelta == 0)
+		{
+			// Same place so no work to do
+			return true;
+		}
+
+		// No matter what, we need to wait for the current async read to finish since we most likely need to issue a new one
+		if (!WaitForAsyncRead())
+		{
+			return false;
+		}
+
+		FilePos = InPos;
+
+		// If the requested location is not within our current serialize buffer, we need to start the whole read process over
+		bool bWithinSerializeBuffer = (PosDelta < 0 && (SerializePos - FMath::Abs(PosDelta) >= 0)) ||
+			(PosDelta > 0 && ((PosDelta + SerializePos) < BufferSize));
+		if (bWithinSerializeBuffer)
+		{
+			// Still within the serialize buffer so just update the position
+			SerializePos += PosDelta;
+		}
+		else
+		{
+			// Reset our EOF tracking and let the read handle setting it if need be
+			bIsAtEOF = false;
+			// Not within our buffer so start a new async read on the serialize buffer
+			OverlappedFilePos = InPos;
+			CopyOverlappedPosition();
+			CurrentAsyncReadBuffer = SerializeBuffer;
+			SerializePos = 0;
+			StartSerializeBufferRead();
+		}
+		return true;
+	}
+
+	virtual bool SeekFromEnd(int64 NewPositionRelativeToEnd = 0) override
+	{
+		check(IsValid());
+		check(NewPositionRelativeToEnd <= 0);
+
+		// Position is negative so this is actually subtracting
+		return Seek(FileSize + NewPositionRelativeToEnd);
+	}
+
+	virtual int64 Tell(void) override
+	{
+		check(IsValid());
+		return FilePos;
+	}
+
+	virtual int64 Size() override
+	{
+		check(IsValid());
+		return FileSize;
+	}
+
+	virtual bool Read(uint8* Dest, int64 BytesToRead) override
+	{
+		check(IsValid());
+		// If zero were requested, quit (some calls like to do zero sized reads)
+		if (BytesToRead <= 0)
+		{
+			return false;
+		}
+
+		if (CurrentAsyncReadBuffer == SerializeBuffer)
+		{
+			// First async read after either construction or a seek
+			if (!WaitForAsyncRead())
+			{
+				return false;
+			}
+			StartStreamBufferRead();
+		}
+
+		check(Dest != nullptr)
+		// While there is data to copy
+		while (BytesToRead > 0)
+		{
+			// Figure out how many bytes we can read from the serialize buffer
+			int64 NumToCopy = FMath::Min<int64>(BytesToRead, BufferSize - SerializePos);
+			if (FilePos + NumToCopy > FileSize)
+			{
+				// Tried to read past the end of the file, so fail
+				return false;
+			}
+			// See if we are at the end of the serialize buffer or not
+			if (NumToCopy > 0)
+			{
+				FMemory::Memcpy(Dest, &Buffers[SerializeBuffer][SerializePos], NumToCopy);
+
+				// Update the internal positions
+				SerializePos += NumToCopy;
+				check(SerializePos <= BufferSize);
+				FilePos += NumToCopy;
+				check(FilePos <= FileSize);
+
+				// Decrement the number of bytes we copied
+				BytesToRead -= NumToCopy;
+
+				// Now offset the dest pointer with the chunk we copied
+				Dest = (uint8*)Dest + NumToCopy;
+			}
+			else
+			{
+				// We've crossed the buffer boundary and now need to make sure the stream buffer read is done
+				if (!WaitForAsyncRead())
+				{
+					return false;
+				}
+				SwapBuffers();
+				StartStreamBufferRead();
+			}
+		}
+		return true;
+	}
+
+	virtual bool Write(const uint8* Source, int64 BytesToWrite) override
+	{
+		check(0 && "This is an async reader only and doesn't support writing");
+		return false;
+	}
+};
+
 /** 
  * Windows file handle implementation
 **/
@@ -260,10 +577,10 @@ public:
 		uint32  Access    = GENERIC_READ;
 		uint32  WinFlags  = FILE_SHARE_READ | (bAllowWrite ? FILE_SHARE_WRITE : 0);
 		uint32  Create    = OPEN_EXISTING;
-		HANDLE Handle    = CreateFileW(*NormalizeFilename(Filename), Access, WinFlags, NULL, Create, FILE_ATTRIBUTE_NORMAL, NULL);
-		if(Handle != INVALID_HANDLE_VALUE)
+		HANDLE Handle    = CreateFileW(*NormalizeFilename(Filename), Access, WinFlags, NULL, Create, FILE_ATTRIBUTE_NORMAL | FILE_FLAG_OVERLAPPED, NULL);
+		if (Handle != INVALID_HANDLE_VALUE)
 		{
-			return new FFileHandleWindows(Handle);
+			return new FAsyncBufferedFileReaderWindows(Handle);
 		}
 		return NULL;
 	}
