#ue4 - Perf: Add vectorized FRotator::Equals() and IsNearlyZero() implementations.

- 3x faster than scalar implementation.
- Includes unit tests.

diff --git a/Engine/Source/Runtime/Core/Private/Tests/Math/UnrealMathTest.cpp b/Engine/Source/Runtime/Core/Private/Tests/Math/UnrealMathTest.cpp
--- a/Engine/Source/Runtime/Core/Private/Tests/Math/UnrealMathTest.cpp
+++ b/Engine/Source/Runtime/Core/Private/Tests/Math/UnrealMathTest.cpp
@@ -400,6 +400,53 @@ FORCENOINLINE FRotator TestQuaternionToRotator(const FQuat& Quat)
 }
 
 
+
+// ROTATOR TESTS
+
+bool TestRotatorEqual0(const FRotator& A, const FRotator& B, const float Tolerance)
+{
+	// This is the version used for a few years (known working version).
+	return (FMath::Abs(FRotator::NormalizeAxis(A.Pitch - B.Pitch)) <= Tolerance)
+		&& (FMath::Abs(FRotator::NormalizeAxis(A.Yaw - B.Yaw)) <= Tolerance)
+		&& (FMath::Abs(FRotator::NormalizeAxis(A.Roll - B.Roll)) <= Tolerance);
+}
+
+bool TestRotatorEqual1(const FRotator& A, const FRotator& B, const float Tolerance)
+{
+	// Test the vectorized method.
+	const VectorRegister RegA = VectorLoadFloat3_W0(&A);
+	const VectorRegister RegB = VectorLoadFloat3_W0(&B);
+	const VectorRegister NormDelta = VectorNormalizeRotator(VectorSubtract(RegA, RegB));
+	const VectorRegister AbsNormDelta = VectorAbs(NormDelta);
+	return !VectorAnyGreaterThan(AbsNormDelta, VectorLoadFloat1(&Tolerance));
+}
+
+bool TestRotatorEqual2(const FRotator& A, const FRotator& B, const float Tolerance)
+{
+	// Test the FRotator method itself. It will likely be an equivalent implementation as 0 or 1 above.
+	return A.Equals(B, Tolerance);
+}
+
+bool TestRotatorEqual3(const FRotator& A, const FRotator& B, const float Tolerance)
+{
+	// Logically equivalent to tests above. Also tests IsNearlyZero().
+	return (A-B).IsNearlyZero(Tolerance);
+}
+
+// Report an error if bComparison is not equal to bExpected.
+void LogRotatorTest(bool bExpected, const TCHAR* TestName, const FRotator& A, const FRotator& B, bool bComparison)
+{
+	const bool bHasPassed = (bComparison == bExpected);
+	if (bHasPassed == false)
+	{
+		UE_LOG(LogUnrealMathTest, Log, TEXT("%s: %s"), bHasPassed ? TEXT("PASSED") : TEXT("FAILED"), TestName);
+		UE_LOG(LogUnrealMathTest, Log, TEXT("%s.Equals(%s) = %d"), *A.ToString(), *B.ToString(), bComparison);
+		GPassing = false;
+	}
+}
+
+
+
 /**
  * Helper debugf function to print out success or failure information for a test
  *
@@ -928,6 +975,38 @@ bool FVectorRegisterAbstractionTest::RunTest(const FString& Parameters)
 	V3 = VectorLoadAligned(&Q3);
 	LogTest( TEXT("VectorQuaternionMultiply2"), TestVectorsEqual(V2, V3, 1e-6f) );
 
+
+	// FRotator tests
+	{
+		//const double StartTime = FPlatformTime::Seconds();
+		const float Nudge = KINDA_SMALL_NUMBER * 0.25f;
+		FRotator RotArray[] ={
+			FRotator(0.f, 0.f, 0.f),
+			FRotator(Nudge, -Nudge, Nudge),
+			FRotator(+180.f, -180.f, +180.f),
+			FRotator(-180.f, +180.f, -180.f),
+			FRotator(+45.0f - Nudge, -120.f + Nudge, +270.f - Nudge),
+			FRotator(-45.0f + Nudge, +120.f - Nudge, -270.f + Nudge),
+			FRotator(+315.f - 360.f, -240.f - 360.f, -90.0f - 360.f),
+			FRotator(-315.f + 360.f, +240.f + 360.f, +90.0f + 360.f),
+		};
+
+		// Equality test
+		const float RotTolerance = KINDA_SMALL_NUMBER;
+		for (auto const& A : RotArray)
+		{
+			for (auto const& B : RotArray)
+			{
+				const bool bExpected = TestRotatorEqual0(A, B, RotTolerance);
+				LogRotatorTest(bExpected, TEXT("TestRotatorEqual1"), A, B, TestRotatorEqual1(A, B, RotTolerance));
+				LogRotatorTest(bExpected, TEXT("TestRotatorEqual2"), A, B, TestRotatorEqual2(A, B, RotTolerance));
+				LogRotatorTest(bExpected, TEXT("TestRotatorEqual3"), A, B, TestRotatorEqual3(A, B, RotTolerance));
+			}
+		}
+		//const double EndTime = FPlatformTime::Seconds();
+		//UE_LOG(LogUnrealMathTest, Warning, TEXT("Rotator test took %.12f secs"), EndTime - StartTime);
+	}
+
 	if (!GPassing)
 	{
 		UE_LOG(LogUnrealMathTest, Fatal,TEXT("VectorIntrinsics Failed."));
diff --git a/Engine/Source/Runtime/Core/Public/Math/Rotator.h b/Engine/Source/Runtime/Core/Public/Math/Rotator.h
--- a/Engine/Source/Runtime/Core/Public/Math/Rotator.h
+++ b/Engine/Source/Runtime/Core/Public/Math/Rotator.h
@@ -480,10 +480,17 @@ FORCEINLINE FRotator FRotator::operator-=( const FRotator &R )
 
 FORCEINLINE bool FRotator::IsNearlyZero(float Tolerance) const
 {
+#if PLATFORM_ENABLE_VECTORINTRINSICS
+	const VectorRegister RegA = VectorLoadFloat3_W0(this);
+	const VectorRegister Norm = VectorNormalizeRotator(RegA);
+	const VectorRegister AbsNorm = VectorAbs(Norm);
+	return !VectorAnyGreaterThan(AbsNorm, VectorLoadFloat1(&Tolerance));
+#else
 	return
 		FMath::Abs(NormalizeAxis(Pitch))<=Tolerance
 		&&	FMath::Abs(NormalizeAxis(Yaw))<=Tolerance
 		&&	FMath::Abs(NormalizeAxis(Roll))<=Tolerance;
+#endif
 }
 
 
@@ -495,9 +502,17 @@ FORCEINLINE bool FRotator::IsZero() const
 
 FORCEINLINE bool FRotator::Equals(const FRotator &R, float Tolerance) const
 {
+#if PLATFORM_ENABLE_VECTORINTRINSICS
+	const VectorRegister RegA = VectorLoadFloat3_W0(this);
+	const VectorRegister RegB = VectorLoadFloat3_W0(&R);
+	const VectorRegister NormDelta = VectorNormalizeRotator(VectorSubtract(RegA, RegB));
+	const VectorRegister AbsNormDelta = VectorAbs(NormDelta);
+	return !VectorAnyGreaterThan(AbsNormDelta, VectorLoadFloat1(&Tolerance));
+#else
 	return (FMath::Abs(NormalizeAxis(Pitch - R.Pitch)) <= Tolerance) 
 		&& (FMath::Abs(NormalizeAxis(Yaw - R.Yaw)) <= Tolerance) 
 		&& (FMath::Abs(NormalizeAxis(Roll - R.Roll)) <= Tolerance);
+#endif
 }
 
 
