Fixed editor per-project user settings getting eradicated on startup

The editor per-project settings in some cases were not being loaded correctly when bGenerateDestIni was false, which caused the file to be altogether emptied. I've changed this code to load the config directly, rather than rely on the nuanced behavior of the FConfigCacheIni helper functions. Also change the code to back up the old settings, rather than destroy them.

diff --git a/Engine/Source/Runtime/Core/Private/Misc/ConfigManifest.cpp b/Engine/Source/Runtime/Core/Private/Misc/ConfigManifest.cpp
--- a/Engine/Source/Runtime/Core/Private/Misc/ConfigManifest.cpp
+++ b/Engine/Source/Runtime/Core/Private/Misc/ConfigManifest.cpp
@@ -186,21 +186,27 @@ void FConfigManifest::MigrateEditorUserSettings()
 	}
 
 	// Handle upgrading editor user settings to the new path
-	FConfigFile EditorUserSettingsConfig;
-	EditorUserSettingsConfig.NoSave = true;
-	FConfigCacheIni::LoadLocalIniFile(EditorUserSettingsConfig, TEXT("EditorUserSettings"), true);
+	FConfigFile OldIni;
+	OldIni.NoSave = true;
+	OldIni.Read(EditorUserSettingsFilename);
 	
-	// Rename the config section
-	MigrateConfigSection(EditorUserSettingsConfig, TEXT("/Script/UnrealEd.EditorUserSettings"), TEXT("/Script/UnrealEd.EditorPerProjectUserSettings"));
+	if (OldIni.Num() != 0)
+	{
+		// Rename the config section
+		MigrateConfigSection(OldIni, TEXT("/Script/UnrealEd.EditorUserSettings"), TEXT("/Script/UnrealEd.EditorPerProjectUserSettings"));
 
-	FConfigFile EditorPerProjectUserSettingsConfig;
-	FConfigCacheIni::LoadLocalIniFile(EditorPerProjectUserSettingsConfig, TEXT("EditorPerProjectUserSettings"), false);
+		const FString EditorPerProjectUserSettingsFilename = ProjectSpecificIniPath(TEXT("EditorPerProjectUserSettings.ini"));
 
-	EditorPerProjectUserSettingsConfig.AddMissingProperties(EditorUserSettingsConfig);
-	if (EditorPerProjectUserSettingsConfig.Write(ProjectSpecificIniPath(TEXT("EditorPerProjectUserSettings.ini")), false))
-	{
-		FPlatformFileManager::Get().GetPlatformFile().DeleteFile(*EditorUserSettingsFilename);
+		FConfigFile NewIni;
+		NewIni.Read(EditorPerProjectUserSettingsFilename);
+		NewIni.AddMissingProperties(OldIni);
+		if (!NewIni.Write(EditorPerProjectUserSettingsFilename, false))
+		{
+			return;
+		}
 	}
+
+	IFileManager::Get().Move(*(EditorUserSettingsFilename + TEXT(".bak")), *EditorUserSettingsFilename);
 }
 
 EConfigManifestVersion FConfigManifest::UpgradeFromVersion(EConfigManifestVersion FromVersion)
