Disable setting thread affinity mask on Android due to issues with big.LITTLE cores
#ue4
#android

diff --git a/Engine/Source/Runtime/Core/Private/Android/AndroidProcess.cpp b/Engine/Source/Runtime/Core/Private/Android/AndroidProcess.cpp
--- a/Engine/Source/Runtime/Core/Private/Android/AndroidProcess.cpp
+++ b/Engine/Source/Runtime/Core/Private/Android/AndroidProcess.cpp
@@ -18,9 +18,11 @@ const TCHAR* FAndroidPlatformProcess::ComputerName()
 
 void FAndroidPlatformProcess::SetThreadAffinityMask( uint64 InAffinityMask )
 {
+	/* Skip setting mask for now due to issues with Big/Little CPUs
 	pid_t ThreadId = gettid();
 	int AffinityMask = (int)InAffinityMask;
 	syscall(__NR_sched_setaffinity, ThreadId, sizeof(AffinityMask), &AffinityMask);
+	*/
 }
 
 const TCHAR* FAndroidPlatformProcess::BaseDir()
