#ue4 - Fix scalar versions of FTransform::TransformFVector4() and TransformFVector4NoScale() setting W to 1 or 2 instead of 0 or 1.

diff --git a/Engine/Source/Runtime/Core/Public/Math/Transform.h b/Engine/Source/Runtime/Core/Public/Math/Transform.h
--- a/Engine/Source/Runtime/Core/Public/Math/Transform.h
+++ b/Engine/Source/Runtime/Core/Public/Math/Transform.h
@@ -1142,7 +1142,7 @@ FORCEINLINE FVector4 FTransform::TransformFVector4NoScale(const FVector4& V) con
 
 	//Transform using QST is following
 	//QST(P) = Q*S*P*-Q + T where Q = quaternion, S = scale, T = translation
-	FVector4 Transform = Rotation.RotateVector(V);
+	FVector4 Transform = FVector4(Rotation.RotateVector(FVector(V)), 0.f);
 	if (V.W == 1.f)
 	{
 		Transform += FVector4(Translation, 1.f);
@@ -1162,7 +1162,7 @@ FORCEINLINE FVector4 FTransform::TransformFVector4(const FVector4& V) const
 	//Transform using QST is following
 	//QST(P) = Q*S*P*-Q + T where Q = quaternion, S = scale, T = translation
 
-	FVector4 Transform = (Rotation.RotateVector(Scale3D*V));
+	FVector4 Transform = FVector4(Rotation.RotateVector(Scale3D*FVector(V)), 0.f);
 	if (V.W == 1.f)
 	{
 		Transform += FVector4(Translation, 1.f);
