Don't use floating window level for regular windows on Mac

diff --git a/Engine/Source/Runtime/Core/Private/Mac/MacWindow.cpp b/Engine/Source/Runtime/Core/Private/Mac/MacWindow.cpp
--- a/Engine/Source/Runtime/Core/Private/Mac/MacWindow.cpp
+++ b/Engine/Source/Runtime/Core/Private/Mac/MacWindow.cpp
@@ -107,14 +107,7 @@ void FMacWindow::Initialize( FMacApplication* const Application, const TSharedRe
 			}
 			else if (Definition->IsRegularWindow)
 			{
-				if (InParent.IsValid())
-				{
-					[WindowHandle setLevel: NSFloatingWindowLevel];
-				}
-				else
-				{
-					[WindowHandle setLevel: NSNormalWindowLevel];
-				}
+				[WindowHandle setLevel: NSNormalWindowLevel];
 			}
 			else if (!Definition->SupportsMaximize && !Definition->SupportsMinimize)
 			{
