Linux: fix failing to import starter content (UE-12878).

diff --git a/Engine/Source/Runtime/Core/Private/Linux/LinuxPlatformMisc.cpp b/Engine/Source/Runtime/Core/Private/Linux/LinuxPlatformMisc.cpp
--- a/Engine/Source/Runtime/Core/Private/Linux/LinuxPlatformMisc.cpp
+++ b/Engine/Source/Runtime/Core/Private/Linux/LinuxPlatformMisc.cpp
@@ -72,7 +72,8 @@ const TCHAR* FLinuxPlatformMisc::RootDir()
 
 void FLinuxPlatformMisc::NormalizePath(FString& InPath)
 {
-	if (InPath.Contains(TEXT("~"), ESearchCase::CaseSensitive))	// case sensitive is quicker, and our substring doesn't care
+	// only expand if path starts with ~, e.g. ~/ should be expanded, /~ should not
+	if (InPath.StartsWith(TEXT("~"), ESearchCase::CaseSensitive))	// case sensitive is quicker, and our substring doesn't care
 	{
 		static bool bHaveHome = false;
 		static TCHAR CachedResult[PlatformMiscLimits::MaxUserHomeDirLength];
diff --git a/Engine/Source/Runtime/Core/Private/Misc/Paths.cpp b/Engine/Source/Runtime/Core/Private/Misc/Paths.cpp
--- a/Engine/Source/Runtime/Core/Private/Misc/Paths.cpp
+++ b/Engine/Source/Runtime/Core/Private/Misc/Paths.cpp
@@ -781,6 +781,15 @@ FString FPaths::ConvertRelativePathToFull(const FString& BasePath, const FString
 
 	FPaths::NormalizeFilename(FullyPathed);
 	FPaths::CollapseRelativeDirectories(FullyPathed);
+
+	if (FullyPathed.Len() == 0)
+	{
+		// Empty path is not absolute, and '/' is the best guess across all the platforms.
+		// This substituion is not valid for Windows of course; however CollapseRelativeDirectories() will not produce an empty
+		// absolute path on Windows as it takes care not to remove the drive letter.
+		FullyPathed = TEXT("/");
+	}
+
 	return FullyPathed;
 }
 
