Fixed benign compile error that was preventing use of non-copyable return types on TFunctions

diff --git a/Engine/Source/Runtime/Core/Public/Templates/Function.h b/Engine/Source/Runtime/Core/Public/Templates/Function.h
--- a/Engine/Source/Runtime/Core/Public/Templates/Function.h
+++ b/Engine/Source/Runtime/Core/Public/Templates/Function.h
@@ -155,9 +155,9 @@ namespace UE4Function_Private
 	#endif
 
 	template <typename T>
-	inline T& FakeCall(T* Ptr)
+	inline T&& FakeCall(T* Ptr)
 	{
-		return *Ptr;
+		return MoveTemp(*Ptr);
 	}
 
 	inline void FakeCall(void* Ptr)
