UE13406 - r.MobileContentScaleFactor doesn't work on Android.
-- Updating the function which is doing the caching so that it checks to see if the value has changed and recalculate if so

diff --git a/Engine/Source/Runtime/Core/Private/Android/AndroidWindow.cpp b/Engine/Source/Runtime/Core/Private/Android/AndroidWindow.cpp
--- a/Engine/Source/Runtime/Core/Private/Android/AndroidWindow.cpp
+++ b/Engine/Source/Runtime/Core/Private/Android/AndroidWindow.cpp
@@ -10,6 +10,8 @@
 // Cached calculated screen resolution
 static int32 WindowWidth = -1;
 static int32 WindowHeight = -1;
+static bool WindowInit = false;
+static float ContentScaleFactor = -1.0f;
 
 FAndroidWindow::~FAndroidWindow()
 {
@@ -70,8 +72,12 @@ int32 FAndroidWindow::GetDepthBufferPreference()
 
 FPlatformRect FAndroidWindow::GetScreenRect()
 {
-	// since orientation and resolution won't change on Android, use cached results if valid
-	if (WindowWidth > 8)
+	// CSF is a multiplier to 1280x720
+	static IConsoleVariable* CVar = IConsoleManager::Get().FindConsoleVariable(TEXT("r.MobileContentScaleFactor"));
+	float RequestedContentScaleFactor = CVar->GetFloat();
+
+	// since orientation and resolution won't change on Android, use cached results if still valid
+	if (WindowInit && RequestedContentScaleFactor == ContentScaleFactor)
 	{
 		FPlatformRect ScreenRect;
 		ScreenRect.Left = 0;
@@ -121,10 +127,6 @@ FPlatformRect FAndroidWindow::GetScreenRect()
 		UE_LOG(LogAndroid, Log, TEXT("Limiting MaxWidth=%d and MaxHeight=%d due to bMobileHDR32bpp (was %dx%d)"), MaxWidth, MaxHeight, OldMaxWidth, OldMaxHeight);
 	}
 
-	// CSF is a multiplier to 1280x720
-	static IConsoleVariable* CVar = IConsoleManager::Get().FindConsoleVariable(TEXT("r.MobileContentScaleFactor")); 
-	float RequestedContentScaleFactor = CVar->GetFloat();
-
 	// 0 means to use native size
 	int32 Width, Height;
 	if (RequestedContentScaleFactor == 0.0f)
@@ -163,6 +165,8 @@ FPlatformRect FAndroidWindow::GetScreenRect()
 	// save for future calls
 	WindowWidth = Width;
 	WindowHeight = Height;
+	WindowInit = true;
+	ContentScaleFactor = RequestedContentScaleFactor;
 
 	return ScreenRect;
 }
