Core - Force exit when a background thread crashes

diff --git a/Engine/Source/Runtime/Core/Private/Windows/WindowsRunnableThread.cpp b/Engine/Source/Runtime/Core/Private/Windows/WindowsRunnableThread.cpp
--- a/Engine/Source/Runtime/Core/Private/Windows/WindowsRunnableThread.cpp
+++ b/Engine/Source/Runtime/Core/Private/Windows/WindowsRunnableThread.cpp
@@ -36,6 +36,8 @@ uint32 FRunnableThreadWin::GuardedRun()
 			GError->HandleError();
 			// Throw an error so that main thread shuts down too (otherwise task graph stalls forever).
 			UE_LOG(LogThreadingWindows, Fatal, TEXT("Runnable thread %s crashed."), *ThreadName);
+
+			FPlatformMisc::RequestExit( true );
 		}
 	}
 	else
