Disable overlapped reads on Windows due to see if it fixes hangs on build machines.

diff --git a/Engine/Source/Runtime/Core/Private/Windows/WindowsPlatformFile.cpp b/Engine/Source/Runtime/Core/Private/Windows/WindowsPlatformFile.cpp
--- a/Engine/Source/Runtime/Core/Private/Windows/WindowsPlatformFile.cpp
+++ b/Engine/Source/Runtime/Core/Private/Windows/WindowsPlatformFile.cpp
@@ -577,11 +577,19 @@ public:
 		uint32  Access    = GENERIC_READ;
 		uint32  WinFlags  = FILE_SHARE_READ | (bAllowWrite ? FILE_SHARE_WRITE : 0);
 		uint32  Create    = OPEN_EXISTING;
+#if USE_OVERLAPPED_READS
 		HANDLE Handle    = CreateFileW(*NormalizeFilename(Filename), Access, WinFlags, NULL, Create, FILE_ATTRIBUTE_NORMAL | FILE_FLAG_OVERLAPPED, NULL);
 		if (Handle != INVALID_HANDLE_VALUE)
 		{
 			return new FAsyncBufferedFileReaderWindows(Handle);
 		}
+#else
+		HANDLE Handle    = CreateFileW(*NormalizeFilename(Filename), Access, WinFlags, NULL, Create, FILE_ATTRIBUTE_NORMAL, NULL);
+		if(Handle != INVALID_HANDLE_VALUE)
+		{
+			return new FFileHandleWindows(Handle);
+		}
+#endif
 		return NULL;
 	}
 	virtual IFileHandle* OpenWrite(const TCHAR* Filename, bool bAppend = false, bool bAllowRead = false) override
