Use Cocoa window number instead of a TWeakPtr<FGenericWindow> to store the context window in MacTextInputMethodSystem

diff --git a/Engine/Source/Runtime/Core/Private/Mac/MacTextInputMethodSystem.cpp b/Engine/Source/Runtime/Core/Private/Mac/MacTextInputMethodSystem.cpp
--- a/Engine/Source/Runtime/Core/Private/Mac/MacTextInputMethodSystem.cpp
+++ b/Engine/Source/Runtime/Core/Private/Mac/MacTextInputMethodSystem.cpp
@@ -18,13 +18,14 @@ namespace
 	public:
 		FTextInputMethodChangeNotifier(const TSharedRef<ITextInputMethodContext>& InContext)
 		:	Context(InContext)
+		,	ContextWindowNumber(0)
 		{
 		}
 		
 		virtual ~FTextInputMethodChangeNotifier() {}
 		
-		void SetContextWindow(TSharedPtr<FGenericWindow> Window);
-		TSharedPtr<FGenericWindow> GetContextWindow();
+		void SetContextWindowNumber(int32 WindowNumber);
+		int32 GetContextWindowNumber() const;
 		
 		virtual void NotifyLayoutChanged(const ELayoutChangeType ChangeType) override;
 		virtual void NotifySelectionChanged() override;
@@ -33,38 +34,36 @@ namespace
 		
 	private:
 		TWeakPtr<ITextInputMethodContext> Context;
-		TWeakPtr<FGenericWindow> ContextWindow;
+		int32 ContextWindowNumber;
 	};
-	
-	void FTextInputMethodChangeNotifier::SetContextWindow(TSharedPtr<FGenericWindow> Window)
+
+	void FTextInputMethodChangeNotifier::SetContextWindowNumber(int32 WindowNumber)
 	{
-		ContextWindow = Window;
+		ContextWindowNumber = WindowNumber;
 	}
-	
-	TSharedPtr<FGenericWindow> FTextInputMethodChangeNotifier::GetContextWindow()
+
+	int32 FTextInputMethodChangeNotifier::GetContextWindowNumber() const
 	{
-		if (!ContextWindow.IsValid())
-		{
-			ContextWindow = Context.Pin()->GetWindow();
-		}
-		return ContextWindow.Pin();
+		return ContextWindowNumber;
 	}
-	
+
 	void FTextInputMethodChangeNotifier::NotifyLayoutChanged(const ELayoutChangeType ChangeType)
 	{
-		if(ChangeType != ELayoutChangeType::Created && GetContextWindow().IsValid())
+		if(ChangeType != ELayoutChangeType::Created && ContextWindowNumber != 0)
 		{
-			FCocoaWindow* CocoaWindow = (FCocoaWindow*)GetContextWindow()->GetOSWindowHandle();
-			
-			MacApplication->InvalidateTextLayout( CocoaWindow );
+			FCocoaWindow* CocoaWindow = (FCocoaWindow*)[NSApp windowWithWindowNumber:ContextWindowNumber];
+			if (CocoaWindow)
+			{
+				MacApplication->InvalidateTextLayout( CocoaWindow );
+			}
 		}
 	}
 	
 	void FTextInputMethodChangeNotifier::NotifySelectionChanged()
 	{
-		if(GetContextWindow().IsValid())
+		if(ContextWindowNumber != 0)
 		{
-			FCocoaWindow* CocoaWindow = (FCocoaWindow*)GetContextWindow()->GetOSWindowHandle();
+			FCocoaWindow* CocoaWindow = (FCocoaWindow*)[NSApp windowWithWindowNumber:ContextWindowNumber];
 			MainThreadCall(^{
 				SCOPED_AUTORELEASE_POOL;
 				if(CocoaWindow && [CocoaWindow openGLView])
@@ -78,9 +77,9 @@ namespace
 	
 	void FTextInputMethodChangeNotifier::NotifyTextChanged(const uint32 BeginIndex, const uint32 OldLength, const uint32 NewLength)
 	{
-		if(GetContextWindow().IsValid())
+		if(ContextWindowNumber != 0)
 		{
-			FCocoaWindow* CocoaWindow = (FCocoaWindow*)GetContextWindow()->GetOSWindowHandle();
+			FCocoaWindow* CocoaWindow = (FCocoaWindow*)[NSApp windowWithWindowNumber:ContextWindowNumber];
 			MainThreadCall(^{
 				SCOPED_AUTORELEASE_POOL;
 				if(CocoaWindow && [CocoaWindow openGLView])
@@ -94,9 +93,9 @@ namespace
 	
 	void FTextInputMethodChangeNotifier::CancelComposition()
 	{
-		if(GetContextWindow().IsValid())
+		if(ContextWindowNumber != 0)
 		{
-			FCocoaWindow* CocoaWindow = (FCocoaWindow*)GetContextWindow()->GetOSWindowHandle();
+			FCocoaWindow* CocoaWindow = (FCocoaWindow*)[NSApp windowWithWindowNumber:ContextWindowNumber];
 			MainThreadCall(^{
 				SCOPED_AUTORELEASE_POOL;
 				if(CocoaWindow && [CocoaWindow openGLView])
@@ -146,8 +145,7 @@ void FMacTextInputMethodSystem::UnregisterContext(const TSharedRef<ITextInputMet
 	
 	TSharedPtr<ITextInputMethodChangeNotifier> Notifier(NotifierRef.Pin());
 	FTextInputMethodChangeNotifier* MacNotifier = (FTextInputMethodChangeNotifier*)Notifier.Get();
-	TSharedPtr<FGenericWindow> GenericWindow = MacNotifier->GetContextWindow();
-	if(GenericWindow.IsValid())
+	if(MacNotifier->GetContextWindowNumber() != 0)
 	{
 		DeactivateContext(Context);
 	}
@@ -171,8 +169,8 @@ void FMacTextInputMethodSystem::ActivateContext(const TSharedRef<ITextInputMetho
 	bool bActivatedContext = false;
 	if(GenericWindow.IsValid())
 	{
-		MacNotifier->SetContextWindow(GenericWindow);
 		FCocoaWindow* CocoaWindow = (FCocoaWindow*)GenericWindow->GetOSWindowHandle();
+		MacNotifier->SetContextWindowNumber([CocoaWindow windowNumber]);
 		bActivatedContext = MainThreadReturn(^{
 			bool bSuccess = false;
 			if(CocoaWindow && [CocoaWindow openGLView])
@@ -206,11 +204,10 @@ void FMacTextInputMethodSystem::DeactivateContext(const TSharedRef<ITextInputMet
 	
 	TSharedPtr<ITextInputMethodChangeNotifier> Notifier(NotifierRef.Pin());
 	FTextInputMethodChangeNotifier* MacNotifier = (FTextInputMethodChangeNotifier*)Notifier.Get();
-	TSharedPtr<FGenericWindow> GenericWindow = MacNotifier->GetContextWindow();
 	bool bDeactivatedContext = false;
-	if(GenericWindow.IsValid())
+	if(MacNotifier->GetContextWindowNumber() != 0)
 	{
-		FCocoaWindow* CocoaWindow = (FCocoaWindow*)GenericWindow->GetOSWindowHandle();
+		FCocoaWindow* CocoaWindow = (FCocoaWindow*)[NSApp windowWithWindowNumber:MacNotifier->GetContextWindowNumber()];
 		bDeactivatedContext = MainThreadReturn(^{
 			bool bSuccess = false;
 			if(CocoaWindow && [CocoaWindow openGLView])
@@ -244,10 +241,9 @@ bool FMacTextInputMethodSystem::IsActiveContext(const TSharedRef<ITextInputMetho
 	
 	TSharedPtr<ITextInputMethodChangeNotifier> Notifier(NotifierRef.Pin());
 	FTextInputMethodChangeNotifier* MacNotifier = (FTextInputMethodChangeNotifier*)Notifier.Get();
-	TSharedPtr<FGenericWindow> GenericWindow = MacNotifier->GetContextWindow();
-	if(GenericWindow.IsValid())
+	if(MacNotifier->GetContextWindowNumber() != 0)
 	{
-		FCocoaWindow* CocoaWindow = (FCocoaWindow*)GenericWindow->GetOSWindowHandle();
+		FCocoaWindow* CocoaWindow = (FCocoaWindow*)[NSApp windowWithWindowNumber:MacNotifier->GetContextWindowNumber()];
 		if(CocoaWindow && [CocoaWindow openGLView])
 		{
 			NSView* GLView = [CocoaWindow openGLView];
