Set window title in FMacWindow::Initialize()

diff --git a/Engine/Source/Runtime/Core/Private/Mac/MacWindow.cpp b/Engine/Source/Runtime/Core/Private/Mac/MacWindow.cpp
--- a/Engine/Source/Runtime/Core/Private/Mac/MacWindow.cpp
+++ b/Engine/Source/Runtime/Core/Private/Mac/MacWindow.cpp
@@ -153,6 +153,8 @@ void FMacWindow::Initialize( FMacApplication* const Application, const TSharedRe
 						WindowHandle.titleVisibility = NSWindowTitleHidden;
 					}
 				}
+
+				SetText(*Definition->Title);
 			}
 			else if(Definition->AppearsInTaskbar)
 			{
