fix for crash on startup of the game
UE-13741
#ios

diff --git a/Engine/Source/Runtime/Core/Private/IOS/IOSAppDelegate.cpp b/Engine/Source/Runtime/Core/Private/IOS/IOSAppDelegate.cpp
--- a/Engine/Source/Runtime/Core/Private/IOS/IOSAppDelegate.cpp
+++ b/Engine/Source/Runtime/Core/Private/IOS/IOSAppDelegate.cpp
@@ -584,7 +584,7 @@ void InstallSignalHandlers()
     [self ToggleSuspend:false];
 	[self ToggleAudioSession:true];
 
-    if (FTaskGraphInterface::IsRunning())
+    if (bEngineInit)
     {
         FGraphEventRef ResignTask = FFunctionGraphTask::CreateAndDispatchWhenReady([&]()
         {
