#ue4 - Refactor vectorized Quaternion Vector rotation to a single method (VectorQuaternionRotateVector), rather than two calls to VectorQuaternionMultiply2().

- This will allow other rotation methods to be implemented more easily.
- Currently the new method uses the same approach of: Q * V * Q^-1 using VectorQuaternionMultiply2().
- Also added VectorQuaternionInverseRotateVector().

diff --git a/Engine/Source/Runtime/Core/Private/Math/TransformVectorized.cpp b/Engine/Source/Runtime/Core/Private/Math/TransformVectorized.cpp
--- a/Engine/Source/Runtime/Core/Private/Math/TransformVectorized.cpp
+++ b/Engine/Source/Runtime/Core/Private/Math/TransformVectorized.cpp
@@ -104,10 +104,8 @@ FTransform FTransform::GetRelativeTransformReverse(const FTransform& Other) cons
 	VectorRegister VInverseRot = VectorQuaternionInverse(Rotation);
 	VectorRegister VRotation = VectorQuaternionMultiply2(Other.Rotation, VInverseRot );
 	
-	// [Q(B)*Q(A)(-1)*T(A)*Q(A)*Q(B)(-1)]	
-	VInverseRot = VectorQuaternionInverse(VRotation);
-	VectorRegister VQT = VectorQuaternionMultiply2(VRotation, Translation);	
-	VectorRegister VR = VectorQuaternionMultiply2(VQT, VInverseRot);
+	// RotatedTranslation
+	VectorRegister VR = VectorQuaternionRotateVector(VRotation, Translation);
 
 	// Translation = T(B)-S(B)/S(A) *[Q(B)*Q(A)(-1)*T(A)*Q(A)*Q(B)(-1)]	
 	VectorRegister VTranslation = VectorSet_W0(VectorSubtract(Other.Translation, VectorMultiply(VScale3D, VR)));
@@ -154,10 +152,11 @@ void FTransform::SetToRelativeTransform(const FTransform& ParentTransform)
 	//VQTranslation = (  ( T(A).X - T(B).X ),  ( T(A).Y - T(B).Y ), ( T(A).Z - T(B).Z), 0.f );
 	VectorRegister VQTranslation = VectorSet_W0(VectorSubtract(Translation, ParentTransform.Translation));
 
-	// Translation = 1/S(B) *[Q(B)(-1)*(T(A)-T(B))*Q(B)]
+	// Inverse RotatedTranslation
 	VectorRegister VInverseParentRot = VectorQuaternionInverse(ParentTransform.Rotation);
-	VectorRegister VQT = VectorQuaternionMultiply2(VInverseParentRot, VQTranslation);
-	VectorRegister VR = VectorQuaternionMultiply2(VQT, ParentTransform.Rotation);
+	VectorRegister VR = VectorQuaternionRotateVector(VInverseParentRot, VQTranslation);
+
+	// Translation = 1/S(B)
 	Translation = VectorMultiply(VR, VSafeScale3D);
 
 	// Rotation = Q(B)(-1) * Q(A)	
@@ -194,10 +193,11 @@ FTransform FTransform::GetRelativeTransform(const FTransform& Other) const
 	//VQTranslation = (  ( T(A).X - T(B).X ),  ( T(A).Y - T(B).Y ), ( T(A).Z - T(B).Z), 0.f );
 	VectorRegister VQTranslation =  VectorSet_W0(VectorSubtract(Translation, Other.Translation));
 
-	// Translation = 1/S(B) *[Q(B)(-1)*(T(A)-T(B))*Q(B)]
+	// Inverse RotatedTranslation
 	VectorRegister VInverseRot = VectorQuaternionInverse(Other.Rotation);
-	VectorRegister VQT = VectorQuaternionMultiply2(VInverseRot, VQTranslation);
-	VectorRegister VR = VectorQuaternionMultiply2(VQT, Other.Rotation);
+	VectorRegister VR = VectorQuaternionRotateVector(VInverseRot, VQTranslation);
+
+	//Translation = 1/S(B)
 	VectorRegister VTranslation = VectorMultiply(VR, VSafeScale3D);
 
 	// Rotation = Q(B)(-1) * Q(A)	
diff --git a/Engine/Source/Runtime/Core/Private/Tests/Math/UnrealMathTest.cpp b/Engine/Source/Runtime/Core/Private/Tests/Math/UnrealMathTest.cpp
--- a/Engine/Source/Runtime/Core/Private/Tests/Math/UnrealMathTest.cpp
+++ b/Engine/Source/Runtime/Core/Private/Tests/Math/UnrealMathTest.cpp
@@ -10,6 +10,7 @@ MS_ALIGN( 16 ) static float GScratch[16] GCC_ALIGN( 16 );
 static float GSum;
 static bool GPassing;
 
+#define MATHTEST_INLINE FORCEINLINE_DEBUGGABLE
 
 /**
  * Tests if two vectors (xyzw) are bitwise equal
@@ -306,7 +307,7 @@ FORCENOINLINE FQuat TestRotatorToQuaternion( const FRotator& Rotator)
 	return RotationQuat;
 }
 
-FORCENOINLINE FVector TestQuaternionRotateVector(const FQuat& Quat, const FVector& Vector)
+MATHTEST_INLINE FVector TestQuaternionRotateVectorScalar(const FQuat& Quat, const FVector& Vector)
 {
 	// (q.W*q.W-qv.qv)v + 2(qv.v)qv + 2 q.W (qv x v)
 	const FVector qv(Quat.X, Quat.Y, Quat.Z);
@@ -317,7 +318,8 @@ FORCENOINLINE FVector TestQuaternionRotateVector(const FQuat& Quat, const FVecto
 	return vOut;
 }
 
-FORCENOINLINE FVector TestQuaternionMultiplyVector(const FQuat& Quat, const FVector& Vector)
+// Q * V * Q^-1
+MATHTEST_INLINE FVector TestQuaternionMultiplyVector(const FQuat& Quat, const FVector& Vector)
 {
 	FQuat VQ(Vector.X, Vector.Y, Vector.Z, 0.f);
 	FQuat VT, VR;
@@ -328,6 +330,18 @@ FORCENOINLINE FVector TestQuaternionMultiplyVector(const FQuat& Quat, const FVec
 	return FVector(VR.X, VR.Y, VR.Z);
 }
 
+MATHTEST_INLINE FVector TestQuaternionRotateVectorRegister(const FQuat& Quat, const FVector &V)
+{
+	const VectorRegister Rotation = *((const VectorRegister*)(&Quat));
+	const VectorRegister InputVectorW0 = VectorLoadFloat3_W0(&V);
+	const VectorRegister RotatedVec = VectorQuaternionRotateVector(Rotation, InputVectorW0);
+
+	FVector Result;
+	VectorStoreFloat3(RotatedVec, &Result);
+	return Result;
+}
+
+
 /**
 * Multiplies two quaternions: The order matters.
 *
@@ -913,29 +927,49 @@ bool FVectorRegisterAbstractionTest::RunTest(const FString& Parameters)
 	}
 
 	{
-		const FRotator Rotator0(30.0f, -45.0f, 90.0f);
-		const FRotator Rotator1(45.0f,  60.0f, 120.0f);
-		const FQuat Quat0(Rotator0);
-		const FQuat Quat1(Rotator1);
+		// Use these Quats...
+		const FQuat TestQuats[] = {
+			FQuat(FQuat::Identity),
+			FQuat(FRotator(30.0f, -45.0f, 90.0f)),
+			FQuat(FRotator(45.0f,  60.0f, 120.0f)),
+			FQuat(FRotator(0.0f, 180.0f, 45.0f)),
+			FQuat(FRotator(-120.0f, -90.0f, 0.0f)),
+			FQuat(FRotator(-0.01f, 0.02f, -0.03f)),
+		};
+
+		// ... to rotate these Vectors...
+		const FVector TestVectors[] = {
+			FVector::ZeroVector,
+			FVector::ForwardVector,
+			FVector::RightVector,
+			FVector::UpVector,
+			FVector(45.0f, -60.0f, 120.0f),
+			FVector(0.45f, -0.60f, 1.20f),
+			FVector(0.57735026918962576451f, 0.57735026918962576451f, 0.57735026918962576451f),
+		};
+
+		// ... and test within this tolerance.
 		const float Tolerance = 1e-4f;
 
-		LogTest( TEXT("Test0 FQuat::RotateVector(const FVector& Vec)"),		TestVectorsEqual3_NoVec(Quat0.RotateVector(FVector::ForwardVector), TestQuaternionRotateVector(Quat0, FVector::ForwardVector),	Tolerance ));
-		LogTest( TEXT("Test0 FQuat::operator*(const FVector& Vec)"),		TestVectorsEqual3_NoVec(Quat0 * FVector::ForwardVector, TestQuaternionMultiplyVector(Quat0, FVector::ForwardVector), Tolerance ));
-		LogTest( TEXT("Test1 FQuat::RotateVector(const FVector& Vec)"),		TestVectorsEqual3_NoVec(Quat0.RotateVector(FVector::RightVector), TestQuaternionRotateVector(Quat0, FVector::RightVector),	Tolerance ));
-		LogTest( TEXT("Test1 FQuat::operator*(const FVector& Vec)"),		TestVectorsEqual3_NoVec(Quat0 * FVector::RightVector, TestQuaternionMultiplyVector(Quat0, FVector::RightVector), Tolerance ));
-		LogTest( TEXT("Test2 FQuat::RotateVector(const FVector& Vec)"),		TestVectorsEqual3_NoVec(Quat0.RotateVector(FVector::UpVector), TestQuaternionRotateVector(Quat0, FVector::UpVector),	Tolerance ));
-		LogTest( TEXT("Test2 FQuat::operator*(const FVector& Vec)"),		TestVectorsEqual3_NoVec(Quat0 * FVector::UpVector, TestQuaternionMultiplyVector(Quat0, FVector::UpVector), Tolerance ));
-		LogTest( TEXT("Test3 FQuat::RotateVector(const FVector& Vec)"),		TestVectorsEqual3_NoVec(Quat0.RotateVector(FVector(45.0f, -60.0f, 120.0f)), TestQuaternionRotateVector(Quat0, FVector(45.0f, -60.0f, 120.0f)),	Tolerance ));
-		LogTest( TEXT("Test3 FQuat::operator*(const FVector& Vec)"),		TestVectorsEqual3_NoVec(Quat0 * FVector(45.0f, -60.0f, 120.0f), TestQuaternionMultiplyVector(Quat0, FVector(45.0f, -60.0f, 120.0f)), Tolerance ));
-		LogTest( TEXT("Test4 FQuat::RotateVector(const FVector& Vec)"),		TestVectorsEqual3_NoVec(Quat1.RotateVector(FVector::ForwardVector), TestQuaternionRotateVector(Quat1, FVector::ForwardVector),	Tolerance ));
-		LogTest( TEXT("Test4 FQuat::operator*(const FVector& Vec)"),		TestVectorsEqual3_NoVec(Quat1 * FVector::ForwardVector, TestQuaternionMultiplyVector(Quat1, FVector::ForwardVector), Tolerance ));
-		LogTest( TEXT("Test5 FQuat::RotateVector(const FVector& Vec)"),		TestVectorsEqual3_NoVec(Quat1.RotateVector(FVector::RightVector), TestQuaternionRotateVector(Quat1, FVector::RightVector),	Tolerance ));
-		LogTest( TEXT("Test5 FQuat::operator*(const FVector& Vec)"),		TestVectorsEqual3_NoVec(Quat1 * FVector::RightVector, TestQuaternionMultiplyVector(Quat1, FVector::RightVector), Tolerance ));
-		LogTest( TEXT("Test6 FQuat::RotateVector(const FVector& Vec)"),		TestVectorsEqual3_NoVec(Quat1.RotateVector(FVector::UpVector), TestQuaternionRotateVector(Quat1, FVector::UpVector),	Tolerance ));
-		LogTest( TEXT("Test6 FQuat::operator*(const FVector& Vec)"),		TestVectorsEqual3_NoVec(Quat1 * FVector::UpVector, TestQuaternionMultiplyVector(Quat1, FVector::UpVector), Tolerance ));
-		LogTest( TEXT("Test7 FQuat::RotateVector(const FVector& Vec)"),		TestVectorsEqual3_NoVec(Quat1.RotateVector(FVector(45.0f, -60.0f, 120.0f)), TestQuaternionRotateVector(Quat1, FVector(45.0f, -60.0f, 120.0f)),	Tolerance ));
-		LogTest( TEXT("Test7 FQuat::operator*(const FVector& Vec)"),		TestVectorsEqual3_NoVec(Quat1 * FVector(45.0f, -60.0f, 120.0f), TestQuaternionMultiplyVector(Quat1, FVector(45.0f, -60.0f, 120.0f)), Tolerance ));
+		// Test Macro. Tests FQuat::RotateVector(_Vec) against _Func(_Vec)
+		#define TEST_QUAT_ROTATE(_QIndex, _VIndex, _Quat, _Vec, _Func, _Tolerance) \
+		{ \
+			const FString _TestName = FString::Printf(TEXT("Test Quat%d: Vec%d: %s"), _QIndex, _VIndex, TEXT(#_Func)); \
+			LogTest( *_TestName, TestVectorsEqual3_NoVec(_Quat.RotateVector(_Vec), _Func(_Quat, _Vec), _Tolerance) ); \
+		}
 
+		// Test loop
+		for (int32 QIndex = 0; QIndex < ARRAY_COUNT(TestQuats); ++QIndex)
+		{
+			const FQuat& Q = TestQuats[QIndex];
+			for (int32 VIndex = 0; VIndex < ARRAY_COUNT(TestVectors); ++VIndex)
+			{
+				const FVector& V = TestVectors[VIndex];
+				TEST_QUAT_ROTATE(QIndex, VIndex, Q, V, TestQuaternionRotateVectorScalar, Tolerance);
+				TEST_QUAT_ROTATE(QIndex, VIndex, Q, V, TestQuaternionRotateVectorRegister, Tolerance);
+				TEST_QUAT_ROTATE(QIndex, VIndex, Q, V, TestQuaternionMultiplyVector, Tolerance);
+			}
+		}
 	}
 
 	{
@@ -978,7 +1012,6 @@ bool FVectorRegisterAbstractionTest::RunTest(const FString& Parameters)
 
 	// FRotator tests
 	{
-		//const double StartTime = FPlatformTime::Seconds();
 		const float Nudge = KINDA_SMALL_NUMBER * 0.25f;
 		FRotator RotArray[] ={
 			FRotator(0.f, 0.f, 0.f),
@@ -1003,8 +1036,6 @@ bool FVectorRegisterAbstractionTest::RunTest(const FString& Parameters)
 				LogRotatorTest(bExpected, TEXT("TestRotatorEqual3"), A, B, TestRotatorEqual3(A, B, RotTolerance));
 			}
 		}
-		//const double EndTime = FPlatformTime::Seconds();
-		//UE_LOG(LogUnrealMathTest, Warning, TEXT("Rotator test took %.12f secs"), EndTime - StartTime);
 	}
 
 	if (!GPassing)
diff --git a/Engine/Source/Runtime/Core/Public/Math/TransformVectorized.h b/Engine/Source/Runtime/Core/Public/Math/TransformVectorized.h
--- a/Engine/Source/Runtime/Core/Public/Math/TransformVectorized.h
+++ b/Engine/Source/Runtime/Core/Public/Math/TransformVectorized.h
@@ -2,9 +2,6 @@
 
 #pragma once
 
-/** This returns Quaternion Inverse of X **/
-#define MAKE_QUATINV_VECTORREGISTER(X) VectorMultiply(GlobalVectorConstants::QINV_SIGN_MASK, X)
-
 /**
  * Transform composed of Scale, Rotation (as a quaternion), and Translation.
  *
@@ -1269,12 +1266,11 @@ private:
 		const VectorRegister InvScale = VectorSet_W0(GetSafeScaleReciprocal(Scale3D , ScalarRegister(GlobalVectorConstants::SmallNumber)));
 
 		// Invert the rotation
-		const VectorRegister InvRotation = MAKE_QUATINV_VECTORREGISTER(Rotation);
+		const VectorRegister InvRotation = VectorQuaternionInverse(Rotation);
 
 		// Invert the translation
 		const VectorRegister ScaledTranslation = VectorMultiply(InvScale, Translation);
-		const VectorRegister t1 = VectorQuaternionMultiply2(InvRotation, ScaledTranslation);
-		const VectorRegister t2 = VectorQuaternionMultiply2(t1, Rotation);
+		const VectorRegister t2 = VectorQuaternionRotateVector(InvRotation, ScaledTranslation);
 		const VectorRegister InvTranslation = VectorSet_W0(VectorNegate(t2));
 
 		return FTransform(InvRotation, InvTranslation, InvScale);
@@ -1339,15 +1335,12 @@ FORCEINLINE void FTransform::Multiply(FTransform* OutTransform, const FTransform
 	const VectorRegister ScaleA = A->Scale3D;
 	const VectorRegister ScaleB = B->Scale3D;
 
-	const VectorRegister QuatBInv = MAKE_QUATINV_VECTORREGISTER(QuatB);
-
 	// RotationResult = B.Rotation * A.Rotation
 	OutTransform->Rotation = VectorQuaternionMultiply2(QuatB, QuatA);
 
-	// TranslateResult = ((B.Rotation * B.Scale * A.Translation) * Inverse(B.Rotation)) + B.Translate
+	// TranslateResult = B.Rotate(B.Scale * A.Translation) + B.Translate
 	const VectorRegister ScaledTransA = VectorMultiply(TranslateA, ScaleB);
-	const VectorRegister Temp = VectorQuaternionMultiply2(QuatB, ScaledTransA);
-	const VectorRegister RotatedTranslate = VectorQuaternionMultiply2(Temp, QuatBInv);
+	const VectorRegister RotatedTranslate = VectorQuaternionRotateVector(QuatB, ScaledTransA);
 	OutTransform->Translation = VectorAdd(RotatedTranslate, TranslateB);
 
 	// ScaleResult = Scale.B * Scale.A
@@ -1395,13 +1388,11 @@ FORCEINLINE FVector4 FTransform::TransformFVector4NoScale(const FVector4& V) con
 	const VectorRegister InputVector = VectorLoadAligned(&V);
 
 	//Transform using QST is following
-	//QST(P) = Q*S*P*-Q + T where Q = quaternion, S = 1.0f, T = translation
-	const VectorRegister InverseRotation = MAKE_QUATINV_VECTORREGISTER(Rotation);	
+	//QST(P) = Q.Rotate(S*P) + T where Q = quaternion, S = 1.0f, T = translation
 
-	//FQuat Transform = Rotation*FQuat(V.X, V.Y, V.Z, 0.f)*Rotation.Inverse();
+	//RotatedVec = Q.Rotate(V.X, V.Y, V.Z, 0.f)
 	const VectorRegister InputVectorW0 = VectorSet_W0(InputVector);	
-	const VectorRegister TempStorage = VectorQuaternionMultiply2(Rotation, InputVectorW0);
-	const VectorRegister RotatedVec = VectorQuaternionMultiply2(TempStorage, InverseRotation);
+	const VectorRegister RotatedVec = VectorQuaternionRotateVector(Rotation, InputVectorW0);
 
 	// NewVect.XYZ += Translation * W
 	// NewVect.W += 1 * W
@@ -1423,14 +1414,12 @@ FORCEINLINE FVector4 FTransform::TransformFVector4(const FVector4& V) const
 	const VectorRegister InputVector = VectorLoadAligned(&V);
 
 	//Transform using QST is following
-	//QST(P) = Q*S*P*-Q + T where Q = quaternion, S = scale, T = translation
-	const VectorRegister InverseRotation = MAKE_QUATINV_VECTORREGISTER(Rotation);	
+	//QST(P) = Q.Rotate(S*P) + T where Q = quaternion, S = scale, T = translation
 
-	//FQuat Transform = Rotation*FQuat(Scale*V.X, Scale*V.Y, Scale*V.Z, 0.f)*Rotation.Inverse();
+	//RotatedVec = Q.Rotate(Scale*V.X, Scale*V.Y, Scale*V.Z, 0.f)
 	const VectorRegister InputVectorW0 = VectorSet_W0(InputVector);
 	const VectorRegister ScaledVec = VectorMultiply(Scale3D, InputVectorW0);
-	const VectorRegister TempStorage = VectorQuaternionMultiply2(Rotation, ScaledVec);
-	const VectorRegister RotatedVec = VectorQuaternionMultiply2(TempStorage, InverseRotation);
+	const VectorRegister RotatedVec = VectorQuaternionRotateVector(Rotation, ScaledVec);
 
 	// NewVect.XYZ += Translation * W
 	// NewVect.W += 1 * W
@@ -1450,13 +1439,11 @@ FORCEINLINE FVector FTransform::TransformPosition(const FVector& V) const
 	const VectorRegister InputVectorW0 = VectorLoadFloat3_W0(&V);
 
 	//Transform using QST is following
-	//QST(P) = Q*S*P*-Q + T where Q = quaternion, S = scale, T = translation
-	const VectorRegister InverseRotation = MAKE_QUATINV_VECTORREGISTER(Rotation);
-
-	//FQuat Transform = Rotation*FQuat(Scale*V.X, Scale*V.Y, Scale*V.Z, 0.f)*Rotation.Inverse();
+	//QST(P) = Q.Rotate(S*P) + T where Q = quaternion, S = scale, T = translation
+	
+	//RotatedVec = Q.Rotate(Scale*V.X, Scale*V.Y, Scale*V.Z, 0.f)
 	const VectorRegister ScaledVec = VectorMultiply(Scale3D, InputVectorW0);
-	const VectorRegister TempStorage = VectorQuaternionMultiply2(Rotation, ScaledVec);
-	const VectorRegister RotatedVec = VectorQuaternionMultiply2(TempStorage, InverseRotation);
+	const VectorRegister RotatedVec = VectorQuaternionRotateVector(Rotation, ScaledVec);
 
 	const VectorRegister TranslatedVec = VectorAdd(RotatedVec, Translation);
 
@@ -1472,12 +1459,10 @@ FORCEINLINE FVector FTransform::TransformPositionNoScale(const FVector& V) const
 	const VectorRegister InputVectorW0 = VectorLoadFloat3_W0(&V);
 
 	//Transform using QST is following
-	//QST(P) = Q*S*P*-Q + T where Q = quaternion, S = 1.0f, T = translation
-	const VectorRegister InverseRotation = MAKE_QUATINV_VECTORREGISTER(Rotation);
+	//QST(P) = Q.Rotate(S*P) + T where Q = quaternion, S = 1.0f, T = translation
 
-	//FQuat Transform = Rotation*FQuat(V.X, V.Y, V.Z, 0.f)*Rotation.Inverse();
-	const VectorRegister TempStorage = VectorQuaternionMultiply2(Rotation, InputVectorW0);
-	const VectorRegister RotatedVec = VectorQuaternionMultiply2(TempStorage, InverseRotation);
+	//RotatedVec = Q.Rotate(V.X, V.Y, V.Z, 0.f)
+	const VectorRegister RotatedVec = VectorQuaternionRotateVector(Rotation, InputVectorW0);
 
 	const VectorRegister TranslatedVec = VectorAdd(RotatedVec, Translation);
 
@@ -1491,12 +1476,10 @@ FORCEINLINE FVector FTransform::TransformVector(const FVector& V) const
 	DiagnosticCheckNaN_All();
 
 	const VectorRegister InputVectorW0 = VectorLoadFloat3_W0(&V);
-	const VectorRegister InverseRotation = MAKE_QUATINV_VECTORREGISTER(Rotation);
 
-	//FQuat Transform = Rotation*FQuat(Scale*V.X, Scale*V.Y, Scale*V.Z, 0.f)*Rotation.Inverse();
+	//RotatedVec = Q.Rotate(Scale*V.X, Scale*V.Y, Scale*V.Z, 0.f)
 	const VectorRegister ScaledVec = VectorMultiply(Scale3D, InputVectorW0);
-	const VectorRegister TempStorage = VectorQuaternionMultiply2(Rotation, ScaledVec);
-	const VectorRegister RotatedVec = VectorQuaternionMultiply2(TempStorage, InverseRotation);
+	const VectorRegister RotatedVec = VectorQuaternionRotateVector(Rotation, ScaledVec);
 
 	FVector Result;
 	VectorStoreFloat3(RotatedVec, &Result);
@@ -1508,11 +1491,9 @@ FORCEINLINE FVector FTransform::TransformVectorNoScale(const FVector& V) const
 	DiagnosticCheckNaN_All();
 
 	const VectorRegister InputVectorW0 = VectorLoadFloat3_W0(&V);
-	const VectorRegister InverseRotation = MAKE_QUATINV_VECTORREGISTER(Rotation);
 
-	//FQuat Transform = Rotation*FQuat(V.X, V.Y, V.Z, 0.f)*Rotation.Inverse()
-	const VectorRegister TempStorage = VectorQuaternionMultiply2(Rotation, InputVectorW0);
-	const VectorRegister RotatedVec = VectorQuaternionMultiply2(TempStorage, InverseRotation);
+	//RotatedVec = Q.Rotate(V.X, V.Y, V.Z, 0.f)
+	const VectorRegister RotatedVec = VectorQuaternionRotateVector(Rotation, InputVectorW0);
 
 	FVector Result;
 	VectorStoreFloat3(RotatedVec, &Result);
@@ -1524,23 +1505,19 @@ FORCEINLINE FVector FTransform::InverseTransformPosition(const FVector &V) const
 {
 	DiagnosticCheckNaN_All();
 
-	VectorRegister InputVector = VectorLoadFloat3_W0(&V);
-
-	// Rotation.Inverse()
-	VectorRegister InverseRotation = MAKE_QUATINV_VECTORREGISTER(Rotation);
+	const VectorRegister InputVector = VectorLoadFloat3_W0(&V);
 
 	// (V-Translation)
 	const VectorRegister TranslatedVec = VectorSet_W0(VectorSubtract(InputVector, Translation));
 
-	// ( Rotation.Inverse() * (V-Translation) ); aka. FVector FQuat::operator*( const FVector& V ) const
-	VectorRegister VT = VectorQuaternionMultiply2(InverseRotation, TranslatedVec);
-	VectorRegister VR = VectorQuaternionMultiply2(VT, Rotation);
+	// ( Rotation.Inverse() * (V-Translation) )
+	const VectorRegister VR = VectorQuaternionInverseRotateVector(Rotation, TranslatedVec);
 
 	// GetSafeScaleReciprocal(Scale3D);
-	VectorRegister SafeReciprocal = GetSafeScaleReciprocal(Scale3D);	
+	const VectorRegister SafeReciprocal = GetSafeScaleReciprocal(Scale3D);	
 
 	// ( Rotation.Inverse() * (V-Translation) ) * GetSafeScaleReciprocal(Scale3D);
-	VectorRegister VResult = VectorMultiply(VR , SafeReciprocal);
+	const VectorRegister VResult = VectorMultiply(VR, SafeReciprocal);
 
 	FVector Result;
 	VectorStoreFloat3(VResult, &Result);
@@ -1552,17 +1529,13 @@ FORCEINLINE FVector FTransform::InverseTransformPositionNoScale(const FVector &V
 {
 	DiagnosticCheckNaN_All();
 
-	VectorRegister InputVector = VectorLoadFloat3_W0(&V);
-
-	// Rotation.Inverse()
-	VectorRegister InverseRotation = MAKE_QUATINV_VECTORREGISTER(Rotation);
+	const VectorRegister InputVector = VectorLoadFloat3_W0(&V);
 
 	// (V-Translation)
 	const VectorRegister TranslatedVec = VectorSet_W0(VectorSubtract(InputVector, Translation));
 
-	// ( Rotation.Inverse() * (V-Translation) ); aka. FVector FQuat::operator*( const FVector& V ) const
-	VectorRegister VT = VectorQuaternionMultiply2(InverseRotation, TranslatedVec);
-	VectorRegister VResult = VectorQuaternionMultiply2(VT, Rotation);
+	// ( Rotation.Inverse() * (V-Translation) )
+	const VectorRegister VResult = VectorQuaternionInverseRotateVector(Rotation, TranslatedVec);
 
 	FVector Result;
 	VectorStoreFloat3(VResult, &Result);
@@ -1575,20 +1548,16 @@ FORCEINLINE FVector FTransform::InverseTransformVector(const FVector &V) const
 {
 	DiagnosticCheckNaN_All();
 
-	VectorRegister InputVector = VectorLoadFloat3_W0(&V);
-
-	// Rotation.Inverse()
-	VectorRegister InverseRotation = MAKE_QUATINV_VECTORREGISTER(Rotation);
+	const VectorRegister InputVector = VectorLoadFloat3_W0(&V);
 
 	// ( Rotation.Inverse() * V ) aka. FVector FQuat::operator*( const FVector& V ) const
-	VectorRegister VT = VectorQuaternionMultiply2(InverseRotation, InputVector);
-	VectorRegister VR = VectorQuaternionMultiply2(VT, Rotation);
+	const VectorRegister VR = VectorQuaternionInverseRotateVector(Rotation, InputVector);
 
 	// GetSafeScaleReciprocal(Scale3D);
-	VectorRegister SafeReciprocal = GetSafeScaleReciprocal(Scale3D);
+	const VectorRegister SafeReciprocal = GetSafeScaleReciprocal(Scale3D);
 
 	// ( Rotation.Inverse() * V) * GetSafeScaleReciprocal(Scale3D);
-	VectorRegister VResult = VectorMultiply(VR , SafeReciprocal);
+	const VectorRegister VResult = VectorMultiply(VR, SafeReciprocal);
 
 	FVector Result;
 	VectorStoreFloat3(VResult, &Result);
@@ -1602,12 +1571,8 @@ FORCEINLINE FVector FTransform::InverseTransformVectorNoScale(const FVector &V)
 
 	VectorRegister InputVector = VectorLoadFloat3_W0(&V);
 
-	// Rotation.Inverse()
-	VectorRegister InverseRotation = MAKE_QUATINV_VECTORREGISTER(Rotation);
-
-	// ( Rotation.Inverse() * V ); aka. FVector FQuat::operator*( const FVector& V ) const
-	VectorRegister VT = VectorQuaternionMultiply2(InverseRotation, InputVector);
-	VectorRegister VResult = VectorQuaternionMultiply2(VT, Rotation);
+	// ( Rotation.Inverse() * V )
+	VectorRegister VResult = VectorQuaternionInverseRotateVector(Rotation, InputVector);
 
 	FVector Result;
 	VectorStoreFloat3(VResult, &Result);
diff --git a/Engine/Source/Runtime/Core/Public/Math/UnrealMathVectorCommon.h b/Engine/Source/Runtime/Core/Public/Math/UnrealMathVectorCommon.h
--- a/Engine/Source/Runtime/Core/Public/Math/UnrealMathVectorCommon.h
+++ b/Engine/Source/Runtime/Core/Public/Math/UnrealMathVectorCommon.h
@@ -170,3 +170,76 @@ FORCEINLINE VectorRegister VectorQuaternionInverse(const VectorRegister& Normali
 {
 	return VectorMultiply(GlobalVectorConstants::QINV_SIGN_MASK, NormalizedQuat);
 }
+
+
+// Return whether the W component is zero. Should only be used for debug purposes (it's slow).
+inline bool DebugVectorWIsZero(const VectorRegister& Vector_W0)
+{
+	union { VectorRegister v; float f[4]; } TestFloats;
+	VectorStoreAligned(Vector_W0, TestFloats.f);
+	return (TestFloats.f[3] == 0.f);
+}
+
+
+/**
+ * Rotate a vector using a unit Quaternion.
+ *
+ * @param Quat Unit Quaternion to use for rotation.
+ * @param Vector_W0 Vector to rotate. W component must be zero.
+ * @return Vector after rotation by Quat.
+ */
+FORCEINLINE VectorRegister VectorQuaternionRotateVector(const VectorRegister& Quat, const VectorRegister& Vector_W0)
+{
+	//checkfSlow(DebugVectorWIsZero(Vector_W0), TEXT("VectorQuaternionRotateVector: Vector_W0.W is not zero before rotation!"));
+
+	// Q * V * Q.Inverse
+	const VectorRegister InverseRotation = VectorQuaternionInverse(Quat);
+	const VectorRegister Temp = VectorQuaternionMultiply2(Quat, Vector_W0);
+	const VectorRegister Rotated = VectorQuaternionMultiply2(Temp, InverseRotation);
+
+	return Rotated;
+}
+
+/**
+ * Rotate a vector using the inverse of a unit Quaternion (rotation in the opposite direction).
+ *
+ * @param Quat Unit Quaternion to use for rotation.
+ * @param Vector_W0 Vector to rotate. W component must be zero.
+ * @return Vector after rotation by the inverse of Quat.
+ */
+FORCEINLINE VectorRegister VectorQuaternionInverseRotateVector(const VectorRegister& Quat, const VectorRegister& Vector_W0)
+{
+	//checkfSlow(DebugVectorWIsZero(Vector_W0), TEXT("VectorQuaternionRotateVector: Vector_W0.W is not zero before rotation!"));
+
+	// Q.Inverse * V * Q
+	const VectorRegister InverseRotation = VectorQuaternionInverse(Quat);
+	const VectorRegister Temp = VectorQuaternionMultiply2(InverseRotation, Vector_W0);
+	const VectorRegister Rotated = VectorQuaternionMultiply2(Temp, Quat);
+
+	return Rotated;
+}
+
+
+/**
+* Rotate a vector using a pointer to a unit Quaternion.
+*
+* @param Result		Pointer to where the result should be stored
+* @param Quat		Pointer to the unit quaternion (must not be the destination)
+* @param Vector_W0	Pointer to the vector (must not be the destination). W component must be zero.
+*/
+FORCEINLINE void VectorQuaternionRotateVectorPtr( void* RESTRICT Result, const void* RESTRICT Quat, const void* RESTRICT Vector_W0)
+{
+	*((VectorRegister *)Result) = VectorQuaternionRotateVector(*((const VectorRegister *)Quat), *((const VectorRegister *)Vector_W0));
+}
+
+/**
+* Rotate a vector using the inverse of a unit Quaternion (rotation in the opposite direction).
+*
+* @param Result		Pointer to where the result should be stored
+* @param Quat		Pointer to the unit quaternion (must not be the destination)
+* @param Vector_W0	Pointer to the vector (must not be the destination). W component must be zero.
+*/
+FORCEINLINE void VectorQuaternionInverseRotateVectorPtr(void* RESTRICT Result, const void* RESTRICT Quat, const void* RESTRICT Vector_W0)
+{
+	*((VectorRegister *)Result) = VectorQuaternionInverseRotateVector(*((const VectorRegister *)Quat), *((const VectorRegister *)Vector_W0));
+}
