More fixes for Android window resoliution caching
#ue4
#android

diff --git a/Engine/Source/Runtime/Core/Private/Android/AndroidWindow.cpp b/Engine/Source/Runtime/Core/Private/Android/AndroidWindow.cpp
--- a/Engine/Source/Runtime/Core/Private/Android/AndroidWindow.cpp
+++ b/Engine/Source/Runtime/Core/Private/Android/AndroidWindow.cpp
@@ -12,6 +12,7 @@ static int32 WindowWidth = -1;
 static int32 WindowHeight = -1;
 static bool WindowInit = false;
 static float ContentScaleFactor = -1.0f;
+static ANativeWindow* LastWindow = NULL;
 
 FAndroidWindow::~FAndroidWindow()
 {
@@ -76,8 +77,26 @@ FPlatformRect FAndroidWindow::GetScreenRect()
 	static IConsoleVariable* CVar = IConsoleManager::Get().FindConsoleVariable(TEXT("r.MobileContentScaleFactor"));
 	float RequestedContentScaleFactor = CVar->GetFloat();
 
-	// since orientation and resolution won't change on Android, use cached results if still valid
-	if (WindowInit && RequestedContentScaleFactor == ContentScaleFactor)
+	ANativeWindow* Window = (ANativeWindow*)FPlatformMisc::GetHardwareWindow();
+	check(Window != NULL);
+
+	if (RequestedContentScaleFactor != ContentScaleFactor)
+	{
+		FPlatformMisc::LowLevelOutputDebugString(TEXT("***** RequestedContentScaleFactor different %d != %d, not using res cache"), RequestedContentScaleFactor, ContentScaleFactor);
+	}
+
+	if (Window != LastWindow)
+	{
+		FPlatformMisc::LowLevelOutputDebugString(TEXT("***** Window different, not using res cache"));
+	}
+
+	if (WindowWidth <= 8)
+	{
+		FPlatformMisc::LowLevelOutputDebugString(TEXT("***** WindowWidth is %d, not using res cache"), WindowWidth);
+	}
+
+	// since orientation won't change on Android, use cached results if still valid
+	if (WindowInit && RequestedContentScaleFactor == ContentScaleFactor && Window == LastWindow && WindowWidth > 8)
 	{
 		FPlatformRect ScreenRect;
 		ScreenRect.Left = 0;
@@ -90,9 +109,6 @@ FPlatformRect FAndroidWindow::GetScreenRect()
 
 	// currently hardcoding resolution
 
-	ANativeWindow* Window = (ANativeWindow*)FPlatformMisc::GetHardwareWindow();
-	check(Window != NULL);
-
 	// get the aspect ratio of the physical screen
 	int32 ScreenWidth, ScreenHeight;
 	CalculateSurfaceSize(Window, ScreenWidth, ScreenHeight);
@@ -167,6 +183,7 @@ FPlatformRect FAndroidWindow::GetScreenRect()
 	WindowHeight = Height;
 	WindowInit = true;
 	ContentScaleFactor = RequestedContentScaleFactor;
+	LastWindow = Window;
 
 	return ScreenRect;
 }
