Merging using UE4-To-UE4-LauncherDev

diff --git a/Engine/Source/Developer/SuperSearch/SuperSearch.Build.cs b/Engine/Source/Developer/SuperSearch/SuperSearch.Build.cs
--- a/Engine/Source/Developer/SuperSearch/SuperSearch.Build.cs
+++ b/Engine/Source/Developer/SuperSearch/SuperSearch.Build.cs
@@ -4,12 +4,12 @@ using UnrealBuildTool;
 
 public class SuperSearch : ModuleRules
 {
-	public SuperSearch(TargetInfo Target)
-	{
-		if (UEBuildConfiguration.bBuildEditor == true)
-		{
-		PrivateDependencyModuleNames.AddRange(
-			new string[] {
+    public SuperSearch(TargetInfo Target)
+    {
+        if (UEBuildConfiguration.bBuildEditor == true)
+        {
+            PrivateDependencyModuleNames.AddRange(
+                new string[] {
 					"Core",
 					"CoreUObject", // @todo Mac: for some reason it's needed to link in debug on Mac
 					"Engine",
@@ -22,19 +22,21 @@ public class SuperSearch : ModuleRules
 					"HTTP",
 					"IntroTutorials"
 				}
-			);
-		}
-		else
-		{
-	   		PrivateDependencyModuleNames.AddRange(
-				new string[] {
+            );
+        }
+        else
+        {
+            PrivateDependencyModuleNames.AddRange(
+                new string[] {
 					"Core",
 					"CoreUObject", // @todo Mac: for some reason it's needed to link in debug on Mac
+                    "InputCore",
 					"Slate",
 					"SlateCore",
-	                "HTTP"
-			}
-		);
-	}
-}
+	                "HTTP",
+                    "Json",
+			    }
+            );
+        }
+    }
 }
diff --git a/Engine/Source/Runtime/Core/Public/Windows/WindowsApplication.h b/Engine/Source/Runtime/Core/Public/Windows/WindowsApplication.h
--- a/Engine/Source/Runtime/Core/Public/Windows/WindowsApplication.h
+++ b/Engine/Source/Runtime/Core/Public/Windows/WindowsApplication.h
@@ -22,7 +22,7 @@ class FGenericApplicationMessageHandler;
 
 namespace ETaskbarProgressState
 {
-	enum Type
+	enum CORE_API Type
 	{
 		//Stops displaying progress and returns the button to its normal state.
 		NoProgress = 0x0,
@@ -60,7 +60,7 @@ namespace ETaskbarProgressState
  * This class can be used to change the appearance of a window's entry in the windows task bar,
  * such as setting an overlay icon or showing a progress indicator.
  */
-class FTaskbarList
+class CORE_API FTaskbarList
 {
 public:
 
@@ -272,7 +272,7 @@ public:
 /**
  * Windows-specific application implementation.
  */
-class FWindowsApplication
+class CORE_API FWindowsApplication
 	: public GenericApplication
 	, public IForceFeedbackSystem
 {
@@ -290,7 +290,7 @@ public:
 	/** Virtual destructor. */
 	virtual ~FWindowsApplication();
 
-public:
+public:	
 
 	/** Called by a window when an OLE Drag and Drop operation occurred on a non-game thread */
 	void DeferDragDropOperation( const FDeferredWindowsDragDropOperation& DeferredDragDropOperation );
@@ -447,11 +447,11 @@ private:
 	bool bHasLoadedInputPlugins;
 
 	TArray<int32> PressedModifierKeys;
-	
+
 	FModifierKeysState CachedModifierKeyState;
 
 	int32 bAllowedToDeferMessageProcessing;
-
+	
 	FAutoConsoleVariableRef CVarDeferMessageProcessing;
 	
 	/** True if we are in the middle of a windows modal size loop */
@@ -464,9 +464,9 @@ private:
 	TSharedPtr<FTaskbarList> TaskbarList;
 
 	// Accessibility shortcut keys
-	STICKYKEYS StartupStickyKeys;
-	TOGGLEKEYS StartupToggleKeys;
-	FILTERKEYS StartupFilterKeys;
+	STICKYKEYS							StartupStickyKeys;
+	TOGGLEKEYS							StartupToggleKeys;
+	FILTERKEYS							StartupFilterKeys;
 };
 
 
