- Android build fix.

diff --git a/Engine/Source/Runtime/Core/Private/Android/AndroidWindow.cpp b/Engine/Source/Runtime/Core/Private/Android/AndroidWindow.cpp
--- a/Engine/Source/Runtime/Core/Private/Android/AndroidWindow.cpp
+++ b/Engine/Source/Runtime/Core/Private/Android/AndroidWindow.cpp
@@ -82,7 +82,7 @@ FPlatformRect FAndroidWindow::GetScreenRect()
 
 	if (RequestedContentScaleFactor != ContentScaleFactor)
 	{
-		FPlatformMisc::LowLevelOutputDebugString(TEXT("***** RequestedContentScaleFactor different %d != %d, not using res cache"), RequestedContentScaleFactor, ContentScaleFactor);
+		FPlatformMisc::LowLevelOutputDebugStringf(TEXT("***** RequestedContentScaleFactor different %d != %d, not using res cache"), RequestedContentScaleFactor, ContentScaleFactor);
 	}
 
 	if (Window != LastWindow)
@@ -92,7 +92,7 @@ FPlatformRect FAndroidWindow::GetScreenRect()
 
 	if (WindowWidth <= 8)
 	{
-		FPlatformMisc::LowLevelOutputDebugString(TEXT("***** WindowWidth is %d, not using res cache"), WindowWidth);
+		FPlatformMisc::LowLevelOutputDebugStringf(TEXT("***** WindowWidth is %d, not using res cache"), WindowWidth);
 	}
 
 	// since orientation won't change on Android, use cached results if still valid
