#ue4 - Perf: convert VectorMultiply() and VectorAdd() to VectorMultiplyAdd() in TransformFVector and variants.

diff --git a/Engine/Source/Runtime/Core/Public/Math/TransformVectorized.h b/Engine/Source/Runtime/Core/Public/Math/TransformVectorized.h
--- a/Engine/Source/Runtime/Core/Public/Math/TransformVectorized.h
+++ b/Engine/Source/Runtime/Core/Public/Math/TransformVectorized.h
@@ -1397,7 +1397,7 @@ FORCEINLINE FVector4 FTransform::TransformFVector4NoScale(const FVector4& V) con
 	// NewVect.XYZ += Translation * W
 	// NewVect.W += 1 * W
 	const VectorRegister WWWW = VectorReplicate(InputVector, 3);
-	const VectorRegister TranslatedVec = VectorAdd(RotatedVec, VectorMultiply(Translation, WWWW));
+	const VectorRegister TranslatedVec = VectorMultiplyAdd(Translation, WWWW, RotatedVec);
 
 	FVector4 NewVectOutput;
 	VectorStoreAligned(TranslatedVec, &NewVectOutput);
@@ -1424,7 +1424,7 @@ FORCEINLINE FVector4 FTransform::TransformFVector4(const FVector4& V) const
 	// NewVect.XYZ += Translation * W
 	// NewVect.W += 1 * W
 	const VectorRegister WWWW = VectorReplicate(InputVector, 3);
-	const VectorRegister TranslatedVec = VectorAdd(RotatedVec, VectorMultiply(Translation, WWWW));
+	const VectorRegister TranslatedVec = VectorMultiplyAdd(Translation, WWWW, RotatedVec);
 
 	FVector4 NewVectOutput;
 	VectorStoreAligned(TranslatedVec, &NewVectOutput);
