#ue4 - Perf: Implement faster Quaternion Vector rotation.

- Scalar FQuat::RotateVector() takes ~18% less time.
- VectorQuaternionRotateVector() takes ~25% less time.
- Vectorized FTransform * FTransform takes ~33% less time.

diff --git a/Engine/Source/Runtime/Core/Public/Math/Quat.h b/Engine/Source/Runtime/Core/Public/Math/Quat.h
--- a/Engine/Source/Runtime/Core/Public/Math/Quat.h
+++ b/Engine/Source/Runtime/Core/Public/Math/Quat.h
@@ -860,27 +860,19 @@ FORCEINLINE FVector FQuat::RotateVector( FVector V ) const
 	VectorQuaternionVector3Rotate(&Result, &V, this);
 	return Result;
 
-	/*
-	// In unit testing this appears to be slower than the non-vectorized version.
-#elif PLATFORM_ENABLE_VECTORINTRINSICS
-	FQuat VQ(V.X, V.Y, V.Z, 0.f);
-	FQuat VT, VR;
-	FQuat I = Inverse();
-	VectorQuaternionMultiply(&VT, this, &VQ);
-	VectorQuaternionMultiply(&VR, &VT, &I);
-
-	return FVector(VR.X, VR.Y, VR.Z);
-	*/
-
 #else
-	// (q.W*q.W-qv.qv)v + 2(qv.v)qv + 2 q.W (qv x v)
 
-	const FVector qv(X, Y, Z);
-	FVector vOut = (2.f * W) * (qv ^ V);
-	vOut += ((W * W) - (qv | qv)) * V;
-	vOut += (2.f * (qv | V)) * qv;
+	// http://people.csail.mit.edu/bkph/articles/Quaternions.pdf
+	// V' = V + 2w(Q x V) + (2Q x (Q x V))
+	// refactor:
+	// V' = V + w(2(Q x V)) + (Q x (2(Q x V)))
+	// T = 2(Q x V);
+	// V' = V + w*(T) + (Q x T)
 
-	return vOut;
+	const FVector Q(X, Y, Z);
+	const FVector T = 2.f * FVector::CrossProduct(Q, V);
+	const FVector Result = V + (W * T) + FVector::CrossProduct(Q, T);
+	return Result;
 #endif
 }
 
@@ -891,7 +883,12 @@ FORCEINLINE FVector FQuat::UnrotateVector( FVector V ) const
 	VectorQuaternionVector3InverseRotate(&Result, &V, this);
 	return Result;
 #else
-	return Inverse().RotateVector(V);
+	//return Inverse().RotateVector(V);
+
+	const FVector Q(-X, -Y, -Z); // Inverse
+	const FVector T = 2.f * FVector::CrossProduct(Q, V);
+	const FVector Result = V + (W * T) + FVector::CrossProduct(Q, T);
+	return Result;
 #endif
 }
 
diff --git a/Engine/Source/Runtime/Core/Public/Math/UnrealMathVectorCommon.h b/Engine/Source/Runtime/Core/Public/Math/UnrealMathVectorCommon.h
--- a/Engine/Source/Runtime/Core/Public/Math/UnrealMathVectorCommon.h
+++ b/Engine/Source/Runtime/Core/Public/Math/UnrealMathVectorCommon.h
@@ -172,31 +172,34 @@ FORCEINLINE VectorRegister VectorQuaternionInverse(const VectorRegister& Normali
 }
 
 
-// Return whether the W component is zero. Should only be used for debug purposes (it's slow).
-inline bool DebugVectorWIsZero(const VectorRegister& Vector_W0)
-{
-	union { VectorRegister v; float f[4]; } TestFloats;
-	VectorStoreAligned(Vector_W0, TestFloats.f);
-	return (TestFloats.f[3] == 0.f);
-}
-
-
 /**
  * Rotate a vector using a unit Quaternion.
  *
  * @param Quat Unit Quaternion to use for rotation.
- * @param Vector_W0 Vector to rotate. W component must be zero.
+ * @param VectorW0 Vector to rotate. W component must be zero.
  * @return Vector after rotation by Quat.
  */
-FORCEINLINE VectorRegister VectorQuaternionRotateVector(const VectorRegister& Quat, const VectorRegister& Vector_W0)
+FORCEINLINE VectorRegister VectorQuaternionRotateVector(const VectorRegister& Quat, const VectorRegister& VectorW0)
 {
-	//checkfSlow(DebugVectorWIsZero(Vector_W0), TEXT("VectorQuaternionRotateVector: Vector_W0.W is not zero before rotation!"));
-
 	// Q * V * Q.Inverse
-	const VectorRegister InverseRotation = VectorQuaternionInverse(Quat);
-	const VectorRegister Temp = VectorQuaternionMultiply2(Quat, Vector_W0);
-	const VectorRegister Rotated = VectorQuaternionMultiply2(Temp, InverseRotation);
-
+	//const VectorRegister InverseRotation = VectorQuaternionInverse(Quat);
+	//const VectorRegister Temp = VectorQuaternionMultiply2(Quat, VectorW0);
+	//const VectorRegister Rotated = VectorQuaternionMultiply2(Temp, InverseRotation);
+
+	// Equivalence of above can be shown to be:
+	// http://people.csail.mit.edu/bkph/articles/Quaternions.pdf
+	// V' = V + 2w(Q x V) + (2Q x (Q x V))
+	// refactor:
+	// V' = V + w(2(Q x V)) + (Q x (2(Q x V)))
+	// T = 2(Q x V);
+	// V' = V + w*(T) + (Q x T)
+
+	const VectorRegister QW = VectorReplicate(Quat, 3);
+	VectorRegister T = VectorCross(Quat, VectorW0);
+	T = VectorAdd(T, T);
+	const VectorRegister VTemp0 = VectorMultiplyAdd(QW, T, VectorW0);
+	const VectorRegister VTemp1 = VectorCross(Quat, T);
+	const VectorRegister Rotated = VectorAdd(VTemp0, VTemp1);
 	return Rotated;
 }
 
@@ -204,19 +207,18 @@ FORCEINLINE VectorRegister VectorQuaternionRotateVector(const VectorRegister& Qu
  * Rotate a vector using the inverse of a unit Quaternion (rotation in the opposite direction).
  *
  * @param Quat Unit Quaternion to use for rotation.
- * @param Vector_W0 Vector to rotate. W component must be zero.
+ * @param VectorW0 Vector to rotate. W component must be zero.
  * @return Vector after rotation by the inverse of Quat.
  */
-FORCEINLINE VectorRegister VectorQuaternionInverseRotateVector(const VectorRegister& Quat, const VectorRegister& Vector_W0)
+FORCEINLINE VectorRegister VectorQuaternionInverseRotateVector(const VectorRegister& Quat, const VectorRegister& VectorW0)
 {
-	//checkfSlow(DebugVectorWIsZero(Vector_W0), TEXT("VectorQuaternionRotateVector: Vector_W0.W is not zero before rotation!"));
-
 	// Q.Inverse * V * Q
-	const VectorRegister InverseRotation = VectorQuaternionInverse(Quat);
-	const VectorRegister Temp = VectorQuaternionMultiply2(InverseRotation, Vector_W0);
-	const VectorRegister Rotated = VectorQuaternionMultiply2(Temp, Quat);
+	//const VectorRegister InverseRotation = VectorQuaternionInverse(Quat);
+	//const VectorRegister Temp = VectorQuaternionMultiply2(InverseRotation, VectorW0);
+	//const VectorRegister Rotated = VectorQuaternionMultiply2(Temp, Quat);
 
-	return Rotated;
+	const VectorRegister QInv = VectorQuaternionInverse(Quat);
+	return VectorQuaternionRotateVector(QInv, VectorW0);
 }
 
 
@@ -225,11 +227,11 @@ FORCEINLINE VectorRegister VectorQuaternionInverseRotateVector(const VectorRegis
 *
 * @param Result		Pointer to where the result should be stored
 * @param Quat		Pointer to the unit quaternion (must not be the destination)
-* @param Vector_W0	Pointer to the vector (must not be the destination). W component must be zero.
+* @param VectorW0	Pointer to the vector (must not be the destination). W component must be zero.
 */
-FORCEINLINE void VectorQuaternionRotateVectorPtr( void* RESTRICT Result, const void* RESTRICT Quat, const void* RESTRICT Vector_W0)
+FORCEINLINE void VectorQuaternionRotateVectorPtr( void* RESTRICT Result, const void* RESTRICT Quat, const void* RESTRICT VectorW0)
 {
-	*((VectorRegister *)Result) = VectorQuaternionRotateVector(*((const VectorRegister *)Quat), *((const VectorRegister *)Vector_W0));
+	*((VectorRegister *)Result) = VectorQuaternionRotateVector(*((const VectorRegister *)Quat), *((const VectorRegister *)VectorW0));
 }
 
 /**
@@ -237,9 +239,9 @@ FORCEINLINE void VectorQuaternionRotateVectorPtr( void* RESTRICT Result, const v
 *
 * @param Result		Pointer to where the result should be stored
 * @param Quat		Pointer to the unit quaternion (must not be the destination)
-* @param Vector_W0	Pointer to the vector (must not be the destination). W component must be zero.
+* @param VectorW0	Pointer to the vector (must not be the destination). W component must be zero.
 */
-FORCEINLINE void VectorQuaternionInverseRotateVectorPtr(void* RESTRICT Result, const void* RESTRICT Quat, const void* RESTRICT Vector_W0)
+FORCEINLINE void VectorQuaternionInverseRotateVectorPtr(void* RESTRICT Result, const void* RESTRICT Quat, const void* RESTRICT VectorW0)
 {
-	*((VectorRegister *)Result) = VectorQuaternionInverseRotateVector(*((const VectorRegister *)Quat), *((const VectorRegister *)Vector_W0));
+	*((VectorRegister *)Result) = VectorQuaternionInverseRotateVector(*((const VectorRegister *)Quat), *((const VectorRegister *)VectorW0));
 }
