An improved way of identifying whether a config should be considered a default write or not.

UE-14447

diff --git a/Engine/Source/Runtime/Core/Private/Misc/ConfigCacheIni.cpp b/Engine/Source/Runtime/Core/Private/Misc/ConfigCacheIni.cpp
--- a/Engine/Source/Runtime/Core/Private/Misc/ConfigCacheIni.cpp
+++ b/Engine/Source/Runtime/Core/Private/Misc/ConfigCacheIni.cpp
@@ -947,9 +947,7 @@ bool FConfigFile::Write( const FString& Filename, bool bDoRemoteWrite/* = true*/
 				const bool bOptionIsFromCommandline = PropertySetFromCommandlineOption(this, SectionName, PropertyName, PropertyValue);
 
 				// If we are writing to a default config file and this property is an array, we need to be careful to remove those from higher up the hierarchy
-				const bool bIsADefaultIniWrite = !Filename.Contains( FPaths::GameSavedDir() ) 
-					&& !Filename.Contains( FPaths::EngineSavedDir() ) 
-					&& FPaths::GetBaseFilename(Filename).StartsWith( TEXT( "Default" ) );
+				const bool bIsADefaultIniWrite = !Filename.Contains( FPaths::GeneratedConfigDir() );
 
 				// Check if the property matches the source configs. We do not wanna write it out if so.
 				if ((bIsADefaultIniWrite || bDifferentNumberOfElements || !DoesConfigPropertyValueMatch(SourceConfigFile, SectionName, PropertyName, PropertyValue)) && !bOptionIsFromCommandline)
