Add functionality to uobject and fconfigfile to support editing one property at a time.
Enabled for Config editor as an initial test bed.

UE-14449

diff --git a/Engine/Source/Editor/DetailCustomizations/Private/ConfigEditorPropertyDetails.cpp b/Engine/Source/Editor/DetailCustomizations/Private/ConfigEditorPropertyDetails.cpp
--- a/Engine/Source/Editor/DetailCustomizations/Private/ConfigEditorPropertyDetails.cpp
+++ b/Engine/Source/Editor/DetailCustomizations/Private/ConfigEditorPropertyDetails.cpp
@@ -57,7 +57,7 @@ void FConfigPropertyHelperDetails::OnPropertyValueChanged(UObject* Object, FProp
 		const FString* FileName = ConfigFileAndPropertySourcePairings.FindKey(Object);
 		if (FileName != nullptr)
 		{
-			Object->UpdateDefaultConfigFile(*FileName);
+			Object->UpdateSinglePropertyInConfigFile(Property, *FileName);
 		}
 	}
 }
diff --git a/Engine/Source/Runtime/Core/Private/Misc/ConfigCacheIni.cpp b/Engine/Source/Runtime/Core/Private/Misc/ConfigCacheIni.cpp
--- a/Engine/Source/Runtime/Core/Private/Misc/ConfigCacheIni.cpp
+++ b/Engine/Source/Runtime/Core/Private/Misc/ConfigCacheIni.cpp
@@ -3160,6 +3160,99 @@ void FConfigFile::UpdateSections(const TCHAR* DiskFilename, const TCHAR* IniRoot
 	Write(DiskFilename, true, NewFile);
 }
 
+bool FConfigFile::UpdateSinglePropertyInSection(const TCHAR* DiskFilename, const TCHAR* PropertyName, const TCHAR* SectionName)
+{
+	// Result of whether the file has been updated on disk.
+	bool bSuccessfullyUpdatedFile = false;
+
+	// Since we don't want any modifications to other sections, or properties, we manually process the file, not read into sections
+	FString DiskFile;
+	if (FFileHelper::LoadFileToString(DiskFile, DiskFilename))
+	{
+		FString NewFile;
+		const FString DecoratedSectionName = FString::Printf(TEXT("[%s]"), SectionName);
+
+		// Track if the these properties belong to the section we are looking for.
+		bool bIsCurrentlyWritingSectionProperty = false;
+		
+		// Track if the section was ever updated.
+		bool bSectionWasUpdated = false;
+		// Track if the property was ever updated.
+		bool bPropertyWasUpdated = false;
+
+		// Iterate through the file, line by line, to find the property we would like to update.
+		FString TheLine;
+		const TCHAR* Ptr = DiskFile.Len() > 0 ? *DiskFile : NULL;
+		while (FParse::Line(&Ptr, TheLine, true))
+		{
+			const bool bIsASectionHeader = (TheLine.Len() > 3) && (TheLine[0] == '[') && (TheLine[TheLine.Len() - 1] == ']');
+			if (bIsASectionHeader)
+			{
+				// Cache off the old value of whether we had been writing a section.
+				const bool bWasWritingSection = bIsCurrentlyWritingSectionProperty;
+
+				if (!bSectionWasUpdated && bWasWritingSection && !bPropertyWasUpdated)
+				{
+					bPropertyWasUpdated = true;
+					// If we have found the property in the file, let's update it with the new value.
+					if (const FConfigSection* LocalSection = this->Find(SectionName))
+					{
+						if (const FString* LocalSectionPropertyStr = LocalSection->Find(PropertyName))
+						{
+							NewFile += FString::Printf(TEXT("%s=%s"), PropertyName, *(*LocalSectionPropertyStr)) + LINE_TERMINATOR + LINE_TERMINATOR;
+						}
+					}
+				}
+
+				bIsCurrentlyWritingSectionProperty = (TheLine == DecoratedSectionName);
+				bSectionWasUpdated |= bIsCurrentlyWritingSectionProperty;
+
+				NewFile += TheLine + LINE_TERMINATOR;
+			}
+			else if (bIsCurrentlyWritingSectionProperty && TheLine.StartsWith(PropertyName))
+			{
+				bPropertyWasUpdated = true;
+				// If we have found the property in the file, let's update it with the new value.
+				if (const FConfigSection* LocalSection = this->Find(SectionName))
+				{
+					if (const FString* LocalSectionPropertyStr = LocalSection->Find(PropertyName))
+					{
+						NewFile += FString::Printf(TEXT("%s=%s"), PropertyName, *(*LocalSectionPropertyStr)) + LINE_TERMINATOR;
+					}
+				}
+			}
+			else
+			{
+				NewFile += TheLine + LINE_TERMINATOR;
+			}
+		}
+
+		if (bSectionWasUpdated == false)
+		{
+			// If we reached the eof without updating the section, write the header
+			NewFile += LINE_TERMINATOR + DecoratedSectionName + LINE_TERMINATOR;
+		}
+
+		if (bPropertyWasUpdated == false)
+		{
+			// If we reached the eof without updating the property, we should write it out.
+			bPropertyWasUpdated = true;
+			if (const FConfigSection* LocalSection = this->Find(SectionName))
+			{
+				if (const FString* LocalSectionPropertyStr = LocalSection->Find(PropertyName))
+				{
+					NewFile += FString::Printf(TEXT("%s=%s"), PropertyName, *(*LocalSectionPropertyStr)) + LINE_TERMINATOR;
+				}
+			}
+		}
+
+		// Write the updated file to disk.
+		bSuccessfullyUpdatedFile = FFileHelper::SaveStringToFile(NewFile, DiskFilename);
+	}
+
+	return bSuccessfullyUpdatedFile;
+}
+
 void ApplyCVarSettingsFromIni(const TCHAR* InSectionName, const TCHAR* InIniFilename, uint32 SetBy)
 {
 	// Lookup the config section for this section and group number
diff --git a/Engine/Source/Runtime/Core/Public/Misc/ConfigCacheIni.h b/Engine/Source/Runtime/Core/Public/Misc/ConfigCacheIni.h
--- a/Engine/Source/Runtime/Core/Public/Misc/ConfigCacheIni.h
+++ b/Engine/Source/Runtime/Core/Public/Misc/ConfigCacheIni.h
@@ -169,6 +169,11 @@ public:
 	 */
 	CORE_API void UpdateSections(const TCHAR* DiskFilename, const TCHAR* IniRootName=NULL);
 
+	/**
+	 * Update a single property in the config file, for the section that is specified.
+	 */
+	CORE_API bool UpdateSinglePropertyInSection(const TCHAR* DiskFilename, const TCHAR* PropertyName, const TCHAR* SectionName);
+
 	
 	/** 
 	 * Check the source hierarchy which was loaded without any user changes from the Config/Saved dir.
diff --git a/Engine/Source/Runtime/CoreUObject/Private/UObject/Obj.cpp b/Engine/Source/Runtime/CoreUObject/Private/UObject/Obj.cpp
--- a/Engine/Source/Runtime/CoreUObject/Private/UObject/Obj.cpp
+++ b/Engine/Source/Runtime/CoreUObject/Private/UObject/Obj.cpp
@@ -1864,6 +1864,41 @@ void UObject::UpdateGlobalUserConfigFile()
 	UpdateSingleSectionOfConfigFile(GetGlobalUserConfigFilename());
 }
 
+void UObject::UpdateSinglePropertyInConfigFile(const UProperty* InProperty, const FString& InConfigIniName)
+{
+	// Arrays and ini files are a mine field, for now we don't support this.
+	if (!InProperty->IsA(UArrayProperty::StaticClass()))
+	{
+		// create a sandbox FConfigCache
+		FConfigCacheIni Config(EConfigCacheType::Temporary);
+
+		// add an empty file to the config so it doesn't read in the original file (see FConfigCacheIni.Find())
+		FConfigFile& NewFile = Config.Add(InConfigIniName, FConfigFile());
+
+		// save the object properties to this file
+		SaveConfig(CPF_Config, *InConfigIniName, &Config);
+
+		// Take the saved section for this object and have the config system process and write out the one property we care about.
+		ensureMsgf(Config.Num() == 1, TEXT("UObject::UpdateDefaultConfig() caused more files than expected in the Sandbox config cache!"));
+
+		TArray<FString> Keys;
+		NewFile.GetKeys(Keys);
+
+		const FString SectionName = Keys[0];
+		const FString PropertyName = InProperty->GetFName().ToString();
+		NewFile.UpdateSinglePropertyInSection(*InConfigIniName, *PropertyName, *SectionName);
+
+		// reload the file, so that it refresh the cache internally.
+		FString FinalIniFileName;
+		GConfig->LoadGlobalIniFile(FinalIniFileName, *GetClass()->ClassConfigName.ToString(), NULL, true);
+	}
+	else
+	{
+		UE_LOG(LogObj, Warning, TEXT("UObject::UpdateSinglePropertyInConfigFile does not support this property type."));
+		return;
+	}
+}
+
 
 void UObject::InstanceSubobjectTemplates( FObjectInstancingGraph* InstanceGraph )
 {
diff --git a/Engine/Source/Runtime/CoreUObject/Public/UObject/UObject.h b/Engine/Source/Runtime/CoreUObject/Public/UObject/UObject.h
--- a/Engine/Source/Runtime/CoreUObject/Public/UObject/UObject.h
+++ b/Engine/Source/Runtime/CoreUObject/Public/UObject/UObject.h
@@ -789,6 +789,11 @@ public:
 	 */
 	void UpdateGlobalUserConfigFile();
 
+	/**
+	 * Saves just the property into the global user ini file for the class (with just the changes from base)
+	 */
+	void UpdateSinglePropertyInConfigFile(const UProperty* InProperty, const FString& InConfigIniName);
+
 private:
 	/**
 	 * Saves just the section(s) for this class into the given ini file for the class (with just the changes from base)
