Linux: re-enable git plugin (UE-13628).

- Crash was apparently caused by fork()ing inside ExecProcess.
- Replacing ExecProcess implementation from PR #1112 (also known as PR #1110) with some changes kills two birds with a single stone.
- Also enabled source control in the Linux editor.

diff --git a/Engine/Plugins/Developer/GitSourceControl/GitSourceControl.uplugin b/Engine/Plugins/Developer/GitSourceControl/GitSourceControl.uplugin
--- a/Engine/Plugins/Developer/GitSourceControl/GitSourceControl.uplugin
+++ b/Engine/Plugins/Developer/GitSourceControl/GitSourceControl.uplugin
@@ -17,8 +17,7 @@
 	[
 		{
 			"Name" : "GitSourceControl",
-			"Type" : "Editor",
-			"BlacklistPlatforms" : [ "Linux" ]
+			"Type" : "Editor"
 		}
 	]
 }
diff --git a/Engine/Source/Developer/SourceControl/Public/ISourceControlProvider.h b/Engine/Source/Developer/SourceControl/Public/ISourceControlProvider.h
--- a/Engine/Source/Developer/SourceControl/Public/ISourceControlProvider.h
+++ b/Engine/Source/Developer/SourceControl/Public/ISourceControlProvider.h
@@ -6,7 +6,7 @@
 #include "SourceControlHelpers.h"
 #include "Runtime/Core/Public/Features/IModularFeature.h"
 
-#define SOURCE_CONTROL_WITH_SLATE			!PLATFORM_LINUX
+#define SOURCE_CONTROL_WITH_SLATE			(!(PLATFORM_LINUX && IS_PROGRAM))
 
 /**
  * Hint for how to execute the operation. Note that asynchronous operations require
diff --git a/Engine/Source/Runtime/Core/Private/Linux/LinuxPlatformProcess.cpp b/Engine/Source/Runtime/Core/Private/Linux/LinuxPlatformProcess.cpp
--- a/Engine/Source/Runtime/Core/Private/Linux/LinuxPlatformProcess.cpp
+++ b/Engine/Source/Runtime/Core/Private/Linux/LinuxPlatformProcess.cpp
@@ -956,123 +956,74 @@ bool FLinuxPlatformProcess::IsApplicationRunning( const TCHAR* ProcName )
 
 bool FLinuxPlatformProcess::ExecProcess( const TCHAR* URL, const TCHAR* Params, int32* OutReturnCode, FString* OutStdOut, FString* OutStdErr )
 {
-	TArray<FString> ArgsArray;
-	FString(Params).ParseIntoArray(ArgsArray, TEXT(" "), true);
-	char *args[ArgsArray.Num()];
-
-	for(int i = 0; i < ArgsArray.Num(); i++)
+	FString CmdLineParams = Params;
+	FString ExecutableFileName = URL;
+	int32 ReturnCode = -1;
+	FString DefaultError;
+	if (!OutStdErr)
 	{
-		args[i] = TCHAR_TO_ANSI(*ArgsArray[i]);
+		OutStdErr = &DefaultError;
 	}
-	pid_t pid;
-	int status;
-	int fd_stdout[2], fd_stderr[2];
 
-	if (pipe(fd_stdout) == -1) 
-	{
-		int ErrNo = errno;
-		UE_LOG(LogHAL, Fatal, TEXT("Creating fd_stdout pipe failed with errno = %d (%s)"), ErrNo,
-			StringCast< TCHAR >(strerror(ErrNo)).Get());
-		return false;
-	}
+	void* PipeRead = nullptr;
+	void* PipeWrite = nullptr;
+	verify(FPlatformProcess::CreatePipe(PipeRead, PipeWrite));
 
-	if (pipe(fd_stderr) == -1) 
-	{
-		int ErrNo = errno;
-		UE_LOG(LogHAL, Fatal, TEXT("Creating fd_stderr pipe failed with errno = %d (%s)"), ErrNo,
-			StringCast< TCHAR >(strerror(ErrNo)).Get());
-		return false;
-	}
+	bool bInvoked = false;
 
-	pid = fork();
-	if(pid < 0)
-	{
-		int ErrNo = errno;
-		UE_LOG(LogHAL, Fatal, TEXT("fork() failed with errno = %d (%s)"), ErrNo,
-			StringCast< TCHAR >(strerror(ErrNo)).Get());
-		return false;
-	}
+	const bool bLaunchDetached = true;
+	const bool bLaunchHidden = false;
+	const bool bLaunchReallyHidden = bLaunchHidden;
 
-	if(pid == 0)
-	{
-		close(fd_stdout[0]);
-		close(fd_stderr[0]);
-		dup2(fd_stdout[1], 1);
-		dup2(fd_stderr[1], 2);
-		close(fd_stdout[1]);
-		close(fd_stderr[1]);
-
-		// TODO Not sure if we need to look up URL in PATH
-		exit(execv(TCHAR_TO_ANSI(URL), args));
-	}
-	else
+	FProcHandle ProcHandle = FPlatformProcess::CreateProc(*ExecutableFileName, *CmdLineParams, bLaunchDetached, bLaunchHidden, bLaunchReallyHidden, NULL, 0, NULL, PipeWrite);
+	if (ProcHandle.IsValid())
 	{
-		// TODO This might deadlock, should use select. Rewrite me. Also doesn't handle all errors correctly.
-		close(fd_stdout[1]);
-		close(fd_stderr[1]);
-		do
+		while (FPlatformProcess::IsProcRunning(ProcHandle))
 		{
-			pid_t wpid = waitpid(pid, &status, 0);
-			if (wpid == -1)
-			{
-				int ErrNo = errno;
-				UE_LOG(LogHAL, Fatal, TEXT("waitpid() failed with errno = %d (%s)"), ErrNo,
-				StringCast< TCHAR >(strerror(ErrNo)).Get());
-				return false;
-			}
-
-			if (WIFEXITED(status))
-			{
-				if(OutReturnCode)
-				{
-					*OutReturnCode = WEXITSTATUS(status);
-				}
-			}
-			else if (WIFSIGNALED(status))
+			FString NewLine = FPlatformProcess::ReadPipe(PipeRead);
+			if (NewLine.Len() > 0)
 			{
-				if(OutReturnCode)
+				if (OutStdOut != nullptr)
 				{
-					*OutReturnCode = WTERMSIG(status);
+					*OutStdOut += NewLine;
 				}
 			}
+			FPlatformProcess::Sleep(0.5);
+		}
 
-		} while (!WIFEXITED(status) && !WIFSIGNALED(status));
-
-		while(1)
+		// read the remainder
+		for(;;)
 		{
-			char buf[100];
-			int size = read(fd_stdout[0], buf, 100);
-			if(size)
-			{
-				if(OutStdErr)
-				{
-					*OutStdErr += FString(buf);
-				}
-			}
-			else
+			FString NewLine = FPlatformProcess::ReadPipe(PipeRead);
+			if (NewLine.Len() <= 0)
 			{
 				break;
 			}
-		}
 
-		while(1)
-		{
-			char buf[100];
-			int size = read(fd_stderr[0], buf, 100);
-			if(size)
+			if (OutStdOut != nullptr)
 			{
-				if(OutStdOut)
-				{
-					*OutStdOut += FString(buf);
-				}
-			}
-			else
-			{
-				break;
+				*OutStdOut += NewLine;
 			}
 		}
-		return true;
+
+		FPlatformProcess::Sleep(0.5);
+
+		bInvoked = true;
+		bool bGotReturnCode = FPlatformProcess::GetProcReturnCode(ProcHandle, &ReturnCode);
+		check(bGotReturnCode);
+		*OutReturnCode = ReturnCode;
+
+		FPlatformProcess::CloseProc(ProcHandle);
+	}
+	else
+	{
+		bInvoked = false;
+		*OutReturnCode = -1;
+		*OutStdOut = "";
+		UE_LOG(LogHAL, Warning, TEXT("Failed to launch Tool. (%s)"), *ExecutableFileName);
 	}
+	FPlatformProcess::ClosePipe(PipeRead, PipeWrite);
+	return bInvoked;
 }
 
 void FLinuxPlatformProcess::LaunchFileInDefaultExternalApplication( const TCHAR* FileName, const TCHAR* Parms, ELaunchVerb::Type Verb )
