Inline calls to ConditionalSleep to avoid TFunction + lambda overhead in FMultiReaderSingleWriterGT.

diff --git a/Engine/Source/Runtime/Core/Private/HAL/ThreadingBase.cpp b/Engine/Source/Runtime/Core/Private/HAL/ThreadingBase.cpp
--- a/Engine/Source/Runtime/Core/Private/HAL/ThreadingBase.cpp
+++ b/Engine/Source/Runtime/Core/Private/HAL/ThreadingBase.cpp
@@ -659,19 +659,17 @@ void FTlsAutoCleanup::Register()
 	}
 }
 
-
-
-FMultiReaderSingleWriterGT::FMultiReaderSingleWriterGT()
-{
-	CanRead = TFunction<bool()>([=]() { return FPlatformAtomics::InterlockedCompareExchange(&CriticalSection.Action, ReadingAction, NoAction) == ReadingAction; });
-	CanWrite = TFunction<bool()>([=]() { return FPlatformAtomics::InterlockedCompareExchange(&CriticalSection.Action, WritingAction, NoAction) == WritingAction; });
-}
-
 void FMultiReaderSingleWriterGT::LockRead()
 {
 	if (!IsInGameThread())
-	{
-		FPlatformProcess::ConditionalSleep(CanRead);
+	{		
+		SCOPE_CYCLE_COUNTER(STAT_Sleep);
+
+		FThreadIdleStats::FScopeIdle Scope;
+		while (!CanRead())
+		{
+			FPlatformProcess::SleepNoStats(0.0f);
+		}
 	}
 	CriticalSection.ReadCounter.Increment();
 }
@@ -692,7 +690,17 @@ void FMultiReaderSingleWriterGT::UnlockRead()
 void FMultiReaderSingleWriterGT::LockWrite()
 {
 	check(IsInGameThread());
-	FPlatformProcess::ConditionalSleep(CanWrite);
+
+	{
+		SCOPE_CYCLE_COUNTER(STAT_Sleep);
+
+		FThreadIdleStats::FScopeIdle Scope;
+		while (!CanWrite())
+		{
+			FPlatformProcess::SleepNoStats(0.0f);
+		}
+	}
+
 	CriticalSection.WriteCounter.Increment();
 }
 
diff --git a/Engine/Source/Runtime/Core/Public/HAL/ThreadingBase.h b/Engine/Source/Runtime/Core/Public/HAL/ThreadingBase.h
--- a/Engine/Source/Runtime/Core/Public/HAL/ThreadingBase.h
+++ b/Engine/Source/Runtime/Core/Public/HAL/ThreadingBase.h
@@ -843,8 +843,6 @@ public:
 class FMultiReaderSingleWriterGT
 {
 public:
-	FMultiReaderSingleWriterGT();
-
 	/** Protect data from modification while reading. If issued on game thread, doesn't wait for write to finish. */
 	void LockRead();
 
@@ -877,8 +875,15 @@ private:
 	static const int32 NoAction = 0;
 	static const int32 ReadingAction = 1;
 
-	TFunction<bool()> CanRead;
-	TFunction<bool()> CanWrite;
+	bool CanRead()
+	{
+		return FPlatformAtomics::InterlockedCompareExchange(&CriticalSection.Action, ReadingAction, NoAction) == ReadingAction;
+	}
+	
+	bool CanWrite()
+	{
+		return FPlatformAtomics::InterlockedCompareExchange(&CriticalSection.Action, WritingAction, NoAction) == WritingAction;
+	}
 };
 
 class FReadScopeLock
