Mac stability improvements

diff --git a/Engine/Source/Runtime/Core/Private/Mac/CocoaTextView.cpp b/Engine/Source/Runtime/Core/Private/Mac/CocoaTextView.cpp
--- a/Engine/Source/Runtime/Core/Private/Mac/CocoaTextView.cpp
+++ b/Engine/Source/Runtime/Core/Private/Mac/CocoaTextView.cpp
@@ -163,10 +163,13 @@
 - (void)deactivateInputMethod
 {
 	SCOPED_AUTORELEASE_POOL;
-	[self unmarkText];
-	IMMContext = NULL;
-	[[self inputContext] deactivate];
-	[[self inputContext] discardMarkedText];
+	if (IMMContext.IsValid())
+	{
+		[self unmarkText];
+		IMMContext = NULL;
+		[[self inputContext] deactivate];
+		[[self inputContext] discardMarkedText];
+	}
 }
 
 - (bool)isActiveInputMethod:(const TSharedRef<ITextInputMethodContext>&)InContext
diff --git a/Engine/Source/Runtime/Core/Private/Mac/CocoaWindow.cpp b/Engine/Source/Runtime/Core/Private/Mac/CocoaWindow.cpp
--- a/Engine/Source/Runtime/Core/Private/Mac/CocoaWindow.cpp
+++ b/Engine/Source/Runtime/Core/Private/Mac/CocoaWindow.cpp
@@ -407,10 +407,6 @@ NSString* NSPerformDragOperation = @"NSPerformDragOperation";
 - (void)windowWillClose:(NSNotification*)Notification
 {
 	SCOPED_AUTORELEASE_POOL;
-	if (MacApplication)
-	{
-		MacApplication->DeferEvent(Notification);
-	}
 	[self setDelegate:nil];
 }
 
diff --git a/Engine/Source/Runtime/Core/Private/Mac/MacApplication.cpp b/Engine/Source/Runtime/Core/Private/Mac/MacApplication.cpp
--- a/Engine/Source/Runtime/Core/Private/Mac/MacApplication.cpp
+++ b/Engine/Source/Runtime/Core/Private/Mac/MacApplication.cpp
@@ -121,7 +121,10 @@ FMacApplication::~FMacApplication()
 
 	CGDisplayRemoveReconfigurationCallback(FMacApplication::OnDisplayReconfiguration, this);
 
-	TextInputMethodSystem->Terminate();
+	if (TextInputMethodSystem.IsValid())
+	{
+		TextInputMethodSystem->Terminate();
+	}
 
 	MacApplication = nullptr;
 }
@@ -169,25 +172,28 @@ void FMacApplication::PumpMessages(const float TimeDelta)
 
 void FMacApplication::ProcessDeferredEvents(const float TimeDelta)
 {
-	TArray<FDeferredMacEvent> EventsToProcess;
+	if (!FPlatformMisc::bIsPumpingMessages)
+	{
+		TArray<FDeferredMacEvent> EventsToProcess;
 
-	EventsMutex.Lock();
-	EventsToProcess.Append(DeferredEvents);
-	DeferredEvents.Empty();
-	EventsMutex.Unlock();
+		EventsMutex.Lock();
+		EventsToProcess.Append(DeferredEvents);
+		DeferredEvents.Empty();
+		EventsMutex.Unlock();
 
-	const bool bAlreadyProcessingDeferredEvents = bIsProcessingDeferredEvents;
-	bIsProcessingDeferredEvents = true;
+		const bool bAlreadyProcessingDeferredEvents = bIsProcessingDeferredEvents;
+		bIsProcessingDeferredEvents = true;
 
-	for (int32 Index = 0; Index < EventsToProcess.Num(); ++Index)
-	{
-		ProcessEvent(EventsToProcess[Index]);
-	}
+		for (int32 Index = 0; Index < EventsToProcess.Num(); ++Index)
+		{
+			ProcessEvent(EventsToProcess[Index]);
+		}
 
-	bIsProcessingDeferredEvents = bAlreadyProcessingDeferredEvents;
+		bIsProcessingDeferredEvents = bAlreadyProcessingDeferredEvents;
 
-	InvalidateTextLayouts();
-	CloseQueuedWindows();
+		InvalidateTextLayouts();
+		CloseQueuedWindows();
+	}
 }
 
 TSharedRef<FGenericWindow> FMacApplication::MakeWindow()
@@ -571,10 +577,6 @@ void FMacApplication::ProcessEvent(const FDeferredMacEvent& Event)
 		{
 			OnWindowDidMove(EventWindow);
 		}
-		else if (Event.NotificationName == NSWindowWillCloseNotification)
-		{
-			OnWindowDidClose(EventWindow);
-		}
 		else if (Event.NotificationName == NSDraggingExited)
 		{
 			MessageHandler->OnDragLeave(EventWindow);
@@ -960,17 +962,6 @@ void FMacApplication::OnWindowDidResize(TSharedRef<FMacWindow> Window)
 	MessageHandler->OnResizingWindow(Window);
 }
 
-void FMacApplication::OnWindowDidClose(TSharedRef<FMacWindow> Window)
-{
-	SCOPED_AUTORELEASE_POOL;
-	if ([Window->GetWindowHandle() isKeyWindow])
-	{
-		MessageHandler->OnWindowActivationChanged(Window, EWindowActivation::Deactivate);
-	}
-	Windows.Remove(Window);
-	MessageHandler->OnWindowClose(Window);
-}
-
 bool FMacApplication::OnWindowDestroyed(TSharedRef<FMacWindow> Window)
 {
 	SCOPED_AUTORELEASE_POOL;
diff --git a/Engine/Source/Runtime/Core/Private/Mac/MacPlatformMisc.cpp b/Engine/Source/Runtime/Core/Private/Mac/MacPlatformMisc.cpp
--- a/Engine/Source/Runtime/Core/Private/Mac/MacPlatformMisc.cpp
+++ b/Engine/Source/Runtime/Core/Private/Mac/MacPlatformMisc.cpp
@@ -282,6 +282,7 @@ FMacMallocCrashHandler* FMacApplicationInfo::CrashMalloc = nullptr;
 
 UpdateCachedMacMenuStateProc FMacPlatformMisc::UpdateCachedMacMenuState = nullptr;
 bool FMacPlatformMisc::bChachedMacMenuStateNeedsUpdate = true;
+bool FMacPlatformMisc::bIsPumpingMessages = false;
 id<NSObject> FMacPlatformMisc::CommandletActivity = nil;
 
 void FMacPlatformMisc::PlatformPreInit()
@@ -581,6 +582,9 @@ void FMacPlatformMisc::PumpMessages( bool bFromMainLoop )
 {
 	if( bFromMainLoop )
 	{
+		const bool bIsAlreadyPumpingMessages = bIsPumpingMessages;
+		bIsPumpingMessages = true;
+
 		ProcessGameThreadEvents();
 
 		if (MacApplication && !MacApplication->IsProcessingDeferredEvents() && IsInGameThread())
@@ -591,6 +595,8 @@ void FMacPlatformMisc::PumpMessages( bool bFromMainLoop )
 				bChachedMacMenuStateNeedsUpdate = false;
 			}
 		}
+
+		bIsPumpingMessages = bIsAlreadyPumpingMessages;
 	}
 }
 
diff --git a/Engine/Source/Runtime/Core/Private/Mac/MacTextInputMethodSystem.cpp b/Engine/Source/Runtime/Core/Private/Mac/MacTextInputMethodSystem.cpp
--- a/Engine/Source/Runtime/Core/Private/Mac/MacTextInputMethodSystem.cpp
+++ b/Engine/Source/Runtime/Core/Private/Mac/MacTextInputMethodSystem.cpp
@@ -208,20 +208,28 @@ void FMacTextInputMethodSystem::DeactivateContext(const TSharedRef<ITextInputMet
 	if(MacNotifier->GetContextWindowNumber() != 0)
 	{
 		FCocoaWindow* CocoaWindow = (FCocoaWindow*)[NSApp windowWithWindowNumber:MacNotifier->GetContextWindowNumber()];
-		bDeactivatedContext = MainThreadReturn(^{
-			bool bSuccess = false;
-			if(CocoaWindow && [CocoaWindow openGLView])
-			{
-				NSView* GLView = [CocoaWindow openGLView];
-				if([GLView isKindOfClass:[FCocoaTextView class]])
+		if (CocoaWindow)
+		{
+			bDeactivatedContext = MainThreadReturn(^{
+				bool bSuccess = false;
+				if(CocoaWindow && [CocoaWindow openGLView])
 				{
-					FCocoaTextView* TextView = (FCocoaTextView*)GLView;
-					[TextView deactivateInputMethod];
-					bSuccess = true;
+					NSView* GLView = [CocoaWindow openGLView];
+					if([GLView isKindOfClass:[FCocoaTextView class]])
+					{
+						FCocoaTextView* TextView = (FCocoaTextView*)GLView;
+						[TextView deactivateInputMethod];
+						bSuccess = true;
+					}
 				}
-			}
-			return bSuccess;
-		}, UE4IMEEventMode);
+				return bSuccess;
+			}, UE4IMEEventMode);
+		}
+		else
+		{
+			// If the window is no longer open, the context is already inactive (handled by MessageHandler->OnWindowActivationChanged() called in FMacApplication::OnWindowDestroyed())
+			bDeactivatedContext = true;
+		}
 	}
 	if(!bDeactivatedContext)
 	{
diff --git a/Engine/Source/Runtime/Core/Public/Mac/MacApplication.h b/Engine/Source/Runtime/Core/Public/Mac/MacApplication.h
--- a/Engine/Source/Runtime/Core/Public/Mac/MacApplication.h
+++ b/Engine/Source/Runtime/Core/Public/Mac/MacApplication.h
@@ -209,7 +209,6 @@ private:
 
 	void OnWindowDidMove(TSharedRef<FMacWindow> Window);
 	void OnWindowDidResize(TSharedRef<FMacWindow> Window);
-	void OnWindowDidClose(TSharedRef<FMacWindow> Window);
 	bool OnWindowDestroyed(TSharedRef<FMacWindow> Window);
 
 	void OnApplicationDidBecomeActive();
diff --git a/Engine/Source/Runtime/Core/Public/Mac/MacPlatformMisc.h b/Engine/Source/Runtime/Core/Public/Mac/MacPlatformMisc.h
--- a/Engine/Source/Runtime/Core/Public/Mac/MacPlatformMisc.h
+++ b/Engine/Source/Runtime/Core/Public/Mac/MacPlatformMisc.h
@@ -199,6 +199,7 @@ struct CORE_API FMacPlatformMisc : public FGenericPlatformMisc
 	static FString GetOperatingSystemId();
 
 	static bool bChachedMacMenuStateNeedsUpdate;
+	static bool bIsPumpingMessages;
 
 	static id<NSObject> CommandletActivity;
 };
diff --git a/Engine/Source/Runtime/OpenGLDrv/Private/Mac/OpenGLMac.cpp b/Engine/Source/Runtime/OpenGLDrv/Private/Mac/OpenGLMac.cpp
--- a/Engine/Source/Runtime/OpenGLDrv/Private/Mac/OpenGLMac.cpp
+++ b/Engine/Source/Runtime/OpenGLDrv/Private/Mac/OpenGLMac.cpp
@@ -338,25 +338,27 @@ FPlatformOpenGLContext* PlatformCreateOpenGLContext(FPlatformOpenGLDevice* Devic
 
 	[Context->OpenGLView setAutoresizingMask:NSViewWidthSizable | NSViewHeightSizable];
 
-	if (FPlatformMisc::IsRunningOnMavericks() && ([Context->WindowHandle styleMask] & NSTexturedBackgroundWindowMask))
-	{
-		NSView* SuperView = [[Context->WindowHandle contentView] superview];
-		[SuperView addSubview:Context->OpenGLView];
-		[SuperView setWantsLayer:YES];
-		[SuperView addSubview:[Context->WindowHandle standardWindowButton:NSWindowCloseButton]];
-		[SuperView addSubview:[Context->WindowHandle standardWindowButton:NSWindowMiniaturizeButton]];
-		[SuperView addSubview:[Context->WindowHandle standardWindowButton:NSWindowZoomButton]];
-	}
-	else
-	{
-		[Context->OpenGLView setWantsLayer:YES];
-		[Context->WindowHandle setContentView:Context->OpenGLView];
-	}
+	MainThreadCall(^{
+		if (FPlatformMisc::IsRunningOnMavericks() && ([Context->WindowHandle styleMask] & NSTexturedBackgroundWindowMask))
+		{
+			NSView* SuperView = [[Context->WindowHandle contentView] superview];
+			[SuperView addSubview:Context->OpenGLView];
+			[SuperView setWantsLayer:YES];
+			[SuperView addSubview:[Context->WindowHandle standardWindowButton:NSWindowCloseButton]];
+			[SuperView addSubview:[Context->WindowHandle standardWindowButton:NSWindowMiniaturizeButton]];
+			[SuperView addSubview:[Context->WindowHandle standardWindowButton:NSWindowZoomButton]];
+		}
+		else
+		{
+			[Context->OpenGLView setWantsLayer:YES];
+			[Context->WindowHandle setContentView:Context->OpenGLView];
+		}
 
-	[[Context->WindowHandle standardWindowButton:NSWindowCloseButton] setAction:@selector(performClose:)];
+		[[Context->WindowHandle standardWindowButton:NSWindowCloseButton] setAction:@selector(performClose:)];
 
-	Context->OpenGLView.layer.magnificationFilter = kCAFilterNearest;
-	Context->OpenGLView.layer.minificationFilter = kCAFilterNearest;
+		Context->OpenGLView.layer.magnificationFilter = kCAFilterNearest;
+		Context->OpenGLView.layer.minificationFilter = kCAFilterNearest;
+	}, NSDefaultRunLoopMode, true);
 
 	return Context;
 }
