Updated warning messages for when no localization paths are specifed in the configs.

diff --git a/Engine/Source/Runtime/Core/Private/Misc/Paths.cpp b/Engine/Source/Runtime/Core/Private/Misc/Paths.cpp
--- a/Engine/Source/Runtime/Core/Private/Misc/Paths.cpp
+++ b/Engine/Source/Runtime/Core/Private/Misc/Paths.cpp
@@ -238,7 +238,7 @@ const TArray<FString>& FPaths::GetEngineLocalizationPaths()
 			GConfig->GetArray( TEXT("Internationalization"), TEXT("LocalizationPaths"), Results, GEngineIni );
 			if(!Results.Num())
 			{
-				UE_LOG(LogInit, Warning, TEXT("We appear to have no engine localization data?"));
+				UE_LOG(LogInit, Warning, TEXT("No paths for engine localization data were specifed in the engine configuration."));
 			}
 			HasInitialized = true;
 		}
@@ -263,7 +263,7 @@ const TArray<FString>& FPaths::GetEditorLocalizationPaths()
 			GConfig->GetArray( TEXT("Internationalization"), TEXT("LocalizationPaths"), Results, GEditorIni );
 			if(!Results.Num())
 			{
-				UE_LOG(LogInit, Warning, TEXT("We appear to have no editor localization data? Editor can't run."));
+				UE_LOG(LogInit, Warning, TEXT("No paths for editor localization data were specifed in the editor configuration."));
 			}
 			HasInitialized = true;
 		}
@@ -288,7 +288,7 @@ const TArray<FString>& FPaths::GetPropertyNameLocalizationPaths()
 			GConfig->GetArray( TEXT("Internationalization"), TEXT("PropertyNameLocalizationPaths"), Results, GEditorIni );
 			if(!Results.Num())
 			{
-				UE_LOG(LogInit, Warning, TEXT("We appear to have no property name localization data? Editor can't run."));
+				UE_LOG(LogInit, Warning, TEXT("No paths for property name localization data were specifed in the editor configuration."));
 			}
 			HasInitialized = true;
 		}
@@ -313,7 +313,7 @@ const TArray<FString>& FPaths::GetToolTipLocalizationPaths()
 			GConfig->GetArray( TEXT("Internationalization"), TEXT("ToolTipLocalizationPaths"), Results, GEditorIni );
 			if(!Results.Num())
 			{
-				UE_LOG(LogInit, Warning, TEXT("We appear to have no tool tip localization data? Editor can't run."));
+				UE_LOG(LogInit, Warning, TEXT("No paths for tooltips localization data were specifed in the editor configuration."));
 			}
 			HasInitialized = true;
 		}
@@ -338,7 +338,7 @@ const TArray<FString>& FPaths::GetGameLocalizationPaths()
 			GConfig->GetArray( TEXT("Internationalization"), TEXT("LocalizationPaths"), Results, GGameIni );
 			if(!Results.Num()) // Failed to find localization path.
 			{
-				UE_LOG(LogPaths, Warning, TEXT("We appear to have no game localization data? Game can't run."));
+				UE_LOG(LogPaths, Warning, TEXT("No paths for game localization data were specifed in the game configuration."));
 			}
 			HasInitialized = true;
 		}
