Fixed issues caused by passing TSharedPtr to dispatch blocks on Mac

diff --git a/Engine/Source/Runtime/Core/Private/Mac/MacApplication.cpp b/Engine/Source/Runtime/Core/Private/Mac/MacApplication.cpp
--- a/Engine/Source/Runtime/Core/Private/Mac/MacApplication.cpp
+++ b/Engine/Source/Runtime/Core/Private/Mac/MacApplication.cpp
@@ -379,7 +379,9 @@ void FMacApplication::DeferEvent(NSObject* Object)
 				DeferredEvent.Window = FindWindowByNSWindow(WindowHandle);
 				if (DeferredEvent.Window.IsValid())
 				{
-					GameThreadCall(^{ OnWindowDidResize(DeferredEvent.Window.Pin().ToSharedRef()); }, @[ NSDefaultRunLoopMode, UE4ResizeEventMode, UE4ShowEventMode, UE4FullscreenEventMode ], true);
+					GameThreadCall(^{
+						TSharedPtr<FMacWindow> Window = FindWindowByNSWindow(WindowHandle); // We cannot use DeferredEvent.Window here because the block doesn't handle TSharedPtr correctly
+						OnWindowDidResize(Window.ToSharedRef()); }, @[ NSDefaultRunLoopMode, UE4ResizeEventMode, UE4ShowEventMode, UE4FullscreenEventMode ], true);
 				}
 				return;
 			}
@@ -761,13 +763,14 @@ void FMacApplication::ProcessMouseDownEvent(const FDeferredMacEvent& Event)
 			if (IsMouseOverTitleBar)
 			{
 				const bool bShouldMinimize = [[NSUserDefaults standardUserDefaults] boolForKey:@"AppleMiniaturizeOnDoubleClick"];
+				FCocoaWindow* WindowHandle = EventWindow->GetWindowHandle();
 				if (bShouldMinimize)
 				{
-					MainThreadCall(^{ [EventWindow->GetWindowHandle() performMiniaturize:nil]; }, NSDefaultRunLoopMode, true);
+					MainThreadCall(^{ [WindowHandle performMiniaturize:nil]; }, NSDefaultRunLoopMode, true);
 				}
 				else if (!FPlatformMisc::IsRunningOnMavericks())
 				{
-					MainThreadCall(^{ [EventWindow->GetWindowHandle() zoom:nil]; }, NSDefaultRunLoopMode, true);
+					MainThreadCall(^{ [WindowHandle zoom:nil]; }, NSDefaultRunLoopMode, true);
 				}
 			}
 			else
@@ -892,7 +895,8 @@ void FMacApplication::ProcessKeyDownEvent(const FDeferredMacEvent& Event)
 		FCocoaMenu* MainMenu = [[NSApp mainMenu] isKindOfClass:[FCocoaMenu class]] ? (FCocoaMenu*)[NSApp mainMenu]: nil;
 		if (MainMenu)
 		{
-			MainThreadCall(^{ [MainMenu highlightKeyEquivalent:Event.Event]; }, NSDefaultRunLoopMode, true);
+			NSEvent* NativeEvent = Event.Event;
+			MainThreadCall(^{ [MainMenu highlightKeyEquivalent:NativeEvent]; }, NSDefaultRunLoopMode, true);
 		}
 	}
 	else
diff --git a/Engine/Source/Runtime/Slate/Private/Framework/MultiBox/Mac/MacMenu.cpp b/Engine/Source/Runtime/Slate/Private/Framework/MultiBox/Mac/MacMenu.cpp
--- a/Engine/Source/Runtime/Slate/Private/Framework/MultiBox/Mac/MacMenu.cpp
+++ b/Engine/Source/Runtime/Slate/Private/Framework/MultiBox/Mac/MacMenu.cpp
@@ -85,6 +85,14 @@ static FCriticalSection GCachedMenuStateCS;
 
 void FSlateMacMenu::UpdateWithMultiBox(const TSharedPtr< FMultiBox > MultiBox)
 {
+	// The dispatch block can't handle TSharedPtr correctly, so we use a small trick to pass MultiBox safely
+	struct FSafeMultiBoxPass
+	{
+		TSharedPtr<FMultiBox> MultiBox;
+	};
+	FSafeMultiBoxPass* SafeMultiBoxPtr = new FSafeMultiBoxPass;
+	SafeMultiBoxPtr->MultiBox = MultiBox;
+
 	MainThreadCall(^{
 		FScopeLock Lock(&GCachedMenuStateCS);
 
@@ -103,9 +111,9 @@ void FSlateMacMenu::UpdateWithMultiBox(const TSharedPtr< FMultiBox > MultiBox)
 		}
 		GCachedMenuState.Reset();
 		
-		if( MultiBox.IsValid() )
+		if( SafeMultiBoxPtr->MultiBox.IsValid() )
 		{
-			const TArray<TSharedRef<const FMultiBlock> >& MenuBlocks = MultiBox->GetBlocks();
+			const TArray<TSharedRef<const FMultiBlock> >& MenuBlocks = SafeMultiBoxPtr->MultiBox->GetBlocks();
 
 			for (int32 Index = 0; Index < MenuBlocks.Num(); Index++)
 			{
@@ -153,6 +161,8 @@ void FSlateMacMenu::UpdateWithMultiBox(const TSharedPtr< FMultiBox > MultiBox)
 			{
 				[HelpMenu release];
 			}
+
+			delete SafeMultiBoxPtr;
 		}
 
 		FPlatformMisc::bChachedMacMenuStateNeedsUpdate = true;
