#ue4 - Perf: Add FTransform methods CopyTranslation(), CopyRotation(), CopyScale3D(). SceneComponent now uses these in CalcNewComponentToWorld().

- Avoids intermediate copies/unaligned loads between VectorRegister and FVector

diff --git a/Engine/Source/Runtime/Core/Public/Math/Transform.h b/Engine/Source/Runtime/Core/Public/Math/Transform.h
--- a/Engine/Source/Runtime/Core/Public/Math/Transform.h
+++ b/Engine/Source/Runtime/Core/Public/Math/Transform.h
@@ -679,6 +679,12 @@ public:
 		DiagnosticCheckNaN_Translate();
 	}
 
+	/** Copy translation from another FTransform. */
+	FORCEINLINE void CopyTranslation(const FTransform& Other)
+	{
+		Translation = Other.Translation;
+	}
+
 	/**
 	 * Concatenates another rotation to this transformation 
 	 * @param DeltaRotation The rotation to concatenate in the following fashion: Rotation = Rotation * DeltaRotation
@@ -709,6 +715,12 @@ public:
 		DiagnosticCheckNaN_Rotate();
 	}
 
+	/** Copy rotation from another FTransform. */
+	FORCEINLINE void CopyRotation(const FTransform& Other)
+	{
+		Rotation = Other.Rotation;
+	}
+
 	/**
 	 * Sets the Scale3D component
 	 * @param NewScale3D The new value for the Scale3D component
@@ -719,6 +731,12 @@ public:
 		DiagnosticCheckNaN_Scale3D();
 	}
 
+	/** Copy scale from another FTransform. */
+	FORCEINLINE void CopyScale3D(const FTransform& Other)
+	{
+		Scale3D = Other.Scale3D;
+	}
+
 	/**
 	 * Sets both the translation and Scale3D components at the same time
 	 * @param NewTranslation The new value for the translation component
diff --git a/Engine/Source/Runtime/Core/Public/Math/TransformVectorized.h b/Engine/Source/Runtime/Core/Public/Math/TransformVectorized.h
--- a/Engine/Source/Runtime/Core/Public/Math/TransformVectorized.h
+++ b/Engine/Source/Runtime/Core/Public/Math/TransformVectorized.h
@@ -728,6 +728,12 @@ public:
 		DiagnosticCheckNaN_Translate();
 	}
 
+	/** Copy translation from another FTransform. */
+	FORCEINLINE void CopyTranslation(const FTransform& Other)
+	{
+		Translation = Other.Translation;
+	}
+
 	/**
 	 * Concatenates another rotation to this transformation 
 	 * @param DeltaRotation The rotation to concatenate in the following fashion: Rotation = Rotation * DeltaRotation
@@ -758,6 +764,12 @@ public:
 		DiagnosticCheckNaN_Rotate();
 	}
 
+	/** Copy rotation from another FTransform. */
+	FORCEINLINE void CopyRotation(const FTransform& Other)
+	{
+		Rotation = Other.Rotation;
+	}
+
 	/**
 	 * Sets the Scale3D component
 	 * @param NewScale3D The new value for the Scale3D component
@@ -768,6 +780,12 @@ public:
 		DiagnosticCheckNaN_Scale3D();
 	}
 
+	/** Copy scale from another FTransform. */
+	FORCEINLINE void CopyScale3D(const FTransform& Other)
+	{
+		Scale3D = Other.Scale3D;
+	}
+
 	/**
 	 * Sets both the translation and Scale3D components at the same time
 	 * @param NewTranslation The new value for the translation component
diff --git a/Engine/Source/Runtime/Engine/Private/Components/SceneComponent.cpp b/Engine/Source/Runtime/Engine/Private/Components/SceneComponent.cpp
--- a/Engine/Source/Runtime/Engine/Private/Components/SceneComponent.cpp
+++ b/Engine/Source/Runtime/Engine/Private/Components/SceneComponent.cpp
@@ -294,17 +294,17 @@ FTransform USceneComponent::CalcNewComponentToWorld(const FTransform& NewRelativ
 
 		if(bAbsoluteLocation)
 		{
-			NewCompToWorld.SetTranslation(NewRelativeTransform.GetTranslation());
+			NewCompToWorld.CopyTranslation(NewRelativeTransform);
 		}
 
 		if(bAbsoluteRotation)
 		{
-			NewCompToWorld.SetRotation(NewRelativeTransform.GetRotation());
+			NewCompToWorld.CopyRotation(NewRelativeTransform);
 		}
 
 		if(bAbsoluteScale)
 		{
-			NewCompToWorld.SetScale3D(NewRelativeTransform.GetScale3D());
+			NewCompToWorld.CopyScale3D(NewRelativeTransform);
 		}
 
 		return NewCompToWorld;
