Fix infinite loop on startup when migrating config files.

diff --git a/Engine/Source/Runtime/Core/Private/Misc/ConfigManifest.cpp b/Engine/Source/Runtime/Core/Private/Misc/ConfigManifest.cpp
--- a/Engine/Source/Runtime/Core/Private/Misc/ConfigManifest.cpp
+++ b/Engine/Source/Runtime/Core/Private/Misc/ConfigManifest.cpp
@@ -54,6 +54,11 @@ FString ProjectAgnosticIniPath(const TCHAR* InLeaf)
 	return FPaths::GameAgnosticSavedDir() / TEXT("Config") / ANSI_TO_TCHAR(FPlatformProperties::PlatformName()) / InLeaf;
 }
 
+FString GetConfigDirectoryToSearch(const FEngineVersion& Version)
+{
+	return FString(FPlatformProcess::UserSettingsDir()) / ENGINE_VERSION_TEXT(EPIC_PRODUCT_IDENTIFIER) / Version.ToString(EVersionComponent::Minor) / TEXT("Saved") / TEXT("Config") / ANSI_TO_TCHAR(FPlatformProperties::PlatformName());
+}
+
 /** Migrates config files from a previous version of the engine. Does nothing on non-installed versions */
 void MigratePreviousEngineInis()
 {
@@ -66,9 +71,12 @@ void MigratePreviousEngineInis()
 	FEngineVersion PreviousVersion(GEngineVersion.GetMajor(), GEngineVersion.GetMinor() - 1, GEngineVersion.GetPatch(), GEngineVersion.GetChangelist(), GEngineVersion.GetBranch());
 
 	IPlatformFile& PlatformFile = FPlatformFileManager::Get().GetPlatformFile();
-	while(PreviousVersion.GetMinor() >= 0)
+	FString Directory;
+
+	// Minor is stored as unsigned int, we have to split the loop to avoid infinite iteration.
+	while(PreviousVersion.GetMinor() > 0)
 	{
-		const FString Directory = FString(FPlatformProcess::UserSettingsDir()) / ENGINE_VERSION_TEXT(EPIC_PRODUCT_IDENTIFIER) / PreviousVersion.ToString(EVersionComponent::Minor) / TEXT("Saved") / TEXT("Config") / ANSI_TO_TCHAR(FPlatformProperties::PlatformName());
+		Directory = GetConfigDirectoryToSearch(PreviousVersion);
 		if (FPaths::DirectoryExists(Directory))
 		{
 			const FString DestDir = ProjectAgnosticIniPath(TEXT(""));
@@ -83,6 +91,17 @@ void MigratePreviousEngineInis()
 
 		PreviousVersion = FEngineVersion(PreviousVersion.GetMajor(), PreviousVersion.GetMinor() - 1, PreviousVersion.GetPatch(), PreviousVersion.GetChangelist(), PreviousVersion.GetBranch());
 	}
+
+	// One last check for version with Minor == 0
+	Directory = GetConfigDirectoryToSearch(PreviousVersion);
+	if (FPaths::DirectoryExists(Directory))
+	{
+		const FString DestDir = ProjectAgnosticIniPath(TEXT(""));
+		if (PlatformFile.CreateDirectoryTree(*DestDir))
+		{
+			PlatformFile.CopyDirectoryTree(*DestDir, *Directory, false);
+		}
+	}
 }
 
 void FConfigManifest::UpgradeFromPreviousVersions()
