Fix the Report.wer AppPath value in OS X crash reports.

diff --git a/Engine/Source/Runtime/Core/Private/Mac/MacPlatformMisc.cpp b/Engine/Source/Runtime/Core/Private/Mac/MacPlatformMisc.cpp
--- a/Engine/Source/Runtime/Core/Private/Mac/MacPlatformMisc.cpp
+++ b/Engine/Source/Runtime/Core/Private/Mac/MacPlatformMisc.cpp
@@ -49,7 +49,7 @@ struct FMacApplicationInfo
 		
 		ExecutableName = FPlatformProcess::ExecutableName();
 		
-		AppPath = FPlatformProcess::GenerateApplicationPath(AppName, FApp::GetBuildConfiguration());
+		AppPath = FString([[NSBundle mainBundle] executablePath]);
 
 		AppBundleID = FString([[NSBundle mainBundle] bundleIdentifier]);
 
