Core - Misc tweak to Windows crash reporting

diff --git a/Engine/Source/Runtime/Core/Private/Windows/WindowsPlatformCrashContext.cpp b/Engine/Source/Runtime/Core/Private/Windows/WindowsPlatformCrashContext.cpp
--- a/Engine/Source/Runtime/Core/Private/Windows/WindowsPlatformCrashContext.cpp
+++ b/Engine/Source/Runtime/Core/Private/Windows/WindowsPlatformCrashContext.cpp
@@ -455,18 +455,14 @@ void CreateExceptionInfoString(EXCEPTION_RECORD* ExceptionRecord)
 		ErrorString += FString::Printf(TEXT("0x%08x"), (uint32)ExceptionRecord->ExceptionCode);
 	}
 
-#if WITH_EDITORONLY_DATA
 	FCString::Strncpy(GErrorExceptionDescription, *ErrorString, ARRAY_COUNT(GErrorExceptionDescription));
-#endif
+
 #undef HANDLE_CASE
 }
 #include "HideWindowsPlatformTypes.h"
 
 int32 ReportCrash( LPEXCEPTION_POINTERS ExceptionInfo )
 {
-	// Switch to malloc crash.
-	FMallocCrash::Get().SetAsGMalloc();
-
 	// Only create a minidump the first time this function is called.
 	// (Can be called the first time from the RenderThread, then a second time from the MainThread.)
 	if( FPlatformAtomics::InterlockedIncrement( &ReportCrashCallCount ) != 1 )
