Linux: ignore unusable POSIX API for host id (UE-8013)

- Removing after discussion with Wes: most systems don't implement it in a useful way anyway, so this is only going to confuse the analytics.
- Also removed hardcoded strings which are also confusing.

diff --git a/Engine/Source/Runtime/Core/Private/Linux/LinuxPlatformMisc.cpp b/Engine/Source/Runtime/Core/Private/Linux/LinuxPlatformMisc.cpp
--- a/Engine/Source/Runtime/Core/Private/Linux/LinuxPlatformMisc.cpp
+++ b/Engine/Source/Runtime/Core/Private/Linux/LinuxPlatformMisc.cpp
@@ -717,20 +717,10 @@ FString FLinuxPlatformMisc::GetOperatingSystemId()
 
 			close(OsGuidFile);
 		}
-		else
-		{
-			// use old POSIX API which doesn't seem to be useful anymore
-			long HostId = gethostid();
-			if (HostId != 0)
-			{
-				CachedResult = FString::Printf(TEXT("baad051dbaad051dbaad051d%08lx"), HostId);
-			}
-		}
 
-		if (CachedResult.Len() == 0)
-		{
-			CachedResult = TEXT("dead051ddead051ddead051ddead051d");
-		}
+		// old POSIX gethostid() is not useful. It is impossible to have globally unique 32-bit GUIDs and most
+		// systems don't try hard implementing it these days (glibc will return a permuted IP address, often 127.0.0.1)
+		// Due to that, we just ignore that call and consider lack of systemd's /etc/machine-id a failure to obtain the host id.
 
 		bHasCachedResult = true;	// even if we failed to read the real one
 	}
