Merging CL #2537837 - Fixed one more reason for a crash on Mac caused by using TSharedRef/Ptr<FMacWindow> on main thread

diff --git a/Engine/Source/Runtime/Core/Private/Mac/MacApplication.cpp b/Engine/Source/Runtime/Core/Private/Mac/MacApplication.cpp
--- a/Engine/Source/Runtime/Core/Private/Mac/MacApplication.cpp
+++ b/Engine/Source/Runtime/Core/Private/Mac/MacApplication.cpp
@@ -75,7 +75,7 @@ FMacApplication::FMacApplication()
 																								usingBlock:^(NSNotification* Notification){ bIsWorkspaceSessionActive = false; }];
 
 	MouseMovedEventMonitor = [NSEvent addGlobalMonitorForEventsMatchingMask:NSMouseMovedMask handler:^(NSEvent* Event) { DeferEvent(Event); }];
-	EventMonitor = [NSEvent addLocalMonitorForEventsMatchingMask:NSAnyEventMask handler:^(NSEvent* IncomingEvent) { return HandleNSEvent(IncomingEvent); }];
+	EventMonitor = [NSEvent addLocalMonitorForEventsMatchingMask:NSAnyEventMask handler:^(NSEvent* Event) { return HandleNSEvent(Event); }];
 
 	CGDisplayRegisterReconfigurationCallback(FMacApplication::OnDisplayReconfiguration, this);
 
@@ -289,12 +289,11 @@ void FMacApplication::CloseWindow(TSharedRef<FMacWindow> Window)
 void FMacApplication::DeferEvent(NSObject* Object)
 {
 	FDeferredMacEvent DeferredEvent;
-	FCocoaWindow* WindowHandle = nullptr;
 
 	if (Object && [Object isKindOfClass:[NSEvent class]])
 	{
 		NSEvent* Event = (NSEvent*)Object;
-		WindowHandle = FindEventWindow(Event);
+		DeferredEvent.Window = FindEventWindow(Event);
 		DeferredEvent.Event = [Event retain];
 		DeferredEvent.Type = [Event type];
 		DeferredEvent.LocationInWindow = FVector2D([Event locationInWindow].x, [Event locationInWindow].y);
@@ -305,10 +304,13 @@ void FMacApplication::DeferEvent(NSObject* Object)
 
 		if (DeferredEvent.Type == NSKeyDown)
 		{
-			DeferredEvent.Window = FindWindowByNSWindow(WindowHandle);
-			if (DeferredEvent.Window.IsValid() && DeferredEvent.Window.Pin()->OnIMKKeyDown(Event))
+			if (DeferredEvent.Window && [DeferredEvent.Window openGLView])
 			{
-				return;
+				FCocoaTextView* View = (FCocoaTextView*)[DeferredEvent.Window openGLView];
+				if (View && [View imkKeyDown:Event])
+				{
+					return;
+				}
 			}
 		}
 
@@ -372,15 +374,14 @@ void FMacApplication::DeferEvent(NSObject* Object)
 		DeferredEvent.NotificationName = [[Notification name] retain];
 		if ([[Notification object] isKindOfClass:[FCocoaWindow class]])
 		{
-			WindowHandle = (FCocoaWindow*)[Notification object];
+			DeferredEvent.Window = (FCocoaWindow*)[Notification object];
 
 			if (DeferredEvent.NotificationName == NSWindowDidResizeNotification)
 			{
-				DeferredEvent.Window = FindWindowByNSWindow(WindowHandle);
-				if (DeferredEvent.Window.IsValid())
+				if (DeferredEvent.Window)
 				{
 					GameThreadCall(^{
-						TSharedPtr<FMacWindow> Window = FindWindowByNSWindow(WindowHandle); // We cannot use DeferredEvent.Window here because the block doesn't handle TSharedPtr correctly
+						TSharedPtr<FMacWindow> Window = FindWindowByNSWindow(DeferredEvent.Window);
 						OnWindowDidResize(Window.ToSharedRef()); }, @[ NSDefaultRunLoopMode, UE4ResizeEventMode, UE4ShowEventMode, UE4FullscreenEventMode ], true);
 				}
 				return;
@@ -388,7 +389,6 @@ void FMacApplication::DeferEvent(NSObject* Object)
 		}
 		else if ([[Notification object] conformsToProtocol:@protocol(NSDraggingInfo)])
 		{
-			WindowHandle = (FCocoaWindow*)[(id<NSDraggingInfo>)[Notification object] draggingDestinationWindow];
 			if (DeferredEvent.NotificationName == NSPrepareForDragOperation)
 			{
 				DeferredEvent.DraggingPasteboard = [[(id<NSDraggingInfo>)[Notification object] draggingPasteboard] retain];
@@ -396,8 +396,6 @@ void FMacApplication::DeferEvent(NSObject* Object)
 		}
 	}
 
-	DeferredEvent.Window = FindWindowByNSWindow(WindowHandle);
-
 	FScopeLock Lock(&EventsMutex);
 	DeferredEvents.Add(DeferredEvent);
 }
@@ -481,6 +479,7 @@ int32 FMacApplication::MTContactCallback(void* Device, void* Data, int32 NumFing
 
 void FMacApplication::ProcessEvent(const FDeferredMacEvent& Event)
 {
+	TSharedPtr<FMacWindow> EventWindow = FindWindowByNSWindow(Event.Window);
 	if (Event.Type)
 	{
 		if (CurrentModifierFlags != Event.ModifierFlags)
@@ -506,23 +505,23 @@ void FMacApplication::ProcessEvent(const FDeferredMacEvent& Event)
 			case NSLeftMouseDragged:
 			case NSRightMouseDragged:
 			case NSOtherMouseDragged:
-				ProcessMouseMovedEvent(Event);
+				ProcessMouseMovedEvent(Event, EventWindow);
 				break;
 
 			case NSLeftMouseDown:
 			case NSRightMouseDown:
 			case NSOtherMouseDown:
-				ProcessMouseDownEvent(Event);
+				ProcessMouseDownEvent(Event, EventWindow);
 				break;
 
 			case NSLeftMouseUp:
 			case NSRightMouseUp:
 			case NSOtherMouseUp:
-				ProcessMouseUpEvent(Event);
+				ProcessMouseUpEvent(Event, EventWindow);
 				break;
 
 			case NSScrollWheel:
-				ProcessScrollWheelEvent(Event);
+				ProcessScrollWheelEvent(Event, EventWindow);
 				break;
 
 			case NSEventTypeMagnify:
@@ -534,7 +533,7 @@ void FMacApplication::ProcessEvent(const FDeferredMacEvent& Event)
 				break;
 
 			case NSKeyDown:
-				ProcessKeyDownEvent(Event);
+				ProcessKeyDownEvent(Event, EventWindow);
 				break;
 
 			case NSKeyUp:
@@ -542,33 +541,32 @@ void FMacApplication::ProcessEvent(const FDeferredMacEvent& Event)
 				break;
 		}
 	}
-	else if (Event.Window.IsValid())
+	else if (EventWindow.IsValid())
 	{
-		TSharedRef<FMacWindow> EventWindow = Event.Window.Pin().ToSharedRef();
 		if (Event.NotificationName == NSWindowWillStartLiveResizeNotification)
 		{
-			MessageHandler->BeginReshapingWindow(EventWindow);
+			MessageHandler->BeginReshapingWindow(EventWindow.ToSharedRef());
 		}
 		else if (Event.NotificationName == NSWindowDidEndLiveResizeNotification)
 		{
-			MessageHandler->FinishedReshapingWindow(EventWindow);
+			MessageHandler->FinishedReshapingWindow(EventWindow.ToSharedRef());
 		}
 		else if (Event.NotificationName == NSWindowDidEnterFullScreenNotification)
 		{
-			OnWindowDidResize(EventWindow);
+			OnWindowDidResize(EventWindow.ToSharedRef());
 		}
 		else if (Event.NotificationName == NSWindowDidExitFullScreenNotification)
 		{
-			OnWindowDidResize(EventWindow);
+			OnWindowDidResize(EventWindow.ToSharedRef());
 		}
 		else if (Event.NotificationName == NSWindowDidBecomeMainNotification)
 		{
-			MessageHandler->OnWindowActivationChanged(EventWindow, EWindowActivation::Activate);
+			MessageHandler->OnWindowActivationChanged(EventWindow.ToSharedRef(), EWindowActivation::Activate);
 			OnWindowsReordered(![NSApp isActive]);
 		}
 		else if (Event.NotificationName == NSWindowDidResignMainNotification)
 		{
-			MessageHandler->OnWindowActivationChanged(EventWindow, EWindowActivation::Deactivate);
+			MessageHandler->OnWindowActivationChanged(EventWindow.ToSharedRef(), EWindowActivation::Deactivate);
 			OnWindowsReordered(![NSApp isActive]);
 		}
 		else if (Event.NotificationName == NSWindowWillMoveNotification)
@@ -577,15 +575,15 @@ void FMacApplication::ProcessEvent(const FDeferredMacEvent& Event)
 		}
 		else if (Event.NotificationName == NSWindowDidMoveNotification)
 		{
-			OnWindowDidMove(EventWindow);
+			OnWindowDidMove(EventWindow.ToSharedRef());
 		}
 		else if (Event.NotificationName == NSDraggingExited)
 		{
-			MessageHandler->OnDragLeave(EventWindow);
+			MessageHandler->OnDragLeave(EventWindow.ToSharedRef());
 		}
 		else if (Event.NotificationName == NSDraggingUpdated)
 		{
-			MessageHandler->OnDragOver(EventWindow);
+			MessageHandler->OnDragOver(EventWindow.ToSharedRef());
 		}
 		else if (Event.NotificationName == NSPrepareForDragOperation)
 		{
@@ -607,17 +605,17 @@ void FMacApplication::ProcessEvent(const FDeferredMacEvent& Event)
 					FileList.Add(ListElement);
 				}
 
-				MessageHandler->OnDragEnterFiles(EventWindow, FileList);
+				MessageHandler->OnDragEnterFiles(EventWindow.ToSharedRef(), FileList);
 			}
 			else if (bHaveText)
 			{
 				NSString* Text = [Event.DraggingPasteboard stringForType:NSPasteboardTypeString];
-				MessageHandler->OnDragEnterText(EventWindow, FString(Text));
+				MessageHandler->OnDragEnterText(EventWindow.ToSharedRef(), FString(Text));
 			}
 		}
 		else if (Event.NotificationName == NSPerformDragOperation)
 		{
-			MessageHandler->OnDragDrop(EventWindow);
+			MessageHandler->OnDragDrop(EventWindow.ToSharedRef());
 		}
 	}
 }
@@ -630,9 +628,8 @@ void FMacApplication::ResendEvent(NSEvent* Event)
 	}, NSDefaultRunLoopMode, true);
 }
 
-void FMacApplication::ProcessMouseMovedEvent(const FDeferredMacEvent& Event)
+void FMacApplication::ProcessMouseMovedEvent(const FDeferredMacEvent& Event, TSharedPtr<FMacWindow> EventWindow)
 {
-	TSharedPtr<FMacWindow> EventWindow = Event.Window.IsValid() ? Event.Window.Pin() : nullptr;
 	if (EventWindow.IsValid() && EventWindow->IsRegularWindow())
 	{
 		bool IsMouseOverTitleBar = false;
@@ -726,15 +723,14 @@ void FMacApplication::ProcessMouseMovedEvent(const FDeferredMacEvent& Event)
 		}
 	}
 
-	if (Event.Window.IsValid() && EventWindow->GetWindowHandle() && !DraggedWindow && !GetCapture())
+	if (EventWindow.IsValid() && EventWindow->GetWindowHandle() && !DraggedWindow && !GetCapture())
 	{
 		MessageHandler->OnCursorSet();
 	}
 }
 
-void FMacApplication::ProcessMouseDownEvent(const FDeferredMacEvent& Event)
+void FMacApplication::ProcessMouseDownEvent(const FDeferredMacEvent& Event, TSharedPtr<FMacWindow> EventWindow)
 {
-	TSharedPtr<FMacWindow> EventWindow = Event.Window.Pin();
 	EMouseButtons::Type Button = Event.Type == NSLeftMouseDown ? EMouseButtons::Left : EMouseButtons::Right;
 	if (Event.Type == NSOtherMouseDown)
 	{
@@ -754,7 +750,7 @@ void FMacApplication::ProcessMouseDownEvent(const FDeferredMacEvent& Event)
 		}
 	}
 
-	if (Event.Window.IsValid())
+	if (EventWindow.IsValid())
 	{
 		if (Button == LastPressedMouseButton && (Event.ClickCount % 2) == 0)
 		{
@@ -783,7 +779,7 @@ void FMacApplication::ProcessMouseDownEvent(const FDeferredMacEvent& Event)
 			MessageHandler->OnMouseDown(EventWindow, Button);
 		}
 
-		if (Event.Window.IsValid() && EventWindow->GetWindowHandle() && !DraggedWindow && !GetCapture())
+		if (EventWindow.IsValid() && EventWindow->GetWindowHandle() && !DraggedWindow && !GetCapture())
 		{
 			MessageHandler->OnCursorSet();
 		}
@@ -792,7 +788,7 @@ void FMacApplication::ProcessMouseDownEvent(const FDeferredMacEvent& Event)
 	LastPressedMouseButton = Button;
 }
 
-void FMacApplication::ProcessMouseUpEvent(const FDeferredMacEvent& Event)
+void FMacApplication::ProcessMouseUpEvent(const FDeferredMacEvent& Event, TSharedPtr<FMacWindow> EventWindow)
 {
 	EMouseButtons::Type Button = Event.Type == NSLeftMouseUp ? EMouseButtons::Left : EMouseButtons::Right;
 	if (Event.Type == NSOtherMouseUp)
@@ -815,7 +811,7 @@ void FMacApplication::ProcessMouseUpEvent(const FDeferredMacEvent& Event)
 
 	MessageHandler->OnMouseUp(Button);
 
-	if (Event.Window.IsValid() && Event.Window.Pin()->GetWindowHandle() && !DraggedWindow && !GetCapture())
+	if (EventWindow.IsValid() && EventWindow->GetWindowHandle() && !DraggedWindow && !GetCapture())
 	{
 		MessageHandler->OnCursorSet();
 	}
@@ -824,7 +820,7 @@ void FMacApplication::ProcessMouseUpEvent(const FDeferredMacEvent& Event)
 	DraggedWindow = nullptr;
 }
 
-void FMacApplication::ProcessScrollWheelEvent(const FDeferredMacEvent& Event)
+void FMacApplication::ProcessScrollWheelEvent(const FDeferredMacEvent& Event, TSharedPtr<FMacWindow> EventWindow)
 {
 	const float DeltaX = (Event.ModifierFlags & NSShiftKeyMask) ? Event.Delta.Y : Event.Delta.X;
 	const float DeltaY = (Event.ModifierFlags & NSShiftKeyMask) ? Event.Delta.X : Event.Delta.Y;
@@ -844,8 +840,7 @@ void FMacApplication::ProcessScrollWheelEvent(const FDeferredMacEvent& Event)
 		MessageHandler->OnMouseWheel(DeltaY);
 	}
 
-	TSharedPtr<FMacWindow> EventWindow = Event.Window.Pin();
-	if (Event.Window.IsValid() && EventWindow->GetWindowHandle() && !DraggedWindow && !GetCapture())
+	if (EventWindow.IsValid() && EventWindow->GetWindowHandle() && !DraggedWindow && !GetCapture())
 	{
 		MessageHandler->OnCursorSet();
 	}
@@ -872,10 +867,10 @@ void FMacApplication::ProcessGestureEvent(const FDeferredMacEvent& Event)
 	}
 }
 
-void FMacApplication::ProcessKeyDownEvent(const FDeferredMacEvent& Event)
+void FMacApplication::ProcessKeyDownEvent(const FDeferredMacEvent& Event, TSharedPtr<FMacWindow> EventWindow)
 {
 	bool bHandled = false;
-	if (!bSystemModalMode && Event.Window.IsValid())
+	if (!bSystemModalMode && EventWindow.IsValid())
 	{
 		const TCHAR Character = ConvertChar([Event.Characters characterAtIndex:0]);
 		const TCHAR CharCode = [Event.CharactersIgnoringModifiers characterAtIndex:0];
diff --git a/Engine/Source/Runtime/Core/Private/Mac/MacWindow.cpp b/Engine/Source/Runtime/Core/Private/Mac/MacWindow.cpp
--- a/Engine/Source/Runtime/Core/Private/Mac/MacWindow.cpp
+++ b/Engine/Source/Runtime/Core/Private/Mac/MacWindow.cpp
@@ -295,11 +295,7 @@ void FMacWindow::Destroy()
 		bIsClosed = true;
 		[WindowHandle setAlphaValue:0.0f];
 		[WindowHandle setBackgroundColor:[NSColor clearColor]];
-		TSharedPtr<FMacWindow> Window = MacApplication->FindWindowByNSWindow(WindowHandle);
-		if (Window.IsValid())
-		{
-			MacApplication->OnWindowDestroyed(Window.ToSharedRef());
-		}
+		MacApplication->OnWindowDestroyed(SharedThis(this));
 		WindowHandle = nullptr;
 	}
 }
@@ -575,19 +571,3 @@ void FMacWindow::OnDisplayReconfiguration(CGDirectDisplayID Display, CGDisplayCh
 		});
 	}
 }
-
-bool FMacWindow::OnIMKKeyDown(NSEvent* Event)
-{
-	if(WindowHandle && [WindowHandle openGLView])
-	{
-		return MainThreadReturn(^{
-			SCOPED_AUTORELEASE_POOL;
-			FCocoaTextView* View = (FCocoaTextView*)[WindowHandle openGLView];
-			return View && [View imkKeyDown:Event];
-		}, UE4IMEEventMode);
-	}
-	else
-	{
-		return false;
-	}
-}
diff --git a/Engine/Source/Runtime/Core/Public/Mac/MacApplication.h b/Engine/Source/Runtime/Core/Public/Mac/MacApplication.h
--- a/Engine/Source/Runtime/Core/Public/Mac/MacApplication.h
+++ b/Engine/Source/Runtime/Core/Public/Mac/MacApplication.h
@@ -12,6 +12,7 @@ struct FDeferredMacEvent
 {
 	FDeferredMacEvent()
 	:	Event(nullptr)
+	,	Window(nullptr)
 	,	Type(0)
 	,	LocationInWindow(FVector2D::ZeroVector)
 	,	ModifierFlags(0)
@@ -92,7 +93,7 @@ struct FDeferredMacEvent
 	// However, in some cases we need the original NSEvent (highlighting menus, resending unhandled key events), so we store it as well.
 	NSEvent* Event;
 
-	TWeakPtr<FMacWindow> Window;
+	FCocoaWindow* Window;
 
 	int32 Type;
 	FVector2D LocationInWindow;
@@ -199,12 +200,12 @@ private:
 	void ProcessEvent(const FDeferredMacEvent& Event);
 	void ResendEvent(NSEvent* Event);
 
-	void ProcessMouseMovedEvent(const FDeferredMacEvent& Event);
-	void ProcessMouseDownEvent(const FDeferredMacEvent& Event);
-	void ProcessMouseUpEvent(const FDeferredMacEvent& Event);
-	void ProcessScrollWheelEvent(const FDeferredMacEvent& Event);
+	void ProcessMouseMovedEvent(const FDeferredMacEvent& Event, TSharedPtr<FMacWindow> EventWindow);
+	void ProcessMouseDownEvent(const FDeferredMacEvent& Event, TSharedPtr<FMacWindow> EventWindow);
+	void ProcessMouseUpEvent(const FDeferredMacEvent& Event, TSharedPtr<FMacWindow> EventWindow);
+	void ProcessScrollWheelEvent(const FDeferredMacEvent& Event, TSharedPtr<FMacWindow> EventWindow);
 	void ProcessGestureEvent(const FDeferredMacEvent& Event);
-	void ProcessKeyDownEvent(const FDeferredMacEvent& Event);
+	void ProcessKeyDownEvent(const FDeferredMacEvent& Event, TSharedPtr<FMacWindow> EventWindow);
 	void ProcessKeyUpEvent(const FDeferredMacEvent& Event);
 
 	void OnWindowDidMove(TSharedRef<FMacWindow> Window);
diff --git a/Engine/Source/Runtime/Core/Public/Mac/MacWindow.h b/Engine/Source/Runtime/Core/Public/Mac/MacWindow.h
--- a/Engine/Source/Runtime/Core/Public/Mac/MacWindow.h
+++ b/Engine/Source/Runtime/Core/Public/Mac/MacWindow.h
@@ -22,8 +22,6 @@ public:
 	void Initialize( class FMacApplication* const Application, const TSharedRef< FGenericWindowDefinition >& InDefinition, const TSharedPtr< FMacWindow >& InParent, const bool bShowImmediately );
 	
 	void OnDisplayReconfiguration(CGDirectDisplayID Display, CGDisplayChangeSummaryFlags Flags);
-	
-	bool OnIMKKeyDown(NSEvent* Event);
 
 public:
 
