Fixed the difference in how different versions of Windows handle overlapped i/o read requests

diff --git a/Engine/Source/Runtime/Core/Private/Windows/WindowsPlatformFile.cpp b/Engine/Source/Runtime/Core/Private/Windows/WindowsPlatformFile.cpp
--- a/Engine/Source/Runtime/Core/Private/Windows/WindowsPlatformFile.cpp
+++ b/Engine/Source/Runtime/Core/Private/Windows/WindowsPlatformFile.cpp
@@ -107,6 +107,18 @@ protected:
 		OverlappedIO.OffsetHigh = LI.HighPart;
 	}
 
+	FORCEINLINE void UpdateFileOffsetAfterRead(uint32 AmountRead)
+	{
+		bHasReadOutstanding = false;
+		OverlappedFilePos += AmountRead;
+		// Update the overlapped structure since it uses this for where to read from
+		CopyOverlappedPosition();
+		if (OverlappedFilePos >= uint64(FileSize))
+		{
+			bIsAtEOF = true;
+		}
+	}
+
 	bool WaitForAsyncRead()
 	{
 		// Check for already being at EOF because we won't issue a read
@@ -117,14 +129,7 @@ protected:
 		uint32 NumRead = 0;
 		if (GetOverlappedResult(Handle, &OverlappedIO, (::DWORD*)&NumRead, true) != false)
 		{
-			bHasReadOutstanding = false;
-			OverlappedFilePos += NumRead;
-			// Update the overlapped structure or the API will get confused
-			CopyOverlappedPosition();
-			if (OverlappedFilePos >= uint64(FileSize))
-			{
-				bIsAtEOF = true;
-			}
+			UpdateFileOffsetAfterRead(NumRead);
 			return true;
 		}
 		else if (GetLastError() == ERROR_HANDLE_EOF)
@@ -143,7 +148,20 @@ protected:
 			CurrentAsyncReadBuffer = BufferToReadInto;
 			uint32 NumRead = 0;
 			// Now kick off an async read
-			ReadFile(Handle, Buffers[BufferToReadInto], BufferSize, (::DWORD*)&NumRead, &OverlappedIO);
+			if (!ReadFile(Handle, Buffers[BufferToReadInto], BufferSize, (::DWORD*)&NumRead, &OverlappedIO))
+			{
+				uint32 ErrorCode = GetLastError();
+				if (ErrorCode != ERROR_IO_PENDING)
+				{
+					bIsAtEOF = true;
+					bHasReadOutstanding = false;
+				}
+			}
+			else
+			{
+				// Read completed immediately
+				UpdateFileOffsetAfterRead(NumRead);
+			}
 		}
 	}
 
@@ -571,25 +589,16 @@ public:
 		}
 		return Result;
 	}
-
 	virtual IFileHandle* OpenRead(const TCHAR* Filename, bool bAllowWrite = false) override
 	{
 		uint32  Access    = GENERIC_READ;
 		uint32  WinFlags  = FILE_SHARE_READ | (bAllowWrite ? FILE_SHARE_WRITE : 0);
 		uint32  Create    = OPEN_EXISTING;
-#if USE_OVERLAPPED_READS
 		HANDLE Handle    = CreateFileW(*NormalizeFilename(Filename), Access, WinFlags, NULL, Create, FILE_ATTRIBUTE_NORMAL | FILE_FLAG_OVERLAPPED, NULL);
 		if (Handle != INVALID_HANDLE_VALUE)
 		{
 			return new FAsyncBufferedFileReaderWindows(Handle);
 		}
-#else
-		HANDLE Handle    = CreateFileW(*NormalizeFilename(Filename), Access, WinFlags, NULL, Create, FILE_ATTRIBUTE_NORMAL, NULL);
-		if(Handle != INVALID_HANDLE_VALUE)
-		{
-			return new FFileHandleWindows(Handle);
-		}
-#endif
 		return NULL;
 	}
 	virtual IFileHandle* OpenWrite(const TCHAR* Filename, bool bAppend = false, bool bAllowRead = false) override
