Fix for generated config directory check in the ::Write code. Caused some users to have issues with generated ini files becoming corrupt.

UE-15031

diff --git a/Engine/Source/Runtime/Core/Private/Misc/ConfigCacheIni.cpp b/Engine/Source/Runtime/Core/Private/Misc/ConfigCacheIni.cpp
--- a/Engine/Source/Runtime/Core/Private/Misc/ConfigCacheIni.cpp
+++ b/Engine/Source/Runtime/Core/Private/Misc/ConfigCacheIni.cpp
@@ -947,7 +947,9 @@ bool FConfigFile::Write( const FString& Filename, bool bDoRemoteWrite/* = true*/
 				const bool bOptionIsFromCommandline = PropertySetFromCommandlineOption(this, SectionName, PropertyName, PropertyValue);
 
 				// If we are writing to a default config file and this property is an array, we need to be careful to remove those from higher up the hierarchy
-				const bool bIsADefaultIniWrite = !Filename.Contains( FPaths::GeneratedConfigDir() );
+				const FString AbsoluteFilename = FPaths::ConvertRelativePathToFull(Filename);
+				const FString AbsoluteGeneratedConfigDir = FPaths::ConvertRelativePathToFull(FPaths::GeneratedConfigDir());
+				const bool bIsADefaultIniWrite = !AbsoluteFilename.Contains(AbsoluteGeneratedConfigDir);
 
 				// Check if the property matches the source configs. We do not wanna write it out if so.
 				if ((bIsADefaultIniWrite || bDifferentNumberOfElements || !DoesConfigPropertyValueMatch(SourceConfigFile, SectionName, PropertyName, PropertyValue)) && !bOptionIsFromCommandline)
