Fixed infinite loop caused by erroneous decrement of unsigned integer
Tidied a previous fix for the same issue more concise logic.
This addresses UE-15012.

diff --git a/Engine/Source/Runtime/Core/Private/Misc/ConfigManifest.cpp b/Engine/Source/Runtime/Core/Private/Misc/ConfigManifest.cpp
--- a/Engine/Source/Runtime/Core/Private/Misc/ConfigManifest.cpp
+++ b/Engine/Source/Runtime/Core/Private/Misc/ConfigManifest.cpp
@@ -54,11 +54,6 @@ FString ProjectAgnosticIniPath(const TCHAR* InLeaf)
 	return FPaths::GameAgnosticSavedDir() / TEXT("Config") / ANSI_TO_TCHAR(FPlatformProperties::PlatformName()) / InLeaf;
 }
 
-FString GetConfigDirectoryToSearch(const FEngineVersion& Version)
-{
-	return FString(FPlatformProcess::UserSettingsDir()) / ENGINE_VERSION_TEXT(EPIC_PRODUCT_IDENTIFIER) / Version.ToString(EVersionComponent::Minor) / TEXT("Saved") / TEXT("Config") / ANSI_TO_TCHAR(FPlatformProperties::PlatformName());
-}
-
 /** Migrates config files from a previous version of the engine. Does nothing on non-installed versions */
 void MigratePreviousEngineInis()
 {
@@ -68,15 +63,14 @@ void MigratePreviousEngineInis()
 		return;
 	}
 
-	FEngineVersion PreviousVersion(GEngineVersion.GetMajor(), GEngineVersion.GetMinor() - 1, GEngineVersion.GetPatch(), GEngineVersion.GetChangelist(), GEngineVersion.GetBranch());
+	int32 MinorVersion = GEngineVersion.GetMinor() - 1;
 
 	IPlatformFile& PlatformFile = FPlatformFileManager::Get().GetPlatformFile();
-	FString Directory;
-
-	// Minor is stored as unsigned int, we have to split the loop to avoid infinite iteration.
-	while(PreviousVersion.GetMinor() > 0)
+	while(MinorVersion >= 0)
 	{
-		Directory = GetConfigDirectoryToSearch(PreviousVersion);
+		const FEngineVersion PreviousVersion(GEngineVersion.GetMajor(), MinorVersion--, 0, 0, FString());
+	
+		const FString Directory = FString(FPlatformProcess::UserSettingsDir()) / ENGINE_VERSION_TEXT(EPIC_PRODUCT_IDENTIFIER) / PreviousVersion.ToString(EVersionComponent::Minor) / TEXT("Saved") / TEXT("Config") / ANSI_TO_TCHAR(FPlatformProperties::PlatformName());
 		if (FPaths::DirectoryExists(Directory))
 		{
 			const FString DestDir = ProjectAgnosticIniPath(TEXT(""));
@@ -88,19 +82,6 @@ void MigratePreviousEngineInis()
 			// If we failed to create the directory tree anyway we don't want to allow the possibility of upgrading from even older versions, so early return regardless
 			return;
 		}
-
-		PreviousVersion = FEngineVersion(PreviousVersion.GetMajor(), PreviousVersion.GetMinor() - 1, PreviousVersion.GetPatch(), PreviousVersion.GetChangelist(), PreviousVersion.GetBranch());
-	}
-
-	// One last check for version with Minor == 0
-	Directory = GetConfigDirectoryToSearch(PreviousVersion);
-	if (FPaths::DirectoryExists(Directory))
-	{
-		const FString DestDir = ProjectAgnosticIniPath(TEXT(""));
-		if (PlatformFile.CreateDirectoryTree(*DestDir))
-		{
-			PlatformFile.CopyDirectoryTree(*DestDir, *Directory, false);
-		}
 	}
 }
 
