Added equality/inequality operator to FInterpCurve.

diff --git a/Engine/Source/Runtime/Core/Public/Math/InterpCurve.h b/Engine/Source/Runtime/Core/Public/Math/InterpCurve.h
--- a/Engine/Source/Runtime/Core/Public/Math/InterpCurve.h
+++ b/Engine/Source/Runtime/Core/Public/Math/InterpCurve.h
@@ -123,6 +123,24 @@ public:
 		return Ar;
 	}
 
+	/**
+	 * Compare equality of two FInterpCurves
+	 */
+	friend bool operator==(const FInterpCurve& Curve1, const FInterpCurve& Curve2)
+	{
+		return (Curve1.Points == Curve2.Points &&
+				Curve1.bIsLooped == Curve2.bIsLooped &&
+				(!Curve1.bIsLooped || Curve1.LoopKeyOffset == Curve2.LoopKeyOffset));
+	}
+
+	/**
+	 * Compare inequality of two FInterpCurves
+	 */
+	friend bool operator!=(const FInterpCurve& Curve1, const FInterpCurve& Curve2)
+	{
+		return !(Curve1 == Curve2);
+	}
+
 	/**
 	 * Finds the lower index of the two points whose input values bound the supplied input value.
 	 */
