Core - Added a comment

diff --git a/Engine/Source/Runtime/Core/Private/GenericPlatform/GenericPlatformMemory.cpp b/Engine/Source/Runtime/Core/Private/GenericPlatform/GenericPlatformMemory.cpp
--- a/Engine/Source/Runtime/Core/Private/GenericPlatform/GenericPlatformMemory.cpp
+++ b/Engine/Source/Runtime/Core/Private/GenericPlatform/GenericPlatformMemory.cpp
@@ -49,6 +49,7 @@ void FGenericPlatformMemory::Init()
 
 void FGenericPlatformMemory::OnOutOfMemory(uint64 Size, uint32 Alignment)
 {
+	// Update memory stats before we enter the crash handler.
 	bIsOOM = true;
 	UE_LOG(LogMemory, Fatal, TEXT("Ran out of memory allocating %llu bytes with alignment %u"), Size, Alignment);
 }
