fix for debug iOS game crashing
UE-15078
#ios

diff --git a/Engine/Source/Runtime/Core/Public/IOS/IOSPlatformMisc.h b/Engine/Source/Runtime/Core/Public/IOS/IOSPlatformMisc.h
--- a/Engine/Source/Runtime/Core/Public/IOS/IOSPlatformMisc.h
+++ b/Engine/Source/Runtime/Core/Public/IOS/IOSPlatformMisc.h
@@ -20,7 +20,7 @@ struct CORE_API FIOSPlatformMisc : public FGenericPlatformMisc
 	static void* GetHardwareWindow();
 
 #if !UE_BUILD_SHIPPING
-	FORCEINLINE static bool IsDebuggerPresent()
+	static bool IsDebuggerPresent()
 	{
 		// Based on http://developer.apple.com/library/mac/#qa/qa1361/_index.html
 
@@ -62,7 +62,7 @@ struct CORE_API FIOSPlatformMisc : public FGenericPlatformMisc
 	}
 
 	/** Prompts for remote debugging if debugger is not attached. Regardless of result, breaks into debugger afterwards. Returns false for use in conditionals. */
-	static FORCEINLINE bool DebugBreakAndPromptForRemoteReturningFalse()
+	FORCEINLINE static bool DebugBreakAndPromptForRemoteReturningFalse()
 	{
 #if !UE_BUILD_SHIPPING
 		if (!IsDebuggerPresent())
