Hot-reload will now always call ShutdownModule even when abondoning DLLs

diff --git a/Engine/Source/Developer/HotReload/Private/HotReload.cpp b/Engine/Source/Developer/HotReload/Private/HotReload.cpp
--- a/Engine/Source/Developer/HotReload/Private/HotReload.cpp
+++ b/Engine/Source/Developer/HotReload/Private/HotReload.cpp
@@ -645,7 +645,8 @@ ECompilationResult::Type FHotReloadModule::DoHotReloadInternal(bool bRecompileFi
 
 			// Abandon the old module.  We can't unload it because various data structures may be living
 			// that have vtables pointing to code that would become invalidated.
-			FModuleManager::Get().AbandonModule(ShortPackageName);
+			const bool bAbandonModule = true;
+			FModuleManager::Get().UnloadOrAbandonModuleWithCallback(ShortPackageName, HotReloadAr, bAbandonModule);
 
 			// Module should never be loaded at this point
 			check(!FModuleManager::Get().IsModuleLoaded(ShortPackageName));
diff --git a/Engine/Source/Runtime/Core/Private/Modules/ModuleManager.cpp b/Engine/Source/Runtime/Core/Private/Modules/ModuleManager.cpp
--- a/Engine/Source/Runtime/Core/Private/Modules/ModuleManager.cpp
+++ b/Engine/Source/Runtime/Core/Private/Modules/ModuleManager.cpp
@@ -560,6 +560,9 @@ void FModuleManager::AbandonModule( const FName InModuleName )
 		// Only if already loaded
 		if( ModuleInfo->Module.IsValid() )
 		{
+			// Allow the module to shut itself down
+			ModuleInfo->Module->ShutdownModule();
+
 			// Release reference to module interface.  This will actually destroy the module object.
 			// @todo UE4 DLL: Could be dangerous in some cases to reset the module interface while abandoning.  Currently not
 			// a problem because script modules don't implement any functionality here.  Possible, we should keep these references
@@ -942,14 +945,14 @@ void FModuleManager::FindModulePathsInDirectory(const FString& InDirectoryName,
 }
 
 
-void FModuleManager::UnloadOrAbandonModuleWithCallback( const FName InModuleName, FOutputDevice &Ar )
+void FModuleManager::UnloadOrAbandonModuleWithCallback(const FName InModuleName, FOutputDevice &Ar, bool bAbandonOnly)
 {
 	auto Module = FindModuleChecked(InModuleName);
 	
 	Module->Module->PreUnloadCallback();
 
 	const bool bIsHotReloadable = DoesLoadedModuleHaveUObjects( InModuleName );
-	if( bIsHotReloadable && Module->Module->SupportsDynamicReloading() )
+	if (!bAbandonOnly && bIsHotReloadable && Module->Module->SupportsDynamicReloading())
 	{
 		if( !UnloadModule( InModuleName ))
 		{
@@ -958,7 +961,11 @@ void FModuleManager::UnloadOrAbandonModuleWithCallback( const FName InModuleName
 	}
 	else
 	{
-		Ar.Logf( TEXT( "Module being reloaded does not support dynamic unloading -- abandoning existing loaded module so that we can load the recompiled version!" ) );
+		// Don't warn if abandoning was the intent here
+		if (!bAbandonOnly)
+		{			
+			Ar.Logf(TEXT("Module being reloaded does not support dynamic unloading -- abandoning existing loaded module so that we can load the recompiled version!"));
+		}
 		AbandonModule( InModuleName );
 	}
 }
diff --git a/Engine/Source/Runtime/Core/Public/Modules/ModuleManager.h b/Engine/Source/Runtime/Core/Public/Modules/ModuleManager.h
--- a/Engine/Source/Runtime/Core/Public/Modules/ModuleManager.h
+++ b/Engine/Source/Runtime/Core/Public/Modules/ModuleManager.h
@@ -211,9 +211,10 @@ public:
 	 *
 	 * @param InModuleName The name of the module to unload.
 	 * @param Ar The archive to receive error messages, if any.
+	 * @param bAbandonOnly Do not try to unload the module, only abandon.
 	 * @see LoadModuleWithCallback
 	 */
-	void UnloadOrAbandonModuleWithCallback( const FName InModuleName, FOutputDevice &Ar );
+	void UnloadOrAbandonModuleWithCallback( const FName InModuleName, FOutputDevice &Ar, bool bAbandonOnly = false);
 
 public:
 
