Fixed the problem with modal dialog buttons not working in some cases on Mac

diff --git a/Engine/Source/Runtime/Core/Private/Mac/MacApplication.cpp b/Engine/Source/Runtime/Core/Private/Mac/MacApplication.cpp
--- a/Engine/Source/Runtime/Core/Private/Mac/MacApplication.cpp
+++ b/Engine/Source/Runtime/Core/Private/Mac/MacApplication.cpp
@@ -172,28 +172,24 @@ void FMacApplication::PumpMessages(const float TimeDelta)
 
 void FMacApplication::ProcessDeferredEvents(const float TimeDelta)
 {
-	if (!FPlatformMisc::bIsPumpingMessages)
-	{
-		TArray<FDeferredMacEvent> EventsToProcess;
+	TArray<FDeferredMacEvent> EventsToProcess;
 
-		EventsMutex.Lock();
-		EventsToProcess.Append(DeferredEvents);
-		DeferredEvents.Empty();
-		EventsMutex.Unlock();
+	EventsMutex.Lock();
+	EventsToProcess.Append(DeferredEvents);
+	DeferredEvents.Empty();
+	EventsMutex.Unlock();
 
-		const bool bAlreadyProcessingDeferredEvents = bIsProcessingDeferredEvents;
-		bIsProcessingDeferredEvents = true;
+	const bool bAlreadyProcessingDeferredEvents = bIsProcessingDeferredEvents;
+	bIsProcessingDeferredEvents = true;
 
-		for (int32 Index = 0; Index < EventsToProcess.Num(); ++Index)
-		{
-			ProcessEvent(EventsToProcess[Index]);
-		}
+	for (int32 Index = 0; Index < EventsToProcess.Num(); ++Index)
+	{
+		ProcessEvent(EventsToProcess[Index]);
+	}
 
-		bIsProcessingDeferredEvents = bAlreadyProcessingDeferredEvents;
+	bIsProcessingDeferredEvents = bAlreadyProcessingDeferredEvents;
 
-		InvalidateTextLayouts();
-		CloseQueuedWindows();
-	}
+	InvalidateTextLayouts();
 }
 
 TSharedRef<FGenericWindow> FMacApplication::MakeWindow()
diff --git a/Engine/Source/Runtime/Core/Private/Mac/MacPlatformMisc.cpp b/Engine/Source/Runtime/Core/Private/Mac/MacPlatformMisc.cpp
--- a/Engine/Source/Runtime/Core/Private/Mac/MacPlatformMisc.cpp
+++ b/Engine/Source/Runtime/Core/Private/Mac/MacPlatformMisc.cpp
@@ -282,7 +282,6 @@ FMacMallocCrashHandler* FMacApplicationInfo::CrashMalloc = nullptr;
 
 UpdateCachedMacMenuStateProc FMacPlatformMisc::UpdateCachedMacMenuState = nullptr;
 bool FMacPlatformMisc::bChachedMacMenuStateNeedsUpdate = true;
-bool FMacPlatformMisc::bIsPumpingMessages = false;
 id<NSObject> FMacPlatformMisc::CommandletActivity = nil;
 
 void FMacPlatformMisc::PlatformPreInit()
@@ -582,9 +581,6 @@ void FMacPlatformMisc::PumpMessages( bool bFromMainLoop )
 {
 	if( bFromMainLoop )
 	{
-		const bool bIsAlreadyPumpingMessages = bIsPumpingMessages;
-		bIsPumpingMessages = true;
-
 		ProcessGameThreadEvents();
 
 		if (MacApplication && !MacApplication->IsProcessingDeferredEvents() && IsInGameThread())
@@ -595,8 +591,6 @@ void FMacPlatformMisc::PumpMessages( bool bFromMainLoop )
 				bChachedMacMenuStateNeedsUpdate = false;
 			}
 		}
-
-		bIsPumpingMessages = bIsAlreadyPumpingMessages;
 	}
 }
 
diff --git a/Engine/Source/Runtime/Core/Public/Mac/MacPlatformMisc.h b/Engine/Source/Runtime/Core/Public/Mac/MacPlatformMisc.h
--- a/Engine/Source/Runtime/Core/Public/Mac/MacPlatformMisc.h
+++ b/Engine/Source/Runtime/Core/Public/Mac/MacPlatformMisc.h
@@ -199,7 +199,6 @@ struct CORE_API FMacPlatformMisc : public FGenericPlatformMisc
 	static FString GetOperatingSystemId();
 
 	static bool bChachedMacMenuStateNeedsUpdate;
-	static bool bIsPumpingMessages;
 
 	static id<NSObject> CommandletActivity;
 };
