PR #1108: Added Polar/Cartesian conversion functions with FVector2D parameters (Contributed by [REDACTED_PERSON])

diff --git a/Engine/Source/Runtime/Core/Private/Math/UnrealMath.cpp b/Engine/Source/Runtime/Core/Private/Math/UnrealMath.cpp
--- a/Engine/Source/Runtime/Core/Private/Math/UnrealMath.cpp
+++ b/Engine/Source/Runtime/Core/Private/Math/UnrealMath.cpp
@@ -3021,4 +3021,16 @@ void FMath::ApplyScaleToFloat(float& Dst, const FVector& DeltaScale, float Magni
 	const float Multiplier = ( DeltaScale.X > 0.0f || DeltaScale.Y > 0.0f || DeltaScale.Z > 0.0f ) ? Magnitude : -Magnitude;
 	Dst += Multiplier * DeltaScale.Size();
 	Dst = FMath::Max( 0.0f, Dst );
+}
+
+void FMath::CartesianToPolar(const FVector2D InCart, FVector2D& OutPolar)
+{
+	OutPolar.X = Sqrt(Square(InCart.X) + Square(InCart.Y));
+	OutPolar.Y = Atan2(InCart.Y, InCart.X);
+}
+
+void FMath::PolarToCartesian(const FVector2D InPolar, FVector2D& OutCart)
+{
+	OutCart.X = InPolar.X * Cos(InPolar.Y);
+	OutCart.Y = InPolar.X * Sin(InPolar.Y);
 }
\ No newline at end of file
diff --git a/Engine/Source/Runtime/Core/Public/Math/UnrealMathUtility.h b/Engine/Source/Runtime/Core/Public/Math/UnrealMathUtility.h
--- a/Engine/Source/Runtime/Core/Public/Math/UnrealMathUtility.h
+++ b/Engine/Source/Runtime/Core/Public/Math/UnrealMathUtility.h
@@ -476,17 +476,22 @@ struct FMath : public FPlatformMath
 	static CORE_API float FixedTurn(float InCurrent, float InDesired, float InDeltaRate);
 
 	/** Converts given Cartesian coordinate pair to Polar coordinate system. */
-	static FORCEINLINE void CartesianToPolar(float X, float Y, float& OutRad, float& OutAng)
+	static FORCEINLINE void CartesianToPolar(const float X, const float Y, float& OutRad, float& OutAng)
 	{
 		OutRad = Sqrt(Square(X) + Square(Y));
 		OutAng = Atan2(Y, X);
 	}
+	/** Converts given Cartesian coordinate pair to Polar coordinate system. */
+	static FORCEINLINE void CartesianToPolar(const FVector2D InCart, FVector2D& OutPolar);
+
 	/** Converts given Polar coordinate pair to Cartesian coordinate system. */
-	static FORCEINLINE void PolarToCartesian(float Rad, float Ang, float& OutX, float& OutY)
+	static FORCEINLINE void PolarToCartesian(const float Rad, const float Ang, float& OutX, float& OutY)
 	{
 		OutX = Rad * Cos(Ang);
 		OutY = Rad * Sin(Ang);
 	}
+	/** Converts given Polar coordinate pair to Cartesian coordinate system. */
+	static FORCEINLINE void PolarToCartesian(const FVector2D InPolar, FVector2D& OutCart);
 
 	/**
 	 * Calculates the dotted distance of vector 'Direction' to coordinate system O(AxisX,AxisY,AxisZ).
