Re-added CloseQueuedWindows(); on Mac, accidentally removed in CL 2541218

diff --git a/Engine/Source/Runtime/Core/Private/Mac/MacApplication.cpp b/Engine/Source/Runtime/Core/Private/Mac/MacApplication.cpp
--- a/Engine/Source/Runtime/Core/Private/Mac/MacApplication.cpp
+++ b/Engine/Source/Runtime/Core/Private/Mac/MacApplication.cpp
@@ -190,6 +190,7 @@ void FMacApplication::ProcessDeferredEvents(const float TimeDelta)
 	bIsProcessingDeferredEvents = bAlreadyProcessingDeferredEvents;
 
 	InvalidateTextLayouts();
+	CloseQueuedWindows();
 }
 
 TSharedRef<FGenericWindow> FMacApplication::MakeWindow()
