#ue4 - Perf: Faster Quat->Rotator conversion.

- We only need to normalize 1 axis and only rarely (when pitch is -90 or 90). ASin and ATan2 produce normalized results (as in [-pi, pi]).

diff --git a/Engine/Source/Runtime/Core/Private/Math/UnrealMath.cpp b/Engine/Source/Runtime/Core/Private/Math/UnrealMath.cpp
--- a/Engine/Source/Runtime/Core/Private/Math/UnrealMath.cpp
+++ b/Engine/Source/Runtime/Core/Private/Math/UnrealMath.cpp
@@ -450,61 +450,20 @@ FRotator FQuat::Rotator() const
 	// but that isn't the case for us, so I went through different testing, and finally found the case 
 	// where both of world lives happily. 
 	const float SINGULARITY_THRESHOLD = 0.4999995f;
-
-#if PLATFORM_ENABLE_VECTORINTRINSICS
-
+	const float RAD_TO_DEG = (180.f)/PI;
 	FRotator RotatorFromQuat;
-	union { VectorRegister v; float f[4]; } VRotatorFromQuat;
 
 	if (SingularityTest < -SINGULARITY_THRESHOLD)
 	{
-		// Pitch
-		VRotatorFromQuat.f[0] = 3.0f*HALF_PI;	// 270 deg
-		// Yaw
-		VRotatorFromQuat.f[1] = FMath::Atan2(YawY, YawX);
-		// Roll
-		VRotatorFromQuat.f[2] = -VRotatorFromQuat.f[1] - (2.f * FMath::Atan2(X, W));
-	}
-	else if (SingularityTest > SINGULARITY_THRESHOLD)
-	{
-		// Pitch
-		VRotatorFromQuat.f[0] = HALF_PI;	// 90 deg
-		// Yaw
-		VRotatorFromQuat.f[1] = FMath::Atan2(YawY, YawX);
-		//Roll
-		VRotatorFromQuat.f[2] = VRotatorFromQuat.f[1] - (2.f * FMath::Atan2(X, W));
-	}
-	else
-	{
-		//Pitch
-		VRotatorFromQuat.f[0] = FMath::FastAsin(2.f*(SingularityTest));
-		// Yaw
-		VRotatorFromQuat.f[1] = FMath::Atan2(YawY, YawX);
-		// Roll
-		VRotatorFromQuat.f[2] = FMath::Atan2(-2.f*(W*X+Y*Z), (1.f-2.f*(FMath::Square(X) + FMath::Square(Y))));
-	}
-
-	VRotatorFromQuat.f[3] = 0.f; // We should initialize this, otherwise the value can be denormalized which is bad for floating point perf.
-	VRotatorFromQuat.v = VectorMultiply(VRotatorFromQuat.v, GlobalVectorConstants::RAD_TO_DEG);
-	VRotatorFromQuat.v = VectorNormalizeRotator(VRotatorFromQuat.v);
-	VectorStoreFloat3(VRotatorFromQuat.v, &RotatorFromQuat);
-
-#else // PLATFORM_ENABLE_VECTORINTRINSICS
-
-	static const float RAD_TO_DEG = (180.f)/PI;
-	FRotator RotatorFromQuat;
-
-	if (SingularityTest < -SINGULARITY_THRESHOLD)
-	{
-		RotatorFromQuat.Pitch = 270.f;
+		RotatorFromQuat.Pitch = -90.f;
 		RotatorFromQuat.Yaw = FMath::Atan2(YawY, YawX) * RAD_TO_DEG;
-		RotatorFromQuat.Roll = -RotatorFromQuat.Yaw - (2.f * FMath::Atan2(X, W) * RAD_TO_DEG);
+		RotatorFromQuat.Roll = FRotator::NormalizeAxis(-RotatorFromQuat.Yaw - (2.f * FMath::Atan2(X, W) * RAD_TO_DEG));
 	}
 	else if (SingularityTest > SINGULARITY_THRESHOLD)
 	{
 		RotatorFromQuat.Pitch = 90.f;
 		RotatorFromQuat.Yaw = FMath::Atan2(YawY, YawX) * RAD_TO_DEG;
-		RotatorFromQuat.Roll = RotatorFromQuat.Yaw - (2.f * FMath::Atan2(X, W) * RAD_TO_DEG);
+		RotatorFromQuat.Roll = FRotator::NormalizeAxis(RotatorFromQuat.Yaw - (2.f * FMath::Atan2(X, W) * RAD_TO_DEG));
 	}
 	else
 	{
@@ -513,11 +472,6 @@ FRotator FQuat::Rotator() const
 		RotatorFromQuat.Roll = FMath::Atan2(-2.f*(W*X+Y*Z), (1.f-2.f*(FMath::Square(X) + FMath::Square(Y)))) * RAD_TO_DEG;
 	}
 
-	RotatorFromQuat.Normalize();
-
-#endif // PLATFORM_ENABLE_VECTORINTRINSICS
-
-
 #if USE_MATRIX_ROTATOR
 	RotatorFromMatrix = RotatorFromMatrix.Clamp();
 
