fix for occasional crash when updating a project setting
UE-15279

diff --git a/Engine/Source/Runtime/Core/Private/Misc/ConfigCacheIni.cpp b/Engine/Source/Runtime/Core/Private/Misc/ConfigCacheIni.cpp
--- a/Engine/Source/Runtime/Core/Private/Misc/ConfigCacheIni.cpp
+++ b/Engine/Source/Runtime/Core/Private/Misc/ConfigCacheIni.cpp
@@ -3185,7 +3185,7 @@ bool FConfigFile::UpdateSinglePropertyInSection(const TCHAR* DiskFilename, const
 		// Iterate through the file, line by line, to find the property we would like to update.
 		FString TheLine;
 		const TCHAR* Ptr = DiskFile.Len() > 0 ? *DiskFile : NULL;
-		while (FParse::Line(&Ptr, TheLine, true))
+		while (Ptr != NULL && FParse::Line(&Ptr, TheLine, true))
 		{
 			const bool bIsASectionHeader = (TheLine.Len() > 3) && (TheLine[0] == '[') && (TheLine[TheLine.Len() - 1] == ']');
 			if (bIsASectionHeader)
