fix for startup crash if home button is pressed and then app is brought back to the foreground
UE-15221
#ios

diff --git a/Engine/Source/Runtime/Core/Private/IOS/IOSAppDelegate.cpp b/Engine/Source/Runtime/Core/Private/IOS/IOSAppDelegate.cpp
--- a/Engine/Source/Runtime/Core/Private/IOS/IOSAppDelegate.cpp
+++ b/Engine/Source/Runtime/Core/Private/IOS/IOSAppDelegate.cpp
@@ -549,11 +549,14 @@ void InstallSignalHandlers()
 	 Sent when the application is about to move from active to inactive state. This can occur for certain types of temporary interruptions (such as an incoming phone call or SMS message) or when the user quits the application and it begins the transition to the background state.
 	 Use this method to pause ongoing tasks, disable timers, and throttle down OpenGL ES frame rates. Games should use this method to pause the game.
 	 */
-    FGraphEventRef ResignTask = FFunctionGraphTask::CreateAndDispatchWhenReady([&]()
-	{
-		FCoreDelegates::ApplicationWillDeactivateDelegate.Broadcast();
-    }, TStatId(), NULL, ENamedThreads::GameThread);
-    FTaskGraphInterface::Get().WaitUntilTaskCompletes(ResignTask);
+    if (bEngineInit)
+    {
+        FGraphEventRef ResignTask = FFunctionGraphTask::CreateAndDispatchWhenReady([&]()
+        {
+            FCoreDelegates::ApplicationWillDeactivateDelegate.Broadcast();
+        }, TStatId(), NULL, ENamedThreads::GameThread);
+        FTaskGraphInterface::Get().WaitUntilTaskCompletes(ResignTask);
+    }
     
 	[self ToggleSuspend:true];
 	[self ToggleAudioSession:false];
