Add support for launching arbitrary URLs on Windows and Mac, rather than only supporting known schemas.

diff --git a/Engine/Source/Runtime/Core/Private/Mac/MacPlatformProcess.cpp b/Engine/Source/Runtime/Core/Private/Mac/MacPlatformProcess.cpp
--- a/Engine/Source/Runtime/Core/Private/Mac/MacPlatformProcess.cpp
+++ b/Engine/Source/Runtime/Core/Private/Mac/MacPlatformProcess.cpp
@@ -193,7 +193,11 @@ void FMacPlatformProcess::LaunchURL( const TCHAR* URL, const TCHAR* Parms, FStri
 
 	UE_LOG(LogMac, Log,  TEXT("LaunchURL %s %s"), URL, Parms?Parms:TEXT("") );
 	NSString* Url = (NSString*)FPlatformString::TCHARToCFString( URL );
-	NSURL* UrlToOpen = [NSURL URLWithString: ([Url hasPrefix: @"http://"] || [Url hasPrefix: @"https://"] || [Url hasPrefix: @"file://"] || [Url hasPrefix: @"mailto:"]) ? Url : [NSString stringWithFormat: @"http://%@", Url]];
+	
+	FString SchemeName;
+	bool bHasSchemeName = FParse::SchemeNameFromURI(URL, SchemeName);
+		
+	NSURL* UrlToOpen = [NSURL URLWithString: (bHasSchemeName ? Url : [NSString stringWithFormat: @"http://%@", Url]];
 	[[NSWorkspace sharedWorkspace] openURL: UrlToOpen];
 	CFRelease( (CFStringRef)Url );
 	if( Error )
diff --git a/Engine/Source/Runtime/Core/Private/Misc/Parse.cpp b/Engine/Source/Runtime/Core/Private/Misc/Parse.cpp
--- a/Engine/Source/Runtime/Core/Private/Misc/Parse.cpp
+++ b/Engine/Source/Runtime/Core/Private/Misc/Parse.cpp
@@ -1209,4 +1209,20 @@ bool FParse::Resolution( const TCHAR* InResolution, uint32& OutX, uint32& OutY )
 {
 	int32 WindowModeDummy;
 	return Resolution(InResolution, OutX, OutY, WindowModeDummy);
-}
\ No newline at end of file
+}
+
+bool FParse::SchemeNameFromURI(const TCHAR* URI, FString& OutSchemeName)
+{
+	for(int32 Idx = 0;;Idx++)
+	{
+		if(!FChar::IsAlpha(URI[Idx]) && !FChar::IsDigit(URI[Idx]) && URI[Idx] != TEXT('+') && URI[Idx] != TEXT('.') && URI[Idx] != TEXT('-'))
+		{
+			if(URI[Idx] == TEXT(':') && Idx > 0)
+			{
+				OutSchemeName = FString(Idx, URI);
+				return true;
+			}
+			return false;
+		}
+	}
+}
diff --git a/Engine/Source/Runtime/Core/Private/Windows/WindowsPlatformProcess.cpp b/Engine/Source/Runtime/Core/Private/Windows/WindowsPlatformProcess.cpp
--- a/Engine/Source/Runtime/Core/Private/Windows/WindowsPlatformProcess.cpp
+++ b/Engine/Source/Runtime/Core/Private/Windows/WindowsPlatformProcess.cpp
@@ -252,10 +252,10 @@ static void LaunchWebURL( const FString& URLParams, FString* Error )
 	}
 }
 
-static void LaunchMailToURL( const TCHAR* MailTo, FString* Error )
+static void LaunchDefaultHandlerForURL( const TCHAR* URL, FString* Error )
 {
-	// ShellExecute will open the default mail app with a "mailto:" link
-	const HINSTANCE Code = ::ShellExecuteW(NULL, TEXT("open"), MailTo, NULL, NULL, SW_SHOWNORMAL);
+	// ShellExecute will open the default handler for a URL
+	const HINSTANCE Code = ::ShellExecuteW(NULL, TEXT("open"), URL, NULL, NULL, SW_SHOWNORMAL);
 	if (Error)
 	{
 		*Error = ((PTRINT)Code <= 32) ? NSLOCTEXT("Core", "UrlFailed", "Failed launching URL").ToString() : TEXT("");
@@ -272,9 +272,11 @@ void FWindowsPlatformProcess::LaunchURL( const TCHAR* URL, const TCHAR* Parms, F
 		*Error = TEXT("");
 	}
 
-	if( FString(URL).StartsWith(TEXT("mailto:")) )
+	// Use the default handler if we have a URI scheme name that doesn't look like a Windows path, and is not http: or https:
+	FString SchemeName;
+	if(FParse::SchemeNameFromURI(URL, SchemeName) && SchemeName.Len() > 1 && SchemeName != TEXT("http") && SchemeName != TEXT("https"))
 	{
-		LaunchMailToURL( URL, Error );
+		LaunchDefaultHandlerForURL(URL, Error);
 	}
 	else
 	{
diff --git a/Engine/Source/Runtime/Core/Public/Misc/Parse.h b/Engine/Source/Runtime/Core/Public/Misc/Parse.h
--- a/Engine/Source/Runtime/Core/Public/Misc/Parse.h
+++ b/Engine/Source/Runtime/Core/Public/Misc/Parse.h
@@ -106,4 +106,7 @@ struct CORE_API FParse
 
 	/** Parses a resolution in the form 1920x1080<f|w|wf>. Same as above, but also attempts to process a fullscreen/windowed flag from the end */
 	static bool Resolution(const TCHAR* InResolution, uint32& OutX, uint32& OutY, int32& OutWindowMode);
+
+	/** Parses the scheme name from a URI */
+	static bool SchemeNameFromURI(const TCHAR* InURI, FString& OutSchemeName);
 };
