CIS fix.

diff --git a/Engine/Source/Runtime/Core/Private/Mac/MacPlatformProcess.cpp b/Engine/Source/Runtime/Core/Private/Mac/MacPlatformProcess.cpp
--- a/Engine/Source/Runtime/Core/Private/Mac/MacPlatformProcess.cpp
+++ b/Engine/Source/Runtime/Core/Private/Mac/MacPlatformProcess.cpp
@@ -197,7 +197,7 @@ void FMacPlatformProcess::LaunchURL( const TCHAR* URL, const TCHAR* Parms, FStri
 	FString SchemeName;
 	bool bHasSchemeName = FParse::SchemeNameFromURI(URL, SchemeName);
 		
-	NSURL* UrlToOpen = [NSURL URLWithString: (bHasSchemeName ? Url : [NSString stringWithFormat: @"http://%@", Url]];
+	NSURL* UrlToOpen = [NSURL URLWithString: bHasSchemeName ? Url : [NSString stringWithFormat: @"http://%@", Url]];
 	[[NSWorkspace sharedWorkspace] openURL: UrlToOpen];
 	CFRelease( (CFStringRef)Url );
 	if( Error )
