More detailed message when uncompression fails. Added the ability to skip uncompress errors during serialization (not to be enabled to by default!)

diff --git a/Engine/Source/Runtime/Core/Private/Misc/Compression.cpp b/Engine/Source/Runtime/Core/Private/Misc/Compression.cpp
--- a/Engine/Source/Runtime/Core/Private/Misc/Compression.cpp
+++ b/Engine/Source/Runtime/Core/Private/Misc/Compression.cpp
@@ -98,7 +98,14 @@ bool appUncompressMemoryZLIB( void* UncompressedBuffer, int32 UncompressedSize,
 	unsigned long ZUncompressedSize	= UncompressedSize;
 	
 	// Uncompress data.
-	bool bOperationSucceeded = uncompress( (uint8*) UncompressedBuffer, &ZUncompressedSize, (const uint8*) CompressedBuffer, ZCompressedSize ) == Z_OK ? true : false;
+	const int32 Result = uncompress((uint8*)UncompressedBuffer, &ZUncompressedSize, (const uint8*)CompressedBuffer, ZCompressedSize);
+	
+	// These warnings will be compiled out in shipping.
+	UE_CLOG(Result == Z_MEM_ERROR, LogCompression, Warning, TEXT("appUncompressMemoryZLIB failed: Error: Z_MEM_ERROR, not enough memory!"));
+	UE_CLOG(Result == Z_BUF_ERROR, LogCompression, Warning, TEXT("appUncompressMemoryZLIB failed: Error: Z_BUF_ERROR, not enough room in the output buffer!"));
+	UE_CLOG(Result == Z_DATA_ERROR, LogCompression, Warning, TEXT("appUncompressMemoryZLIB failed: Error: Z_DATA_ERROR, input data was corrupted or incomplete!"));
+
+	const bool bOperationSucceeded = (Result == Z_OK);
 
 	// Sanity check to make sure we uncompressed as much data as we expected to.
 	check( UncompressedSize == ZUncompressedSize );
@@ -239,6 +246,26 @@ bool FCompression::UncompressMemory( ECompressionFlags Flags, void* Uncompressed
 	{
 		case COMPRESS_ZLIB:
 			bUncompressSucceeded = appUncompressMemoryZLIB(UncompressedBuffer, UncompressedSize, CompressedBuffer, CompressedSize);
+			if (!bUncompressSucceeded)
+			{
+				// This is only to skip serialization errors caused by asset corruption 
+				// that can be fixed during re-save, should never be disabled by default!
+				static struct FFailOnUncompressErrors
+				{
+					bool Value;
+					FFailOnUncompressErrors()
+						: Value(true) // fail by default
+					{
+						GConfig->GetBool(TEXT("Core.System"), TEXT("FailOnUncompressErrors"), Value, GEngineIni);
+					}
+				} FailOnUncompressErrors;
+				if (!FailOnUncompressErrors.Value)
+				{
+					bUncompressSucceeded = true;
+				}
+				// Always log an error
+				UE_LOG(LogCompression, Error, TEXT("FCompression::UncompressMemory - Failed to uncompress memory (%d/%d), this may indicate the asset is corrupt!"), CompressedSize, UncompressedSize);
+			}
 			break;
 		default:
 			UE_LOG(LogCompression, Warning, TEXT("FCompression::UncompressMemory - This compression type not supported"));
