Fixing module manager not finding latest hot-reload DLLs when starting up the editor.

UE-13842

diff --git a/Engine/Source/Runtime/Core/Private/Modules/ModuleManager.cpp b/Engine/Source/Runtime/Core/Private/Modules/ModuleManager.cpp
--- a/Engine/Source/Runtime/Core/Private/Modules/ModuleManager.cpp
+++ b/Engine/Source/Runtime/Core/Private/Modules/ModuleManager.cpp
@@ -166,6 +166,7 @@ bool FindNewestModuleFile(TArray<FString>& FilesToSearch, const FDateTime& Newer
 			{
 				bFound = true;
 				NewestFoundFileTime = FoundFileTime;
+				OutFilename = FPaths::GetCleanFilename(FoundFilePath);
 			}
 		}
 		else
@@ -266,7 +267,7 @@ void FModuleManager::AddModule(const FName InModuleName)
 	// Search for module files
 	TArray<FString> FoundFiles;
 	IFileManager::Get().FindFiles(FoundFiles, *ModuleFileSearchString, true, false);
-	if (FoundFiles.Num() != 0)
+	if (FoundFiles.Num() == 0)
 	{
 		return;
 	}
