Slate - Now using a more accurate Gamma -> Linear space conversion function for converting FColor to FLinearColor.  I've generated a new lookup table using the more accurate conversion, Color > 0.04045 ? pow( Color * (1.0 / 1.055) + 0.0521327, 2.4 ) : Color * (1.0 / 12.92);  Instead of the pow(Color, 2.2);

diff --git a/Engine/Source/Runtime/Core/Private/Math/Color.cpp b/Engine/Source/Runtime/Core/Private/Math/Color.cpp
--- a/Engine/Source/Runtime/Core/Private/Math/Color.cpp
+++ b/Engine/Source/Runtime/Core/Private/Math/Color.cpp
@@ -361,38 +361,61 @@ void ComputeAndFixedColorAndIntensity(const FLinearColor& InLinearColor,FColor&
 /**
  * Pow table for fast FColor -> FLinearColor conversion.
  *
- * FMath::Pow( i / 255.f, 2.2f )
+ * Color > 0.04045 ? pow( Color * (1.0 / 1.055) + 0.0521327, 2.4 ) : Color * (1.0 / 12.92);
  */
 float FLinearColor::PowOneOver255Table[256] = 
 {
-	0,5.07705190066176E-06,2.33280046660989E-05,5.69217657121931E-05,0.000107187362341244,0.000175123977503027,0.000261543754548491,0.000367136269815943,0.000492503787191433,
-	0.000638182842167022,0.000804658499513058,0.000992374304074325,0.0012017395224384,0.00143313458967186,0.00168691531678928,0.00196341621339647,0.00226295316070643,
-	0.00258582559623417,0.00293231832393836,0.00330270303200364,0.00369723957890013,0.00411617709328275,0.00455975492252602,0.00502820345685554,0.00552174485023966,
-	0.00604059365484981,0.00658495738258168,0.00715503700457303,0.00775102739766061,0.00837311774514858,0.00902149189801213,0.00969632870165823,0.0103978022925553,
-	0.0111260823683832,0.0118813344348137,0.0126637200315821,0.0134733969401426,0.0143105193748841,0.0151752381596252,0.0160677008908869,0.01698805208925,0.0179364333399502,
-	0.0189129834237215,0.0199178384387857,0.0209511319147811,0.0220129949193365,0.0231035561579214,0.0242229420675342,0.0253712769047346,0.0265486828284729,0.027755279978126,
-	0.0289911865471078,0.0302565188523887,0.0315513914002264,0.0328759169483838,0.034230206565082,0.0356143696849188,0.0370285141619602,0.0384727463201946,0.0399471710015256,
-	0.0414518916114625,0.0429870101626571,0.0445526273164214,0.0461488424223509,0.0477757535561706,0.049433457555908,0.0511220500564934,0.052841625522879,0.0545922772817603,
-	0.0563740975519798,0.0581871774736854,0.0600316071363132,0.0619074756054558,0.0638148709486772,0.0657538802603301,0.0677245896854243,0.0697270844425988,0.0717614488462391,
-	0.0738277663277846,0.0759261194562648,0.0780565899581019,0.080219258736215,0.0824142058884592,0.0846415107254295,0.0869012517876603,0.0891935068622478,0.0915183529989195,
-	0.0938758665255778,0.0962661230633397,0.0986891975410945,0.1011451642096,0.103634096655137,0.106156067812744,0.108711149979039,0.11129941482466,0.113920933406333,
-	0.116575776178572,0.119264013005047,0.121985713169619,0.124740945387051,0.127529777813422,0.130352278056244,0.1332085131843,0.136098549737202,0.139022453734703,
-	0.141980290685736,0.144972125597231,0.147998022982685,0.151058046870511,0.154152260812165,0.157280727890073,0.160443510725344,0.16364067148529,0.166872271890766,
-	0.170138373223312,0.173439036332135,0.176774321640903,0.18014428915439,0.183548998464951,0.186988508758844,0.190462878822409,0.193972167048093,0.19751643144034,
-	0.201095729621346,0.204710118836677,0.208359655960767,0.212044397502288,0.215764399609395,0.219519718074868,0.223310408341127,0.227136525505149,0.230998124323267,
-	0.23489525921588,0.238827984272048,0.242796353254002,0.24680041960155,0.2508402364364,0.254915856566385,0.259027332489606,0.263174716398492,0.267358060183772,
-	0.271577415438375,0.275832833461245,0.280124365261085,0.284452061560024,0.288815972797219,0.293216149132375,0.297652640449211,0.302125496358853,0.306634766203158,
-	0.311180499057984,0.315762743736397,0.32038154879181,0.325036962521076,0.329729032967515,0.334457807923889,0.339223334935327,0.344025661302187,0.348864834082879,
-	0.353740900096629,0.358653905926199,0.363603897920553,0.368590922197487,0.373615024646202,0.37867625092984,0.383774646487975,0.388910256539059,0.394083126082829,
-	0.399293299902674,0.404540822567962,0.409825738436323,0.415148091655907,0.420507926167587,0.425905285707146,0.43134021380741,0.436812753800359,0.442322948819202,
-	0.44787084180041,0.453456475485731,0.45907989242416,0.46474113497389,0.470440245304218,0.47617726539744,0.481952237050698,0.487765201877811,0.493616201311074,
-	0.49950527660303,0.505432468828216,0.511397818884879,0.517401367496673,0.523443155214325,0.529523222417277,0.535641609315311,0.541798355950137,0.547993502196972,
-	0.554227087766085,0.560499152204328,0.566809734896638,0.573158875067523,0.579546611782525,0.585972983949661,0.592438030320847,0.598941789493296,0.605484299910907,
-	0.612065599865624,0.61868572749878,0.625344720802427,0.632042617620641,0.638779455650817,0.645555272444934,0.652370105410821,0.659223991813387,0.666116968775851,
-	0.673049073280942,0.680020342172095,0.687030812154625,0.694080519796882,0.701169501531402,0.708297793656032,0.715465432335048,0.722672453600255,0.729918893352071,
-	0.737204787360605,0.744530171266715,0.751895080583051,0.759299550695091,0.766743616862161,0.774227314218442,0.781750677773962,0.789313742415586,0.796916542907978,
-	0.804559113894567,0.81224148989849,0.819963705323528,0.827725794455034,0.835527791460841,0.843369730392169,0.851251645184515,0.859173569658532,0.867135537520905,
-	0.875137582365205,0.883179737672745,0.891262036813419,0.899384513046529,0.907547199521614,0.915750129279253,0.923993335251873,0.932276850264543,0.940600707035753,
-	0.948964938178195,0.957369576199527,0.96581465350313,0.974300202388861,0.982826255053791,0.99139284359294,1
+	0,
+	0.000303526983548838, 0.000607053967097675, 0.000910580950646512, 0.00121410793419535, 0.00151763491774419,
+	0.00182116190129302, 0.00212468888484186, 0.0024282158683907, 0.00273174285193954, 0.00303526983548838,
+	0.00334653564113713, 0.00367650719436314, 0.00402471688178252, 0.00439144189356217, 0.00477695332960869,
+	0.005181516543916, 0.00560539145834456, 0.00604883284946662, 0.00651209061157708, 0.00699540999852809,
+	0.00749903184667767, 0.00802319278093555, 0.0085681254056307, 0.00913405848170623, 0.00972121709156193,
+	0.0103298227927056, 0.0109600937612386, 0.0116122449260844, 0.012286488094766, 0.0129830320714536,
+	0.0137020827679224, 0.0144438433080002, 0.0152085141260192, 0.0159962930597398, 0.0168073754381669,
+	0.0176419541646397, 0.0185002197955389, 0.0193823606149269, 0.0202885627054049, 0.0212190100154473,
+	0.0221738844234532, 0.02315336579873, 0.0241576320596103, 0.0251868592288862, 0.0262412214867272,
+	0.0273208912212394, 0.0284260390768075, 0.0295568340003534, 0.0307134432856324, 0.0318960326156814,
+	0.0331047661035236, 0.0343398063312275, 0.0356013143874111, 0.0368894499032755, 0.0382043710872463,
+	0.0395462347582974, 0.0409151963780232, 0.0423114100815264, 0.0437350287071788, 0.0451862038253117,
+	0.0466650857658898, 0.0481718236452158, 0.049706565391714, 0.0512694577708345, 0.0528606464091205,
+	0.0544802758174765, 0.0561284894136735, 0.0578054295441256, 0.0595112375049707, 0.0612460535624849,
+	0.0630100169728596, 0.0648032660013696, 0.0666259379409563, 0.0684781691302512, 0.070360094971063,
+	0.0722718499453493, 0.0742135676316953, 0.0761853807213167, 0.0781874210336082, 0.0802198195312533,
+	0.0822827063349132, 0.0843762107375113, 0.0865004612181274, 0.0886555854555171, 0.0908417103412699,
+	0.0930589619926197, 0.0953074657649191, 0.0975873462637915, 0.0998987273569704, 0.102241732185838,
+	0.104616483176675, 0.107023102051626, 0.109461709839399, 0.1119324268857, 0.114435372863418,
+	0.116970666782559, 0.119538426999953, 0.122138771228724, 0.124771816547542, 0.127437679409664,
+	0.130136475651761, 0.132868320502552, 0.135633328591233, 0.138431613955729, 0.141263290050755,
+	0.144128469755705, 0.147027265382362, 0.149959788682454, 0.152926150855031, 0.155926462553701,
+	0.158960833893705, 0.162029374458845, 0.16513219330827, 0.168269398983119, 0.171441099513036,
+	0.174647402422543, 0.17788841473729, 0.181164242990184, 0.184474993227387, 0.187820771014205,
+	0.191201681440861, 0.194617829128147, 0.198069318232982, 0.201556252453853, 0.205078735036156,
+	0.208636868777438, 0.212230756032542, 0.215860498718652, 0.219526198320249, 0.223227955893977,
+	0.226965872073417, 0.23074004707378, 0.23455058069651, 0.238397572333811, 0.242281120973093,
+	0.246201325201334, 0.250158283209375, 0.254152092796134, 0.258182851372752, 0.262250655966664,
+	0.266355603225604, 0.270497789421545, 0.274677310454565, 0.278894261856656, 0.283148738795466,
+	0.287440836077983, 0.291770648154158, 0.296138269120463, 0.300543792723403, 0.304987312362961,
+	0.309468921095997, 0.313988711639584, 0.3185467763743, 0.323143207347467, 0.32777809627633,
+	0.332451534551205, 0.337163613238559, 0.341914423084057, 0.346704054515559, 0.351532597646068,
+	0.356400142276637, 0.361306777899234, 0.36625259369956, 0.371237678559833, 0.376262121061519,
+	0.381326009488037, 0.386429431827418, 0.39157247577492, 0.396755228735618, 0.401977777826949,
+	0.407240209881218, 0.41254261144808, 0.417885068796976, 0.423267667919539, 0.428690494531971,
+	0.434153634077377, 0.439657171728079, 0.445201192387887, 0.450785780694349, 0.456411021020965,
+	0.462076997479369, 0.467783793921492, 0.473531493941681, 0.479320180878805, 0.485149937818323,
+	0.491020847594331, 0.496932992791578, 0.502886455747457, 0.50888131855397, 0.514917663059676,
+	0.520995570871595, 0.527115123357109, 0.533276401645826, 0.539479486631421, 0.545724458973463,
+	0.552011399099209, 0.558340387205378, 0.56471150325991, 0.571124827003694, 0.577580437952282,
+	0.584078415397575, 0.590618838409497, 0.597201785837643, 0.603827336312907, 0.610495568249093,
+	0.617206559844509, 0.623960389083534, 0.630757133738175, 0.637596871369601, 0.644479679329661,
+	0.651405634762384, 0.658374814605461, 0.665387295591707, 0.672443154250516, 0.679542466909286,
+	0.686685309694841, 0.693871758534824, 0.701101889159085, 0.708375777101046, 0.71569349769906,
+	0.723055126097739, 0.730460737249286, 0.737910405914797, 0.745404206665559, 0.752942213884326,
+	0.760524501766589, 0.768151144321824, 0.775822215374732, 0.783537788566466, 0.791297937355839,
+	0.799102735020525, 0.806952254658248, 0.81484656918795, 0.822785751350956, 0.830769873712124,
+	0.838799008660978, 0.846873228412837, 0.854992605009927, 0.863157210322481, 0.871367116049835,
+	0.879622393721502, 0.887923114698241, 0.896269350173118, 0.904661171172551, 0.913098648557343,
+	0.921581853023715, 0.930110855104312, 0.938685725169219, 0.947306533426946, 0.955973349925421,
+	0.964686244552961, 0.973445287039244, 0.982250546956257, 0.991102093719252, 1.0,
 };
 
