Changing the TypeContainers provider TMap to use a FString for compraison rather then a TCHAR* so strings from other sources still match. For example, when hot reloading.

Merging using UE4-To-UE4-LauncherDev

diff --git a/Engine/Source/Runtime/Core/Public/Misc/TypeContainer.h b/Engine/Source/Runtime/Core/Public/Misc/TypeContainer.h
--- a/Engine/Source/Runtime/Core/Public/Misc/TypeContainer.h
+++ b/Engine/Source/Runtime/Core/Public/Misc/TypeContainer.h
@@ -389,7 +389,7 @@ private:
 	FCriticalSection CriticalSection;
 
 	/** Maps class names to instance providers. */
-	TMap<const TCHAR*, TSharedPtr<IInstanceProvider>> Providers;
+	TMap<FString, TSharedPtr<IInstanceProvider>> Providers;
 };
 
 
