Fixed high precision mouse handling on a secondary monitor that's not verically aligned to the primary one

diff --git a/Engine/Source/Runtime/Core/Private/Mac/MacApplication.cpp b/Engine/Source/Runtime/Core/Private/Mac/MacApplication.cpp
--- a/Engine/Source/Runtime/Core/Private/Mac/MacApplication.cpp
+++ b/Engine/Source/Runtime/Core/Private/Mac/MacApplication.cpp
@@ -652,7 +652,7 @@ void FMacApplication::ProcessMouseMovedEvent(const FDeferredMacEvent& Event, TSh
 		// Find the screen the cursor is currently on.
 		NSEnumerator *ScreenEnumerator = [[NSScreen screens] objectEnumerator];
 		NSScreen *Screen;
-		while ((Screen = [ScreenEnumerator nextObject]) && !NSMouseInRect(NSMakePoint(HighPrecisionMousePos.X / MouseScaling.X, HighPrecisionMousePos.Y / MouseScaling.Y), Screen.frame, NO))
+		while ((Screen = [ScreenEnumerator nextObject]) && !NSMouseInRect(NSMakePoint(HighPrecisionMousePos.X / MouseScaling.X, FPlatformMisc::ConvertSlateYPositionToCocoa(HighPrecisionMousePos.Y) / MouseScaling.Y), Screen.frame, NO))
 			;
 
 		// Account for warping delta's
@@ -675,14 +675,11 @@ void FMacApplication::ProcessMouseMovedEvent(const FDeferredMacEvent& Event, TSh
 			VisibleFrame.origin.y += 5;
 			VisibleFrame.size.height -= 10;
 		}
-		NSRect FullFrame = [Screen frame];
-		VisibleFrame.origin.y = (FullFrame.origin.y+FullFrame.size.height) - (VisibleFrame.origin.y + VisibleFrame.size.height);
-
 		int32 ScaledX = (int32)(HighPrecisionMousePos.X / MouseScaling.X);
-		int32 ScaledY = (int32)(HighPrecisionMousePos.Y / MouseScaling.Y);
+		int32 ScaledY = (int32)(FPlatformMisc::ConvertSlateYPositionToCocoa(HighPrecisionMousePos.Y) / MouseScaling.Y);
 		int32 ClampedPosX = FMath::Clamp(ScaledX, (int32)VisibleFrame.origin.x, (int32)(VisibleFrame.origin.x + VisibleFrame.size.width)-1);
 		int32 ClampedPosY = FMath::Clamp(ScaledY, (int32)VisibleFrame.origin.y, (int32)(VisibleFrame.origin.y + VisibleFrame.size.height)-1);
-		MacCursor->SetPosition(ClampedPosX * MouseScaling.X, ClampedPosY * MouseScaling.Y);
+		MacCursor->SetPosition(ClampedPosX * MouseScaling.X, FPlatformMisc::ConvertSlateYPositionToCocoa(ClampedPosY) * MouseScaling.Y);
 
 		// Forward the delta on to Slate
 		MessageHandler->OnRawMouseMove(Delta.X, Delta.Y);
