Fixed external drag-drops not working on mac

The window was not getting set for the deferred event which ultimately resulted in the notification being ignored.
This addresses UE-15367 (On Mac, unable to drag file into Content Browser to import)

diff --git a/Engine/Source/Runtime/Core/Private/Mac/MacApplication.cpp b/Engine/Source/Runtime/Core/Private/Mac/MacApplication.cpp
--- a/Engine/Source/Runtime/Core/Private/Mac/MacApplication.cpp
+++ b/Engine/Source/Runtime/Core/Private/Mac/MacApplication.cpp
@@ -386,6 +386,13 @@ void FMacApplication::DeferEvent(NSObject* Object)
 		}
 		else if ([[Notification object] conformsToProtocol:@protocol(NSDraggingInfo)])
 		{
+            NSWindow* NotificationWindow = [(id<NSDraggingInfo>)[Notification object] draggingDestinationWindow];
+            
+            if (NotificationWindow && [NotificationWindow isKindOfClass:[FCocoaWindow class]])
+            {
+           		DeferredEvent.Window = (FCocoaWindow*)NotificationWindow;
+            }
+            
 			if (DeferredEvent.NotificationName == NSPrepareForDragOperation)
 			{
 				DeferredEvent.DraggingPasteboard = [[(id<NSDraggingInfo>)[Notification object] draggingPasteboard] retain];
