Allow pak files outside of obb
#ue-14247
#ue4
#android

diff --git a/Engine/Source/Runtime/Core/Private/Android/AndroidFile.cpp b/Engine/Source/Runtime/Core/Private/Android/AndroidFile.cpp
--- a/Engine/Source/Runtime/Core/Private/Android/AndroidFile.cpp
+++ b/Engine/Source/Runtime/Core/Private/Android/AndroidFile.cpp
@@ -1111,6 +1111,14 @@ public:
 #if LOG_ANDROID_FILE_MANIFEST
 			FPlatformMisc::LowLevelOutputDebugStringf(TEXT("Failed to find time stamp in NonUFSManifest for file '%s'"), Filename);
 #endif
+
+			// pak file outside of obb may not be in manifest so check if it exists
+			if (AssetPath.EndsWith(".pak"))
+			{
+				// return local file access timestamp (if exists)
+				return GetAccessTimeStamp(Filename, true);
+			}
+
 			return FDateTime::MinValue();
 
 #endif
