fix for project on the command-line for Mac

diff --git a/Engine/Source/Runtime/Core/Private/Mac/MacPlatformMisc.cpp b/Engine/Source/Runtime/Core/Private/Mac/MacPlatformMisc.cpp
--- a/Engine/Source/Runtime/Core/Private/Mac/MacPlatformMisc.cpp
+++ b/Engine/Source/Runtime/Core/Private/Mac/MacPlatformMisc.cpp
@@ -1758,4 +1758,37 @@ const TCHAR* FMacPlatformMisc::GetCompleteCommandLine()
         }
     }
 	return *CommandLine;
-}
\ No newline at end of file
+}
+
+
+void FMacPlatformMisc::MergeDefaultArgumentsIntoCommandLine(FString& CommandLine, FString DefaultArguments)
+{
+    // Normalize the arguments by removing any newlines and unnecessary spaces
+    DefaultArguments.ReplaceInline(TEXT("\r"), TEXT(" "));
+    DefaultArguments.ReplaceInline(TEXT("\n"), TEXT(" "));
+    DefaultArguments = DefaultArguments.Trim();
+    
+    // We need to make sure that the project is always the first argument, otherwise it's not parsed correctly in LaunchEngineLoop (which also
+    // interferes with passing a map name on the command-line). It's possible that UE4CommandLine.txt contains the game name (for packaged games)
+    // or that the command line contains the game name (for the editor), so we have to be careful in which order they're stitched together.
+    if (DefaultArguments.Len() > 0)
+    {
+        // Measure the length of the project name in the default command line.
+        int32 ProjectNameLength = 0;
+        const TCHAR* RemainingArguments = *DefaultArguments;
+        FString FirstArgument;
+        if (FParse::Token(RemainingArguments, FirstArgument, false) && !FirstArgument.StartsWith("-"))
+        {
+            ProjectNameLength = RemainingArguments - *DefaultArguments;
+        }
+        
+        // Build a new command line consisting of the application name, project name, command line arguments, and default arguments.
+        FString NewCommandLine;
+        NewCommandLine.Append(*DefaultArguments, ProjectNameLength);
+        NewCommandLine.AppendChar(TEXT(' '));
+        NewCommandLine.Append(*CommandLine);
+        NewCommandLine.AppendChar(TEXT(' '));
+        NewCommandLine.Append(*DefaultArguments + ProjectNameLength);
+        CommandLine = NewCommandLine;
+    }
+}
diff --git a/Engine/Source/Runtime/Core/Public/Mac/MacPlatformMisc.h b/Engine/Source/Runtime/Core/Public/Mac/MacPlatformMisc.h
--- a/Engine/Source/Runtime/Core/Public/Mac/MacPlatformMisc.h
+++ b/Engine/Source/Runtime/Core/Public/Mac/MacPlatformMisc.h
@@ -203,6 +203,8 @@ struct CORE_API FMacPlatformMisc : public FGenericPlatformMisc
 	static id<NSObject> CommandletActivity;
 
 	static const TCHAR* GetCompleteCommandLine();
+    
+    static void MergeDefaultArgumentsIntoCommandLine(FString& CommandLine, FString DefaultArguments);
 };
 
 #ifdef __OBJC__
