Fixes for JNIEnv thread crashes
#ue-14317
#ue4
#android

diff --git a/Engine/Source/Runtime/Core/Private/Android/AndroidApplication.cpp b/Engine/Source/Runtime/Core/Private/Android/AndroidApplication.cpp
--- a/Engine/Source/Runtime/Core/Private/Android/AndroidApplication.cpp
+++ b/Engine/Source/Runtime/Core/Private/Android/AndroidApplication.cpp
@@ -152,6 +152,7 @@ void FAndroidApplication::InitializeJavaEnv( JavaVM* VM, jint Version, jobject G
 
 static void JavaEnvDestructor(void*)
 {
+	FPlatformMisc::LowLevelOutputDebugStringf(TEXT("*** JavaEnvDestructor: %d"), FPlatformTLS::GetCurrentThreadId());
 	FAndroidApplication::DetachJavaEnv();
 }
 
@@ -167,7 +168,10 @@ JNIEnv* FAndroidApplication::GetJavaEnv( bool bRequireGlobalThis /*= true*/ )
 	{
 		return nullptr;
 	}
-#else
+#endif
+#if 0
+	// not reliable at the moment.. revisit later
+
 	// Magic static - *should* be thread safe
 	//Android & pthread specific, bind a destructor for thread exit
 	static uint32 TlsSlot = 0;
@@ -191,6 +195,36 @@ JNIEnv* FAndroidApplication::GetJavaEnv( bool bRequireGlobalThis /*= true*/ )
 	}
 
 	return (!bRequireGlobalThis || (GlobalObjectRef != nullptr)) ? Env : nullptr;
+#else
+	// register a destructor to detach this thread
+	static uint32 TlsSlot = 0;
+	if (TlsSlot == 0)
+	{
+		pthread_key_create((pthread_key_t*)&TlsSlot, &JavaEnvDestructor);
+	}
+
+	JNIEnv* Env = nullptr;
+	jint GetEnvResult = CurrentJavaVM->GetEnv((void **)&Env, CurrentJavaVersion);
+	if (GetEnvResult == JNI_EDETACHED)
+	{
+		// attach to this thread
+		jint AttachResult = CurrentJavaVM->AttachCurrentThread(&Env, NULL);
+		if (AttachResult == JNI_ERR)
+		{
+			FPlatformMisc::LowLevelOutputDebugString(L"UNIT TEST -- Failed to attach thread to get the JNI environment!");
+			check(false);
+			return nullptr;
+		}
+		FPlatformTLS::SetTlsValue(TlsSlot, (void*)Env);
+	}
+	else if (GetEnvResult != JNI_OK)
+	{
+		FPlatformMisc::LowLevelOutputDebugStringf(L"UNIT TEST -- Failed to get the JNI environment! Result = %d", GetEnvResult);
+		check(false);
+		return nullptr;
+
+	}
+	return Env;
 #endif
 }
 
diff --git a/Engine/Source/Runtime/Core/Private/Android/AndroidJava.cpp b/Engine/Source/Runtime/Core/Private/Android/AndroidJava.cpp
--- a/Engine/Source/Runtime/Core/Private/Android/AndroidJava.cpp
+++ b/Engine/Source/Runtime/Core/Private/Android/AndroidJava.cpp
@@ -5,9 +5,9 @@
 #include "AndroidApplication.h"
 
 FJavaClassObject::FJavaClassObject(FName ClassName, const char* CtorSig, ...)
-	: BoundThread(0)
 {
-	BindObjectToThread();
+	JNIEnv*	JEnv = FAndroidApplication::GetJavaEnv();
+
 	Class = FAndroidApplication::FindJavaClass(ClassName.GetPlainANSIString());
 	check(Class);
 	jmethodID Constructor = JEnv->GetMethodID(Class, "<init>", CtorSig);
@@ -26,11 +26,13 @@ FJavaClassObject::FJavaClassObject(FName ClassName, const char* CtorSig, ...)
 
 FJavaClassObject::~FJavaClassObject()
 {
+	JNIEnv*	JEnv = FAndroidApplication::GetJavaEnv();
 	JEnv->DeleteGlobalRef(Object);
 }
 
 FJavaClassMethod FJavaClassObject::GetClassMethod(const char* MethodName, const char* FuncSig)
 {
+	JNIEnv*	JEnv = FAndroidApplication::GetJavaEnv();
 	FJavaClassMethod Method;
 	Method.Method = JEnv->GetMethodID(Class, MethodName, FuncSig);
 	Method.Name = MethodName;
@@ -40,24 +42,10 @@ FJavaClassMethod FJavaClassObject::GetClassMethod(const char* MethodName, const
 	return Method;
 }
 
-void FJavaClassObject::BindObjectToThread()
-{
-	uint32 ThisThreadID = FPlatformTLS::GetCurrentThreadId();
-	if (BoundThread != ThisThreadID)
-	{
-		JEnv = FAndroidApplication::GetJavaEnv();
-		BoundThread = ThisThreadID;
-	}
-}
-
-void FJavaClassObject::UnbindObjectFromThread()
-{
-	BoundThread = 0;
-}
-
 template<>
 void FJavaClassObject::CallMethod<void>(FJavaClassMethod Method, ...)
 {
+	JNIEnv*	JEnv = FAndroidApplication::GetJavaEnv();
 	va_list Params;
 	va_start(Params, Method);
 	JEnv->CallVoidMethodV(Object, Method.Method, Params);
@@ -68,6 +56,7 @@ void FJavaClassObject::CallMethod<void>(FJavaClassMethod Method, ...)
 template<>
 bool FJavaClassObject::CallMethod<bool>(FJavaClassMethod Method, ...)
 {
+	JNIEnv*	JEnv = FAndroidApplication::GetJavaEnv();
 	va_list Params;
 	va_start(Params, Method);
 	bool RetVal = JEnv->CallBooleanMethodV(Object, Method.Method, Params);
@@ -79,6 +68,7 @@ bool FJavaClassObject::CallMethod<bool>(FJavaClassMethod Method, ...)
 template<>
 int FJavaClassObject::CallMethod<int>(FJavaClassMethod Method, ...)
 {
+	JNIEnv*	JEnv = FAndroidApplication::GetJavaEnv();
 	va_list Params;
 	va_start(Params, Method);
 	int RetVal = JEnv->CallIntMethodV(Object, Method.Method, Params);
@@ -90,6 +80,7 @@ int FJavaClassObject::CallMethod<int>(FJavaClassMethod Method, ...)
 template<>
 jobject FJavaClassObject::CallMethod<jobject>(FJavaClassMethod Method, ...)
 {
+	JNIEnv*	JEnv = FAndroidApplication::GetJavaEnv();
 	va_list Params;
 	va_start(Params, Method);
 	jobject val = JEnv->CallObjectMethodV(Object, Method.Method, Params);
@@ -103,6 +94,7 @@ jobject FJavaClassObject::CallMethod<jobject>(FJavaClassMethod Method, ...)
 template<>
 int64 FJavaClassObject::CallMethod<int64>(FJavaClassMethod Method, ...)
 {
+	JNIEnv*	JEnv = FAndroidApplication::GetJavaEnv();
 	va_list Params;
 	va_start(Params, Method);
 	int64 RetVal = JEnv->CallLongMethodV(Object, Method.Method, Params);
@@ -114,6 +106,7 @@ int64 FJavaClassObject::CallMethod<int64>(FJavaClassMethod Method, ...)
 template<>
 FString FJavaClassObject::CallMethod<FString>(FJavaClassMethod Method, ...)
 {
+	JNIEnv*	JEnv = FAndroidApplication::GetJavaEnv();
 	va_list Params;
 	va_start(Params, Method);
 	jstring RetVal = static_cast<jstring>(
@@ -128,15 +121,17 @@ FString FJavaClassObject::CallMethod<FString>(FJavaClassMethod Method, ...)
 
 jstring FJavaClassObject::GetJString(const FString& String)
 {
+	JNIEnv*	JEnv = FAndroidApplication::GetJavaEnv();
 	auto StringCastObj = StringCast<ANSICHAR>(*String);
-	jstring local = FAndroidApplication::GetJavaEnv()->NewStringUTF(StringCastObj.Get());
-	jstring result = (jstring)FAndroidApplication::GetJavaEnv()->NewGlobalRef(local);
-	FAndroidApplication::GetJavaEnv()->DeleteLocalRef(local);
+	jstring local = JEnv->NewStringUTF(StringCastObj.Get());
+	jstring result = (jstring)JEnv->NewGlobalRef(local);
+	JEnv->DeleteLocalRef(local);
 	return result;
 }
 
 void FJavaClassObject::VerifyException()
 {
+	JNIEnv*	JEnv = FAndroidApplication::GetJavaEnv();
 	if (JEnv->ExceptionCheck())
 	{
 		JEnv->ExceptionDescribe();
diff --git a/Engine/Source/Runtime/Core/Private/Android/AndroidJavaMediaPlayer.cpp b/Engine/Source/Runtime/Core/Private/Android/AndroidJavaMediaPlayer.cpp
--- a/Engine/Source/Runtime/Core/Private/Android/AndroidJavaMediaPlayer.cpp
+++ b/Engine/Source/Runtime/Core/Private/Android/AndroidJavaMediaPlayer.cpp
@@ -33,113 +33,96 @@ FJavaAndroidMediaPlayer::FJavaAndroidMediaPlayer(bool swizzlePixels)
 
 int32 FJavaAndroidMediaPlayer::GetDuration()
 {
-	BindObjectToThread();
 	return CallMethod<int32>(GetDurationMethod);
 }
 
 void FJavaAndroidMediaPlayer::Reset()
 {
-	BindObjectToThread();
 	CallMethod<void>(ResetMethod);
 }
 
 void FJavaAndroidMediaPlayer::Stop()
 {
-	BindObjectToThread();
 	CallMethod<void>(StopMethod);
 }
 
 int32 FJavaAndroidMediaPlayer::GetCurrentPosition()
 {
-	BindObjectToThread();
 	int32 position = CallMethod<int32>(GetCurrentPositionMethod);
 	return position;
 }
 
 bool FJavaAndroidMediaPlayer::IsLooping()
 {
-	BindObjectToThread();
 	return CallMethod<bool>(IsLoopingMethod);
 }
 
 bool FJavaAndroidMediaPlayer::IsPlaying()
 {
-	BindObjectToThread();
 	return CallMethod<bool>(IsPlayingMethod);
 }
 
 void FJavaAndroidMediaPlayer::SetDataSource(const FString & Url)
 {
-	BindObjectToThread();
 	CallMethod<void>(SetDataSourceURLMethod, GetJString(Url));
 }
 
 bool FJavaAndroidMediaPlayer::SetDataSource(const FString& MoviePathOnDevice, int64 offset, int64 size)
 {
-	BindObjectToThread();
 	return CallMethod<bool>(SetDataSourceFileMethod, GetJString(MoviePathOnDevice), offset, size);
 }
 
 bool FJavaAndroidMediaPlayer::SetDataSource(jobject AssetMgr, const FString& AssetPath, int64 offset, int64 size)
 {
-	BindObjectToThread();
 	return CallMethod<bool>(SetDataSourceAssetMethod, AssetMgr, GetJString(AssetPath), offset, size);
 }
 
 void FJavaAndroidMediaPlayer::Prepare()
 {
-	BindObjectToThread();
 	CallMethod<void>(PrepareMethod);
 }
 
 void FJavaAndroidMediaPlayer::SeekTo(int32 Milliseconds)
 {
-	BindObjectToThread();
 	CallMethod<void>(SeekToMethod, Milliseconds);
 }
 
 void FJavaAndroidMediaPlayer::SetLooping(bool Looping)
 {
-	BindObjectToThread();
 	CallMethod<void>(SetLoopingMethod, Looping);
 }
 
 void FJavaAndroidMediaPlayer::Release()
 {
-	BindObjectToThread();
 	CallMethod<void>(ReleaseMethod);
 }
 
 int32 FJavaAndroidMediaPlayer::GetVideoHeight()
 {
-	BindObjectToThread();
 	return CallMethod<int32>(GetVideoHeightMethod);
 }
 
 int32 FJavaAndroidMediaPlayer::GetVideoWidth()
 {
-	BindObjectToThread();
 	return CallMethod<int32>(GetVideoWidthMethod);
 }
 
 void FJavaAndroidMediaPlayer::SetVideoEnabled(bool enabled /*= true*/)
 {
-	BindObjectToThread();
 	CallMethod<void>(SetVideoEnabledMethod, enabled);
 }
 
 void FJavaAndroidMediaPlayer::SetAudioEnabled(bool enabled /*= true*/)
 {
-	BindObjectToThread();
 	CallMethod<void>(SetAudioEnabledMethod, enabled);
 }
 
 bool FJavaAndroidMediaPlayer::GetVideoLastFrameData(void* & outPixels, int64 & outCount)
 {
-	BindObjectToThread();
 	jobject buffer = CallMethod<jobject>(GetVideoLastFrameDataMethod);
 	if (nullptr != buffer)
 	{
+		JNIEnv*	JEnv = FAndroidApplication::GetJavaEnv();
 		outPixels = JEnv->GetDirectBufferAddress(buffer);
 		outCount = JEnv->GetDirectBufferCapacity(buffer);
 		JEnv->DeleteGlobalRef(buffer);
@@ -153,19 +136,16 @@ bool FJavaAndroidMediaPlayer::GetVideoLastFrameData(void* & outPixels, int64 & o
 
 void FJavaAndroidMediaPlayer::Start()
 {
-	BindObjectToThread();
 	CallMethod<void>(StartMethod);
 }
 
 void FJavaAndroidMediaPlayer::Pause()
 {
-	BindObjectToThread();
 	CallMethod<void>(PauseMethod);
 }
 
 bool FJavaAndroidMediaPlayer::GetVideoLastFrame(int32 destTexture)
 {
-	BindObjectToThread();
 	return CallMethod<bool>(GetVideoLastFrameMethod, destTexture);
 }
 
diff --git a/Engine/Source/Runtime/Core/Private/Android/AndroidJavaMessageBox.cpp b/Engine/Source/Runtime/Core/Private/Android/AndroidJavaMessageBox.cpp
--- a/Engine/Source/Runtime/Core/Private/Android/AndroidJavaMessageBox.cpp
+++ b/Engine/Source/Runtime/Core/Private/Android/AndroidJavaMessageBox.cpp
@@ -15,31 +15,26 @@ FJavaAndroidMessageBox::FJavaAndroidMessageBox()
 
 void FJavaAndroidMessageBox::SetCaption(const FString & Text)
 {
-	BindObjectToThread();
 	CallMethod<void>(SetCaptionMethod, GetJString(Text));
 }
 
 void FJavaAndroidMessageBox::SetText(const FString & Text)
 {
-	BindObjectToThread();
 	CallMethod<void>(SetTextMethod, GetJString(Text));
 }
 
 void FJavaAndroidMessageBox::AddButton(const FString & Text)
 {
-	BindObjectToThread();
 	CallMethod<void>(AddButtonMethod, GetJString(Text));
 }
 
 void FJavaAndroidMessageBox::Clear()
 {
-	BindObjectToThread();
 	CallMethod<void>(ClearMethod);
 }
 
 int32 FJavaAndroidMessageBox::Show()
 {
-	BindObjectToThread();
 	return CallMethod<int32>(ShowMethod);
 }
 
diff --git a/Engine/Source/Runtime/Core/Public/Android/AndroidJava.h b/Engine/Source/Runtime/Core/Public/Android/AndroidJava.h
--- a/Engine/Source/Runtime/Core/Public/Android/AndroidJava.h
+++ b/Engine/Source/Runtime/Core/Public/Android/AndroidJava.h
@@ -35,10 +35,6 @@ public:
 		return Object;
 	}
 
-	void BindObjectToThread();
-
-	void UnbindObjectFromThread();
-
 	static jstring GetJString(const FString& String);
 
 	void VerifyException();
@@ -46,10 +42,8 @@ public:
 protected:
 
 	FName			ClassName;
-	JNIEnv*			JEnv;
 	jobject			Object;
 	jclass			Class;
-	uint32			BoundThread;
 
 private:
 	FJavaClassObject(const FJavaClassObject& rhs);
