Core - Added ToHexBlob and FromHexBlob

diff --git a/Engine/Source/Runtime/Core/Private/Containers/String.cpp b/Engine/Source/Runtime/Core/Private/Containers/String.cpp
--- a/Engine/Source/Runtime/Core/Private/Containers/String.cpp
+++ b/Engine/Source/Runtime/Core/Private/Containers/String.cpp
@@ -336,47 +336,24 @@ void FString::AppendInt( int32 InNum )
 	*this += TempNum + TempAt;
 }
 
-/**
- * Converts a string into a boolean value
- *   1, "True", "Yes", GTrue, GYes, and non-zero integers become true
- *   0, "False", "No", GFalse, GNo, and unparsable values become false
- *
- * @param Source the string to parse
- *
- * @return The boolean value
- */
+
 bool FString::ToBool() const
 {
 	return FCString::ToBool(**this);
 }
 
-/**
- * Converts a buffer to a string by hex-ifying the elements
- *
- * @param SrcBuffer the buffer to stringify
- * @param SrcSize the number of bytes to convert
- *
- * @return the blob in string form
- */
 FString FString::FromBlob(const uint8* SrcBuffer,const uint32 SrcSize)
 {
 	FString Result;
+	Result.Reserve( SrcSize * 3 );
 	// Convert and append each byte in the buffer
 	for (uint32 Count = 0; Count < SrcSize; Count++)
 	{
-		Result += FString::Printf(TEXT("%03d"),(uint32)SrcBuffer[Count]);
+		Result += FString::Printf(TEXT("%03d"),(uint8)SrcBuffer[Count]);
 	}
 	return Result;
 }
 
-/**
- * Converts a string into a buffer
- *
- * @param DestBuffer the buffer to fill with the string data
- * @param DestSize the size of the buffer in bytes (must be at least string len / 2)
- *
- * @return true if the conversion happened, false otherwise
- */
 bool FString::ToBlob(const FString& Source,uint8* DestBuffer,const uint32 DestSize)
 {
 	// Make sure the buffer is at least half the size and that the string is an
@@ -387,7 +364,7 @@ bool FString::ToBlob(const FString& Source,uint8* DestBuffer,const uint32 DestSi
 		TCHAR ConvBuffer[4];
 		ConvBuffer[3] = TEXT('\0');
 		int32 WriteIndex = 0;
-		// Walk the string 2 chars at a time
+		// Walk the string 3 chars at a time
 		for (int32 Index = 0; Index < Source.Len(); Index += 3, WriteIndex++)
 		{
 			ConvBuffer[0] = Source[Index];
@@ -400,6 +377,41 @@ bool FString::ToBlob(const FString& Source,uint8* DestBuffer,const uint32 DestSi
 	return false;
 }
 
+FString FString::FromHexBlob( const uint8* SrcBuffer, const uint32 SrcSize )
+{
+	FString Result;
+	Result.Reserve( SrcSize * 2 );
+	// Convert and append each byte in the buffer
+	for (uint32 Count = 0; Count < SrcSize; Count++)
+	{
+		Result += FString::Printf( TEXT( "%02X" ), (uint8)SrcBuffer[Count] );
+	}
+	return Result;
+}
+
+bool FString::ToHexBlob( const FString& Source, uint8* DestBuffer, const uint32 DestSize )
+{
+	// Make sure the buffer is at least half the size and that the string is an
+	// even number of characters long
+	if (DestSize >= (uint32)(Source.Len() / 2) &&
+		 (Source.Len() % 2) == 0)
+	{
+		TCHAR ConvBuffer[3];
+		ConvBuffer[2] = TEXT( '\0' );
+		int32 WriteIndex = 0;
+		// Walk the string 2 chars at a time
+		TCHAR* End = nullptr;
+		for (int32 Index = 0; Index < Source.Len(); Index += 2, WriteIndex++)
+		{
+			ConvBuffer[0] = Source[Index];
+			ConvBuffer[1] = Source[Index + 1];
+			DestBuffer[WriteIndex] = FCString::Strtoi( ConvBuffer, &End, 16 );
+		}
+		return true;
+	}
+	return false;
+}
+
 FString FString::SanitizeFloat( double InFloat )
 {
 	// Avoids negative zero
diff --git a/Engine/Source/Runtime/Core/Public/Containers/UnrealString.h b/Engine/Source/Runtime/Core/Public/Containers/UnrealString.h
--- a/Engine/Source/Runtime/Core/Public/Containers/UnrealString.h
+++ b/Engine/Source/Runtime/Core/Public/Containers/UnrealString.h
@@ -1386,7 +1386,7 @@ public:
 	bool ToBool() const;
 
 	/**
-	 * Converts a buffer to a string by hex-ifying the elements
+	 * Converts a buffer to a string
 	 *
 	 * @param SrcBuffer the buffer to stringify
 	 * @param SrcSize the number of bytes to convert
@@ -1399,12 +1399,32 @@ public:
 	 * Converts a string into a buffer
 	 *
 	 * @param DestBuffer the buffer to fill with the string data
-	 * @param DestSize the size of the buffer in bytes (must be at least string len / 2)
+	 * @param DestSize the size of the buffer in bytes (must be at least string len / 3)
 	 *
 	 * @return true if the conversion happened, false otherwise
 	 */
 	static bool ToBlob(const FString& Source,uint8* DestBuffer,const uint32 DestSize);
 
+	/**
+	 * Converts a buffer to a string by hex-ifying the elements
+	 *
+	 * @param SrcBuffer the buffer to stringify
+	 * @param SrcSize the number of bytes to convert
+	 *
+	 * @return the blob in string form
+	 */
+	static FString FromHexBlob(const uint8* SrcBuffer,const uint32 SrcSize);
+
+	/**
+	 * Converts a string into a buffer
+	 *
+	 * @param DestBuffer the buffer to fill with the string data
+	 * @param DestSize the size of the buffer in bytes (must be at least string len / 2)
+	 *
+	 * @return true if the conversion happened, false otherwise
+	 */
+	static bool ToHexBlob(const FString& Source,uint8* DestBuffer,const uint32 DestSize);
+
 	/**
 	 * Converts a float string with the trailing zeros stripped
 	 * For example - 1.234 will be "1.234" rather than "1.234000"
