Stats - Misc processing tweak

diff --git a/Engine/Source/Runtime/Core/Private/Stats/Stats2.cpp b/Engine/Source/Runtime/Core/Private/Stats/Stats2.cpp
--- a/Engine/Source/Runtime/Core/Private/Stats/Stats2.cpp
+++ b/Engine/Source/Runtime/Core/Private/Stats/Stats2.cpp
@@ -702,13 +702,13 @@ public:
 		static double LastTime = -1.0;
 		bool bShouldProcess = false;
 
+		const int32 MaxIncomingPackets = 16;
 		if( FThreadStats::bIsRawStatsActive )
 		{
 			// For raw stats we process every 24MB of packet data to minimize the stats messages memory usage.
 			//const bool bShouldProcessRawStats = IncomingData.Packets.Num() > 10;
 			const int32 MaxIncomingMessages = 24*1024*1024/sizeof(FStatMessage);
-			const int32 MaxIncomingPackets = 16;
-
+			
 			int32 IncomingDataMessages = 0;
 			for( int32 Index = 0; Index < IncomingData.Packets.Num(); ++Index )
 			{
@@ -719,8 +719,8 @@ public:
 		}
 		else
 		{
-			// For regular stats we won't process more than every 5ms.
-			bShouldProcess = bReadyToProcess && FPlatformTime::Seconds() - LastTime > 0.005f; 
+			// For regular stats we won't process more than every 5ms or every 16 packets.
+			bShouldProcess = bReadyToProcess && (FPlatformTime::Seconds() - LastTime > 0.005f || IncomingData.Packets.Num() > MaxIncomingPackets);
 		}
 
 		if( bShouldProcess )
@@ -1052,26 +1052,26 @@ void FThreadStats::CheckForCollectingStartupStats()
 		CmdLine = CmdLine.Mid(Index + StatCmds.Len());
 	}
 	
-	if (FParse::Param(FCommandLine::Get(), TEXT("LoadTimeStats")))
+	if (FParse::Param( FCommandLine::Get(), TEXT( "LoadTimeStats" ) ))
 	{
-		DirectStatsCommand(TEXT("stat group enable LinkerLoad"));
-		DirectStatsCommand(TEXT("stat group enable AsyncLoad"));
+		DirectStatsCommand( TEXT( "stat group enable LinkerLoad" ) );
+		DirectStatsCommand( TEXT( "stat group enable AsyncLoad" ) );
 		DirectStatsCommand(TEXT("stat dumpsum -start -ms=250 -num=240"));
 	}
-	if (FParse::Param(FCommandLine::Get(), TEXT("LoadTimeFile")))
+	if (FParse::Param( FCommandLine::Get(), TEXT( "LoadTimeFile" ) ))
 	{
-		DirectStatsCommand(TEXT("stat group enable LinkerLoad"));
-		DirectStatsCommand(TEXT("stat group enable AsyncLoad"));
-		DirectStatsCommand(TEXT("stat startfile"));
+		DirectStatsCommand( TEXT( "stat group enable LinkerLoad" ) );
+		DirectStatsCommand( TEXT( "stat group enable AsyncLoad" ) );
+		DirectStatsCommand( TEXT( "stat startfile" ) );
 	}
 
 	// Now we can safely enable malloc profiler.
-	if( FStatsMallocProfilerProxy::HasMemoryProfilerToken() )
+	if (FStatsMallocProfilerProxy::HasMemoryProfilerToken())
 	{
 		// Enable all available groups and enable malloc profiler.
-		IStatGroupEnableManager::Get().StatGroupEnableManagerCommand( TEXT("all") );
+		IStatGroupEnableManager::Get().StatGroupEnableManagerCommand( TEXT( "all" ) );
 		FStatsMallocProfilerProxy::Get()->SetState( true );
-		DirectStatsCommand(TEXT("stat startfileraw"), true);
+		DirectStatsCommand( TEXT( "stat startfileraw" ), true );
 	}
 }
 
