Make FMsg::Logf exit application on fatal error.

diff --git a/Engine/Source/Runtime/Core/Private/Misc/OutputDevice.cpp b/Engine/Source/Runtime/Core/Private/Misc/OutputDevice.cpp
--- a/Engine/Source/Runtime/Core/Private/Misc/OutputDevice.cpp
+++ b/Engine/Source/Runtime/Core/Private/Misc/OutputDevice.cpp
@@ -337,7 +337,7 @@ void FDebug::EnsureFailed(const ANSICHAR* Expr, const ANSICHAR* File, int32 Line
 
 #endif // DO_CHECK || DO_GUARD_SLOW
 
-void VARARGS FError::LowLevelFatal(const ANSICHAR* File, int32 Line, const TCHAR* Format, ... )
+void VARARGS FError::LowLevelFatal(const ANSICHAR* File, int32 Line, const TCHAR* Format, ...)
 {
 	TCHAR DescriptionString[4096];
 	GET_VARARGS( DescriptionString, ARRAY_COUNT(DescriptionString), ARRAY_COUNT(DescriptionString)-1, Format, Format );
@@ -478,7 +478,7 @@ static FCriticalSection					MsgLogfStaticBufferGuard;
 /** Increased from 4096 to fix crashes in the renderthread without autoreporter. */
 static TCHAR							MsgLogfStaticBuffer[8192];
 
-VARARG_BODY(void, FMsg::Logf, const TCHAR*, VARARG_EXTRA(const ANSICHAR* File) VARARG_EXTRA(int32 Line) VARARG_EXTRA(const class FName& Category) VARARG_EXTRA(ELogVerbosity::Type Verbosity))
+void FMsg::Logf_Impl(const ANSICHAR* File, int32 Line, const FName& Category, ELogVerbosity::Type Verbosity, bool bShouldAssertOnFatal, const TCHAR* Fmt, va_list args)
 {
 #if !NO_LOGGING
 	if (Verbosity != ELogVerbosity::Fatal)
@@ -498,12 +498,47 @@ VARARG_BODY(void, FMsg::Logf, const TCHAR*, VARARG_EXTRA(const ANSICHAR* File) V
 				break;
 			}
 		default:
+		{
+			LogDevice = GLog;
+		}
+		break;
+		}
+
+		int32 BufferSize = 1024;
+		TCHAR* Buffer = NULL;
+		int32 Result = -1;
+
+		/* allocate some stack space to use on the first pass, which matches most strings */
+		TCHAR StackBuffer[512];
+		TCHAR* AllocatedBuffer = NULL;
+
+		/* first, try using the stack buffer */
+		Buffer = StackBuffer;
+
+		Result = FCString::GetVarArgs(Buffer, ARRAY_COUNT(StackBuffer), ARRAY_COUNT(StackBuffer) - 1, Fmt, args);
+		if (Result >= ARRAY_COUNT(StackBuffer))
+		{
+			Result = -1;
+		}
+
+		/* if that fails, then use heap allocation to make enough space */
+		while (Result == -1)
+		{
+			FMemory::SystemFree(AllocatedBuffer);
+			/* We need to use malloc here directly as GMalloc might not be safe. */
+			Buffer = AllocatedBuffer = (TCHAR*)FMemory::SystemMalloc(BufferSize * sizeof(TCHAR));
+			Result = FCString::GetVarArgs(Buffer, BufferSize, BufferSize - 1, Fmt, args);
+			if (Result >= ARRAY_COUNT(StackBuffer))
 			{
-				LogDevice = GLog;
+				Result = -1;
 			}
-			break;
-		}
-		GROWABLE_LOGF(LogDevice->Log(Category, Verbosity, Buffer))
+			BufferSize *= 2;
+		};
+		Buffer[Result] = 0;
+
+		LogDevice->Log(Category, Verbosity, Buffer);
+		FMemory::SystemFree(AllocatedBuffer);
+
 	}
 	else
 	{
@@ -514,18 +549,41 @@ VARARG_BODY(void, FMsg::Logf, const TCHAR*, VARARG_EXTRA(const ANSICHAR* File) V
 			// @todo: implement platform independent sprintf_S
 			// We're using one big shared static buffer here, so guard against re-entry
 			FScopeLock MsgLock(&MsgLogfStaticBufferGuard);
+
 			// Print to a large static buffer so we can keep the stack allocation below 16K
-			GET_VARARGS(MsgLogfStaticBuffer, ARRAY_COUNT(MsgLogfStaticBuffer), ARRAY_COUNT(MsgLogfStaticBuffer) - 1, Fmt, Fmt);
+			FCString::GetVarArgs(MsgLogfStaticBuffer, ARRAY_COUNT(MsgLogfStaticBuffer), ARRAY_COUNT(MsgLogfStaticBuffer) - 1, Fmt, args);
+
 			// Copy the message to the stack-allocated buffer)
 			FCString::Strncpy(Message, MsgLogfStaticBuffer, ARRAY_COUNT(Message) - 1);
 			Message[ARRAY_COUNT(Message) - 1] = '\0';
 		}
 
 		StaticFailDebug(TEXT("Fatal error:"), File, Line, Message);
+
+		if (bShouldAssertOnFatal)
+		{
+			FDebug::AssertFailed("", File, Line, Message);
+		}
 	}
 #endif
 }
 
+VARARG_BODY(void, FMsg::Logf, const TCHAR*, VARARG_EXTRA(const ANSICHAR* File) VARARG_EXTRA(int32 Line) VARARG_EXTRA(const class FName& Category) VARARG_EXTRA(ELogVerbosity::Type Verbosity))
+{
+	va_list va;
+	va_start(va, Fmt);
+	FMsg::Logf_Impl(File, Line, Category, Verbosity, true, Fmt, va);
+	va_end(va);
+}
+
+VARARG_BODY(void, FMsg::Logf_Internal, const TCHAR*, VARARG_EXTRA(const ANSICHAR* File) VARARG_EXTRA(int32 Line) VARARG_EXTRA(const class FName& Category) VARARG_EXTRA(ELogVerbosity::Type Verbosity))
+{
+	va_list va;
+	va_start(va, Fmt);
+	FMsg::Logf_Impl(File, Line, Category, Verbosity, false, Fmt, va);
+	va_end(va);
+}
+
 /** Sends a formatted message to a remote tool. */
 void VARARGS FMsg::SendNotificationStringf( const TCHAR *Fmt, ... )
 {
diff --git a/Engine/Source/Runtime/Core/Public/Misc/AssertionMacros.h b/Engine/Source/Runtime/Core/Public/Misc/AssertionMacros.h
--- a/Engine/Source/Runtime/Core/Public/Misc/AssertionMacros.h
+++ b/Engine/Source/Runtime/Core/Public/Misc/AssertionMacros.h
@@ -307,7 +307,7 @@ struct FTCharArrayTester
 			{ \
 				if (!CategoryName.IsSuppressed(ELogVerbosity::Verbosity)) \
 				{ \
-					FMsg::Logf(__FILE__, __LINE__, CategoryName.GetCategoryName(), ELogVerbosity::Verbosity, Format, ##__VA_ARGS__); \
+					FMsg::Logf_Internal(__FILE__, __LINE__, CategoryName.GetCategoryName(), ELogVerbosity::Verbosity, Format, ##__VA_ARGS__); \
 					if (ELogVerbosity::Verbosity == ELogVerbosity::Fatal) \
 					{\
 						_DebugBreakAndPromptForRemote(); \
@@ -328,7 +328,7 @@ struct FTCharArrayTester
 				{ \
 					if (Condition) \
 					{ \
-						FMsg::Logf(__FILE__, __LINE__, CategoryName.GetCategoryName(), ELogVerbosity::Verbosity, Format, ##__VA_ARGS__); \
+						FMsg::Logf_Internal(__FILE__, __LINE__, CategoryName.GetCategoryName(), ELogVerbosity::Verbosity, Format, ##__VA_ARGS__); \
 						if (ELogVerbosity::Verbosity == ELogVerbosity::Fatal) \
 						{\
 							_DebugBreakAndPromptForRemote(); \
diff --git a/Engine/Source/Runtime/Core/Public/Misc/OutputDevice.h b/Engine/Source/Runtime/Core/Public/Misc/OutputDevice.h
--- a/Engine/Source/Runtime/Core/Public/Misc/OutputDevice.h
+++ b/Engine/Source/Runtime/Core/Public/Misc/OutputDevice.h
@@ -372,6 +372,11 @@ struct CORE_API FMsg
 
 	/** Log function */
 	VARARG_DECL( static void, static void, {}, Logf, VARARG_NONE, const TCHAR*, VARARG_EXTRA(const ANSICHAR* File) VARARG_EXTRA(int32 Line) VARARG_EXTRA(const class FName& Category) VARARG_EXTRA(ELogVerbosity::Type Verbosity), VARARG_EXTRA(File) VARARG_EXTRA(Line) VARARG_EXTRA(Category) VARARG_EXTRA(Verbosity) );
+
+	VARARG_DECL(static void, static void, {}, Logf_Internal, VARARG_NONE, const TCHAR*, VARARG_EXTRA(const ANSICHAR* File) VARARG_EXTRA(int32 Line) VARARG_EXTRA(const class FName& Category) VARARG_EXTRA(ELogVerbosity::Type Verbosity), VARARG_EXTRA(File) VARARG_EXTRA(Line) VARARG_EXTRA(Category) VARARG_EXTRA(Verbosity));
+
+private:
+	static void Logf_Impl(const ANSICHAR* File, int32 Line, const FName& Category, ELogVerbosity::Type Verbosity, bool bShouldAssertOnFatal, const TCHAR* Fmt, va_list args);
 };
 
 /**
