Added IsValidIndex to FString and had the [] operators use this to verify the index they were given was valid

Fixed some fallout from this change in existing code.

diff --git a/Engine/Source/Runtime/Core/Private/Containers/String.cpp b/Engine/Source/Runtime/Core/Private/Containers/String.cpp
--- a/Engine/Source/Runtime/Core/Private/Containers/String.cpp
+++ b/Engine/Source/Runtime/Core/Private/Containers/String.cpp
@@ -80,7 +80,8 @@ int32 FString::Find(const TCHAR* SubStr, ESearchCase::Type SearchCase, ESearchDi
 FString FString::ToUpper() const
 {
 	FString New( **this );
-	for( int32 i=0; i < New.Data.Num(); i++ )
+	const int32 StringLength = Len();
+	for( int32 i=0; i < StringLength; ++i )
 	{
 		New[i] = FChar::ToUpper(New[i]);
 	}
@@ -90,7 +91,8 @@ FString FString::ToUpper() const
 FString FString::ToLower() const
 {
 	FString New( **this );
-	for( int32 i=0; i < New.Data.Num(); i++ )
+	const int32 StringLength = Len();
+	for( int32 i=0; i < StringLength; ++i )
 	{
 		New[i] = FChar::ToLower(New[i]);
 	}
diff --git a/Engine/Source/Runtime/Core/Public/Containers/Array.h b/Engine/Source/Runtime/Core/Public/Containers/Array.h
--- a/Engine/Source/Runtime/Core/Public/Containers/Array.h
+++ b/Engine/Source/Runtime/Core/Public/Containers/Array.h
@@ -680,8 +680,7 @@ public:
 	}
 
 	/**
-	 * Tests if index is valid, i.e. greater than zero and less than number of
-	 * elements in array.
+	 * Tests if index is valid, i.e. than or equal to zero, and less than the number of elements in the array.
 	 *
 	 * @param Index Index to test.
 	 *
diff --git a/Engine/Source/Runtime/Core/Public/Containers/UnrealString.h b/Engine/Source/Runtime/Core/Public/Containers/UnrealString.h
--- a/Engine/Source/Runtime/Core/Public/Containers/UnrealString.h
+++ b/Engine/Source/Runtime/Core/Public/Containers/UnrealString.h
@@ -200,7 +200,8 @@ public:
 	 */
 	FORCEINLINE TCHAR& operator[]( int32 Index )
 	{
-		return Data[Index];
+		checkf(IsValidIndex(Index), TEXT("String index out of bounds: Index %i from a string with a length of %i"), Index, Len());
+		return Data.GetData()[Index];
 	}
 
 	/**
@@ -211,7 +212,8 @@ public:
 	 */
 	FORCEINLINE const TCHAR& operator[]( int32 Index ) const
 	{
-		return Data[Index];
+		checkf(IsValidIndex(Index), TEXT("String index out of bounds: Index %i from a string with a length of %i"), Index, Len());
+		return Data.GetData()[Index];
 	}
 
 	/**
@@ -298,6 +300,17 @@ public:
 		Data.Shrink();
 	}
 
+	/**
+	 * Tests if index is valid, i.e. greater than or equal to zero, and less than the number of characters in this string (excluding the null terminator).
+	 *
+	 * @param Index Index to test.
+	 *
+	 * @returns True if index is valid. False otherwise.
+	 */
+	FORCEINLINE bool IsValidIndex(int32 Index) const
+	{
+		return Index >= 0 && Index < Len();
+	}
 
 	/**
 	 * Get pointer to the string
@@ -309,7 +322,6 @@ public:
 		return Data.Num() ? Data.GetData() : TEXT("");
 	}
 
-
 	/** 
 	 *Get string as array of TCHARS 
 	 *
diff --git a/Engine/Source/Runtime/Slate/Private/Framework/Text/SlateTextRun.cpp b/Engine/Source/Runtime/Slate/Private/Framework/Text/SlateTextRun.cpp
--- a/Engine/Source/Runtime/Slate/Private/Framework/Text/SlateTextRun.cpp
+++ b/Engine/Source/Runtime/Slate/Private/Framework/Text/SlateTextRun.cpp
@@ -54,7 +54,7 @@ FVector2D FSlateTextRun::Measure( int32 BeginIndex, int32 EndIndex, float Scale
 int8 FSlateTextRun::GetKerning(int32 CurrentIndex, float Scale) const
 {
 	const int32 PreviousIndex = CurrentIndex - 1;
-	if ( PreviousIndex < 0 )
+	if ( PreviousIndex < 0 || CurrentIndex == Text->Len() )
 	{
 		return 0;
 	}
