Core: Fixed a compilation error if CopyFromWeakArray was used

diff --git a/Engine/Source/Runtime/Core/Public/UObject/WeakObjectPtrTemplates.h b/Engine/Source/Runtime/Core/Public/UObject/WeakObjectPtrTemplates.h
--- a/Engine/Source/Runtime/Core/Public/UObject/WeakObjectPtrTemplates.h
+++ b/Engine/Source/Runtime/Core/Public/UObject/WeakObjectPtrTemplates.h
@@ -309,15 +309,15 @@ template<class T> struct TIsPODType<TAutoWeakObjectPtr<T> > { enum { Value = tru
 template<class T> struct TIsZeroConstructType<TAutoWeakObjectPtr<T> > { enum { Value = true }; };
 template<class T> struct TIsWeakPointerType<TAutoWeakObjectPtr<T> > { enum { Value = true }; };
 
-template<class T, class U>
-void CopyFromWeakArray(T& Dest,const U& Src)
+template<typename DestArrayType, typename SourceArrayType>
+void CopyFromWeakArray(DestArrayType& Dest, const SourceArrayType& Src)
 {
 	Dest.Empty(Src.Num());
 	for (int32 Index = 0; Index < Src.Num(); Index++)
 	{
-		if (Src(Index).Get())
+		if (auto Value = Src[Index].Get())
 		{
-			Dest.AddItem(Src(Index).Get());
+			Dest.Add(Value);
 		}
 	}
-}
+}
\ No newline at end of file
