Fixed Xbox One build failure due to C++11 feature used in InternationalizationMetadata.h.

diff --git a/Engine/Source/Runtime/Core/Public/Internationalization/InternationalizationMetadata.h b/Engine/Source/Runtime/Core/Public/Internationalization/InternationalizationMetadata.h
--- a/Engine/Source/Runtime/Core/Public/Internationalization/InternationalizationMetadata.h
+++ b/Engine/Source/Runtime/Core/Public/Internationalization/InternationalizationMetadata.h
@@ -46,8 +46,6 @@ public:
 
 protected:
 	FLocMetadataValue() : Type(ELocMetadataType::None) {}
-	FLocMetadataValue( const FLocMetadataValue& ) = delete;  
-	FLocMetadataValue& operator=( const FLocMetadataValue& ) = delete;
 
 	virtual bool EqualTo( const FLocMetadataValue& Other ) const = 0;
 	virtual bool LessThan( const FLocMetadataValue& Other ) const = 0;
@@ -55,6 +53,10 @@ protected:
 	virtual FString GetTypeString() const = 0;
 
 	void ErrorMessage(const FString& InType);
+
+private:
+	FLocMetadataValue( const FLocMetadataValue& );
+	FLocMetadataValue& operator=( const FLocMetadataValue& );
 };
 
 
