Core - UE-15731/Any stat command in the editor crashes with a bogus assert inside FString::ChrN

diff --git a/Engine/Source/Runtime/Core/Private/Containers/String.cpp b/Engine/Source/Runtime/Core/Private/Containers/String.cpp
--- a/Engine/Source/Runtime/Core/Private/Containers/String.cpp
+++ b/Engine/Source/Runtime/Core/Private/Containers/String.cpp
@@ -465,7 +465,7 @@ FString FString::ChrN( int32 NumCharacters, TCHAR Char )
 	{
 		Temp[Cx] = Char;
 	}
-	Temp[NumCharacters]=0;
+	Temp.Data[NumCharacters]=0;
 	return Temp;
 }
 
