Fix Xbox One compile.

diff --git a/Engine/Source/Runtime/Core/Private/Internationalization/InternationalizationMetadata.cpp b/Engine/Source/Runtime/Core/Private/Internationalization/InternationalizationMetadata.cpp
--- a/Engine/Source/Runtime/Core/Private/Internationalization/InternationalizationMetadata.cpp
+++ b/Engine/Source/Runtime/Core/Private/Internationalization/InternationalizationMetadata.cpp
@@ -7,7 +7,11 @@ DEFINE_LOG_CATEGORY_STATIC(LogInternationalizationMetadata, Log, All);
 
 const FString FLocMetadataObject::COMPARISON_MODIFIER_PREFIX = TEXT("*");
 
+#if PLATFORM_COMPILER_HAS_DEFAULTED_FUNCTIONS
 FLocMetadataValue::~FLocMetadataValue() = default;
+#else
+FLocMetadataValue::~FLocMetadataValue() {}
+#endif
 
 void FLocMetadataValue::ErrorMessage( const FString& InType )
 {
