CrashReport - Removed ProgramCounter, SymbolDisplacement from the callstack line

diff --git a/Engine/Source/Runtime/Core/Private/GenericPlatform/GenericPlatformStackWalk.cpp b/Engine/Source/Runtime/Core/Private/GenericPlatform/GenericPlatformStackWalk.cpp
--- a/Engine/Source/Runtime/Core/Private/GenericPlatform/GenericPlatformStackWalk.cpp
+++ b/Engine/Source/Runtime/Core/Private/GenericPlatform/GenericPlatformStackWalk.cpp
@@ -30,10 +30,10 @@ bool FGenericPlatformStackWalk::SymbolInfoToHumanReadableString( const FProgramC
 {
 	const int32 MAX_TEMP_SPRINTF = 256;
 	// Valid callstack line 
-	// ModuleName!FunctionName {ProgramCounter} + offset bytes [Filename:LineNumber]
+	// ModuleName!FunctionName [Filename:LineNumber]
 	// 
 	// Invalid callstack line
-	// ModuleName! {ProgramCounter} + offset bytes
+	// ModuleName!
 	// 
 	if( HumanReadableString && HumanReadableStringSize > 0 )
 	{
@@ -45,7 +45,7 @@ bool FGenericPlatformStackWalk::SymbolInfoToHumanReadableString( const FProgramC
 		const UPTRINT RealPos = FMath::Max( (UPTRINT)Pos0, (UPTRINT)Pos1 );
 		const ANSICHAR* StrippedModuleName = RealPos > 0 ? (const ANSICHAR*)(RealPos + 1) : SymbolInfo.ModuleName;
 
-		//FCStringAnsi::Sprintf( StackLine, "%s!%s {0x%016llx} + %i bytes [%s:%i]", StrippedModuleName, (const ANSICHAR*)SymbolInfo.FunctionName, SymbolInfo.ProgramCounter, SymbolInfo.SymbolDisplacement, (const ANSICHAR*)SymbolInfo.Filename, SymbolInfo.LineNumber );
+		//FCStringAnsi::Sprintf( StackLine, "%s!%s [%s:%i]", StrippedModuleName, (const ANSICHAR*)SymbolInfo.FunctionName, (const ANSICHAR*)SymbolInfo.Filename, SymbolInfo.LineNumber );
 		FCStringAnsi::Strcat( StackLine, MAX_SPRINTF, StrippedModuleName );
 		
 		const bool bHasValidFunctionName = FCStringAnsi::Strlen( SymbolInfo.FunctionName ) > 0;
@@ -55,12 +55,6 @@ bool FGenericPlatformStackWalk::SymbolInfoToHumanReadableString( const FProgramC
 			FCStringAnsi::Strcat( StackLine, MAX_SPRINTF, SymbolInfo.FunctionName );
 		}
 
-		{
-			ANSICHAR ProgramCounterAndSymbolDisplacement[MAX_TEMP_SPRINTF] = {0};
-			FCStringAnsi::Snprintf( ProgramCounterAndSymbolDisplacement, MAX_TEMP_SPRINTF, " {0x%016llx} + %i bytes", SymbolInfo.ProgramCounter, SymbolInfo.SymbolDisplacement );
-			FCStringAnsi::Strcat( StackLine, MAX_SPRINTF, ProgramCounterAndSymbolDisplacement );
-		}
-
 		const bool bHasValidFilename = FCStringAnsi::Strlen( SymbolInfo.Filename ) > 0 && SymbolInfo.LineNumber > 0;
 		if( bHasValidFilename )
 		{
@@ -83,10 +77,10 @@ bool FGenericPlatformStackWalk::SymbolInfoToHumanReadableString( const FProgramC
 bool FGenericPlatformStackWalk::SymbolInfoToHumanReadableStringEx( const FProgramCounterSymbolInfoEx& SymbolInfo, FString& out_HumanReadableString )
 {
 	// Valid callstack line 
-	// ModuleName!FunctionName {ProgramCounter} + offset bytes [Filename:LineNumber]
+	// ModuleName!FunctionName [Filename:LineNumber]
 	// 
 	// Invalid callstack line
-	// ModuleName! {ProgramCounter} + offset bytes
+	// ModuleName! {ProgramCounter}
 	
 	// Strip module path.
 	const TCHAR* Pos0 = FCString::Strrchr( *SymbolInfo.ModuleName, '\\' );
@@ -103,8 +97,6 @@ bool FGenericPlatformStackWalk::SymbolInfoToHumanReadableStringEx( const FProgra
 		out_HumanReadableString += SymbolInfo.FunctionName;
 	}
 
-	//out_HumanReadableString += FString::Printf( TEXT( " {0x%016llx} + %i bytes" ), SymbolInfo.ProgramCounter, SymbolInfo.SymbolDisplacement );
-
 	const bool bHasValidFilename = SymbolInfo.Filename.Len() > 0 && SymbolInfo.LineNumber > 0;
 	if( bHasValidFilename )
 	{
