Fix CIS.

diff --git a/Engine/Source/Runtime/Core/Private/Internationalization/GatherableTextData.cpp b/Engine/Source/Runtime/Core/Private/Internationalization/GatherableTextData.cpp
--- a/Engine/Source/Runtime/Core/Private/Internationalization/GatherableTextData.cpp
+++ b/Engine/Source/Runtime/Core/Private/Internationalization/GatherableTextData.cpp
@@ -1,6 +1,6 @@
 // Copyright 1998-2015 Epic Games, Inc. All Rights Reserved.
 
-#include "CorePrivatePch.h"
+#include "CorePrivatePCH.h"
 #include "Internationalization/GatherableTextData.h"
 
 FArchive& operator<<(FArchive& Archive, FTextSourceSiteContext& This)
@@ -30,4 +30,4 @@ FArchive& operator<<(FArchive& Archive, FGatherableTextData& This)
 	Archive << This.SourceSiteContexts;
 
 	return Archive;
-}
\ No newline at end of file
+}
