Added support for binding Lambdas to multicast delegates

diff --git a/Engine/Source/Runtime/Core/Public/Delegates/DelegateSignatureImpl.inl b/Engine/Source/Runtime/Core/Public/Delegates/DelegateSignatureImpl.inl
--- a/Engine/Source/Runtime/Core/Public/Delegates/DelegateSignatureImpl.inl
+++ b/Engine/Source/Runtime/Core/Public/Delegates/DelegateSignatureImpl.inl
@@ -1302,6 +1302,38 @@ public:
 		return Add( FDelegate::CreateStatic( InFunc, PAYLOAD_TEMPLATE_PASSIN_FourVars ) );
 	}
 
+	/**
+	 * Adds a C++ lambda delegate
+	 * technically this works for any functor types, but lambdas are the primary use case
+	 *
+	 * @param	InFunctor	Functor (e.g. Lambda)
+	 */
+	template<typename FunctorType>
+	inline FDelegateHandle AddLambda(FunctorType&& InFunctor)
+	{
+		return Add(FDelegate::CreateLambda(Forward<FunctorType>(InFunctor)));
+	}
+	template<typename FunctorType, PAYLOAD_TEMPLATE_DECL_OneVar>
+	inline FDelegateHandle AddLambda(FunctorType&& InFunctor, PAYLOAD_TEMPLATE_ARGS_OneVar)
+	{
+		return Add(FDelegate::CreateLambda(Forward<FunctorType>(InFunctor), PAYLOAD_TEMPLATE_PASSIN_OneVar));
+	}
+	template<typename FunctorType, PAYLOAD_TEMPLATE_DECL_TwoVars>
+	inline FDelegateHandle AddLambda(FunctorType&& InFunctor, PAYLOAD_TEMPLATE_ARGS_TwoVars)
+	{
+		return Add(FDelegate::CreateLambda(Forward<FunctorType>(InFunctor), PAYLOAD_TEMPLATE_PASSIN_TwoVars));
+	}
+	template<typename FunctorType, PAYLOAD_TEMPLATE_DECL_ThreeVars>
+	inline FDelegateHandle AddLambda(FunctorType&& InFunctor, PAYLOAD_TEMPLATE_ARGS_ThreeVars)
+	{
+		return Add(FDelegate::CreateLambda(Forward<FunctorType>(InFunctor), PAYLOAD_TEMPLATE_PASSIN_ThreeVars));
+	}
+	template<typename FunctorType, PAYLOAD_TEMPLATE_DECL_FourVars>
+	inline FDelegateHandle AddLambda(FunctorType&& InFunctor, PAYLOAD_TEMPLATE_ARGS_FourVars)
+	{
+		return Add(FDelegate::CreateLambda(Forward<FunctorType>(InFunctor), PAYLOAD_TEMPLATE_PASSIN_FourVars));
+	}
+
 	/**
 	 * Adds a raw C++ pointer delegate.
 	 *
diff --git a/Engine/Source/Runtime/Core/Public/Delegates/DelegateSignatureImpl_Variadics.inl b/Engine/Source/Runtime/Core/Public/Delegates/DelegateSignatureImpl_Variadics.inl
--- a/Engine/Source/Runtime/Core/Public/Delegates/DelegateSignatureImpl_Variadics.inl
+++ b/Engine/Source/Runtime/Core/Public/Delegates/DelegateSignatureImpl_Variadics.inl
@@ -599,6 +599,18 @@ public:
 		return Add(FDelegate::CreateStatic(InFunc, Vars...));
 	}
 
+	/**
+	 * Adds a C++ lambda delegate
+	 * technically this works for any functor types, but lambdas are the primary use case
+	 *
+	 * @param	InFunctor	Functor (e.g. Lambda)
+	 */
+	template<typename FunctorType, typename... VarTypes>
+	inline FDelegateHandle AddLambda(FunctorType&& InFunctor, VarTypes... Vars)
+	{
+		return Add(FDelegate::CreateLambda(Forward<FunctorType>(InFunctor), Vars...));
+	}
+
 	/**
 	 * Adds a raw C++ pointer delegate.
 	 *
