The text localization manager now exposes a method so that in-memory localizations can be updated or added to by loading a specified localization resource file.

diff --git a/Engine/Source/Runtime/Core/Private/Internationalization/TextLocalizationManager.cpp b/Engine/Source/Runtime/Core/Private/Internationalization/TextLocalizationManager.cpp
--- a/Engine/Source/Runtime/Core/Private/Internationalization/TextLocalizationManager.cpp
+++ b/Engine/Source/Runtime/Core/Private/Internationalization/TextLocalizationManager.cpp
@@ -651,6 +651,17 @@ void FTextLocalizationManager::LoadFromManifestAndArchives( const FString& Confi
 	UpdateFromLocalizations(LocalizationEntryTrackers);
 }
 
+void FTextLocalizationManager::UpdateFromLocalizationResource(const FString& LocalizationResourceFilePath)
+{
+	TArray<FLocalizationEntryTracker> LocalizationEntryTrackers;
+	
+	FLocalizationEntryTracker LocalizationEntryTracker = LocalizationEntryTrackers[LocalizationEntryTrackers.Add(FLocalizationEntryTracker())];
+	LocalizationEntryTracker.LoadFromFile(LocalizationResourceFilePath);
+	LocalizationEntryTracker.DetectAndLogConflicts();
+
+	UpdateFromLocalizations(LocalizationEntryTrackers);
+}
+
 void FTextLocalizationManager::OnCultureChanged()
 {
 	const bool ShouldLoadEditor = bIsInitialized && WITH_EDITOR;
diff --git a/Engine/Source/Runtime/Core/Public/Internationalization/TextLocalizationManager.h b/Engine/Source/Runtime/Core/Public/Internationalization/TextLocalizationManager.h
--- a/Engine/Source/Runtime/Core/Public/Internationalization/TextLocalizationManager.h
+++ b/Engine/Source/Runtime/Core/Public/Internationalization/TextLocalizationManager.h
@@ -144,6 +144,9 @@ public:
 	/** Loads localizations for the current culture based on a configuration file specifying which manifests and archives to use. Effectively generates a localization resource in memory. */
 	void LoadFromManifestAndArchives(const FString& ConfigFilePath, IInternationalizationArchiveSerializer& ArchiveSerializer, IInternationalizationManifestSerializer& ManifestSerializer);
 
+	/** Updates display string entries and adds new display string entries based on localizations found in a specified localization resource. */
+	void UpdateFromLocalizationResource(const FString& LocalizationResourceFilePath);
+
 	/**	Returns the current text revision number. This value can be cached when caching information from the text localization manager.
 	 *	If the revision does not match, cached information may be invalid and should be recached. */
 	int32 GetTextRevision() const { return TextRevisionCounter; }
@@ -159,7 +162,7 @@ private:
 	/** Loads localization resources for the specified culture, optionally loading localization resources that are editor-specific or game-specific. */
 	void LoadLocalizationResourcesForCulture(const FString& CultureName, const bool ShouldLoadEditor, const bool ShouldLoadGame);
 
-	/** Updates display string entries based on specified localizations and adds new display string entries. */
+	/** Updates display string entries and adds new display string entries based on provided localizations. */
 	void UpdateFromLocalizations(const TArray<FLocalizationEntryTracker>& LocalizationEntryTrackers);
 
 	/** Dirties the text revision counter by incrementing it, causing a revision mismatch for any information cached before this happens.  */
