UE-15716 - Hide pointer option is not given by the browser.

Play catch up with some recent changes how slate behaves (more correctly now) on window actiivation - this had a side effect on browser pointer lock which behaves entirely differently to rest of the platforms.  This needs a second look after 4.8 ships.

diff --git a/Engine/Source/Runtime/Core/Private/HTML5/HTML5Application.cpp b/Engine/Source/Runtime/Core/Private/HTML5/HTML5Application.cpp
--- a/Engine/Source/Runtime/Core/Private/HTML5/HTML5Application.cpp
+++ b/Engine/Source/Runtime/Core/Private/HTML5/HTML5Application.cpp
@@ -123,20 +123,20 @@ void FHTML5Application::PollGameDeviceState( const float TimeDelta )
 					case SDL_WINDOWEVENT_ENTER:
 						{
 							UE_LOG(LogHTML5Application, Verbose, TEXT("WindowEnter"));
-							MessageHandler->OnCursorSet();
-							MessageHandler->OnWindowActivationChanged(ApplicationWindow, EWindowActivation::Activate); 
+						//	MessageHandler->OnWindowActivationChanged(ApplicationWindow, EWindowActivation::Activate); 
 							WarmUpTicks = 0; 
 						}
 						break;
 					case SDL_WINDOWEVENT_LEAVE:
 						{
 							UE_LOG(LogHTML5Application, Verbose, TEXT("WindowLeave"));
-							MessageHandler->OnWindowActivationChanged(ApplicationWindow, EWindowActivation::Deactivate);
+						 //	MessageHandler->OnWindowActivationChanged(ApplicationWindow, EWindowActivation::Deactivate);
 						}
 						break;
 					case SDL_WINDOWEVENT_FOCUS_GAINED:
 						{
 							UE_LOG(LogHTML5Application, Verbose, TEXT("WindowFocusGained"));
+							MessageHandler->OnCursorSet();
 							MessageHandler->OnWindowActivationChanged(ApplicationWindow, EWindowActivation::Activate);
 									WarmUpTicks = 0;
 						}
@@ -195,7 +195,7 @@ void FHTML5Application::PollGameDeviceState( const float TimeDelta )
 		else
 		{
 			UE_LOG(LogHTML5Application, Verbose, TEXT("Exit pointer lock"));
-			emscripten_exit_pointerlock(); 
+			//emscripten_exit_pointerlock(); 
 		}
 #endif 
 
