Added process enumerators to FPlatformProcess. Usage:

1. Create object of type FPlatformProcess::FProcEnumerator.
2. Call FProcEnumerator::MoveNext() to move enumerator to the next position. It will return false if enumeration is done and GetCurrent() is no longer valid.
3. Use FProcEnumerator::GetCurrent() to get FProcEnumInfo.
4. Use FProcEnumInfo getters to get process properties that are interesting to you.

diff --git a/Engine/Source/Runtime/Core/Private/Linux/LinuxPlatformProcess.cpp b/Engine/Source/Runtime/Core/Private/Linux/LinuxPlatformProcess.cpp
--- a/Engine/Source/Runtime/Core/Private/Linux/LinuxPlatformProcess.cpp
+++ b/Engine/Source/Runtime/Core/Private/Linux/LinuxPlatformProcess.cpp
@@ -1041,3 +1041,111 @@ void FLinuxPlatformProcess::ExploreFolder( const TCHAR* FilePath )
 	// TODO This is broken, not an explore action but should be fine if called on a directory
 	FLinuxPlatformProcess::LaunchFileInDefaultExternalApplication(FilePath, NULL, ELaunchVerb::Edit);
 }
+
+/**
+ * Private struct to store implementation specific data.
+ */
+struct FProcEnumData
+{
+	// Array of processes.
+	TArray<FLinuxPlatformProcess::FProcEnumInfo> Processes;
+
+	// Current process id.
+	uint32 CurrentProcIndex;
+};
+
+FLinuxPlatformProcess::FProcEnumerator::FProcEnumerator()
+{
+  	Data = new FProcEnumData;
+	Data->CurrentProcIndex = -1;
+	
+	TArray<uint32> PIDs;
+	
+	class FPIDsCollector : public IPlatformFile::FDirectoryVisitor
+	{
+	public:
+		FPIDsCollector(TArray<uint32>& PIDs)
+			: PIDs(PIDs)
+		{ }
+
+		bool Visit(const TCHAR* FilenameOrDirectory, bool bIsDirectory)
+		{
+			FString StrPID = FPaths::GetBaseFilename(FilenameOrDirectory);
+		  
+			if (bIsDirectory && FCString::IsNumeric(*StrPID))
+			{
+				PIDs.Add(FCString::Atoi(*StrPID));
+			}
+			
+			return true;
+		}
+
+		TArray<uint32>& PIDs;
+	} PIDsCollector(PIDs);
+
+	IPlatformFile::GetPlatformPhysical().IterateDirectory(TEXT("/proc"), PIDsCollector);
+	
+	for(auto PID : PIDsCollector.PIDs)
+	{
+		Data->Processes.Add(FProcEnumInfo(PID));
+	}
+}
+
+FLinuxPlatformProcess::FProcEnumerator::~FProcEnumerator()
+{
+	delete Data;
+}
+
+FLinuxPlatformProcess::FProcEnumInfo FLinuxPlatformProcess::FProcEnumerator::GetCurrent() const
+{
+	return Data->Processes[Data->CurrentProcIndex];
+}
+
+bool FLinuxPlatformProcess::FProcEnumerator::MoveNext()
+{
+	if (Data->CurrentProcIndex + 1 == Data->Processes.Num())
+	{
+		return false;
+	}
+
+	++Data->CurrentProcIndex;
+
+	return true;
+}
+
+FLinuxPlatformProcess::FProcEnumInfo::FProcEnumInfo(uint32 PID)
+	: PID(PID)
+{
+
+}
+
+uint32 FLinuxPlatformProcess::FProcEnumInfo::GetPID() const
+{
+	return PID;
+}
+
+uint32 FLinuxPlatformProcess::FProcEnumInfo::GetParentPID() const
+{
+	char Buf[256];
+	uint32 DummyNumber;
+	char DummyChar;
+	uint32 ParentPID;
+	
+	sprintf(Buf, "/proc/%d/stat", GetPID());
+	
+	FILE* FilePtr = fopen(Buf, "r");
+	fscanf(FilePtr, "%d %s %c %d", &DummyNumber, Buf, &DummyChar, &ParentPID);	
+	fclose(FilePtr);
+
+	return ParentPID;
+}
+
+FString FLinuxPlatformProcess::FProcEnumInfo::GetFullPath() const
+{
+	return GetApplicationName(GetPID());
+}
+
+FString FLinuxPlatformProcess::FProcEnumInfo::GetName() const
+{
+	return FPaths::GetCleanFilename(GetFullPath());
+}
diff --git a/Engine/Source/Runtime/Core/Private/Mac/MacPlatformProcess.cpp b/Engine/Source/Runtime/Core/Private/Mac/MacPlatformProcess.cpp
--- a/Engine/Source/Runtime/Core/Private/Mac/MacPlatformProcess.cpp
+++ b/Engine/Source/Runtime/Core/Private/Mac/MacPlatformProcess.cpp
@@ -548,24 +548,15 @@ void FMacPlatformProcess::TerminateProc( FProcHandle& ProcessHandle, bool KillTr
 	{
 		int32 ProcessID = [(NSTask*)ProcessHandle.Get() processIdentifier];
 
-		int32 Mib[4] = { CTL_KERN, KERN_PROC, KERN_PROC_ALL, 0 };
-		size_t BufferSize = 0;
-		if (sysctl(Mib, 4, NULL, &BufferSize, NULL, 0) != -1 && BufferSize > 0)
+		FProcEnumerator ProcEnumerator;
+
+		while (ProcEnumerator.MoveNext())
 		{
-			struct kinfo_proc* Processes = (struct kinfo_proc*)FMemory::Malloc(BufferSize);
-			if (sysctl(Mib, 4, Processes, &BufferSize, NULL, 0) != -1)
+			auto Current = ProcEnumerator.GetCurrent();
+			if (Current.GetParentPID() == ProcessID)
 			{
-				uint32 ProcCount = (uint32)(BufferSize / sizeof(struct kinfo_proc));
-				for (uint32 Index = 0; Index < ProcCount; Index++)
-				{
-					if (Processes[Index].kp_eproc.e_ppid == ProcessID)
-					{
-						kill(Processes[Index].kp_proc.p_pid, SIGTERM);
-					}
-				}
+				kill(Current.GetPID(), SIGTERM);
 			}
-
-			FMemory::Free(Processes);
 		}
 	}
 
@@ -590,38 +581,6 @@ bool FMacPlatformProcess::GetProcReturnCode( FProcHandle& ProcessHandle, int32*
 	return true;
 }
 
-bool FMacPlatformProcess::IsApplicationRunning( const TCHAR* ProcName )
-{
-	const FString ProcString = FPaths::GetBaseFilename(ProcName);
-	uint32 ThisProcessID = getpid();
-	bool bFound = false;
-	
-	int32 Mib[4] = { CTL_KERN, KERN_PROC, KERN_PROC_ALL, 0 };
-	size_t BufferSize = 0;
-	if (sysctl(Mib, 4, NULL, &BufferSize, NULL, 0) != -1 && BufferSize > 0)
-	{
-		char Buffer[MAX_PATH];
-		struct kinfo_proc* Processes = (struct kinfo_proc*)FMemory::Malloc(BufferSize);
-		if (sysctl(Mib, 4, Processes, &BufferSize, NULL, 0) != -1)
-		{
-			uint32 ProcCount = (uint32)(BufferSize / sizeof(struct kinfo_proc));
-			for (uint32 Index = 0; Index < ProcCount; Index++)
-			{
-				proc_pidpath(Processes[Index].kp_proc.p_pid, Buffer, sizeof(Buffer));
-				const FString TestProcString = FPaths::GetBaseFilename(Buffer);
-				if (Processes[Index].kp_proc.p_pid != ThisProcessID && TestProcString == ProcString)
-				{
-					bFound = true;
-				}
-			}
-		}
-		
-		FMemory::Free(Processes);
-	}
-    
-	return bFound;
-}
-
 bool FMacPlatformProcess::IsApplicationRunning( uint32 ProcessId )
 {
 	errno = 0;
@@ -1009,6 +968,98 @@ bool FMacPlatformProcess::ReadPipeToArray(void* ReadPipe, TArray<uint8>& Output)
 	return false;
 }
 
+bool FMacPlatformProcess::IsApplicationRunning(const TCHAR* ProcName)
+{
+	const FString ProcString = FPaths::GetCleanFilename(ProcName);
+	uint32 ThisProcessID = getpid();
+
+	FProcEnumerator ProcEnumerator;
+
+	while (ProcEnumerator.MoveNext())
+	{
+		auto Current = ProcEnumerator.GetCurrent();
+		if (Current.GetPID() != ThisProcessID && Current.GetName() == ProcString)
+		{
+			return true;
+		}
+	}
+
+	return false;
+}
+
+FMacPlatformProcess::FProcEnumerator::FProcEnumerator()
+{
+	int32 Mib[4] = { CTL_KERN, KERN_PROC, KERN_PROC_ALL, 0 };
+	size_t BufferSize = 0;
+
+	Processes = nullptr;
+	ProcCount = 0;
+	CurrentProcIndex = -1;
+
+	if (sysctl(Mib, 4, NULL, &BufferSize, NULL, 0) != -1 && BufferSize > 0)
+	{
+		char Buffer[MAX_PATH];
+		Processes = (struct kinfo_proc*)FMemory::Malloc(BufferSize);
+		if (sysctl(Mib, 4, Processes, &BufferSize, NULL, 0) != -1)
+		{
+			ProcCount = (uint32)(BufferSize / sizeof(struct kinfo_proc));
+		}
+	}
+}
+
+FMacPlatformProcess::FProcEnumerator::~FProcEnumerator()
+{
+	if (Processes != nullptr)
+	{
+		FMemory::Free(Processes);
+	}
+}
+
+FMacPlatformProcess::FProcEnumInfo::FProcEnumInfo(struct kinfo_proc ProcInfo)
+	: ProcInfo(ProcInfo)
+{
+
+}
+
+bool FMacPlatformProcess::FProcEnumerator::MoveNext()
+{
+	if (CurrentProcIndex == ProcCount)
+	{
+		return false;
+	}
+
+	++CurrentProcIndex;
+	return true;
+}
+
+FMacPlatformProcess::FProcEnumInfo FMacPlatformProcess::FProcEnumerator::GetCurrent() const
+{
+	return FProcEnumInfo(Processes[CurrentProcIndex]);
+}
+
+uint32 FMacPlatformProcess::FProcEnumInfo::GetPID() const
+{
+	return ProcInfo.kp_proc.p_pid;
+}
+
+uint32 FMacPlatformProcess::FProcEnumInfo::GetParentPID() const
+{
+	return ProcInfo.kp_eproc.e_ppid;
+}
+
+FString FMacPlatformProcess::FProcEnumInfo::GetFullPath() const
+{
+	char Buffer[MAX_PATH];
+	proc_pidpath(GetPID(), Buffer, sizeof(Buffer));
+
+	return Buffer;
+}
+
+FString FMacPlatformProcess::FProcEnumInfo::GetName() const
+{
+	return FPaths::GetCleanFilename(GetFullPath());
+}
+
 #include "MacPlatformRunnableThread.h"
 
 FRunnableThread* FMacPlatformProcess::CreateRunnableThread()
diff --git a/Engine/Source/Runtime/Core/Private/Windows/WindowsPlatformProcess.cpp b/Engine/Source/Runtime/Core/Private/Windows/WindowsPlatformProcess.cpp
--- a/Engine/Source/Runtime/Core/Private/Windows/WindowsPlatformProcess.cpp
+++ b/Engine/Source/Runtime/Core/Private/Windows/WindowsPlatformProcess.cpp
@@ -8,7 +8,6 @@
 	#include <shellapi.h>
 	#include <shlobj.h>
 	#include <LM.h>
-	#include <tlhelp32.h>
 	#include <Psapi.h>
 
 	namespace ProcessConstants
@@ -1247,4 +1246,58 @@ FProcHandle FWindowsPlatformProcess::OpenProcess(uint32 ProcessID)
 	return FProcHandle(::OpenProcess(PROCESS_ALL_ACCESS, 0, ProcessID));
 }
 
+FWindowsPlatformProcess::FProcEnumerator::FProcEnumerator()
+{
+	SnapshotHandle = ::CreateToolhelp32Snapshot(TH32CS_SNAPPROCESS, 0);
+	CurrentEntry.dwSize = 0;
+}
+
+FWindowsPlatformProcess::FProcEnumerator::~FProcEnumerator()
+{
+	::CloseHandle(SnapshotHandle);
+}
+
+FWindowsPlatformProcess::FProcEnumInfo::FProcEnumInfo(const PROCESSENTRY32& Info)
+	: Info(Info)
+{
+
+}
+
+bool FWindowsPlatformProcess::FProcEnumerator::MoveNext()
+{
+	if (CurrentEntry.dwSize == 0)
+	{
+		CurrentEntry.dwSize = sizeof(PROCESSENTRY32);
+
+		return ::Process32First(SnapshotHandle, &CurrentEntry) == TRUE;
+	}
+
+	return ::Process32Next(SnapshotHandle, &CurrentEntry) == TRUE;
+}
+
+FWindowsPlatformProcess::FProcEnumInfo FWindowsPlatformProcess::FProcEnumerator::GetCurrent() const
+{
+	return FProcEnumInfo(CurrentEntry);
+}
+
+uint32 FWindowsPlatformProcess::FProcEnumInfo::GetPID() const
+{
+	return Info.th32ProcessID;
+}
+
+uint32 FWindowsPlatformProcess::FProcEnumInfo::GetParentPID() const
+{
+	return Info.th32ParentProcessID;
+}
+
+FString FWindowsPlatformProcess::FProcEnumInfo::GetName() const
+{
+	return Info.szExeFile;
+}
+
+FString FWindowsPlatformProcess::FProcEnumInfo::GetFullPath() const
+{
+	return GetApplicationName(GetPID());
+}
+
 #include "HideWindowsPlatformTypes.h"
diff --git a/Engine/Source/Runtime/Core/Public/Linux/LinuxPlatformProcess.h b/Engine/Source/Runtime/Core/Public/Linux/LinuxPlatformProcess.h
--- a/Engine/Source/Runtime/Core/Public/Linux/LinuxPlatformProcess.h
+++ b/Engine/Source/Runtime/Core/Public/Linux/LinuxPlatformProcess.h
@@ -185,6 +185,61 @@ protected:
  */
 struct CORE_API FLinuxPlatformProcess : public FGenericPlatformProcess
 {
+	struct FProcEnumInfo;
+
+	/**
+	 * Process enumerator.
+	 */
+	class FProcEnumerator
+	{
+	public:
+		// Constructor
+		FProcEnumerator();
+
+		// Destructor
+		~FProcEnumerator();
+
+		// Gets current process enumerator info.
+		FProcEnumInfo GetCurrent() const;
+		
+		/**
+		 * Moves current to the next process.
+		 *
+		 * @returns True if succeeded. False otherwise.
+		 */
+		bool MoveNext();
+	private:
+		// Private implementation data.
+		struct FProcEnumData* Data;
+	};
+
+	/**
+	 * Process enumeration info structure.
+	 */
+	struct FProcEnumInfo
+	{
+		friend FLinuxPlatformProcess::FProcEnumerator::FProcEnumerator();
+
+		// Gets process PID.
+		uint32 GetPID() const;
+
+		// Gets parent process PID.
+		uint32 GetParentPID() const;
+
+		// Gets process name. I.e. exec name.
+		FString GetName() const;
+
+		// Gets process full image path. I.e. full path of the exec file.
+		FString GetFullPath() const;
+
+	private:
+		// Private constructor.
+		FProcEnumInfo(uint32 PID);
+
+		// Current process PID.
+		uint32 PID;
+	};
+
 	static void* GetDllHandle( const TCHAR* Filename );
 	static void FreeDllHandle( void* DllHandle );
 	static void* GetDllExport( void* DllHandle, const TCHAR* ProcName );
diff --git a/Engine/Source/Runtime/Core/Public/Mac/MacPlatformProcess.h b/Engine/Source/Runtime/Core/Public/Mac/MacPlatformProcess.h
--- a/Engine/Source/Runtime/Core/Public/Mac/MacPlatformProcess.h
+++ b/Engine/Source/Runtime/Core/Public/Mac/MacPlatformProcess.h
@@ -28,6 +28,67 @@ public:
 **/
 struct CORE_API FMacPlatformProcess : public FGenericPlatformProcess
 {
+	struct FProcEnumInfo;
+
+	/**
+	 * Process enumerator.
+	 */
+	class FProcEnumerator
+	{
+	public:
+		// Constructor
+		FProcEnumerator();
+		// Destructor
+		~FProcEnumerator();
+
+		// Gets current process enumerator info.
+		FProcEnumInfo GetCurrent() const;
+
+		/**
+		 * Moves current to the next process.
+		 *
+		 * @returns True if succeeded. False otherwise.
+		 */
+		bool MoveNext();
+
+	private:
+		// Array of process info structures.
+		struct kinfo_proc* Processes;
+
+		// Count of processes.
+		uint32 ProcCount;
+
+		// Current process index.
+		uint32 CurrentProcIndex;
+	};
+
+	/**
+	 * Process enumeration info structure.
+	 */
+	struct FProcEnumInfo
+	{
+		friend FProcEnumInfo FMacPlatformProcess::FProcEnumerator::GetCurrent() const;
+
+		// Gets process PID.
+		uint32 GetPID() const;
+
+		// Gets parent process PID.
+		uint32 GetParentPID() const;
+
+		// Gets process name. I.e. exec name.
+		FString GetName() const;
+
+		// Gets process full image path. I.e. full path of the exec file.
+		FString GetFullPath() const;
+
+	private:
+		// Private constructor.
+		FProcEnumInfo(struct kinfo_proc ProcInfo);
+
+		// Process info struct.
+		struct kinfo_proc ProcInfo;
+	};
+
 	static void* GetDllHandle( const TCHAR* Filename );
 	static void FreeDllHandle( void* DllHandle );
 	static void* GetDllExport( void* DllHandle, const TCHAR* ProcName );
diff --git a/Engine/Source/Runtime/Core/Public/Windows/WindowsPlatformProcess.h b/Engine/Source/Runtime/Core/Public/Windows/WindowsPlatformProcess.h
--- a/Engine/Source/Runtime/Core/Public/Windows/WindowsPlatformProcess.h
+++ b/Engine/Source/Runtime/Core/Public/Windows/WindowsPlatformProcess.h
@@ -5,6 +5,8 @@
 #include "GenericPlatform/GenericPlatformProcess.h"
 #include "Windows/WindowsSystemIncludes.h"
 
+#include <tlhelp32.h>
+
 
 /** Windows implementation of the process handle. */
 struct FProcHandle : public TProcHandle<HANDLE, nullptr>
@@ -55,6 +57,64 @@ struct CORE_API FWindowsPlatformProcess
 		HANDLE			Semaphore;
 	};
 
+	struct FProcEnumInfo;
+
+	/**
+	 * Process enumerator.
+	 */
+	class FProcEnumerator
+	{
+	public:
+		// Constructor
+		FProcEnumerator();
+
+		// Destructor
+		~FProcEnumerator();
+
+		// Gets current process enumerator info.
+		FProcEnumInfo GetCurrent() const;
+
+		/**
+		 * Moves current to the next process.
+		 *
+		 * @returns True if succeeded. False otherwise.
+		 */
+		bool MoveNext();
+	private:
+		// Process info structure for current process.
+		PROCESSENTRY32 CurrentEntry;
+
+		// Processes state snapshot handle.
+		HANDLE SnapshotHandle;
+	};
+
+	/**
+	 * Process enumeration info structure.
+	 */
+	struct FProcEnumInfo
+	{
+		friend FProcEnumInfo FProcEnumerator::GetCurrent() const;
+	public:
+		// Gets process PID.
+		uint32 GetPID() const;
+
+		// Gets parent process PID.
+		uint32 GetParentPID() const;
+
+		// Gets process name. I.e. exec name.
+		FString GetName() const;
+
+		// Gets process full image path. I.e. full path of the exec file.
+		FString GetFullPath() const;
+
+	private:
+		// Private constructor.
+		FProcEnumInfo(const PROCESSENTRY32& Info);
+
+		// Process info structure.
+		PROCESSENTRY32 Info;
+	};
+
 public:
 
 	// FGenericPlatformProcess interface
