Fixed shadowing variable error on CIS.

diff --git a/Engine/Source/Runtime/Core/Private/Mac/MacPlatformProcess.cpp b/Engine/Source/Runtime/Core/Private/Mac/MacPlatformProcess.cpp
--- a/Engine/Source/Runtime/Core/Private/Mac/MacPlatformProcess.cpp
+++ b/Engine/Source/Runtime/Core/Private/Mac/MacPlatformProcess.cpp
@@ -1015,8 +1015,8 @@ FMacPlatformProcess::FProcEnumerator::~FProcEnumerator()
 	}
 }
 
-FMacPlatformProcess::FProcEnumInfo::FProcEnumInfo(struct kinfo_proc ProcInfo)
-	: ProcInfo(ProcInfo)
+FMacPlatformProcess::FProcEnumInfo::FProcEnumInfo(struct kinfo_proc InProcInfo)
+	: ProcInfo(InProcInfo)
 {
 
 }
diff --git a/Engine/Source/Runtime/Core/Public/Mac/MacPlatformProcess.h b/Engine/Source/Runtime/Core/Public/Mac/MacPlatformProcess.h
--- a/Engine/Source/Runtime/Core/Public/Mac/MacPlatformProcess.h
+++ b/Engine/Source/Runtime/Core/Public/Mac/MacPlatformProcess.h
@@ -83,7 +83,7 @@ struct CORE_API FMacPlatformProcess : public FGenericPlatformProcess
 
 	private:
 		// Private constructor.
-		FProcEnumInfo(struct kinfo_proc ProcInfo);
+		FProcEnumInfo(struct kinfo_proc InProcInfo);
 
 		// Process info struct.
 		struct kinfo_proc ProcInfo;
