Fixed shadowing variable errors on Linux and Windows.

diff --git a/Engine/Source/Runtime/Core/Private/Linux/LinuxPlatformProcess.cpp b/Engine/Source/Runtime/Core/Private/Linux/LinuxPlatformProcess.cpp
--- a/Engine/Source/Runtime/Core/Private/Linux/LinuxPlatformProcess.cpp
+++ b/Engine/Source/Runtime/Core/Private/Linux/LinuxPlatformProcess.cpp
@@ -1064,8 +1064,8 @@ FLinuxPlatformProcess::FProcEnumerator::FProcEnumerator()
 	class FPIDsCollector : public IPlatformFile::FDirectoryVisitor
 	{
 	public:
-		FPIDsCollector(TArray<uint32>& PIDs)
-			: PIDs(PIDs)
+		FPIDsCollector(TArray<uint32>& InPIDsToCollect)
+			: PIDsToCollect(InPIDsToCollect)
 		{ }
 
 		bool Visit(const TCHAR* FilenameOrDirectory, bool bIsDirectory)
@@ -1080,7 +1080,7 @@ FLinuxPlatformProcess::FProcEnumerator::FProcEnumerator()
 			return true;
 		}
 
-		TArray<uint32>& PIDs;
+		TArray<uint32>& PIDsToCollect;
 	} PIDsCollector(PIDs);
 
 	IPlatformFile::GetPlatformPhysical().IterateDirectory(TEXT("/proc"), PIDsCollector);
@@ -1113,8 +1113,8 @@ bool FLinuxPlatformProcess::FProcEnumerator::MoveNext()
 	return true;
 }
 
-FLinuxPlatformProcess::FProcEnumInfo::FProcEnumInfo(uint32 PID)
-	: PID(PID)
+FLinuxPlatformProcess::FProcEnumInfo::FProcEnumInfo(uint32 InPID)
+	: PID(InPID)
 {
 
 }
diff --git a/Engine/Source/Runtime/Core/Private/Windows/WindowsPlatformProcess.cpp b/Engine/Source/Runtime/Core/Private/Windows/WindowsPlatformProcess.cpp
--- a/Engine/Source/Runtime/Core/Private/Windows/WindowsPlatformProcess.cpp
+++ b/Engine/Source/Runtime/Core/Private/Windows/WindowsPlatformProcess.cpp
@@ -1257,8 +1257,8 @@ FWindowsPlatformProcess::FProcEnumerator::~FProcEnumerator()
 	::CloseHandle(SnapshotHandle);
 }
 
-FWindowsPlatformProcess::FProcEnumInfo::FProcEnumInfo(const PROCESSENTRY32& Info)
-	: Info(Info)
+FWindowsPlatformProcess::FProcEnumInfo::FProcEnumInfo(const PROCESSENTRY32& InInfo)
+	: Info(InInfo)
 {
 
 }
diff --git a/Engine/Source/Runtime/Core/Public/Linux/LinuxPlatformProcess.h b/Engine/Source/Runtime/Core/Public/Linux/LinuxPlatformProcess.h
--- a/Engine/Source/Runtime/Core/Public/Linux/LinuxPlatformProcess.h
+++ b/Engine/Source/Runtime/Core/Public/Linux/LinuxPlatformProcess.h
@@ -234,7 +234,7 @@ struct CORE_API FLinuxPlatformProcess : public FGenericPlatformProcess
 
 	private:
 		// Private constructor.
-		FProcEnumInfo(uint32 PID);
+		FProcEnumInfo(uint32 InPID);
 
 		// Current process PID.
 		uint32 PID;
diff --git a/Engine/Source/Runtime/Core/Public/Windows/WindowsPlatformProcess.h b/Engine/Source/Runtime/Core/Public/Windows/WindowsPlatformProcess.h
--- a/Engine/Source/Runtime/Core/Public/Windows/WindowsPlatformProcess.h
+++ b/Engine/Source/Runtime/Core/Public/Windows/WindowsPlatformProcess.h
@@ -109,7 +109,7 @@ struct CORE_API FWindowsPlatformProcess
 
 	private:
 		// Private constructor.
-		FProcEnumInfo(const PROCESSENTRY32& Info);
+		FProcEnumInfo(const PROCESSENTRY32& InInfo);
 
 		// Process info structure.
 		PROCESSENTRY32 Info;
