Another CIS fix.

diff --git a/Engine/Source/Runtime/Core/Private/Linux/LinuxPlatformProcess.cpp b/Engine/Source/Runtime/Core/Private/Linux/LinuxPlatformProcess.cpp
--- a/Engine/Source/Runtime/Core/Private/Linux/LinuxPlatformProcess.cpp
+++ b/Engine/Source/Runtime/Core/Private/Linux/LinuxPlatformProcess.cpp
@@ -1074,18 +1074,19 @@ FLinuxPlatformProcess::FProcEnumerator::FProcEnumerator()
 		  
 			if (bIsDirectory && FCString::IsNumeric(*StrPID))
 			{
-				PIDs.Add(FCString::Atoi(*StrPID));
+				PIDsToCollect.Add(FCString::Atoi(*StrPID));
 			}
 			
 			return true;
 		}
 
+	private:
 		TArray<uint32>& PIDsToCollect;
 	} PIDsCollector(PIDs);
 
 	IPlatformFile::GetPlatformPhysical().IterateDirectory(TEXT("/proc"), PIDsCollector);
 	
-	for(auto PID : PIDsCollector.PIDs)
+	for (auto PID : PIDs)
 	{
 		Data->Processes.Add(FProcEnumInfo(PID));
 	}
