CrashReport - Do not call WerReportCloseHandle as it is causing issues with ensures

diff --git a/Engine/Source/Runtime/Core/Private/Windows/WindowsPlatformCrashContext.cpp b/Engine/Source/Runtime/Core/Private/Windows/WindowsPlatformCrashContext.cpp
--- a/Engine/Source/Runtime/Core/Private/Windows/WindowsPlatformCrashContext.cpp
+++ b/Engine/Source/Runtime/Core/Private/Windows/WindowsPlatformCrashContext.cpp
@@ -314,7 +314,8 @@ int32 ReportCrashUsingCrashReportClient(EXCEPTION_POINTERS* ExceptionInfo, const
 			WerReportSubmit( ReportHandle, WerConsentAlwaysPrompt, WER_SUBMIT_QUEUE | WER_SUBMIT_BYPASS_DATA_THROTTLING, &SubmitResult );
 
 			// Cleanup
-			WerReportCloseHandle( ReportHandle );
+			// This method sometimes calls WndProc during an ensure causing an assert.
+			//WerReportCloseHandle( ReportHandle );
 		}
 
 		// Build machines do not upload these automatically since it is not okay to have lingering processes after the build completes.
