No longer using O_APPEND flag when opening files for write on Linux and Mac.
Allows opening of existing files to seek and write.

diff --git a/Engine/Source/Runtime/Core/Private/Apple/ApplePlatformFile.cpp b/Engine/Source/Runtime/Core/Private/Apple/ApplePlatformFile.cpp
--- a/Engine/Source/Runtime/Core/Private/Apple/ApplePlatformFile.cpp
+++ b/Engine/Source/Runtime/Core/Private/Apple/ApplePlatformFile.cpp
@@ -436,11 +436,7 @@ IFileHandle* FApplePlatformFile::OpenRead(const TCHAR* Filename, bool bAllowWrit
 IFileHandle* FApplePlatformFile::OpenWrite(const TCHAR* Filename, bool bAppend, bool bAllowRead)
 {
 	int Flags = O_CREAT;
-	if (bAppend)
-	{
-		Flags |= O_APPEND;
-	}
-	else
+	if (!bAppend)
 	{
 		Flags |= O_TRUNC;
 	}
diff --git a/Engine/Source/Runtime/Core/Private/Linux/LinuxPlatformFile.cpp b/Engine/Source/Runtime/Core/Private/Linux/LinuxPlatformFile.cpp
--- a/Engine/Source/Runtime/Core/Private/Linux/LinuxPlatformFile.cpp
+++ b/Engine/Source/Runtime/Core/Private/Linux/LinuxPlatformFile.cpp
@@ -781,10 +781,6 @@ IFileHandle* FLinuxPlatformFile::OpenRead(const TCHAR* Filename, bool bAllowWrit
 IFileHandle* FLinuxPlatformFile::OpenWrite(const TCHAR* Filename, bool bAppend, bool bAllowRead)
 {
 	int Flags = O_CREAT | O_CLOEXEC;	// prevent children from inheriting this
-	if (bAppend)
-	{
-		Flags |= O_APPEND;
-	}
 
 	if (bAllowRead)
 	{
