Adding smoke test for file access to enforce platform consistency of EFileWriter::FILEWRITE_Append.
Currently disabled to not block certain platforms until fixed.

diff --git a/Engine/Source/Runtime/Core/Private/Tests/Misc/FileTest.cpp b/Engine/Source/Runtime/Core/Private/Tests/Misc/FileTest.cpp
new file mode 100644
--- /dev/null
+++ b/Engine/Source/Runtime/Core/Private/Tests/Misc/FileTest.cpp
@@ -0,0 +1,79 @@
+// Copyright 1998-2015 Epic Games, Inc. All Rights Reserved.
+
+#include "CorePrivatePCH.h"
+#include "AutomationTest.h"
+
+IMPLEMENT_SIMPLE_AUTOMATION_TEST(FFileTests, "System.Core.Misc.File", EAutomationTestFlags::ATF_SmokeTest)
+
+// These file tests are designed to ensure expected file writing behavior, as well as cross-platform consistency
+
+bool FFileTests::RunTest( const FString& Parameters )
+{
+	// Disabled for now, pending changes to make all platforms consistent.
+	return true;
+
+	const FString TempFilename = FPaths::EngineSavedDir() / FGuid::NewGuid().ToString();
+	TArray<uint8> TestData;
+	TArray<uint8> ReadData;
+	uint8 One = 1;
+	FArchive* TestFile = nullptr;
+
+	// We will test the platform abstraction class, IFileManager
+	IFileManager& FileManager = IFileManager::Get();
+
+	// Ensure always starting off without test file existing
+	FileManager.Delete(*TempFilename);
+	check(FileManager.FileExists(*TempFilename) == false);
+
+	// Check a new file can be created
+	TestData.AddZeroed(64);
+	TestFile = FileManager.CreateFileWriter(*TempFilename, 0);
+	check(TestFile != nullptr);
+	TestFile->Serialize(TestData.GetData(), TestData.Num());
+	delete TestFile;
+
+	// Confirm same data
+	check(FFileHelper::LoadFileToArray(ReadData, *TempFilename));
+ 	if(ReadData != TestData)
+	{
+		return false;
+	}
+
+	// Using append flag should open the file, and writing data immediatly should append to the end.
+	// We should also be capable of seeking writing.
+	TestData.Add(One);
+	TestData[10] = One;
+	TestFile = FileManager.CreateFileWriter(*TempFilename, EFileWrite::FILEWRITE_Append);
+	check(TestFile != nullptr);
+	TestFile->Serialize(&One, 1);
+	TestFile->Seek(10);
+	TestFile->Serialize(&One, 1);
+	delete TestFile;
+
+	// Confirm same data
+	check(FFileHelper::LoadFileToArray(ReadData, *TempFilename));
+	if(ReadData != TestData)
+	{
+		return false;
+	}
+
+	// No flags should clobber existing file
+	TestData.Empty();
+	TestData.Add(One);
+	TestFile = FileManager.CreateFileWriter(*TempFilename, 0);
+	check(TestFile != nullptr);
+	TestFile->Serialize(&One, 1);
+	delete TestFile;
+
+	// Confirm same data
+	check(FFileHelper::LoadFileToArray(ReadData, *TempFilename));
+	if(ReadData != TestData)
+	{
+		return false;
+	}
+
+	// Delete temp file
+	FileManager.Delete(*TempFilename);
+
+	return true;
+}
