Stats - Added stats for triggering the event

diff --git a/Engine/Source/Runtime/Core/Public/Async/AsyncWork.h b/Engine/Source/Runtime/Core/Public/Async/AsyncWork.h
--- a/Engine/Source/Runtime/Core/Public/Async/AsyncWork.h
+++ b/Engine/Source/Runtime/Core/Public/Async/AsyncWork.h
@@ -313,6 +313,8 @@ class FAsyncTask
 		check(QueuedPool);
 		if (DoneEvent)
 		{
+			FScopeCycleCounter Scope( Task.GetStatId(), true );
+			DECLARE_SCOPE_CYCLE_COUNTER( TEXT( "FAsyncTask::FinishThreadedWork" ), STAT_FAsyncTask_FinishThreadedWork, STATGROUP_ThreadPoolAsyncTasks );		
 			DoneEvent->Trigger();
 		}
 	}
