OpenWrite doesn't use O_APPEND but still seeks to end of file if append requested on Android
#ue4
#android

diff --git a/Engine/Source/Runtime/Core/Private/Android/AndroidFile.cpp b/Engine/Source/Runtime/Core/Private/Android/AndroidFile.cpp
--- a/Engine/Source/Runtime/Core/Private/Android/AndroidFile.cpp
+++ b/Engine/Source/Runtime/Core/Private/Android/AndroidFile.cpp
@@ -1276,11 +1276,7 @@ public:
 		PathToAndroidPaths(LocalPath, AssetPath, Filename, AllowLocal);
 
 		int Flags = O_CREAT;
-		if (bAppend)
-		{
-			Flags |= O_APPEND;
-		}
-		else
+		if (!bAppend)
 		{
 			Flags |= O_TRUNC;
 		}
