Fix to TBitArray::GetAllocatedSize and TBitArray::CountBytes.
All such calculations changed to be consistent.

diff --git a/Engine/Source/Runtime/Core/Public/Containers/BitArray.h b/Engine/Source/Runtime/Core/Public/Containers/BitArray.h
--- a/Engine/Source/Runtime/Core/Public/Containers/BitArray.h
+++ b/Engine/Source/Runtime/Core/Public/Containers/BitArray.h
@@ -227,7 +227,7 @@ public:
 		NumBits = MaxBits = Copy.NumBits;
 		if(NumBits)
 		{
-			const int32 NumDWORDs = (MaxBits + NumBitsPerDWORD - 1) / NumBitsPerDWORD;
+			const int32 NumDWORDs = FMath::DivideAndRoundUp(MaxBits, NumBitsPerDWORD);
 			Realloc(0);
 			FMemory::Memcpy(GetData(),Copy.GetData(),NumDWORDs * sizeof(uint32));
 		}
@@ -271,7 +271,7 @@ public:
 		}
 
 		// calc the number of dwords for all the bits
-		const int32 NumDWORDs = (BitArray.NumBits + NumBitsPerDWORD - 1) / NumBitsPerDWORD; 
+		const int32 NumDWORDs = FMath::DivideAndRoundUp(BitArray.NumBits, NumBitsPerDWORD);
 
 		// serialize the data as one big chunk
 		Ar.Serialize(BitArray.GetData(), NumDWORDs * sizeof(uint32));
@@ -294,8 +294,8 @@ public:
 		{
 			// Allocate memory for the new bits.
 			const uint32 MaxDWORDs = AllocatorInstance.CalculateSlack(
-				(NumBits + NumBitsPerDWORD - 1) / NumBitsPerDWORD,
-				(MaxBits + NumBitsPerDWORD - 1) / NumBitsPerDWORD,
+				FMath::DivideAndRoundUp(NumBits, NumBitsPerDWORD),
+				FMath::DivideAndRoundUp(MaxBits, NumBitsPerDWORD),
 				sizeof(uint32)
 				);
 			MaxBits = MaxDWORDs * NumBitsPerDWORD;
@@ -342,7 +342,7 @@ public:
 		if(InNumBits)
 		{
 			NumBits = InNumBits;
-			FMemory::Memset(GetData(),Value ? 0xff : 0,(NumBits + NumBitsPerDWORD - 1) / NumBitsPerDWORD * sizeof(uint32));
+			FMemory::Memset(GetData(),Value ? 0xff : 0, FMath::DivideAndRoundUp(NumBits, NumBitsPerDWORD) * sizeof(uint32));
 		}
 	}
 
@@ -411,16 +411,16 @@ public:
 	 */
 	uint32 GetAllocatedSize( void ) const
 	{
-		return (MaxBits / NumBitsPerDWORD) * sizeof(uint32);
+		return FMath::DivideAndRoundUp(MaxBits, NumBitsPerDWORD) * sizeof(uint32);
 	}
 
 	/** Tracks the container's memory use through an archive. */
 	void CountBytes(FArchive& Ar)
 	{
 		Ar.CountBytes(
-			(NumBits / NumBitsPerDWORD) * sizeof(uint32),
-			(MaxBits / NumBitsPerDWORD) * sizeof(uint32)
-			);
+			FMath::DivideAndRoundUp(NumBits, NumBitsPerDWORD) * sizeof(uint32),
+			FMath::DivideAndRoundUp(MaxBits, NumBitsPerDWORD) * sizeof(uint32)
+		);
 	}
 
 	/**
@@ -431,7 +431,7 @@ public:
 	{
 		// Iterate over the array until we see a word with a zero bit.
 		uint32* RESTRICT DwordArray = GetData();
-		const int32 DwordCount = (Num() + NumBitsPerDWORD - 1) / NumBitsPerDWORD;
+		const int32 DwordCount = FMath::DivideAndRoundUp(Num(), NumBitsPerDWORD);
 		int32 DwordIndex = 0;
 		while ( DwordIndex < DwordCount && DwordArray[DwordIndex] == 0xffffffff )
 		{
@@ -644,8 +644,8 @@ private:
 
 	void Realloc(int32 PreviousNumBits)
 	{
-		const int32 PreviousNumDWORDs = (PreviousNumBits + NumBitsPerDWORD - 1) / NumBitsPerDWORD;
-		const int32 MaxDWORDs = (MaxBits + NumBitsPerDWORD - 1) / NumBitsPerDWORD;
+		const int32 PreviousNumDWORDs = FMath::DivideAndRoundUp(PreviousNumBits, NumBitsPerDWORD);
+		const int32 MaxDWORDs = FMath::DivideAndRoundUp(MaxBits, NumBitsPerDWORD);
 
 		AllocatorInstance.ResizeAllocation(PreviousNumDWORDs,MaxDWORDs,sizeof(uint32));
 
@@ -936,8 +936,8 @@ public:
 		{
 			// Allocate memory for the new bits.
 			const uint32 MaxDWORDs = AllocatorInstance.CalculateSlack(
-				(NumBits + NumBitsPerDWORD - 1) / NumBitsPerDWORD,
-				(MaxBits + NumBitsPerDWORD - 1) / NumBitsPerDWORD,
+				FMath::DivideAndRoundUp(NumBits, NumBitsPerDWORD),
+				FMath::DivideAndRoundUp(MaxBits, NumBitsPerDWORD),
 				sizeof(uint32)
 				);
 			MaxBits = MaxDWORDs * NumBitsPerDWORD;
@@ -989,8 +989,8 @@ private:
 
 	void Realloc(int32 PreviousNumBits)
 	{
-		const int32 PreviousNumDWORDs = (PreviousNumBits + NumBitsPerDWORD - 1) / NumBitsPerDWORD;
-		const int32 MaxDWORDs = (MaxBits + NumBitsPerDWORD - 1) / NumBitsPerDWORD;
+		const int32 PreviousNumDWORDs = FMath::DivideAndRoundUp(PreviousNumBits, NumBitsPerDWORD);
+		const int32 MaxDWORDs = FMath::DivideAndRoundUp(MaxBits, NumBitsPerDWORD);
 
 		AllocatorInstance.ResizeAllocation(PreviousNumDWORDs, MaxDWORDs, sizeof(uint32));
 
diff --git a/Engine/Source/Runtime/Core/Public/Containers/ContainerAllocationPolicies.h b/Engine/Source/Runtime/Core/Public/Containers/ContainerAllocationPolicies.h
--- a/Engine/Source/Runtime/Core/Public/Containers/ContainerAllocationPolicies.h
+++ b/Engine/Source/Runtime/Core/Public/Containers/ContainerAllocationPolicies.h
@@ -590,9 +590,9 @@ struct TAllocatorTraits<TFixedAllocator<NumInlineElements>> : TAllocatorTraitsBa
 	enum { SupportsMove = true };
 };
 
-
-enum { NumBitsPerDWORD = 32 };
-enum { NumBitsPerDWORDLogTwo = 5 };
+// We want these to be correctly typed as int32, but we don't want them to have linkage, so we make them macros
+#define NumBitsPerDWORD ((int32)32)
+#define NumBitsPerDWORDLogTwo ((int32)5)
 
 //
 // Sparse array allocation definitions
