Core - Removed unused stats

diff --git a/Engine/Source/Runtime/Core/Private/Async/TaskGraph.cpp b/Engine/Source/Runtime/Core/Private/Async/TaskGraph.cpp
--- a/Engine/Source/Runtime/Core/Private/Async/TaskGraph.cpp
+++ b/Engine/Source/Runtime/Core/Private/Async/TaskGraph.cpp
@@ -8,7 +8,6 @@ DEFINE_LOG_CATEGORY_STATIC(LogTaskGraph, Log, All);
 
 DEFINE_STAT(STAT_FReturnGraphTask);
 DEFINE_STAT(STAT_FTriggerEventGraphTask);
-DEFINE_STAT(STAT_UnknownGraphTask);
 DEFINE_STAT(STAT_ParallelFor);
 DEFINE_STAT(STAT_ParallelForTask);
 
diff --git a/Engine/Source/Runtime/Core/Public/Async/TaskGraphInterfaces.h b/Engine/Source/Runtime/Core/Public/Async/TaskGraphInterfaces.h
--- a/Engine/Source/Runtime/Core/Public/Async/TaskGraphInterfaces.h
+++ b/Engine/Source/Runtime/Core/Public/Async/TaskGraphInterfaces.h
@@ -22,7 +22,6 @@ DECLARE_STATS_GROUP(TEXT("Task Graph Tasks"), STATGROUP_TaskGraphTasks, STATCAT_
 
 DECLARE_CYCLE_STAT_EXTERN(TEXT("FReturnGraphTask"), STAT_FReturnGraphTask, STATGROUP_TaskGraphTasks, CORE_API);
 DECLARE_CYCLE_STAT_EXTERN(TEXT("FTriggerEventGraphTask"), STAT_FTriggerEventGraphTask, STATGROUP_TaskGraphTasks, CORE_API);
-DECLARE_CYCLE_STAT_EXTERN(TEXT("Unknown Graph Task"), STAT_UnknownGraphTask, STATGROUP_TaskGraphTasks, CORE_API);
 DECLARE_CYCLE_STAT_EXTERN(TEXT("ParallelFor"), STAT_ParallelFor, STATGROUP_TaskGraphTasks, CORE_API);
 DECLARE_CYCLE_STAT_EXTERN(TEXT("ParallelForTask"), STAT_ParallelForTask, STATGROUP_TaskGraphTasks, CORE_API);
 
