Seek to end of file when opening with append flag set.

diff --git a/Engine/Source/Runtime/Core/Private/WinRT/WinRTFile.cpp b/Engine/Source/Runtime/Core/Private/WinRT/WinRTFile.cpp
--- a/Engine/Source/Runtime/Core/Private/WinRT/WinRTFile.cpp
+++ b/Engine/Source/Runtime/Core/Private/WinRT/WinRTFile.cpp
@@ -360,7 +360,12 @@ public:
 		HANDLE Handle	 = CreateFile2(*NormalizeFilename(Filename), Access, WinFlags, Create, NULL);
 		if(Handle != INVALID_HANDLE_VALUE)
 		{
-			return new FFileHandleWinRT(Handle);
+			FFileHandleWinRT *PlatformFileHandle = new FFileHandleWinRT(Handle);
+			if (bAppend)
+			{
+				PlatformFileHandle->SeekFromEnd(0);
+			}
+			return PlatformFileHandle;
 		}
 		return NULL;
 	}
diff --git a/Engine/Source/Runtime/Core/Private/Windows/WindowsPlatformFile.cpp b/Engine/Source/Runtime/Core/Private/Windows/WindowsPlatformFile.cpp
--- a/Engine/Source/Runtime/Core/Private/Windows/WindowsPlatformFile.cpp
+++ b/Engine/Source/Runtime/Core/Private/Windows/WindowsPlatformFile.cpp
@@ -609,7 +609,12 @@ public:
 		HANDLE Handle    = CreateFileW(*NormalizeFilename(Filename), Access, WinFlags, NULL, Create, FILE_ATTRIBUTE_NORMAL, NULL);
 		if(Handle != INVALID_HANDLE_VALUE)
 		{
-			return new FFileHandleWindows(Handle);
+			FFileHandleWindows *PlatformFileHandle = new FFileHandleWindows(Handle);
+			if (bAppend)
+			{
+				PlatformFileHandle->SeekFromEnd(0);
+			}
+			return PlatformFileHandle;
 		}
 		return NULL;
 	}
