Async: Added method for querying readiness of futures.

diff --git a/Engine/Source/Runtime/Core/Public/Async/Future.h b/Engine/Source/Runtime/Core/Public/Async/Future.h
--- a/Engine/Source/Runtime/Core/Public/Async/Future.h
+++ b/Engine/Source/Runtime/Core/Public/Async/Future.h
@@ -143,10 +143,22 @@ class TFutureBase
 {
 public:
 
+	/**
+	 * Checks whether this future object has its value set.
+	 *
+	 * @return true if this future has a shared state and the value has been set, false otherwise.
+	 * @see IsValid
+	 */
+	bool IsReady() const
+	{
+		return State->IsComplete();
+	}
+
 	/**
 	 * Checks whether this future object has a valid state.
 	 *
 	 * @return true if the state is valid, false otherwise.
+	 * @see IsReady
 	 */
 	bool IsValid() const
 	{
