Async: Code cleanup pass

diff --git a/Engine/Source/Runtime/Core/Public/Async/Async.h b/Engine/Source/Runtime/Core/Public/Async/Async.h
--- a/Engine/Source/Runtime/Core/Public/Async/Async.h
+++ b/Engine/Source/Runtime/Core/Public/Async/Async.h
@@ -46,7 +46,7 @@ inline void SetPromise(TPromise<void>& Promise, TFunction<void()> Function)
 /**
  * Base class for asynchronous functions that are executed in the Task Graph system.
  */
-class FAsyncTaskBase
+class FAsyncGraphTaskBase
 {
 public:
 
@@ -86,8 +86,8 @@ public:
  * Template for asynchronous functions that are executed in the Task Graph system.
  */
 template<typename ResultType>
-class TAsyncTask
-	: public FAsyncTaskBase
+class TAsyncGraphTask
+	: public FAsyncGraphTaskBase
 {
 public:
 
@@ -97,7 +97,7 @@ public:
 	 * @param InFunction The function to execute asynchronously.
 	 * @param InPromise The promise object used to return the function's result.
 	 */
-	TAsyncTask(TFunction<ResultType()>&& InFunction, TPromise<ResultType>&& InPromise)
+	TAsyncGraphTask(TFunction<ResultType()>&& InFunction, TPromise<ResultType>&& InPromise)
 		: Function(MoveTemp(InFunction))
 		, Promise(MoveTemp(InPromise))
 	{ }
@@ -224,12 +224,12 @@ private:
 /**
  * Helper struct used to generate unique ids for the stats.
  */
-struct FAsyncIndex
+struct FAsyncThreadIndex
 {
 	CORE_API static int32 GetNext()
 	{
-		static FThreadSafeCounter TAsyncThreadIndex;
-		return TAsyncThreadIndex.Add(1);
+		static FThreadSafeCounter ThreadIndex;
+		return ThreadIndex.Add(1);
 	}
 };
 
@@ -276,13 +276,12 @@ template<typename ResultType>
 TFuture<ResultType> Async(EAsyncExecution Execution, TFunction<ResultType()> Function)
 {
 	TPromise<ResultType> Promise;
-	TFuture<ResultType> Future = Promise.GetFuture();
 
 	switch (Execution)
 	{
 	case EAsyncExecution::TaskGraph:
 		{
-			TGraphTask<TAsyncTask<ResultType>>::CreateTask().ConstructAndDispatchWhenReady(MoveTemp(Function), MoveTemp(Promise));
+			TGraphTask<TAsyncGraphTask<ResultType>>::CreateTask().ConstructAndDispatchWhenReady(MoveTemp(Function), MoveTemp(Promise));
 		}
 		break;
 	
@@ -291,7 +290,7 @@ TFuture<ResultType> Async(EAsyncExecution Execution, TFunction<ResultType()> Fun
 			TPromise<FRunnableThread*> ThreadPromise;
 			TAsyncRunnable<ResultType>* Runnable = new TAsyncRunnable<ResultType>(MoveTemp(Function), MoveTemp(Promise), ThreadPromise.GetFuture());
 			
-			const FString TAsyncThreadName = FString::Printf(TEXT("TAsync %d"), FAsyncIndex::GetNext());
+			const FString TAsyncThreadName = FString::Printf(TEXT("TAsync %d"), FAsyncThreadIndex::GetNext());
 			FRunnableThread* RunnableThread = FRunnableThread::Create(Runnable, *TAsyncThreadName);
 
 			check(RunnableThread != nullptr);
@@ -310,7 +309,7 @@ TFuture<ResultType> Async(EAsyncExecution Execution, TFunction<ResultType()> Fun
 		check(false); // not implemented yet!
 	}
 
-	return MoveTemp(Future);
+	return MoveTemp(Promise.GetFuture());
 }
 
 
