Async: Refactored async task API

diff --git a/Engine/Source/Runtime/Core/Public/Async/AsyncResult.h b/Engine/Source/Runtime/Core/Public/Async/AsyncResult.h
new file mode 100644
--- /dev/null
+++ b/Engine/Source/Runtime/Core/Public/Async/AsyncResult.h
@@ -0,0 +1,121 @@
+// Copyright 1998-2015 Epic Games, Inc. All Rights Reserved.
+
+#pragma once
+
+#include "IAsyncProgress.h"
+#include "IAsyncTask.h"
+#include "Future.h"
+
+
+/**
+ * Template for asynchronous return values.
+ *
+ * @param ResultType The return value's type.
+ */
+template<typename ResultType>
+class TAsyncResult
+{
+public:
+
+	/** Default constructor. */
+	TAsyncResult() { }
+
+	/**
+	 * Creates and initializes a new instance with the given synchronous result value.
+	 *
+	 * @param Result The result value to set.
+	 */
+	TAsyncResult(const ResultType& Result)
+	{
+		TPromise<ResultType> Promise;
+		Promise.SetValue(Result);
+		Future = Promise.GetFuture();
+	}
+
+	/**
+	 * Creates and initializes a new instance
+	 *
+	 * @param InFuture The future to set.
+	 * @param InProgress The progress object (optional).
+	 * @param InTask The task object object (optional).
+	 */
+	TAsyncResult(TFuture<ResultType>&& InFuture, const TSharedPtr<IAsyncProgress>& InProgress, const TSharedPtr<IAsyncTask>& InTask)
+		: Future(MoveTemp(InFuture))
+		, Progress(InProgress)
+		, Task(InTask)
+	{ }
+
+	/** Move constructor. */
+	TAsyncResult(TAsyncResult&& Other)
+		: Future(MoveTemp(Other.Future))
+		, Progress(Other.Progress)
+		, Task(Other.Task)
+	{
+		Other.Progress.Reset();
+		Other.Task.Reset();
+	}
+
+public:
+
+	/** Move assignment operator. */
+	TAsyncResult& operator=(TAsyncResult&& Other)
+	{
+		Future = MoveTemp(Other.Future);
+		Progress = Other.Progress;
+		Task = Other.Task;
+
+		Other.Progress.Reset();
+		Other.Task.Reset();
+	}
+
+public:
+
+	/**
+	 * Gets the future that will hold the result.
+	 *
+	 * @return The future.
+	 */
+	const TFuture<ResultType>& GetFuture() const
+	{
+		return Future;
+	}
+
+	/**
+	 * Get an object that indicates the progress of the task that is computing the result.
+	 *
+	 * @return Task progress object, or nullptr if not available.
+	 */
+	TSharedPtr<IAsyncProgress> GetProgess() const
+	{
+		return Progress;
+	}
+
+	/**
+	 * Get the asynchronous task that is computing the result.
+	 *
+	 * @return Task object, or nullptr if not available.
+	 */
+	TSharedPtr<IAsyncTask> GetTask() const
+	{
+		return Task;
+	}
+
+private:
+
+	/** Hidden copy constructor (async results cannot be copied). */
+	TAsyncResult(const TAsyncResult&);
+
+	/** Hidden copy assignment (async results cannot be copied). */
+	TAsyncResult& operator=(const TAsyncResult&);
+
+private:
+
+	/** The future. */
+	TFuture<ResultType> Future;
+
+	/** The progress object. */
+	TSharedPtr<IAsyncProgress> Progress;
+
+	/** The task object. */
+	TSharedPtr<IAsyncTask> Task;
+};
diff --git a/Engine/Source/Runtime/Core/Public/Async/IAsyncProgress.h b/Engine/Source/Runtime/Core/Public/Async/IAsyncProgress.h
new file mode 100644
--- /dev/null
+++ b/Engine/Source/Runtime/Core/Public/Async/IAsyncProgress.h
@@ -0,0 +1,43 @@
+// Copyright 1998-2015 Epic Games, Inc. All Rights Reserved.
+
+#pragma once
+
+
+/**
+ * Interface for checking the progress of asynchronous tasks.
+ */
+class IAsyncProgress
+{
+public:
+
+	/**
+	 * Gets the task's completion percentage, if available (0.0 to 1.0, or unset).
+	 *
+	 * Tasks should make a best effort to calculate a completion percentage. If the
+	 * calculation is expensive, the value should be cached and only updated in regular
+	 * intervals or whenever is appropriate. If a completion percentage cannot be
+	 * computed at all, the optional return value should be unset.
+	 *
+	 * @return Completion percentage value.
+	 * @see GetStatusText
+	 */
+	virtual TOptional<float> GetCompletion() = 0;
+
+	/**
+	 * Gets a human readable text for the task's current status.
+	 *
+	 * @return Status text.
+	 * @see GetCompletion
+	 */
+	virtual FText GetStatusText() = 0;
+
+public:
+
+	/** A delegate that is executed if the task's progress changed. */
+	virtual FSimpleDelegate& OnProgressChanged() = 0;
+
+public:
+
+	/** Virtual destructor. */
+	virtual ~IAsyncProgress() { }
+};
diff --git a/Engine/Source/Runtime/Core/Public/Async/IAsyncTask.h b/Engine/Source/Runtime/Core/Public/Async/IAsyncTask.h
new file mode 100644
--- /dev/null
+++ b/Engine/Source/Runtime/Core/Public/Async/IAsyncTask.h
@@ -0,0 +1,59 @@
+// Copyright 1998-2015 Epic Games, Inc. All Rights Reserved.
+
+#pragma once
+
+
+/**
+ * Enumerates possible states of slow running tasks.
+ */
+enum class EAsyncTaskState
+{
+	/** Task has been canceled. */
+	Cancelled,
+
+	/** Task has completed execution. */
+	Completed,
+
+	/** Task execution failed. */
+	Failed,
+
+	/** Task execution is still pending. */
+	Pending,
+
+	/** Task is executing. */
+	Running,
+};
+
+
+/**
+ * Interface for asynchronous tasks.
+ *
+ * A asynchronous task is a unit of work that runs in parallel to the caller and may take a
+ * considerable amount of time to complete, i.e. several seconds, minutes or even hours.
+ * This interface provides mechanisms for tracking and canceling such tasks.
+ */
+class IAsyncTask
+{
+public:
+
+	/**
+	 * Cancel this task.
+	 *
+	 * If the task is already running, it should make a best effort to abort execution
+	 * as soon as possible. This method should be implemented in such a way that it returns
+	 * immediately and does not block the caller.
+	 */
+	virtual void Cancel() = 0;
+
+	/**
+	 * Gets the current state of the task.
+	 *
+	 * @return Task state.
+	 */
+	virtual EAsyncTaskState GetTaskState() = 0;
+
+public:
+
+	/** Virtual destructor. */
+	virtual ~IAsyncTask() { }
+};
diff --git a/Engine/Source/Runtime/Core/Public/Async/ISlowTask.h b/Engine/Source/Runtime/Core/Public/Async/ISlowTask.h
deleted file mode 100644
diff --git a/Engine/Source/Runtime/Core/Public/Async/SlowResult.h b/Engine/Source/Runtime/Core/Public/Async/SlowResult.h
deleted file mode 100644
diff --git a/Engine/Source/Runtime/Portal/Services/Public/Application/IPortalAppWindow.h b/Engine/Source/Runtime/Portal/Services/Public/Application/IPortalAppWindow.h
new file mode 100644
--- /dev/null
+++ b/Engine/Source/Runtime/Portal/Services/Public/Application/IPortalAppWindow.h
@@ -0,0 +1,51 @@
+// Copyright 1998-2015 Epic Games, Inc. All Rights Reserved.
+
+#pragma once
+
+#include "AsyncResult.h"
+
+
+/**
+ * Interface for the Portal application's marketplace services.
+ */
+class IPortalAppWindow
+{
+public:
+
+	/**
+	 * Opens the Portal application and shows the Friends window.
+	 *
+	 * @return true on success, false otherwise.
+	 */
+	virtual TAsyncResult<bool> OpenFriends() = 0;
+
+	/**
+	 * Opens the Portal application and shows the Library window.
+	 *
+	 * @return true on success, false otherwise.
+	 */
+	virtual TAsyncResult<bool> OpenLibrary() = 0;
+
+	/**
+	 * Opens the Portal application and shows the Marketplace.
+	 *
+	 * @return true on success, false otherwise.
+	 */
+	virtual TAsyncResult<bool> OpenMarketplace() = 0;
+
+	/**
+	 * Opens the Portal application and shows a particular marketplace item.
+	 *
+	 * @param ItemId The identifier of the marketplace item to open.
+	 * @return true on success, false otherwise.
+	 */
+	virtual TAsyncResult<bool> OpenMarketItem(const FString& ItemId) = 0;
+
+public:
+
+	/** Virtual destructor. */
+	virtual ~IPortalAppWindow() { }
+};
+
+
+Expose_TNameOf(IPortalAppWindow)
diff --git a/Engine/Source/Runtime/Portal/Services/Public/IPortalServicesModule.h b/Engine/Source/Runtime/Portal/Services/Public/IPortalServicesModule.h
--- a/Engine/Source/Runtime/Portal/Services/Public/IPortalServicesModule.h
+++ b/Engine/Source/Runtime/Portal/Services/Public/IPortalServicesModule.h
@@ -3,9 +3,9 @@
 #pragma once
 
 #include "ModuleInterface.h"
+#include "TypeContainer.h"
 
 
-class FTypeContainer;
 class IPortalService;
 
 
@@ -23,4 +23,23 @@ public:
 	 * @return The type container.
 	 */
 	virtual FTypeContainer& GetServiceContainer() = 0;
+
+public:
+
+	/**
+	 * Gets a shared pointer to an instance of the specified service.
+	 *
+	 * @param T The type of service to get an instance for.
+	 * @return The service instance, or nullptr if not available.
+	 */
+	template<class T>
+	TSharedPtr<T> GetService()
+	{
+		return GetServiceContainer().GetInstance<T>();
+	}
+
+public:
+
+	/** Virtual destructor. */
+	virtual ~IPortalServicesModule() { }
 };
diff --git a/Engine/Source/Runtime/Portal/Services/Public/Package/IPortalPackageInstaller.h b/Engine/Source/Runtime/Portal/Services/Public/Package/IPortalPackageInstaller.h
--- a/Engine/Source/Runtime/Portal/Services/Public/Package/IPortalPackageInstaller.h
+++ b/Engine/Source/Runtime/Portal/Services/Public/Package/IPortalPackageInstaller.h
@@ -2,23 +2,7 @@
 
 #pragma once
 
-#include "SlowResult.h"
-#include "TypeContainer.h"
-
-
-/**
- * Structure for package identifiers.
- *
- * Package identifiers are used to uniquely identify a package.
- */
-struct FPortalPackageId
-{
-	/** The application name. */
-	FString AppName;
-
-	/** The application build label. */
-	FString BuildLabel;
-};
+#include "AsyncResult.h"
 
 
 /**
@@ -31,23 +15,29 @@ public:
 	/**
 	 * Install the specified package using the given request object.
 	 *
-	 * @param Path The path at which the package should be installed.
-	 * @param Id The identifier of the package to install.
+	 * @param AppName The name of the application to install.
+	 * @param BuildLabel The application's build label.
 	 * @return The result of the task.
 	 * @see UninstallPackage
 	 */
-	virtual TSlowResult<bool> InstallPackage(const FString& Path, const FPortalPackageId& Id) = 0;
+	virtual TAsyncResult<bool> Install(const FString& Path, const FString& AppName, const FString& BuildLabel) = 0;
 
 	/**
 	 * Attempts to uninstall the specified package using the given request object.
 	 *
 	 * @param Path The path at which the package is installed.
-	 * @param Id The identifier of the package to uninstall.
+	 * @param AppName The name of the application to uninstall.
+	 * @param BuildLabel The application's build label.
 	 * @param RemoveUserFiles Whether user created files should be removed as well.
 	 * @return The result of the task.
 	 * @see InstallPackage
 	 */
-	virtual TSlowResult<bool> UninstallPackage(const FString& Path, const FString& Id, bool RemoveUserFiles) = 0;
+	virtual TAsyncResult<bool> Uninstall(const FString& Path, const FString& AppName, const FString& BuildLabel, bool RemoveUserFiles) = 0;
+
+public:
+
+	/** Virtual destructor. */
+	virtual ~IPortalPackageInstaller() { }
 };
 
 
