CrashReport - Removed obsolete comment

diff --git a/Engine/Source/Runtime/Core/Private/Windows/WindowsPlatformCrashContext.cpp b/Engine/Source/Runtime/Core/Private/Windows/WindowsPlatformCrashContext.cpp
--- a/Engine/Source/Runtime/Core/Private/Windows/WindowsPlatformCrashContext.cpp
+++ b/Engine/Source/Runtime/Core/Private/Windows/WindowsPlatformCrashContext.cpp
@@ -466,7 +466,6 @@ int32 ReportCrash( LPEXCEPTION_POINTERS ExceptionInfo )
 
 	FCoreDelegates::OnHandleSystemError.Broadcast();
 
-	// Note: ReportCrashUsingCrashReportClient reads the callstack from GErrorHist
 	const SIZE_T StackTraceSize = 65535;
 	ANSICHAR* StackTrace = (ANSICHAR*) GMalloc->Malloc( StackTraceSize );
 	StackTrace[0] = 0;
