Merging Mac right-lick emulation

diff --git a/Engine/Source/Runtime/Core/Private/Mac/MacApplication.cpp b/Engine/Source/Runtime/Core/Private/Mac/MacApplication.cpp
--- a/Engine/Source/Runtime/Core/Private/Mac/MacApplication.cpp
+++ b/Engine/Source/Runtime/Core/Private/Mac/MacApplication.cpp
@@ -46,6 +46,7 @@ FMacApplication::FMacApplication()
 ,	bSystemModalMode(false)
 ,	ModifierKeysFlags(0)
 ,	CurrentModifierFlags(0)
+,	bEmulatingRightClick(false)
 ,	bIsWorkspaceSessionActive(true)
 {
 	TextInputMethodSystem = MakeShareable(new FMacTextInputMethodSystem);
@@ -329,6 +330,18 @@ void FMacApplication::DeferEvent(NSObject* Object)
 			case NSOtherMouseUp:
 				DeferredEvent.ButtonNumber = [Event buttonNumber];
 				DeferredEvent.ClickCount = [Event clickCount];
+				if (DeferredEvent.Type == NSLeftMouseDown && (DeferredEvent.ModifierFlags & NSControlKeyMask))
+				{
+					bEmulatingRightClick = true;
+					DeferredEvent.Type = NSRightMouseDown;
+					DeferredEvent.ButtonNumber = 2;
+				}
+				else if (DeferredEvent.Type == NSLeftMouseUp && bEmulatingRightClick)
+				{
+					bEmulatingRightClick = false;
+					DeferredEvent.Type = NSRightMouseUp;
+					DeferredEvent.ButtonNumber = 2;
+				}
 				break;
 
 			case NSScrollWheel:
diff --git a/Engine/Source/Runtime/Core/Public/Mac/MacApplication.h b/Engine/Source/Runtime/Core/Public/Mac/MacApplication.h
--- a/Engine/Source/Runtime/Core/Public/Mac/MacApplication.h
+++ b/Engine/Source/Runtime/Core/Public/Mac/MacApplication.h
@@ -276,6 +276,8 @@ private:
 	/** The current set of Cocoa modifier flags, used to detect when Mission Control has been invoked & returned so that we can synthesis the modifier events it steals */
 	NSUInteger CurrentModifierFlags;
 
+	bool bEmulatingRightClick;
+
 	TArray<FCocoaWindow*> WindowsToClose;
 
 	TArray<FCocoaWindow*> WindowsRequiringTextInvalidation;
