Async: Reverted a small change that broke Async()

diff --git a/Engine/Source/Runtime/Core/Public/Async/Async.h b/Engine/Source/Runtime/Core/Public/Async/Async.h
--- a/Engine/Source/Runtime/Core/Public/Async/Async.h
+++ b/Engine/Source/Runtime/Core/Public/Async/Async.h
@@ -276,6 +276,7 @@ template<typename ResultType>
 TFuture<ResultType> Async(EAsyncExecution Execution, TFunction<ResultType()> Function)
 {
 	TPromise<ResultType> Promise;
+	TFuture<ResultType> Future = Promise.GetFuture();
 
 	switch (Execution)
 	{
@@ -309,7 +310,7 @@ TFuture<ResultType> Async(EAsyncExecution Execution, TFunction<ResultType()> Fun
 		check(false); // not implemented yet!
 	}
 
-	return MoveTemp(Promise.GetFuture());
+	return MoveTemp(Future);
 }
 
 
