help command improvements, added filtering fro Exec commands

diff --git a/Engine/Source/Runtime/Core/Private/Misc/Parse.cpp b/Engine/Source/Runtime/Core/Private/Misc/Parse.cpp
--- a/Engine/Source/Runtime/Core/Private/Misc/Parse.cpp
+++ b/Engine/Source/Runtime/Core/Private/Misc/Parse.cpp
@@ -54,7 +54,7 @@ class FConsoleVariableDumpVisitor
 {
 public:
 	// @param Name must not be 0
-	// @param CVar must not be 0
+	// @param " must not be 0
 	static void OnConsoleVariable(const TCHAR *Name, IConsoleObject* CVar,TSet<FString>& Sink)
 	{
 		if(CVar->TestFlags(ECVF_Unregistered))
@@ -132,9 +132,9 @@ void ConsoleCommandLibrary_DumpLibraryHTML(UWorld* InWorld, FExec& SubSystem, co
 			// headline
 			LazyPrintf.PushParam(TEXT("Unreal Engine 4 Console Variables and Commands"));
 			// generated by
-			LazyPrintf.PushParam(TEXT("Unreal Engine 4 'DumpConsoleCommands'"));
+			LazyPrintf.PushParam(TEXT("Unreal Engine 4 console command 'Help'"));
 			// version
-			LazyPrintf.PushParam(TEXT("0.94"));
+			LazyPrintf.PushParam(TEXT("0.95"));
 			// date
 			LazyPrintf.PushParam(*FDateTime::Now().ToString());
 
@@ -154,8 +154,19 @@ void ConsoleCommandLibrary_DumpLibraryHTML(UWorld* InWorld, FExec& SubSystem, co
 
 					Help = Help.ReplaceCharWithEscapedChar();
 
+					const TCHAR* ElementType = TEXT("Unknown");
+
+					if(Element->AsVariable())
+					{
+						ElementType = TEXT("Var"); 
+					}
+					else if(Element->AsCommand())
+					{
+						ElementType = TEXT("Cmd"); 
+					}
+
 					//{name: "r.SetRes", help:"To change the screen/window resolution."},
-					FString DataLine = FString::Printf(TEXT("{name: \"%s\", help:\"%s\"},\r\n"), *Name, *Help);
+					FString DataLine = FString::Printf(TEXT("{name: \"%s\", help:\"%s\", type:\"%s\"},\r\n"), *Name, *Help, ElementType);
 
 					AllData += DataLine;
 				}
@@ -163,7 +174,7 @@ void ConsoleCommandLibrary_DumpLibraryHTML(UWorld* InWorld, FExec& SubSystem, co
 				{
 					// Exec command (better we change them to use the new method as it has better help and is more convenient to use)
 					//{name: "", help:"To change the screen/window resolution."},
-					FString DataLine = FString::Printf(TEXT("{name: \"%s\", help:\"Sorry: Exec commands have no help\"},\r\n"), *Name);
+					FString DataLine = FString::Printf(TEXT("{name: \"%s\", help:\"Sorry: Exec commands have no help\", type:\"Exec\"},\r\n"), *Name);
 
 					AllData += DataLine;
 				}
