Merging using UE4-To-UE4-LauncherDev

#Desc:
#OPP-3468 Mac: Main window can stop appearing

#proj Runtime.Core

#branch LauncherDev

#change Mac window now correctly returns false from GetRestoredDimensions if the window is invalid.

diff --git a/Engine/Source/Runtime/Core/Private/Mac/MacWindow.cpp b/Engine/Source/Runtime/Core/Private/Mac/MacWindow.cpp
--- a/Engine/Source/Runtime/Core/Private/Mac/MacWindow.cpp
+++ b/Engine/Source/Runtime/Core/Private/Mac/MacWindow.cpp
@@ -440,17 +440,24 @@ bool FMacWindow::IsVisible() const
 
 bool FMacWindow::GetRestoredDimensions(int32& X, int32& Y, int32& Width, int32& Height)
 {
-	SCOPED_AUTORELEASE_POOL;
-
-	NSRect Frame = [WindowHandle frame];
-
-	X = Frame.origin.x;
-	Y = FPlatformMisc::ConvertSlateYPositionToCocoa(Frame.origin.y) - Frame.size.height + 1;
-
-	Width = Frame.size.width;
-	Height = Frame.size.height;
-
-	return true;
+	if (WindowHandle)
+	{
+		SCOPED_AUTORELEASE_POOL;
+		
+		NSRect Frame = [WindowHandle frame];
+		
+		X = Frame.origin.x;
+		Y = FPlatformMisc::ConvertSlateYPositionToCocoa(Frame.origin.y) - Frame.size.height + 1;
+		
+		Width = Frame.size.width;
+		Height = Frame.size.height;
+		
+		return true;
+	}
+	else
+	{
+		return false;
+	}
 }
 
 void FMacWindow::SetWindowFocus()
