Added safety checks to ensure that out-of-order scoped tasks are destroyed correctly

This potentially addresses UE-12075

diff --git a/Engine/Source/Runtime/Core/Public/Misc/FeedbackContext.h b/Engine/Source/Runtime/Core/Public/Misc/FeedbackContext.h
--- a/Engine/Source/Runtime/Core/Public/Misc/FeedbackContext.h
+++ b/Engine/Source/Runtime/Core/Public/Misc/FeedbackContext.h
@@ -212,7 +212,16 @@ public:
 			FScopedSlowTaskStack& Stack = *Context.ScopeStack;
 			checkSlow(Stack.Num() != 0 && Stack.Last() == this);
 
-			Stack.Pop(false);
+			auto* Task = Stack.Last();
+			if (ensureMsg(Task == this, TEXT("Out-of-order scoped task construction/destruction")))
+			{
+				Stack.Pop(false);
+			}
+			else
+			{
+				Stack.RemoveSingleSwap(this, false);
+			}
+
 			if (Stack.Num() != 0)
 			{
 				// Stop anything else contributing to the parent frame
