PR #1080: [ Kismet Math Library ] New Node ~ Map Range Clamped, a Direct Implementation of GetMappedRangeValue that clamps the output value to the output range  Rama (Contributed by [REDACTED_PERSON])

diff --git a/Engine/Config/BaseEngine.ini b/Engine/Config/BaseEngine.ini
--- a/Engine/Config/BaseEngine.ini
+++ b/Engine/Config/BaseEngine.ini
@@ -1228,6 +1228,9 @@ bUseStreamingPause=false
 +TaggedPropertyRedirects=(ClassName="SoundNodeWavePlayer",OldPropertyName="SoundWave",NewPropertyName="SoundWaveAssetPtr")
 +TaggedPropertyRedirects=(ClassName="SoundNodeModPlayer",OldPropertyName="SoundMod",NewPropertyName="SoundModAssetPtr")
 
+;Rename MapRange to MapRangeUnclamped - DEPRECATED(4.9)
++K2FieldRedirects=(OldFieldName="KismetMathLibrary.MapRange",NewFieldName="KismetMathLibrary.MapRangeUnclamped")
+
 [CoreUObject.Metadata]
 ; Note: UnrealHeaderTool should be rerun after making changes to MetadataRedirects to catch any keys specified in class headers
 +MetadataRedirects=(OldKey="K2Protected", NewKey="BlueprintProtected")
diff --git a/Engine/Source/Developer/AudioFormatOpus/Private/AudioFormatOpus.cpp b/Engine/Source/Developer/AudioFormatOpus/Private/AudioFormatOpus.cpp
--- a/Engine/Source/Developer/AudioFormatOpus/Private/AudioFormatOpus.cpp
+++ b/Engine/Source/Developer/AudioFormatOpus/Private/AudioFormatOpus.cpp
@@ -475,7 +475,7 @@ public:
 		// will use it as a multiplier to decide how much smaller than the original the
 		// compressed data should be
 		int32 OriginalBitRate = QualityInfo.SampleRate * QualityInfo.NumChannels * SAMPLE_SIZE * 8;
-		return (float)OriginalBitRate * FMath::GetMappedRangeValue(FVector2D(1, 100), FVector2D(0.04, 0.25), QualityInfo.Quality);
+		return (float)OriginalBitRate * FMath::GetMappedRangeValueClamped(FVector2D(1, 100), FVector2D(0.04, 0.25), QualityInfo.Quality);
 	}
 
 	void SerializeHeaderData(FMemoryWriter& CompressedData, uint16 SampleRate, uint32 TrueSampleCount, uint8 NumChannels, uint16 NumFrames) const
diff --git a/Engine/Source/Runtime/Core/Public/Math/UnrealMathUtility.h b/Engine/Source/Runtime/Core/Public/Math/UnrealMathUtility.h
--- a/Engine/Source/Runtime/Core/Public/Math/UnrealMathUtility.h
+++ b/Engine/Source/Runtime/Core/Public/Math/UnrealMathUtility.h
@@ -544,18 +544,25 @@ struct FMath : public FPlatformMath
 	/** Basically a Vector2d version of Lerp. */
 	static float GetRangeValue(FVector2D const& Range, float Pct);
 
-	/**
-	 * For the given value in the input range, returns the corresponding value in the output range.
-	 * Useful for mapping one value range to another value range.  Output is clamped to the OutputRange.
-	 * e.g. given that velocities [50..100] correspond to a sound volume of [0.2..1.4], this makes it easy to 
-	 *      find the volume for a velocity of 77.
-	 */
-	static FORCEINLINE float GetMappedRangeValue(FVector2D const& InputRange, FVector2D const& OutputRange, float Value)
+	DEPRECATED(4.9, "GetMappedRangeValue is deprecated. Use GetMappedRangeValueClamped instead.")
+	static FORCEINLINE float GetMappedRangeValue(const FVector2D& InputRange, const FVector2D& OutputRange, const float Value)
+	{
+		return GetMappedRangeValueClamped(InputRange, OutputRange, Value);
+	}
+
+	/** For the given Value clamped to the Input Range, returns the corresponding value in the Output Range. */
+	static FORCEINLINE float GetMappedRangeValueClamped(const FVector2D& InputRange, const FVector2D& OutputRange, const float Value)
 	{
-		float const ClampedPct = Clamp<float>(GetRangePct(InputRange, Value), 0.f, 1.f);
+		const float ClampedPct = Clamp<float>(GetRangePct(InputRange, Value), 0.f, 1.f);
 		return GetRangeValue(OutputRange, ClampedPct);
 	}
 
+	/** Transform the given Value relative to the input range to the Output Range. */
+	static FORCEINLINE float GetMappedRangeValueUnclamped(const FVector2D& InputRange, const FVector2D& OutputRange, const float Value)
+	{
+		return GetRangeValue(OutputRange, GetRangePct(InputRange, Value));
+	}
+
 	/** Performs a linear interpolation between two values, Alpha ranges from 0-1 */
 	template< class T, class U > 
 	static FORCEINLINE_DEBUGGABLE T Lerp( const T& A, const T& B, const U& Alpha )
diff --git a/Engine/Source/Runtime/Engine/Classes/Kismet/KismetMathLibrary.h b/Engine/Source/Runtime/Engine/Classes/Kismet/KismetMathLibrary.h
--- a/Engine/Source/Runtime/Engine/Classes/Kismet/KismetMathLibrary.h
+++ b/Engine/Source/Runtime/Engine/Classes/Kismet/KismetMathLibrary.h
@@ -504,9 +504,13 @@ class ENGINE_API UKismetMathLibrary : public UBlueprintFunctionLibrary
 	static float NormalizeToRange(float Value, float RangeMin, float RangeMax);
 
 	/** Returns Value mapped from one range into another.  (e.g. 20 normalized from the range 10->50 to 20->40 would result in 25) */
-	UFUNCTION(BlueprintPure, Category="Math|Float")
-	static float MapRange(float Value, float InRangeA, float InRangeB, float OutRangeA, float OutRangeB);
+	UFUNCTION(BlueprintPure, Category="Math|Float", meta=(Keywords = "get mapped value"))
+	static float MapRangeUnclamped(float Value, float InRangeA, float InRangeB, float OutRangeA, float OutRangeB);
 
+	/** Returns Value mapped from one range into another where the Value is clamped to the Input Range.  (e.g. 0.5 normalized from the range 0->1 to 0->50 would result in 25) */
+	UFUNCTION(BlueprintPure, Category="Math|Float", meta=(Keywords = "get mapped value"))
+	static float MapRangeClamped(float Value, float InRangeA, float InRangeB, float OutRangeA, float OutRangeB);
+	
 	/** Multiplies the input value by pi. */
 	UFUNCTION(BlueprintPure, meta=(Keywords = "* multiply"), Category="Math|Float")
 	static float MultiplyByPi(float Value);
diff --git a/Engine/Source/Runtime/Engine/Private/KismetMathLibrary.cpp b/Engine/Source/Runtime/Engine/Private/KismetMathLibrary.cpp
--- a/Engine/Source/Runtime/Engine/Private/KismetMathLibrary.cpp
+++ b/Engine/Source/Runtime/Engine/Private/KismetMathLibrary.cpp
@@ -547,14 +547,14 @@ float UKismetMathLibrary::NormalizeToRange(float Value, float RangeMin, float Ra
 	return (Value - RangeMin) / (RangeMax - RangeMin);
 }
 
-float UKismetMathLibrary::MapRange(float Value, float InRangeA, float InRangeB, float OutRangeA, float OutRangeB)
+float UKismetMathLibrary::MapRangeUnclamped(float Value, float InRangeA, float InRangeB, float OutRangeA, float OutRangeB)
 {
-	if (InRangeB == InRangeA)
-	{
-		return OutRangeA;
-	}
+	return FMath::GetMappedRangeValueUnclamped(FVector2D(InRangeA,InRangeB),FVector2D(OutRangeA,OutRangeB),Value);
+}
 
-	return (Value - InRangeA) * (OutRangeB - OutRangeA) / (InRangeB - InRangeA) + OutRangeA;
+float UKismetMathLibrary::MapRangeClamped(float Value, float InRangeA, float InRangeB, float OutRangeA, float OutRangeB)
+{ 
+	return FMath::GetMappedRangeValueClamped(FVector2D(InRangeA,InRangeB),FVector2D(OutRangeA,OutRangeB),Value);
 }
 
 float UKismetMathLibrary::MultiplyByPi(float Value)
