PR #1088: Fix empty callstack report on a nullptr call when CrashReportClient is not used (Contributed by [REDACTED_PERSON])

diff --git a/Engine/Source/Runtime/Core/Private/Windows/WindowsPlatformStackWalk.cpp b/Engine/Source/Runtime/Core/Private/Windows/WindowsPlatformStackWalk.cpp
--- a/Engine/Source/Runtime/Core/Private/Windows/WindowsPlatformStackWalk.cpp
+++ b/Engine/Source/Runtime/Core/Private/Windows/WindowsPlatformStackWalk.cpp
@@ -106,8 +106,9 @@ static int32 CaptureStackTraceHelper( uint64 *BackTrace, uint32 MaxDepth, CONTEX
 				break;
 			}
 
-			// Stop if the frame pointer or address is NULL.
-			if( StackFrame64.AddrFrame.Offset == 0 || StackFrame64.AddrPC.Offset == 0 )
+			// Stop if the frame pointer is NULL.
+			// Note that the thread's PC 'StackFrame64.AddrPC.Offset' COULD be 0 in case something calls a nullptr.
+			if( StackFrame64.AddrFrame.Offset == 0 )
 			{
 				break;
 			}
