Cleaned up nonsense comments in TBaseMulticastDelegate::RemoveDelegateInstance(FDelegateHandle Handle) (mostly left from the original DEPRECATED_RemoveDelegateInstance and no longer apply)

diff --git a/Engine/Source/Runtime/Core/Public/Delegates/DelegateSignatureImpl_Variadics.inl b/Engine/Source/Runtime/Core/Public/Delegates/DelegateSignatureImpl_Variadics.inl
--- a/Engine/Source/Runtime/Core/Public/Delegates/DelegateSignatureImpl_Variadics.inl
+++ b/Engine/Source/Runtime/Core/Public/Delegates/DelegateSignatureImpl_Variadics.inl
@@ -1067,27 +1067,22 @@ protected:
 	{
 		const TArray<IDelegateInstance*>& InvocationList = GetInvocationList();
 
-		// NOTE: We assume that this method is never called with a nullptr object, in which case the
-		//       the following algorithm would break down (it would remove the first found instance
-		//       of a matching function binding, which is not necessarily the instance we wish to remove).
-
 		for (int32 InvocationListIndex = 0; InvocationListIndex < InvocationList.Num(); ++InvocationListIndex)
 		{
-			// this down-cast is OK! allows for managing invocation list in the base class without requiring virtual functions
-			TDelegateInstanceInterface*& DelegateInstanceRef = (TDelegateInstanceInterface*&)InvocationList[InvocationListIndex];
+			// InvocationList is const, so we const_cast to be able to null the entry
+			// TODO: This is horrible, can we get the base class to do it?
+			IDelegateInstance*& DelegateInstanceRef = const_cast<IDelegateInstance*&>(InvocationList[InvocationListIndex]);
 
-			// NOTE: We must do a deep compare here, not just compare delegate pointers, because multiple
-			//       delegate pointers can refer to the exact same object and method
 			if ((DelegateInstanceRef != nullptr) && DelegateInstanceRef->GetHandle() == Handle)
 			{
 				delete DelegateInstanceRef;
 				DelegateInstanceRef = nullptr;
 
-				break;	// no need to continue, as we never allow the same delegate to be bound twice
+				break; // each delegate binding has a unique handle, so once we find it, we can stop
 			}
 		}
 
-		const_cast<TBaseMulticastDelegate*>(this)->CompactInvocationList();
+		CompactInvocationList();
 	}
 
 	/**
