#ue4 - Fix some Vector constants not including the 4th term.

diff --git a/Engine/Source/Runtime/Core/Public/Math/UnrealMathVectorConstants.h b/Engine/Source/Runtime/Core/Public/Math/UnrealMathVectorConstants.h
--- a/Engine/Source/Runtime/Core/Public/Math/UnrealMathVectorConstants.h
+++ b/Engine/Source/Runtime/Core/Public/Math/UnrealMathVectorConstants.h
@@ -25,9 +25,9 @@ namespace GlobalVectorConstants
 	static const VectorRegister QMULTI_SIGN_MASK1 = MakeVectorRegister( 1.f, 1.f, -1.f, -1.f );
 	static const VectorRegister QMULTI_SIGN_MASK2 = MakeVectorRegister( -1.f, 1.f, 1.f, -1.f );
 
-	static const VectorRegister DEG_TO_RAD = MakeVectorRegister(PI/(180.f), PI/(180.f), PI/(180.f), 0.0f);
-	static const VectorRegister DEG_TO_RAD_HALF = MakeVectorRegister((PI/180.f)*0.5f, (PI/180.f)*0.5f, (PI/180.f)*0.5f, 0.0f);
-	static const VectorRegister RAD_TO_DEG = MakeVectorRegister((180.f)/PI, (180.f)/PI, (180.f)/PI, 0.0f);
+	static const VectorRegister DEG_TO_RAD = MakeVectorRegister(PI/(180.f), PI/(180.f), PI/(180.f), PI/(180.f));
+	static const VectorRegister DEG_TO_RAD_HALF = MakeVectorRegister((PI/180.f)*0.5f, (PI/180.f)*0.5f, (PI/180.f)*0.5f, (PI/180.f)*0.5f);
+	static const VectorRegister RAD_TO_DEG = MakeVectorRegister((180.f)/PI, (180.f)/PI, (180.f)/PI, (180.f)/PI);
 
 	/** Bitmask to AND out the XYZ components in a vector */
 	static const VectorRegister XYZMask = MakeVectorRegister((uint32)0xffffffff, (uint32)0xffffffff, (uint32)0xffffffff, (uint32)0x00000000);
