Added support for basic string-expression lexing, compiling, and evaluation.

See ExpressionParserExamples.cpp and BasicMathExpressionEvaluator.cpp for example usage.

diff --git a/Engine/Source/Runtime/Core/Private/Math/BasicMathExpressionEvaluator.cpp b/Engine/Source/Runtime/Core/Private/Math/BasicMathExpressionEvaluator.cpp
new file mode 100644
--- /dev/null
+++ b/Engine/Source/Runtime/Core/Private/Math/BasicMathExpressionEvaluator.cpp
@@ -0,0 +1,236 @@
+// Copyright 1998-2015 Epic Games, Inc. All Rights Reserved.
+
+#include "CorePrivatePCH.h"
+#include "BasicMathExpressionEvaluator.h"
+
+#define LOCTEXT_NAMESPACE "BasicMathExpressionEvaluator"
+
+namespace ExpressionParser
+{
+	const TCHAR* const FSubExpressionStart::Moniker = TEXT("(");
+	const TCHAR* const FSubExpressionEnd::Moniker = TEXT(")");
+	const TCHAR* const FPlus::Moniker = TEXT("+");
+	const TCHAR* const FMinus::Moniker = TEXT("-");
+	const TCHAR* const FStar::Moniker = TEXT("*");
+	const TCHAR* const FForwardSlash::Moniker = TEXT("/");
+	const TCHAR* const FPercent::Moniker = TEXT("%");
+	const TCHAR* const FSquareRoot::Moniker = TEXT("sqrt");
+	const TCHAR* const FPower::Moniker = TEXT("^");
+
+	TOptional<FStringToken> ParseNumber(const FTokenStream& InStream, FStringToken* Accumulate)
+	{
+		enum class EState { LeadIn, Integer, Dot, Fractional };
+
+		EState State = EState::LeadIn;
+		return InStream.ParseToken([&](TCHAR InC){
+
+			if (State == EState::LeadIn)
+			{
+				if (FChar::IsDigit(InC))
+				{
+					State = EState::Integer;
+					return EParseState::Continue;
+				}
+				else
+				{
+					// Not a number
+					return EParseState::Cancel;
+				}
+			}
+			else if (State == EState::Integer)
+			{
+				if (FChar::IsDigit(InC))
+				{
+					return EParseState::Continue;
+				}
+				else if (InC == '.')
+				{
+					State = EState::Dot;
+					return EParseState::Continue;
+				}
+			}
+			else if (State == EState::Dot)
+			{
+				if (FChar::IsDigit(InC))
+				{
+					State = EState::Fractional;
+					return EParseState::Continue;
+				}
+			}
+			else if (State == EState::Fractional)
+			{
+				if (FChar::IsDigit(InC))
+				{
+					return EParseState::Continue;
+				}
+			}
+
+			return EParseState::StopBefore;
+		}, Accumulate);
+	}
+
+	TOptional<FExpressionError> ConsumeNumber(FExpressionTokenConsumer& Consumer)
+	{
+		auto& Stream = Consumer.GetStream();
+
+		TOptional<FStringToken> Token = ParseNumber(Stream);
+		
+		if (Token.IsSet())
+		{
+			double Value = FCString::Atod(*Token.GetValue().GetString());
+			Consumer.Add(Token.GetValue(), FExpressionNode(Value));
+		}
+
+		return TOptional<FExpressionError>();
+	}
+
+}
+
+FBasicMathExpressionEvaluator::FBasicMathExpressionEvaluator()
+{
+	using namespace ExpressionParser;
+
+	TokenDefinitions.IgnoreWhitespace();
+	TokenDefinitions.DefineToken(&ConsumeSymbol<FSubExpressionStart>);
+	TokenDefinitions.DefineToken(&ConsumeSymbol<FSubExpressionEnd>);
+	TokenDefinitions.DefineToken(&ConsumeSymbol<FPlus>);
+	TokenDefinitions.DefineToken(&ConsumeSymbol<FMinus>);
+	TokenDefinitions.DefineToken(&ConsumeSymbol<FStar>);
+	TokenDefinitions.DefineToken(&ConsumeSymbol<FForwardSlash>);
+	TokenDefinitions.DefineToken(&ConsumeSymbol<FPercent>);
+	TokenDefinitions.DefineToken(&ConsumeSymbol<FSquareRoot>);
+	TokenDefinitions.DefineToken(&ConsumeSymbol<FPower>);
+	TokenDefinitions.DefineToken(&ConsumeNumber);
+
+	Grammar.DefineGrouping<FSubExpressionStart, FSubExpressionEnd>();
+	Grammar.DefinePreUnaryOperator<FPlus>();
+	Grammar.DefinePreUnaryOperator<FMinus>();
+	Grammar.DefinePreUnaryOperator<FSquareRoot>();
+	Grammar.DefineBinaryOperator<FPlus>(5);
+	Grammar.DefineBinaryOperator<FMinus>(5);
+	Grammar.DefineBinaryOperator<FStar>(4);
+	Grammar.DefineBinaryOperator<FForwardSlash>(4);
+	Grammar.DefineBinaryOperator<FPercent>(4);
+	Grammar.DefineBinaryOperator<FPower>(4);;
+
+	JumpTable.MapPreUnary<FPlus>([](double N)			{ return N; });
+	JumpTable.MapPreUnary<FMinus>([](double N)			{ return -N; });
+	JumpTable.MapPreUnary<FSquareRoot>([](double A)		{ return double(FMath::Sqrt(A)); });
+
+	JumpTable.MapBinary<FPlus>([](double A, double B)	{ return A + B; });
+	JumpTable.MapBinary<FMinus>([](double A, double B)	{ return A - B; });
+	JumpTable.MapBinary<FStar>([](double A, double B)	{ return A * B; });
+	JumpTable.MapBinary<FPower>([](double A, double B)	{ return double(FMath::Pow(A, B)); });
+
+	JumpTable.MapBinary<FForwardSlash>([](double A, double B) -> FExpressionResult {
+		if (B == 0)
+		{
+			return MakeError(LOCTEXT("DivisionByZero", "Division by zero"));
+		}
+
+		return MakeValue(A / B);
+	});
+	JumpTable.MapBinary<FPercent>([](double A, double B) -> FExpressionResult {
+		if (B == 0)
+		{
+			return MakeError(LOCTEXT("ModZero", "Modulo zero"));
+		}
+
+		return MakeValue(double(FMath::Fmod(A, B)));
+	});
+}
+
+TValueOrError<double, FExpressionError> FBasicMathExpressionEvaluator::Evaluate(const TCHAR* InExpression) const
+{
+	TValueOrError<FExpressionNode, FExpressionError> Result = ExpressionParser::Evaluate(InExpression, TokenDefinitions, Grammar, JumpTable);
+	if (!Result.IsValid())
+	{
+		return MakeError(Result.GetError());
+	}
+
+	auto& Node = Result.GetValue();
+
+	if (const auto* Numeric = Node.Cast<double>())
+	{
+		return MakeValue(*Numeric);
+	}
+	else
+	{
+		return MakeError(LOCTEXT("UnrecognizedResult", "Unrecognized result returned from expression"));
+	}
+}
+
+bool TestExpression(FAutomationTestBase* Test, const TCHAR* Expression, double Expected)
+{
+	FBasicMathExpressionEvaluator Parser;
+
+	TValueOrError<double, FExpressionError> Result = Parser.Evaluate(Expression);
+	if (!Result.IsValid())
+	{
+		Test->AddError(Result.GetError().Text.ToString());
+		return false;
+	}
+	else if (Result.GetValue() != Expected)
+	{
+		Test->AddError(FString::Printf(TEXT("Expression '%s' evaluated incorrectly: %.f != %.f"), Expression, Result.GetValue(), Expected));
+		return false;
+	}
+	return true;
+}
+
+bool TestInvalidExpression(FAutomationTestBase* Test, const TCHAR* Expression)
+{
+	FBasicMathExpressionEvaluator Parser;
+
+	TValueOrError<double, FExpressionError> Result = Parser.Evaluate(Expression);
+	if (Result.IsValid())
+	{
+		Test->AddError(FString::Printf(TEXT("Invalid expression '%s' did not report an error"), Expression));
+		return false;
+	}
+	else
+	{
+		Test->AddLogItem(Result.GetError().Text.ToString());
+	}
+
+	return true;
+}
+
+IMPLEMENT_SIMPLE_AUTOMATION_TEST(FBasicMathExpressionEvaluatorTest, "System.Core.Math.Expressions", EAutomationTestFlags::ATF_SmokeTest)
+bool FBasicMathExpressionEvaluatorTest::RunTest( const FString& Parameters )
+{
+	bool bResult = true;
+
+	bResult = bResult && TestExpression(this, TEXT("+1"), 1);
+	bResult = bResult && TestExpression(this, TEXT("-20"), -20);
+	bResult = bResult && TestExpression(this, TEXT("-+-2"), 2);
+	bResult = bResult && TestExpression(this, TEXT("1 + 2"), 3);
+	bResult = bResult && TestExpression(this, TEXT("1+2*3"), 7);
+	bResult = bResult && TestExpression(this, TEXT("1+2*3*4+1"), 1+2*3*4+1);
+	bResult = bResult && TestExpression(this, TEXT("1*2+3"), 1*2+3);
+	bResult = bResult && TestExpression(this, TEXT("1+2*3*4+1"), 1+2*3*4+1);
+
+	bResult = bResult && TestExpression(this, TEXT("2^2"), 4);
+	bResult = bResult && TestExpression(this, TEXT("sqrt(4)"), 2);
+	bResult = bResult && TestExpression(this, TEXT("4*sqrt(4)+10"), 18);
+	bResult = bResult && TestExpression(this, TEXT("8%6"), 2);
+
+	// Leading/Trailing whitespace
+	bResult = bResult && TestExpression(this, TEXT(" 1+2 "), 1+2);
+
+	// Test grouping
+	bResult = bResult && TestExpression(this, TEXT("(1+2)*3*4+1"), (1+2)*3*4+1);
+	bResult = bResult && TestExpression(this, TEXT("(1+2)*3*(4+1)"), (1+2)*3*(4+1));
+	bResult = bResult && TestExpression(this, TEXT("((1+2) / (3+1) + 2) * 3"), ((1.0+2) / (3+1) + 2) * 3);
+
+	// Test that some invalid expressions report errors and don't crash etc
+	bResult = bResult && TestInvalidExpression(this, TEXT("gobbledegook"));
+	bResult = bResult && TestInvalidExpression(this, TEXT("50**10"));
+	bResult = bResult && TestInvalidExpression(this, TEXT("*1"));
+	bResult = bResult && TestInvalidExpression(this, TEXT("+"));
+	bResult = bResult && TestInvalidExpression(this, TEXT("{+}"));
+
+	return bResult;
+}
+
+#undef LOCTEXT_NAMESPACE
\ No newline at end of file
diff --git a/Engine/Source/Runtime/Core/Private/Misc/ExpressionParser.cpp b/Engine/Source/Runtime/Core/Private/Misc/ExpressionParser.cpp
new file mode 100644
--- /dev/null
+++ b/Engine/Source/Runtime/Core/Private/Misc/ExpressionParser.cpp
@@ -0,0 +1,721 @@
+
+#include "CorePrivatePCH.h"
+#include "ExpressionParser.h"
+
+#define LOCTEXT_NAMESPACE "ExpressionParser"
+
+FTokenStream::FTokenStream(const TCHAR* In)
+	: Start(In)
+	, End(Start + FCString::Strlen(Start))
+	, ReadPos(In)
+{
+
+}
+
+bool FTokenStream::IsReadPosValid(const TCHAR* InPos, int32 MinNumChars) const
+{
+	return InPos >= Start && InPos <= End - MinNumChars;
+}
+
+TCHAR FTokenStream::PeekChar(int32 Offset) const
+{
+	if (ReadPos + Offset < End)
+	{
+		return *(ReadPos + Offset);
+	}
+
+	return '\0';
+}
+
+int32 FTokenStream::CharsRemaining() const
+{
+	return End - ReadPos;
+}
+
+bool FTokenStream::IsEmpty() const
+{
+	return ReadPos == End;
+}
+
+int32 FTokenStream::GetPosition() const
+{
+	return ReadPos - Start;
+}
+
+FString FTokenStream::GetErrorContext() const
+{
+	const TCHAR* StartPos = ReadPos;
+	const TCHAR* EndPos = StartPos;
+
+	// Skip over any leading whitespace
+	while (FChar::IsWhitespace(*EndPos))
+	{
+		EndPos++;
+	}
+
+	// Read until next whitespace or end of string
+	while (!FChar::IsWhitespace(*EndPos) && *EndPos != '\0')
+	{
+		EndPos++;
+	}
+
+	static const int32 MaxChars = 32;
+	FString Context(FMath::Min(int32(EndPos - StartPos), MaxChars), StartPos);
+	if (EndPos - StartPos > MaxChars)
+	{
+		Context += TEXT("...");
+	}
+	return Context;
+}
+
+/** Parse out a token */
+TOptional<FStringToken> FTokenStream::ParseToken(const TFunctionRef<EParseState(TCHAR)>& Pred, FStringToken* Accumulate) const
+{
+	const TCHAR* OptReadPos = Accumulate ? Accumulate->GetTokenEndPos() : ReadPos;
+
+	if (!IsReadPosValid(OptReadPos))
+	{
+		return TOptional<FStringToken>();
+	}
+
+	FStringToken Token(OptReadPos, 0, OptReadPos - Start);
+
+	while (Token.GetTokenEndPos() != End)
+	{
+		const EParseState State = Pred(*Token.GetTokenEndPos());
+		
+		if (State == EParseState::Cancel)
+		{
+			return TOptional<FStringToken>();
+		}
+		
+		if (State == EParseState::Continue || State == EParseState::StopAfter)
+		{
+			// Need to include this character in this token
+			++Token.TokenEnd;
+		}
+
+		if (State == EParseState::StopAfter || State == EParseState::StopBefore)
+		{
+			// Finished parsing the token
+			break;
+		}
+	}
+
+	if (Token.IsValid())
+	{
+		if (Accumulate)
+		{
+			Accumulate->Accumulate(Token);
+		}
+		return Token;
+	}
+	return TOptional<FStringToken>();
+}
+
+TOptional<FStringToken> FTokenStream::ParseSymbol(FStringToken* Accumulate) const
+{
+	const TCHAR* OptReadPos = Accumulate ? Accumulate->GetTokenEndPos() : ReadPos;
+
+	if (!IsReadPosValid(OptReadPos))
+	{
+		return TOptional<FStringToken>();
+	}
+	
+	FStringToken Token(OptReadPos, 0, OptReadPos - Start);
+	++Token.TokenEnd;
+
+	if (Accumulate)
+	{
+		Accumulate->Accumulate(Token);
+	}
+
+	return Token;
+}
+
+TOptional<FStringToken> FTokenStream::ParseSymbol(TCHAR Symbol, FStringToken* Accumulate) const
+{
+	const TCHAR* OptReadPos = Accumulate ? Accumulate->GetTokenEndPos() : ReadPos;
+
+	if (!IsReadPosValid(OptReadPos))
+	{
+		return TOptional<FStringToken>();
+	}
+	
+	FStringToken Token(OptReadPos, 0, OptReadPos - Start);
+
+	if (*Token.TokenEnd == Symbol)
+	{
+		++Token.TokenEnd;
+
+		if (Accumulate)
+		{
+			Accumulate->Accumulate(Token);
+		}
+
+		return Token;
+	}
+
+	return TOptional<FStringToken>();
+}
+
+TOptional<FStringToken> FTokenStream::ParseToken(const TCHAR* Symbol, FStringToken* Accumulate) const
+{
+	const TCHAR* OptReadPos = Accumulate ? Accumulate->GetTokenEndPos() : ReadPos;
+
+	const int32 Len = FCString::Strlen(Symbol);
+	if (!IsReadPosValid(OptReadPos, Len))
+	{
+		return TOptional<FStringToken>();
+	}
+
+	if (*OptReadPos != *Symbol)
+	{
+		return TOptional<FStringToken>();		
+	}
+
+	FStringToken Token(OptReadPos, 0, OptReadPos - Start);
+	
+	if (FCString::Strncmp(Token.GetTokenEndPos(), Symbol, Len) == 0)
+	{
+		Token.TokenEnd += Len;
+
+		if (Accumulate)
+		{
+			Accumulate->Accumulate(Token);
+		}
+
+		return Token;
+	}
+
+	return TOptional<FStringToken>();
+}
+
+TOptional<FStringToken> FTokenStream::ParseTokenIgnoreCase(const TCHAR* Symbol, FStringToken* Accumulate) const
+{
+	const TCHAR* OptReadPos = Accumulate ? Accumulate->GetTokenEndPos() : ReadPos;
+
+	const int32 Len = FCString::Strlen(Symbol);
+	if (!IsReadPosValid(OptReadPos, Len))
+	{
+		return TOptional<FStringToken>();
+	}
+
+	FStringToken Token(OptReadPos, 0, OptReadPos - Start);
+	
+	if (FCString::Strnicmp(OptReadPos, Symbol, Len) == 0)
+	{
+		Token.TokenEnd += Len;
+
+		if (Accumulate)
+		{
+			Accumulate->Accumulate(Token);
+		}
+
+		return Token;
+	}
+
+	return TOptional<FStringToken>();
+}
+
+TOptional<FStringToken> FTokenStream::ParseWhitespace(FStringToken* Accumulate) const
+{
+	const TCHAR* OptReadPos = Accumulate ? Accumulate->GetTokenEndPos() : ReadPos;
+
+	if (IsReadPosValid(OptReadPos))
+	{
+		return ParseToken([](TCHAR InC){ return FChar::IsWhitespace(InC) ? EParseState::Continue : EParseState::StopBefore; });
+	}
+
+	return TOptional<FStringToken>();
+}
+
+TOptional<FStringToken> FTokenStream::GenerateToken(int32 NumChars, FStringToken* Accumulate) const
+{
+	const TCHAR* OptReadPos = Accumulate ? Accumulate->GetTokenEndPos() : ReadPos;
+
+	if (IsReadPosValid(OptReadPos, NumChars))
+	{
+		FStringToken Token(OptReadPos, 0, OptReadPos - Start);
+		Token.TokenEnd += NumChars;
+		if (Accumulate)
+		{
+			Accumulate->Accumulate(Token);
+		}
+		return Token;
+	}
+
+	return TOptional<FStringToken>();	
+}
+
+void FTokenStream::SetReadPos(const FStringToken& Token)
+{
+	if (ensure(IsReadPosValid(Token.TokenEnd, 0)))
+	{
+		ReadPos = Token.TokenEnd;
+	}
+}
+
+FExpressionResult FOperatorJumpTable::ExecBinary(const FExpressionToken& Operator, const FExpressionToken& L, const FExpressionToken& R) const
+{
+	FOperatorFunctionID ID = { Operator.Node.GetTypeId(), L.Node.GetTypeId(), R.Node.GetTypeId() };
+	if (const auto* Func = BinaryOps.Find(ID))
+	{
+		return (*Func)(L.Node, R.Node);
+	}
+
+	FFormatOrderedArguments Args;
+	Args.Add(FText::FromString(Operator.Context.GetString()));
+	Args.Add(FText::FromString(L.Context.GetString()));
+	Args.Add(FText::FromString(R.Context.GetString()));
+	return MakeError(FText::Format(LOCTEXT("UnaryExecutionError", "Binary operator {0} cannot operate on {1} and {2}"), Args));
+}
+
+FExpressionResult FOperatorJumpTable::ExecPreUnary(const FExpressionToken& Operator, const FExpressionToken& R) const
+{
+	FOperatorFunctionID ID = { Operator.Node.GetTypeId(), FGuid(), R.Node.GetTypeId() };
+	if (const auto* Func = PreUnaryOps.Find(ID))
+	{
+		return (*Func)(R.Node);
+	}
+
+	FFormatOrderedArguments Args;
+	Args.Add(FText::FromString(Operator.Context.GetString()));
+	Args.Add(FText::FromString(R.Context.GetString()));
+	return MakeError(FText::Format(LOCTEXT("UnaryExecutionError", "Pre-unary operator {0} cannot operate on {1}"), Args));
+}
+
+FExpressionResult FOperatorJumpTable::ExecPostUnary(const FExpressionToken& Operator, const FExpressionToken& L) const
+{
+	FOperatorFunctionID ID = { Operator.Node.GetTypeId(), L.Node.GetTypeId(), FGuid() };
+	if (const auto* Func = PostUnaryOps.Find(ID))
+	{
+		return (*Func)(L.Node);
+	}
+
+	FFormatOrderedArguments Args;
+	Args.Add(FText::FromString(Operator.Context.GetString()));
+	Args.Add(FText::FromString(L.Context.GetString()));
+	return MakeError(FText::Format(LOCTEXT("UnaryExecutionError", "Post-unary operator {0} cannot operate on {1}"), Args));
+}
+
+FExpressionTokenConsumer::FExpressionTokenConsumer(const TCHAR* InExpression)
+	: Stream(InExpression)
+{}
+
+TArray<FExpressionToken> FExpressionTokenConsumer::Extract()
+{
+	TArray<FExpressionToken> Swapped;
+	Swap(Swapped, Tokens);
+	return MoveTemp(Swapped);
+}
+
+void FExpressionTokenConsumer::Add(const FStringToken& SourceToken, FExpressionNode Node)
+{
+	Stream.SetReadPos(SourceToken);
+	Tokens.Add(FExpressionToken(SourceToken, MoveTemp(Node)));
+}
+
+void FTokenDefinitions::DefineToken(const TFunction<FExpressionDefinition>& Definition)
+{
+	Definitions.Emplace(Definition);
+}
+
+TOptional<FExpressionError> FTokenDefinitions::ConsumeToken(FExpressionTokenConsumer& Consumer) const
+{
+	auto& Stream = Consumer.GetStream();
+	
+	// Skip over whitespace
+	if (bIgnoreWhitespace)
+	{
+		TOptional<FStringToken> Whitespace = Stream.ParseWhitespace();
+		if (Whitespace.IsSet())
+		{
+			Stream.SetReadPos(Whitespace.GetValue());
+		}
+	}
+
+	if (Stream.IsEmpty())
+	{
+		// Trailing whitespace in the expression.
+		return TOptional<FExpressionError>();
+	}
+
+	const auto* Pos = Stream.GetRead();
+
+	// Try each token in turn. First come first served.
+	for (const auto& Def : Definitions)
+	{
+		// Call the token definition
+		auto Error = Def(Consumer);
+		if (Error.IsSet())
+		{
+			return Error;
+		}
+		// If the stream has moved on, the definition added one or more tokens, so 
+		else if (Stream.GetRead() != Pos)
+		{
+			return TOptional<FExpressionError>();
+		}
+	}
+
+	// No token definition matched the stream at its current position - fatal error
+	FFormatOrderedArguments Args;
+	Args.Add(FText::FromString(Consumer.GetStream().GetErrorContext()));
+	Args.Add(Consumer.GetStream().GetPosition());
+	return FExpressionError(FText::Format(LOCTEXT("LexicalError", "Unrecognized token '{0}' at character {1}"), Args));
+}
+
+TOptional<FExpressionError> FTokenDefinitions::ConsumeTokens(FExpressionTokenConsumer& Consumer) const
+{
+	auto& Stream = Consumer.GetStream();
+	while(!Stream.IsEmpty())
+	{
+		auto Error = ConsumeToken(Consumer);
+		if (Error.IsSet())
+		{
+			return Error;
+		}
+	}
+
+	return TOptional<FExpressionError>();
+}
+
+const FGuid* FExpressionGrammar::GetGrouping(const FGuid& TypeId) const
+{
+	return Groupings.Find(TypeId);
+}
+
+bool FExpressionGrammar::HasPreUnaryOperator(const FGuid& InTypeId) const
+{
+	return PreUnaryOperators.Contains(InTypeId);
+}
+
+bool FExpressionGrammar::HasPostUnaryOperator(const FGuid& InTypeId) const
+{
+	return PostUnaryOperators.Contains(InTypeId);
+}
+
+const int* FExpressionGrammar::GetBinaryOperatorPrecedence(const FGuid& InTypeId) const
+{
+	return BinaryOperators.Find(InTypeId);
+}
+
+struct FExpressionCompiler
+{
+	FExpressionCompiler(const FExpressionGrammar& InGrammar, const TArray<FExpressionToken>& InTokens)
+		: Grammar(InGrammar), Tokens(InTokens)
+	{
+		CurrentTokenIndex = 0;
+		Commands.Reserve(Tokens.Num());
+	}
+
+	TValueOrError<TArray<FExpressionToken>, FExpressionError> Compile()
+	{
+		auto Error = CompileGroup(nullptr, nullptr);
+		if (Error.IsSet())
+		{
+			return MakeError(Error.GetValue());
+		}
+		return MakeValue(MoveTemp(Commands));
+	}
+
+	TOptional<FExpressionError> CompileGroup(const FExpressionToken* GroupStart, const FGuid* StopAt)
+	{
+		enum class EState { PreUnary, PostUnary, Binary };
+
+		TArray<FWrappedOperator> OperatorStack;
+		OperatorStack.Reserve(Tokens.Num() - CurrentTokenIndex);
+
+		bool bFoundEndOfGroup = StopAt == nullptr;
+
+		// Start off looking for a unary operator
+		EState State = EState::PreUnary;
+		for (; CurrentTokenIndex < Tokens.Num(); ++CurrentTokenIndex)
+		{
+			auto& Token = Tokens[CurrentTokenIndex];
+			const auto& TypeId = Token.Node.GetTypeId();
+
+			if (const FGuid* GroupingEnd = Grammar.GetGrouping(TypeId))
+			{
+				// Ignore this token
+				CurrentTokenIndex++;
+
+				const auto MarkerIndex = Commands.Num();
+
+				FGroupMarker Marker(Token);
+
+				// Add the group marker
+				Commands.Add(FExpressionToken(Token.Context, Marker));
+
+				// Start of group - recurse
+				auto Error = CompileGroup(&Token, GroupingEnd);
+				
+				Marker.NumTokens = Commands.Num() - MarkerIndex;
+				Commands[MarkerIndex].Node = Marker;
+
+				if (Error.IsSet())
+				{
+					return Error;
+				}
+
+				State = EState::PostUnary;
+			}
+			else if (StopAt && TypeId == *StopAt)
+			{
+				// End of group
+				bFoundEndOfGroup = true;
+				break;
+			}
+			else if (State == EState::PreUnary)
+			{
+				if (Grammar.HasPreUnaryOperator(TypeId))
+				{
+					// Make this a unary op
+					OperatorStack.Add(FWrappedOperator(FWrappedOperator::PreUnary, Token));
+				}
+				else if (Grammar.GetBinaryOperatorPrecedence(TypeId))
+				{
+					return FExpressionError(FText::Format(LOCTEXT("SyntaxError_NoBinaryOperand", "Syntax error: No operand specified for operator '{0}'"), FText::FromString(Token.Context.GetString())));
+				}
+				else if (Grammar.HasPostUnaryOperator(TypeId))
+				{
+					// Found a post-unary operator for the preceeding token
+					State = EState::PostUnary;
+
+					// Pop off any pending unary operators
+					while (OperatorStack.Num() > 0 && OperatorStack.Last().Precedence <= 0)
+					{
+						auto Operator = OperatorStack.Pop(false);
+						Commands.Add(FExpressionToken(Operator.WrappedToken.Context, Operator));
+					}
+
+					// Make this a post-unary op
+					OperatorStack.Add(FWrappedOperator(FWrappedOperator::PostUnary, Token));
+				}
+				else
+				{
+					// Not an operator, so treat it as an ordinary token
+					Commands.Add(Token);
+					State = EState::PostUnary;
+				}
+			}
+			else if (State == EState::PostUnary)
+			{
+				if (Grammar.HasPostUnaryOperator(TypeId))
+				{
+					// Pop off any pending unary operators
+					while (OperatorStack.Num() > 0 && OperatorStack.Last().Precedence <= 0)
+					{
+						auto Operator = OperatorStack.Pop(false);
+						Commands.Add(FExpressionToken(Operator.WrappedToken.Context, Operator));
+					}
+
+					// Make this a post-unary op
+					OperatorStack.Add(FWrappedOperator(FWrappedOperator::PostUnary, Token));
+				}
+				else
+				{
+					// Checking for binary operators
+					if (const int32* Prec = Grammar.GetBinaryOperatorPrecedence(TypeId))
+					{
+						// Pop off anything of higher precedence than this one onto the command stack
+						while (OperatorStack.Num() > 0 && OperatorStack.Last().Precedence < *Prec)
+						{
+							auto Operator = OperatorStack.Pop(false);
+							Commands.Add(FExpressionToken(Operator.WrappedToken.Context, Operator));
+						}
+
+						// Add the operator itself to the op stack
+						OperatorStack.Add(FWrappedOperator(FWrappedOperator::Binary, Token, *Prec));
+
+						// Check for a unary op again
+						State = EState::PreUnary;
+					}
+					else
+					{
+						// Just add the token. It's possible that this is a syntax error (there's no binary operator specified between two tokens),
+						// But we don't have enough information at this point to say whether or not it is an error
+						Commands.Add(Token);
+						State = EState::PreUnary;
+					}
+				}
+			}
+		}
+
+		if (!bFoundEndOfGroup)
+		{
+			return FExpressionError(FText::Format(LOCTEXT("SyntaxError_UnmatchedGroup", "Syntax error: Reached end of expression before matching end of group '{0}' at line {1}:{2}"),
+				FText::FromString(GroupStart->Context.GetString()),
+				FText::AsNumber(GroupStart->Context.GetLineNumber()),
+				FText::AsNumber(GroupStart->Context.GetCharacterIndex())
+			));
+		}
+
+		// Pop everything off the operator stack, onto the command stack
+		while (OperatorStack.Num() > 0)
+		{
+			auto Operator = OperatorStack.Pop(false);
+			Commands.Add(FExpressionToken(Operator.WrappedToken.Context, Operator));
+		}
+
+		return TOptional<FExpressionError>();
+	}
+
+
+private:
+
+	int32 CurrentTokenIndex;
+
+	/** Working structures */
+	TArray<FExpressionToken> Commands;
+
+private:
+	/** Const data provided by the parser */
+	const FExpressionGrammar& 			Grammar;
+	const TArray<FExpressionToken>& Tokens;
+};
+
+namespace ExpressionParser
+{
+	TValueOrError<TArray<FExpressionToken>, FExpressionError> Lex(const TCHAR* InExpression, const FTokenDefinitions& TokenDefinitions)
+	{
+		FExpressionTokenConsumer TokenConsumer(InExpression);
+		
+		TOptional<FExpressionError> Error = TokenDefinitions.ConsumeTokens(TokenConsumer);
+		if (Error.IsSet())
+		{
+			return MakeError(Error.GetValue());
+		}
+		else
+		{
+			return MakeValue(TokenConsumer.Extract());
+		}
+	}
+
+	TValueOrError<TArray<FExpressionToken>, FExpressionError> Compile(const TCHAR* InExpression, const FTokenDefinitions& TokenDefinitions, const FExpressionGrammar& InGrammar)
+	{
+		TValueOrError<TArray<FExpressionToken>, FExpressionError> Result = Lex(InExpression, TokenDefinitions);
+
+		if (!Result.IsValid())
+		{
+			return MakeError(Result.GetError());
+		}
+
+		return Compile(Result.GetValue(), InGrammar);
+	}
+
+	ResultType Compile(const TArray<FExpressionToken>& InTokens, const FExpressionGrammar& InGrammar)
+	{
+		return FExpressionCompiler(InGrammar, InTokens).Compile();
+	}
+
+	FExpressionResult Evaluate(const TCHAR* InExpression, const FTokenDefinitions& TokenDefinitions, const FExpressionGrammar& InGrammar, const FOperatorJumpTable& InJumpTable)
+	{
+		TValueOrError<TArray<FExpressionToken>, FExpressionError> CompilationResult = Compile(InExpression, TokenDefinitions, InGrammar);
+
+		if (!CompilationResult.IsValid())
+		{
+			return MakeError(CompilationResult.GetError());
+		}
+
+		auto& Tokens = CompilationResult.GetValue();
+
+		TArray<FExpressionToken> Stack;
+		TArray<FExpressionToken> Popped;
+
+		for (int32 Index = 0; Index < Tokens.Num(); ++Index)
+		{
+			auto& Token = Tokens[Index];
+
+			// We can ignore group markers for evaluation
+			if (Token.Node.Cast<FGroupMarker>())
+			{
+				continue;
+			}
+
+			if (const auto* Op = Token.Node.Cast<FWrappedOperator>())
+			{
+				switch(Op->Type)
+				{
+				case FWrappedOperator::Binary:
+					if (Stack.Num() >= 2)
+					{
+						// Binary
+						auto R = Stack.Pop();
+						auto L = Stack.Pop();
+
+						auto OpResult = InJumpTable.ExecBinary(Op->WrappedToken, L, R);
+						if (OpResult.IsValid())
+						{
+							// Inherit the LHS context
+							Stack.Push(FExpressionToken(L.Context, OpResult.GetValue()));
+						}
+						else
+						{
+							FFormatOrderedArguments Args;
+							Args.Add(FText::FromString(Op->WrappedToken.Context.GetString()));
+							Args.Add(FText::FromString(L.Context.GetString()));
+							Args.Add(FText::FromString(R.Context.GetString()));
+							return MakeError(FText::Format(LOCTEXT("InvalidBinaryOp", "Cannot perform operator {0} on '{1}' and '{2}'"), Args));
+						}
+					}
+					else
+					{
+						FFormatOrderedArguments Args;
+						Args.Add(FText::FromString(Op->WrappedToken.Context.GetString()));
+						return MakeError(FText::Format(LOCTEXT("SyntaxError_NoUnaryOperand", "Not enough operands for binary operator {0}"), Args));
+					}
+					break;
+				
+				case FWrappedOperator::PostUnary:
+				case FWrappedOperator::PreUnary:
+
+					if (Stack.Num() >= 1)
+					{
+						auto Operand = Stack.Pop();
+
+						FExpressionResult OpResult = (Op->Type == FWrappedOperator::PreUnary) ?
+							InJumpTable.ExecPreUnary(Op->WrappedToken, Operand) :
+							InJumpTable.ExecPostUnary(Op->WrappedToken, Operand);
+
+						if (OpResult.IsValid())
+						{
+							// Inherit the LHS context
+							Stack.Push(FExpressionToken(Op->WrappedToken.Context, OpResult.GetValue()));
+						}
+						else
+						{
+							FFormatOrderedArguments Args;
+							Args.Add(FText::FromString(Op->WrappedToken.Context.GetString()));
+							Args.Add(FText::FromString(Operand.Context.GetString()));
+							return MakeError(FText::Format(LOCTEXT("InvalidUnaryOp", "Cannot perform operator {0} on '{1}'"), Args));
+						}			
+					}
+					else
+					{
+						FFormatOrderedArguments Args;
+						Args.Add(FText::FromString(Op->WrappedToken.Context.GetString()));
+						return MakeError(FText::Format(LOCTEXT("SyntaxError_NoUnaryOperand", "No operand for unary operator {0}"), Args));
+					}
+					break;
+				}
+			}
+			else
+			{
+				Stack.Push(Token);
+			}
+		}
+
+		if (Stack.Num() == 1)
+		{
+			return MakeValue(Stack[0].Node);
+		}
+		else
+		{
+			return MakeError(LOCTEXT("SyntaxError_InvalidExpression", "Could not evaluate expression"));
+		}
+	}
+}
+
+#undef LOCTEXT_NAMESPACE
\ No newline at end of file
diff --git a/Engine/Source/Runtime/Core/Private/Misc/ExpressionParserExamples.cpp b/Engine/Source/Runtime/Core/Private/Misc/ExpressionParserExamples.cpp
new file mode 100644
--- /dev/null
+++ b/Engine/Source/Runtime/Core/Private/Misc/ExpressionParserExamples.cpp
@@ -0,0 +1,262 @@
+
+// Don't include this code in shipping or perf test builds
+#if !(UE_BUILD_SHIPPING || UE_BUILD_TEST)
+
+#include "CorePrivatePCH.h"
+#include "ExpressionParser.h"
+#include "BasicMathExpressionEvaluator.h"
+
+#define LOCTEXT_NAMESPACE "ExpressionParserExamples"
+
+class FExampleMathEvaluator
+{
+public:
+
+	struct FExampleAdd {};
+
+	/** Very simple math expression parser that supports addition of numbers */
+	FExampleMathEvaluator()
+	{
+		using namespace ExpressionParser;
+
+		TokenDefinitions.IgnoreWhitespace();
+
+		// Define a tokenizer that will tokenize numbers in the expression
+		TokenDefinitions.DefineToken(&ConsumeNumber);
+
+		// Define a tokenizer for the arithmetic add
+		TokenDefinitions.DefineToken([](FExpressionTokenConsumer& Consumer){
+
+			// Called at the start of every new token - see if it's a + character
+			TOptional<FStringToken> Token = Consumer.GetStream().ParseSymbol('+');
+			if (Token.IsSet())
+			{
+				// Add the token to the consumer. This moves the read position in the stream to the end of the token.
+				Consumer.Add(Token.GetValue(), FExampleAdd());
+			}
+
+			// Optionally return an error if there is a good reason to
+			return TOptional<FExpressionError>();
+		});
+
+
+		// Define our FExampleAdd token as a binary operator
+		Grammar.DefineBinaryOperator<FExampleAdd>(5);
+		// And provide an operator overload for it
+		JumpTable.MapBinary<FExampleAdd>([](double A, double B) { return A + B;	});
+	}
+
+	/** Evaluate the given expression, resulting in either a double value, or an error */
+	TValueOrError<double, FExpressionError> Evaluate(const TCHAR* InExpression) const
+	{
+		TValueOrError<FExpressionNode, FExpressionError> Result = ExpressionParser::Evaluate(InExpression, TokenDefinitions, Grammar, JumpTable);
+		if (!Result.IsValid())
+		{
+			return MakeError(Result.GetError());
+		}
+		
+		if (const double* Numeric = Result.GetValue().Cast<double>())
+		{
+			return MakeValue(*Numeric);
+		}
+		else
+		{
+			return MakeError(LOCTEXT("UnrecognizedResult", "Unrecognized result returned from expression"));
+		}
+	}
+	
+private:
+	FTokenDefinitions TokenDefinitions;
+	FExpressionGrammar Grammar;
+	FOperatorJumpTable JumpTable;
+};
+DEFINE_EXPRESSION_NODE_TYPE(FExampleMathEvaluator::FExampleAdd, 0x4BF093DC, 0xA92247B5, 0xAE53A9B3, 0x64A84D2C)
+
+
+/** A simple string formatter */
+class FExampleStringFormatter
+{
+public:
+	/** Token representing a literal string inside the string */
+	struct FStringLiteral
+	{
+		FStringLiteral(FStringToken InString) : String(InString), Len(InString.GetTokenEndPos() - InString.GetTokenStartPos()) {}
+		FStringToken String;
+		int32 Len;
+	};
+	/** Token representing a user-defined token, such as {Argument} */
+	struct FFormatToken
+	{
+		FFormatToken(FStringToken InIdentifier) : Identifier(InIdentifier), Len(Identifier.GetTokenEndPos() - Identifier.GetTokenStartPos()) {}
+		FStringToken Identifier;
+		int32 Len;
+	};
+
+	FExampleStringFormatter()
+	{
+		using namespace ExpressionParser;
+
+		// Define a tokenizer that will parse user supplied tokens (FFormatToken) in the expression
+		TokenDefinitions.DefineToken([](FExpressionTokenConsumer& Consumer) -> TOptional<FExpressionError> {
+			auto& Stream = Consumer.GetStream();
+
+			TOptional<FStringToken> OpeningChar = Stream.ParseSymbol('{');
+			if (!OpeningChar.IsSet())
+			{
+				return TOptional<FExpressionError>();
+			}
+
+			FStringToken& EntireToken = OpeningChar.GetValue();
+
+			// Optional whitespace
+			Stream.ParseToken([](TCHAR InC) { return FChar::IsWhitespace(InC) ? EParseState::Continue : EParseState::StopBefore; }, &EntireToken);
+
+			// The identifier itself
+			TOptional<FStringToken> Identifier = Stream.ParseToken([](TCHAR InC) { return FChar::IsWhitespace(InC) || InC == '}' ? EParseState::StopBefore : EParseState::Continue; }, &EntireToken);
+
+			if (!Identifier.IsSet())
+			{
+				FFormatOrderedArguments Args;
+				Args.Add(FText::FromString(FString(EntireToken.GetTokenEndPos()).Left(10) + TEXT("...")));
+				return FExpressionError(FText::Format(LOCTEXT("UnrecognizedIdentifier", "Unrecognized identifier symbol: '{0}'"), Args));
+			}
+
+			// Optional whitespace
+			Stream.ParseToken([](TCHAR InC) { return FChar::IsWhitespace(InC) ? EParseState::Continue : EParseState::StopBefore; }, &EntireToken);
+			
+			if (!Stream.ParseSymbol('}', &EntireToken).IsSet())
+			{
+				FFormatOrderedArguments Args;
+				Args.Add(FText::FromString(FString(EntireToken.GetTokenEndPos() - EntireToken.GetTokenStartPos(), EntireToken.GetTokenStartPos())));
+				return FExpressionError(FText::Format(LOCTEXT("InvalidTokenDefinition", "Invalid token definition: '{0}'"), Args));
+			}
+
+			// Add the token to the consumer. This moves the read position in the stream to the end of the token.
+			Consumer.Add(EntireToken, FFormatToken(Identifier.GetValue()));
+
+			return TOptional<FExpressionError>();
+		});
+
+		// Define a tokenizer for anything else
+		TokenDefinitions.DefineToken([](FExpressionTokenConsumer& Consumer){
+			
+			// Called at the start of every new token - see if it's a { character
+			TOptional<FStringToken> Token = Consumer.GetStream().ParseToken([](TCHAR InC){
+				return InC == '{' ? EParseState::StopBefore : EParseState::Continue;
+			});
+
+			if (Token.IsSet())
+			{
+				// Add the token to the consumer. This moves the read position in the stream to the end of the token.
+				Consumer.Add(Token.GetValue(), FStringLiteral(Token.GetValue()));
+			}
+
+			return TOptional<FExpressionError>();
+		});
+	}
+
+	/** Format the specified text using the specified arguments. Replaces instances of { Argument } with keys in the map matching 'Argument' */
+	TValueOrError<FText, FExpressionError> Format(FText InExpression, const TMap<FString, FText>& Args) const
+	{
+		const FString& String = InExpression.ToString();
+
+		TValueOrError<TArray<FExpressionToken>, FExpressionError> Result = ExpressionParser::Lex(*String, TokenDefinitions);
+		if (!Result.IsValid())
+		{
+			return MakeError(Result.GetError());
+		}
+		
+		// This code deliberately tries to reallocate as little as possible
+		FString Formatted;
+		Formatted.Reserve(1024);
+		for (auto& Token : Result.GetValue())
+		{
+			if (const auto* Literal = Token.Node.Cast<FStringLiteral>())
+			{
+				Formatted.AppendChars(Literal->String.GetTokenStartPos(), Literal->Len);
+			}
+			else if (const auto* FormatToken = Token.Node.Cast<FFormatToken>())
+			{
+				const FText* Arg = nullptr;
+				for (auto& Pair : Args)
+				{
+					if (Pair.Key.Len() == FormatToken->Len && FCString::Strnicmp(FormatToken->Identifier.GetTokenStartPos(), *Pair.Key, FormatToken->Len) == 0)
+					{
+						Arg = &Pair.Value;
+						break;
+					}
+				}
+
+				if (Arg)
+				{
+					Formatted.AppendChars(FormatToken->Identifier.GetTokenStartPos(), FormatToken->Len);
+				}
+				else
+				{
+					return MakeError(FText::Format(LOCTEXT("UnspecifiedToken", "Unspecified token: {0}"), FText::FromString(FormatToken->Identifier.GetString())));
+				}
+			}
+		}
+
+		return MakeValue(FText::FromString(Formatted));
+	}
+
+private:
+	FTokenDefinitions TokenDefinitions;
+};
+
+DEFINE_EXPRESSION_NODE_TYPE(FExampleStringFormatter::FStringLiteral, 0x03ED3A25, 0x85D94664, 0x8A8001A1, 0xDCC637F7)
+DEFINE_EXPRESSION_NODE_TYPE(FExampleStringFormatter::FFormatToken, 0xAAB48E5B, 0xEDA94853, 0xA951ED2D, 0x0A8E795D)
+
+/** Unit test for the string formatter */
+/*
+DEFINE_LOG_CATEGORY_STATIC(LogExampleStringTests, Log, All);
+
+#include "TextHistory.h"
+
+IMPLEMENT_SIMPLE_AUTOMATION_TEST(FExampleStringFormatterTest, "System.Core.Expression Parser", EAutomationTestFlags::ATF_Editor)
+bool FExampleStringFormatterTest::RunTest( const FString& Parameters )
+{
+	FExampleStringFormatter Formatter;
+
+	double Time = FPlatformTime::Seconds();
+
+	const FText Pattern = LOCTEXT("Pattern", "This is some text containing two arguments, { Argument1 } and {Argument2}.");
+
+
+	FFormatNamedArguments TextArgs;
+	TextArgs.Add(TEXT("Argument1"), LOCTEXT("Arg1", "Argument1"));
+	TextArgs.Add(TEXT("Argument2"), LOCTEXT("Arg2", "Argument2"));
+
+	{
+
+		TMap<FString, FText> Args;
+		Args.Add(TEXT("Argument1"), LOCTEXT("Arg1", "Argument1"));
+		Args.Add(TEXT("Argument2"), LOCTEXT("Arg2", "Argument2"));
+
+		for (int32 Bla = 0; Bla < 1000000; ++Bla)
+		{
+			MakeShareable(new FTextHistory_NamedFormat(Pattern, TextArgs));
+			Formatter.Format(Pattern, Args);
+		}
+	}
+
+	UE_LOG(LogExampleStringTests, Log, TEXT("Expression parser took %fs"), FPlatformTime::Seconds() - Time);
+
+	Time = FPlatformTime::Seconds();
+	{
+		for (int32 Bla = 0; Bla < 1000000; ++Bla)
+		{
+			FText::Format(Pattern, TextArgs);
+		}
+	}
+	UE_LOG(LogExampleStringTests, Log, TEXT("FText::Format took %fs"), FPlatformTime::Seconds() - Time);
+
+	return true;
+}
+
+*/
+
+#undef LOCTEXT_NAMESPACE
+
+#endif //!(UE_BUILD_SHIPPING || UE_BUILD_TEST)
\ No newline at end of file
diff --git a/Engine/Source/Runtime/Core/Public/Math/BasicMathExpressionEvaluator.h b/Engine/Source/Runtime/Core/Public/Math/BasicMathExpressionEvaluator.h
new file mode 100644
--- /dev/null
+++ b/Engine/Source/Runtime/Core/Public/Math/BasicMathExpressionEvaluator.h
@@ -0,0 +1,63 @@
+// Copyright 1998-2015 Epic Games, Inc. All Rights Reserved.
+
+#pragma once
+
+#include "ExpressionParser.h"
+
+#define DEFINE_EXPRESSION_OPERATOR_NODE(TYPE, ...) \
+namespace ExpressionParser {\
+	struct CORE_API TYPE { static const TCHAR* const Moniker; }; \
+}\
+DEFINE_EXPRESSION_NODE_TYPE(ExpressionParser::TYPE, __VA_ARGS__)
+
+/** Define some expression types for basic arithmetic */
+DEFINE_EXPRESSION_NODE_TYPE(double, 0x8444A8A3, 0x19AE4E13, 0xBCFA75EE, 0x39982BD6)
+
+DEFINE_EXPRESSION_OPERATOR_NODE(FSubExpressionStart, 0xCC40A083, 0xADBF46E2, 0xA93D12BB, 0x525D7417)
+DEFINE_EXPRESSION_OPERATOR_NODE(FSubExpressionEnd, 0x125E4C67, 0x96EB48C4, 0x8894E09C, 0xB3CD56BF)
+
+DEFINE_EXPRESSION_OPERATOR_NODE(FPlus, 0x6F88756B, 0xF9234263, 0x9B13614F, 0x2706074B)
+DEFINE_EXPRESSION_OPERATOR_NODE(FMinus, 0xE6240779, 0xEE2849CF, 0x95A0E648, 0x22D58713)
+DEFINE_EXPRESSION_OPERATOR_NODE(FStar, 0xF287B3BF, 0x35DF4141, 0xA8B4F57B, 0x7E06DF47)
+DEFINE_EXPRESSION_OPERATOR_NODE(FForwardSlash, 0xF99670F8, 0x74794256, 0xBB0CAE6D, 0xC67CD5B6)
+DEFINE_EXPRESSION_OPERATOR_NODE(FPercent, 0x936E4434, 0x0A014F2D, 0xBEEC90D3, 0x4D3ECEA2)
+DEFINE_EXPRESSION_OPERATOR_NODE(FSquareRoot, 0xE7C03E11, 0x9DE84B4B, 0xBA4C2B76, 0x69BF028E)
+DEFINE_EXPRESSION_OPERATOR_NODE(FPower, 0x93388F8D, 0x1D9B4DFE, 0xBD4D6CC4, 0x12D1DE99)
+	
+namespace ExpressionParser
+{
+	/** Parse a number from the given stream, optionally from a specific read position */
+	CORE_API TOptional<FStringToken> ParseNumber(const FTokenStream& InStream, FStringToken* Accumulate = nullptr);
+
+	/** Consume a number from the specified consumer's stream, if one exists at the current read position */
+	CORE_API TOptional<FExpressionError> ConsumeNumber(FExpressionTokenConsumer& Consumer);
+
+	/** Consume a symbol from the specified consumer's stream, if one exists at the current read position */
+	template<typename TSymbol>
+	TOptional<FExpressionError> ConsumeSymbol(FExpressionTokenConsumer& Consumer)
+	{
+		TOptional<FStringToken> Token = Consumer.GetStream().ParseToken(TSymbol::Moniker);
+		if (Token.IsSet())
+		{
+			Consumer.Add(Token.GetValue(), TSymbol());
+		}
+
+		return TOptional<FExpressionError>();
+	}
+}
+
+/** A basic math expression evaluator */
+class FBasicMathExpressionEvaluator
+{
+public:
+	/** Constructor that sets up the parser's lexer and compiler */
+	FBasicMathExpressionEvaluator();
+
+	/** Evaluate the given expression, resulting in either a double value, or an error */
+	TValueOrError<double, FExpressionError> Evaluate(const TCHAR* InExpression) const;
+	
+private:
+	FTokenDefinitions TokenDefinitions;
+	FExpressionGrammar Grammar;
+	FOperatorJumpTable JumpTable;
+};
\ No newline at end of file
diff --git a/Engine/Source/Runtime/Core/Public/Misc/ExpressionParser.h b/Engine/Source/Runtime/Core/Public/Misc/ExpressionParser.h
new file mode 100644
--- /dev/null
+++ b/Engine/Source/Runtime/Core/Public/Misc/ExpressionParser.h
@@ -0,0 +1,30 @@
+
+#pragma once
+
+#include "ExpressionParserTypes.h"
+
+/** An expression parser, responsible for lexing, compiling, and evaluating expressions.
+ *	The parser supports 3 functions:
+ *		1. Lexing the expression into a set of user defined tokens,
+ *		2. Compiling the tokenized expression to an efficient reverse-polish execution order,
+ *		3. Evaluating the compiled tokens
+ *
+ *  See ExpressionParserExamples.cpp for example usage.
+ */
+
+namespace ExpressionParser
+{
+	typedef TValueOrError< TArray<FExpressionToken>, FExpressionError > ResultType;
+
+	/** Lex the specified string, using the specified grammar */
+	CORE_API ResultType Lex(const TCHAR* InExpression, const FTokenDefinitions& TokenDefinitions);
+	
+	/** Compile the specified expression into an array of Reverse-Polish order nodes for evaluation, according to our grammar definition */
+	CORE_API ResultType Compile(const TCHAR* InExpression, const FTokenDefinitions& TokenDefinitions, const FExpressionGrammar& InGrammar);
+
+	/** Compile the specified tokens into an array of Reverse-Polish order nodes for evaluation, according to our grammar definition */
+	CORE_API ResultType Compile(const TArray<FExpressionToken>& InTokens, const FExpressionGrammar& InGrammar);
+
+	/** Evaluate the specific expression using our grammar definition, and return the result */
+	CORE_API FExpressionResult Evaluate(const TCHAR* Expression, const FTokenDefinitions& TokenDefinitions, const FExpressionGrammar& InGrammar, const FOperatorJumpTable& InJumpTable);
+}
\ No newline at end of file
diff --git a/Engine/Source/Runtime/Core/Public/Misc/ExpressionParserTypes.h b/Engine/Source/Runtime/Core/Public/Misc/ExpressionParserTypes.h
new file mode 100644
--- /dev/null
+++ b/Engine/Source/Runtime/Core/Public/Misc/ExpressionParserTypes.h
@@ -0,0 +1,583 @@
+
+#pragma once
+
+#include "ValueOrError.h"
+
+class FExpressionNode;
+struct FExpressionError;
+
+typedef TValueOrError<FExpressionNode, FExpressionError> FExpressionResult;
+
+/** Simple error structure used for reporting parse errors */
+struct FExpressionError
+{
+	FExpressionError(const FText& InText) : Text(InText) {}
+
+	FText Text;
+};
+
+/** Simple struct that defines a specific token contained in an FTokenStream  */
+class FStringToken
+{
+public:
+	FStringToken() : TokenStart(nullptr), TokenEnd(nullptr), LineNumber(0), CharacterIndex(0) {}
+
+	/** Get the string representation of this token */
+	FString GetString() const { return FString(TokenEnd - TokenStart, TokenStart); }
+
+	/** Check if this token is valid */
+	bool IsValid() const { return TokenEnd != TokenStart; }
+
+	/** Get the position of the start and end of this token in the stream */
+	const TCHAR* GetTokenStartPos() const { return TokenStart; }
+	const TCHAR* GetTokenEndPos() const { return TokenEnd; }
+
+	/** Contextual information about this token */
+	int32 GetCharacterIndex() const { return CharacterIndex; }
+	int32 GetLineNumber() const { return LineNumber; }
+	
+	/** Accumulate another token into this one */
+	void Accumulate(const FStringToken& InToken) { if (InToken.TokenEnd > TokenEnd) { TokenEnd = InToken.TokenEnd; } }
+
+protected:
+	friend class FTokenStream;
+
+	FStringToken(const TCHAR* InStart, int32 Line = 0, int32 Character = 0)
+		: TokenStart(InStart), TokenEnd(InStart), LineNumber(Line), CharacterIndex(Character)
+	{}
+
+	/** The start of the token */
+	const TCHAR* TokenStart;
+	/** The end of the token */
+	const TCHAR* TokenEnd;
+	/** Line number and Character index */
+	int32 LineNumber, CharacterIndex;
+};
+
+/** Enum specifying how to treat the currently parsing character. */
+enum class EParseState
+{
+	/** Include this character in the token and continue consuming */
+	Continue,
+	/** Include this character in the token and stop consuming */
+	StopAfter,
+	/** Exclude this character from the token and stop consuming */
+	StopBefore, 
+	/** Cancel parsing this token, and return nothing. */
+	Cancel,
+};
+
+/** A token stream wraps up a raw string, providing accessors into it for consuming tokens */
+class CORE_API FTokenStream
+{
+public:
+
+	/**
+	 * Parse out a token using the supplied predicate.
+	 * Will keep consuming characters into the resulting token provided the predicate returns EParseState::Continue or EParseState::StopAfter.
+	 * Optionally supply a token to accumulate into
+	 * Returns a string token for the stream, or empty on error
+	 */
+	TOptional<FStringToken> ParseToken(const TFunctionRef<EParseState(TCHAR)>& Pred, FStringToken* Accumulate = nullptr) const;
+
+	/** Attempt parse out the specified pre-defined string from the current read position (or accumulating into the specified existing token) */
+	TOptional<FStringToken> ParseToken(const TCHAR* Symbol, FStringToken* Accumulate = nullptr) const;
+	TOptional<FStringToken> ParseTokenIgnoreCase(const TCHAR* Symbol, FStringToken* Accumulate = nullptr) const;
+
+	/** Return a string token for the next character in the stream (or accumulating into the specified existing token) */
+	TOptional<FStringToken> ParseSymbol(FStringToken* Accumulate = nullptr) const;
+
+	/** Attempt parse out the specified pre-defined string from the current read position (or accumulating into the specified existing token) */
+	TOptional<FStringToken> ParseSymbol(TCHAR Symbol, FStringToken* Accumulate = nullptr) const;
+
+	/** Parse a whitespace token */
+	TOptional<FStringToken> ParseWhitespace(FStringToken* Accumulate = nullptr) const;
+
+	/** Generate a token for the specified number of chars, at the current read position (or end of Accumulate) */
+	TOptional<FStringToken> GenerateToken(int32 NumChars, FStringToken* Accumulate = nullptr) const;
+
+public:
+
+	/** Constructor. The stream is only valid for the lifetime of the string provided */
+	FTokenStream(const TCHAR* In);
+	
+	/** Peek at the character at the specified offset from the current read position */
+	TCHAR PeekChar(int32 Offset = 0) const;
+
+	/** Get the number of characters remaining in the stream after the current read position */
+	int32 CharsRemaining() const;
+
+	/** Check if it is valid to read (the optional number of characters) from the specified position */
+	bool IsReadPosValid(const TCHAR* InPos, int32 MinNumChars = 1) const;
+
+	/** Check if the stream is empty */
+	bool IsEmpty() const;
+
+	/** Get the current read position from the start of the stream */
+	int32 GetPosition() const;
+
+	const TCHAR* GetStart() const { return Start; }
+	const TCHAR* GetRead() const { return ReadPos; }
+	const TCHAR* GetEnd() const { return End; }
+
+	/** Get the error context from the current read position */
+	FString GetErrorContext() const;
+
+	/** Set the current read position to the character proceeding the specified token */
+	void SetReadPos(const FStringToken& Token);
+
+private:
+
+	/** The start of the expression */
+	const TCHAR* Start;
+	/** The end of the expression */
+	const TCHAR* End;
+	/** The current read position in the expression */
+	const TCHAR* ReadPos;
+};
+
+/** Helper macro to define the necessary template specialization for a particular expression node type */
+/** Variable length arguments are passed the FGuid constructor. Must be unique per type */
+#define DEFINE_EXPRESSION_NODE_TYPE(TYPE, ...) \
+template<> struct TGetExpressionNodeTypeId<TYPE>\
+{\
+	static const FGuid& GetTypeId()\
+	{\
+		static FGuid Global(__VA_ARGS__);\
+		return Global;\
+	}\
+};
+template<typename T> struct TGetExpressionNodeTypeId;
+
+/**
+ * A node in an expression.
+ * 	Can be constructed from any C++ type that has a corresponding DEFINE_EXPRESSION_NODE_TYPE.
+ * 	Evaluation behaviour (unary/binary operator etc) is defined in the expression grammar, rather than the type itself.
+ */
+class FExpressionNode
+{
+public:
+	/** Construction from client expression data type */
+	template<typename T, typename U = typename TEnableIf<!TPointerIsConvertibleFromTo<T, FExpressionNode>::Value>::Type>
+	FExpressionNode(const T& In) : TypeId(TGetExpressionNodeTypeId<T>::GetTypeId())
+	{
+		if (sizeof(T) <= SizeOfInlineBytes)
+		{
+			FMemory::Memcpy(InlineBytes, &In, sizeof(T));
+		}
+		else
+		{
+			BigData = MakeShareable(new T(MoveTemp(In)));
+		}
+	}
+
+	/** Copy construction/assignment */
+	FExpressionNode(const FExpressionNode& In) { *this = In; }
+	FExpressionNode& operator=(const FExpressionNode& In)
+	{
+		TypeId = In.TypeId;
+		if (In.BigData.IsValid())
+		{
+			BigData = In.BigData;
+		}
+		else
+		{
+			BigData = nullptr;
+			FMemory::Memcpy(InlineBytes, In.InlineBytes, SizeOfInlineBytes);
+		}
+		return *this;
+	}
+
+	/** Get the type identifier of this node */
+	const FGuid& GetTypeId() const { return TypeId; }
+
+	/** Cast this node to the specified type. Will asset if the types do not match. */
+	template<typename T>
+	const T* Cast() const
+	{
+		if (TypeId == TGetExpressionNodeTypeId<T>::GetTypeId())
+		{
+			if (BigData.IsValid())
+			{
+				return reinterpret_cast<const T*>(BigData.Get());
+			}
+			else
+			{
+				return reinterpret_cast<const T*>(InlineBytes);
+			}
+		}
+		return nullptr;
+	}
+
+private:
+
+	static const int32 SizeOfInlineBytes = 128 - sizeof(FGuid) - sizeof(TSharedPtr<void>);
+
+	/** TypeID - 16 bytes */
+	FGuid TypeId;
+	/** SharedPtr for large types - 2*sizeof(void*) */
+	TSharedPtr<void> BigData;
+	uint8 InlineBytes[SizeOfInlineBytes];
+};
+
+/** A specific token in a stream. Comprises an expression node, and the stream token it was created from */
+class FExpressionToken
+{
+public:
+	FExpressionToken(const FStringToken& InContext, FExpressionNode InNode)
+		: Node(MoveTemp(InNode))
+		, Context(InContext)
+	{
+	}
+
+	FExpressionToken(FExpressionToken&& In) : Node(MoveTemp(In.Node)), Context(In.Context) {}
+	FExpressionToken& operator=(FExpressionToken&& In) { Node = MoveTemp(In.Node); Context = MoveTemp(In.Context); return *this; }
+
+	FExpressionToken(const FExpressionToken& In) : Node(In.Node), Context(In.Context) {}
+	FExpressionToken& operator=(const FExpressionToken& In) { Node = In.Node; Context = In.Context; return *this; }
+
+	FExpressionNode Node;
+	FStringToken Context;
+};
+
+/** Built-in wrapper tokens used for compilation */
+struct FGroupMarker
+{
+	FGroupMarker(FExpressionToken InWrappedToken) : WrappedToken(MoveTemp(InWrappedToken)), NumTokens(0) {}
+
+	template<typename T>	bool 		IsA() const 	{ return WrappedToken.Node.GetTypeId() == TGetExpressionNodeTypeId<T>::GetTypeId(); }
+	template<typename T>	const T* 	Cast() const	{ return WrappedToken.Node.Cast<T>(); }
+
+	FExpressionToken WrappedToken;
+	int32 NumTokens;
+};
+DEFINE_EXPRESSION_NODE_TYPE(FGroupMarker, 0xEBF6DBA6, 0xE81B4684, 0x938EF46A, 0xD2FC0052)
+
+struct FWrappedOperator
+{
+	enum EType { PreUnary, PostUnary, Binary };
+	FWrappedOperator(EType InType, FExpressionToken InWrappedToken, int32 InPrec = 0) : Type(InType), WrappedToken(MoveTemp(InWrappedToken)), Precedence(InPrec) {}
+
+	template<typename T>	bool 		IsA() const 	{ return WrappedToken.Node.GetTypeId() == TGetExpressionNodeTypeId<T>::GetTypeId(); }
+	template<typename T>	const T* 	Cast() const	{ return WrappedToken.Node.Cast<T>(); }
+
+	EType Type;
+	FExpressionToken WrappedToken;
+	int32 Precedence;
+};
+DEFINE_EXPRESSION_NODE_TYPE(FWrappedOperator, 0x449B352B, 0x17B14A76, 0xBD264E7B, 0x9669D8F4)
+
+namespace impl
+{
+	template <typename> struct TParams;
+	template <typename> struct TParamsImpl;
+
+	template <typename Ret_, typename T, typename Arg1_>
+	struct TParamsImpl<Ret_ (T::*)(Arg1_)> { typedef Ret_ Ret; typedef Arg1_ Arg1; };
+	template <typename Ret_, typename T, typename Arg1_>
+	struct TParamsImpl<Ret_ (T::*)(Arg1_) const> { typedef Ret_ Ret; typedef Arg1_ Arg1; };
+
+	template <typename Ret_, typename T, typename Arg1_, typename Arg2_>
+	struct TParamsImpl<Ret_ (T::*)(Arg1_, Arg2_)> { typedef Ret_ Ret; typedef Arg1_ Arg1; typedef Arg2_ Arg2; };
+	template <typename Ret_, typename T, typename Arg1_, typename Arg2_>
+	struct TParamsImpl<Ret_ (T::*)(Arg1_, Arg2_) const> { typedef Ret_ Ret; typedef Arg1_ Arg1; typedef Arg2_ Arg2; };
+
+	template <typename T>
+	struct TParams : TParamsImpl<decltype(&T::operator())> {};
+}
+
+/** Jump table specifying how to execute an operator with different types */
+struct FOperatorJumpTable
+{
+	/** Execute the specified token as a unary operator, if such an overload exists */
+	FExpressionResult ExecPreUnary(const FExpressionToken& Operator, const FExpressionToken& R) const;
+	/** Execute the specified token as a unary operator, if such an overload exists */
+	FExpressionResult ExecPostUnary(const FExpressionToken& Operator, const FExpressionToken& L) const;
+	/** Execute the specified token as a binary operator, if such an overload exists */
+	FExpressionResult ExecBinary(const FExpressionToken& Operator, const FExpressionToken& L, const FExpressionToken& R) const;
+
+	/**
+	 * Map an expression node to a unary operator with the specified implementation.
+	 * This overload accepts callable types that return an FExpressionResult (and thus support returning an error)
+	 * Example usage that binds a '!' token to a function that attempts to do a boolean 'not'
+	 *		JumpTable.MapUnary<FExclamation>([](bool A) -> FExpressionResult { return MakeValue(!A); });
+	 */
+	template<typename OperatorType, typename FuncType>
+	typename TEnableIf<TIsSame<FExpressionResult, typename impl::TParams<FuncType>::Ret>::Value>::Type
+		MapPreUnary(FuncType InFunc)
+	{
+		typedef typename TRemoveConst<typename TRemoveReference<typename impl::TParams<FuncType>::Arg1>::Type>::Type OperandType;
+
+		FOperatorFunctionID ID = {
+			TGetExpressionNodeTypeId<OperatorType>::GetTypeId(),
+			FGuid(),
+			TGetExpressionNodeTypeId<OperandType>::GetTypeId()
+		};
+
+		PreUnaryOps.Add(ID, [=](const FExpressionNode& InOperand) -> FExpressionResult {
+			return InFunc(*InOperand.Cast<OperandType>());
+		});
+	}
+
+	/**
+	 * Map an expression node to a unary operator with the specified implementation.
+	 * This overload accepts callable types that return anything other than an FExpressionResult. Such functions cannot return errors.
+	 * The returned type is passed into an FExpressionNode, and as such must be a valid expression node type (see DEFINE_EXPRESSION_NODE_TYPE).
+	 * Example usage that binds a 'subtract' token to a function that negates the integer:
+	 *		JumpTable.MapUnary<FSubtractToken>([](int32 A){ return -A; });
+	 */
+	template<typename OperatorType, typename FuncType>
+	typename TEnableIf<!TIsSame<FExpressionResult, typename impl::TParams<FuncType>::Ret>::Value>::Type
+		MapPreUnary(FuncType InFunc)
+	{
+		typedef typename TRemoveConst<typename TRemoveReference<typename impl::TParams<FuncType>::Arg1>::Type>::Type OperandType;
+
+		FOperatorFunctionID ID = {
+			TGetExpressionNodeTypeId<OperatorType>::GetTypeId(),
+			FGuid(),
+			TGetExpressionNodeTypeId<OperandType>::GetTypeId()
+		};
+
+		PreUnaryOps.Add(ID, [=](const FExpressionNode& InOperand) -> FExpressionResult {
+			return MakeValue(InFunc(*InOperand.Cast<OperandType>()));
+		});
+	}
+
+	/**
+	 * Map an expression node to a unary operator with the specified implementation.
+	 * This overload accepts callable types that return an FExpressionResult (and thus support returning an error)
+	 * Example usage that binds a '!' token to a function that attempts to do a boolean 'not'
+	 *		JumpTable.MapUnary<FExclamation>([](bool A) -> FExpressionResult { return MakeValue(!A); });
+	 */
+	template<typename OperatorType, typename FuncType>
+	typename TEnableIf<TIsSame<FExpressionResult, typename impl::TParams<FuncType>::Ret>::Value>::Type
+		MapPostUnary(FuncType InFunc)
+	{
+		typedef typename TRemoveConst<typename TRemoveReference<typename impl::TParams<FuncType>::Arg1>::Type>::Type OperandType;
+
+		FOperatorFunctionID ID = {
+			TGetExpressionNodeTypeId<OperatorType>::GetTypeId(),
+			TGetExpressionNodeTypeId<OperandType>::GetTypeId(),
+			FGuid()
+		};
+
+		PostUnaryOps.Add(ID, [=](const FExpressionNode& InOperand) -> FExpressionResult {
+			return InFunc(*InOperand.Cast<OperandType>());
+		});
+	}
+
+	/**
+	 * Map an expression node to a unary operator with the specified implementation.
+	 * This overload accepts callable types that return anything other than an FExpressionResult. Such functions cannot return errors.
+	 * The returned type is passed into an FExpressionNode, and as such must be a valid expression node type (see DEFINE_EXPRESSION_NODE_TYPE).
+	 * Example usage that binds a 'subtract' token to a function that negates the integer:
+	 *		JumpTable.MapUnary<FSubtractToken>([](int32 A){ return -A; });
+	 */
+	template<typename OperatorType, typename FuncType>
+	typename TEnableIf<!TIsSame<FExpressionResult, typename impl::TParams<FuncType>::Ret>::Value>::Type
+		MapPostUnary(FuncType InFunc)
+	{
+		typedef typename TRemoveConst<typename TRemoveReference<typename impl::TParams<FuncType>::Arg1>::Type>::Type OperandType;
+
+		FOperatorFunctionID ID = {
+			TGetExpressionNodeTypeId<OperatorType>::GetTypeId(),
+			TGetExpressionNodeTypeId<OperandType>::GetTypeId(),
+			FGuid()
+		};
+
+		PostUnaryOps.Add(ID, [=](const FExpressionNode& InOperand) -> FExpressionResult {
+			return MakeValue(InFunc(*InOperand.Cast<OperandType>()));
+		});
+	}
+
+	/**
+	 * Map an expression node to a unary operator with the specified implementation.
+	 * This overload accepts callable types that return an FExpressionResult (and thus support returning an error)
+	 * Example usage that binds a 'slash' token to a function that subtracts two integers:
+	 *		JumpTable.MapBinary<FSlashToken>([](double A, double B){
+	 			if (B == 0) {
+	 				return MakeError(LOCTEXT("DivisionByZero", "Division by zero"));
+	 			} else {
+	 				return MakeValue(A / B);
+	 			}
+	 		});
+	 */
+	template<typename OperatorType, typename FuncType>
+	typename TEnableIf<TIsSame<FExpressionResult, typename impl::TParams<FuncType>::Ret>::Value>::Type
+		MapBinary(FuncType InFunc)
+	{
+		typedef typename TRemoveConst<typename TRemoveReference<typename impl::TParams<FuncType>::Arg1>::Type>::Type LeftOperandType;
+		typedef typename TRemoveConst<typename TRemoveReference<typename impl::TParams<FuncType>::Arg2>::Type>::Type RightOperandType;
+
+		FOperatorFunctionID ID = {
+			TGetExpressionNodeTypeId<OperatorType>::GetTypeId(),
+			TGetExpressionNodeTypeId<LeftOperandType>::GetTypeId(),
+			TGetExpressionNodeTypeId<RightOperandType>::GetTypeId()
+		};
+
+		BinaryOps.Add(ID, [=](const FExpressionNode& InLeftOperand, const FExpressionNode& InRightOperand){
+			return InFunc(*InLeftOperand.Cast<LeftOperandType>(), *InRightOperand.Cast<RightOperandType>());
+		});
+	}
+
+	/**
+	 * Map an expression node to a binary operator with the specified implementation.
+	 * This overload accepts callable types that return anything other than an FExpressionResult. Such functions cannot return errors.
+	 * The returned type is passed into an FExpressionNode, and as such must be a valid expression node type (see DEFINE_EXPRESSION_NODE_TYPE).
+	 * Example usage that binds a 'subtract' token to a function that subtracts two integers:
+	 *		JumpTable.MapBinary<FSubtractToken>([](int32 A, int32 B){ return A - B; });
+	 */
+	template<typename OperatorType, typename FuncType>
+	typename TEnableIf<!TIsSame<FExpressionResult, typename impl::TParams<FuncType>::Ret>::Value>::Type
+		MapBinary(FuncType InFunc)
+	{
+		typedef typename TRemoveConst<typename TRemoveReference<typename impl::TParams<FuncType>::Arg1>::Type>::Type LeftOperandType;
+		typedef typename TRemoveConst<typename TRemoveReference<typename impl::TParams<FuncType>::Arg2>::Type>::Type RightOperandType;
+
+		FOperatorFunctionID ID = {
+			TGetExpressionNodeTypeId<OperatorType>::GetTypeId(),
+			TGetExpressionNodeTypeId<LeftOperandType>::GetTypeId(),
+			TGetExpressionNodeTypeId<RightOperandType>::GetTypeId()
+		};
+
+		BinaryOps.Add(ID, [=](const FExpressionNode& InLeftOperand, const FExpressionNode& InRightOperand){
+			return MakeValue(InFunc(*InLeftOperand.Cast<LeftOperandType>(), *InRightOperand.Cast<RightOperandType>()));
+		});
+	}
+
+private:
+	/** Struct used to identify a function for a specific operator overload */
+	struct FOperatorFunctionID
+	{
+		FGuid OperatorType;
+		FGuid LeftOperandType;
+		FGuid RightOperandType;
+
+		friend bool operator==(const FOperatorFunctionID& A, const FOperatorFunctionID& B)
+		{
+			return A.OperatorType == B.OperatorType &&
+				A.LeftOperandType == B.LeftOperandType &&
+				A.RightOperandType == B.RightOperandType;
+		}
+
+		friend uint32 GetTypeHash(const FOperatorFunctionID& In)
+		{
+			const uint32 Hash = HashCombine(GetTypeHash(In.OperatorType), GetTypeHash(In.LeftOperandType));
+			return HashCombine(GetTypeHash(In.RightOperandType), Hash);
+		}
+	};
+
+	/** Maps of unary/binary operators */
+	TMap<FOperatorFunctionID, TFunction<FExpressionResult(const FExpressionNode&)>> PreUnaryOps;
+	TMap<FOperatorFunctionID, TFunction<FExpressionResult(const FExpressionNode&)>> PostUnaryOps;
+	TMap<FOperatorFunctionID, TFunction<FExpressionResult(const FExpressionNode&, const FExpressionNode&)>> BinaryOps;
+};
+
+/** Class used to consume tokens from a string */
+class CORE_API FExpressionTokenConsumer
+{
+public:
+	/** Construction from a raw string. The consumer is only valid as long as the string is valid */
+	FExpressionTokenConsumer(const TCHAR* InExpression);
+
+	/** Extract the list of tokens from this consumer */
+	TArray<FExpressionToken> Extract();
+
+	/** Add an expression node to the consumer, specifying the FStringToken this node relates to.
+	 *	Adding a node to the consumer will move its stream read position to the end of the added token.
+	 */
+	void Add(const FStringToken& SourceToken, FExpressionNode Node);
+
+	/** Get the expression stream */
+	FTokenStream& GetStream() { return Stream; }
+
+private:
+	FExpressionTokenConsumer(const FExpressionTokenConsumer&);
+	FExpressionTokenConsumer& operator=(const FExpressionTokenConsumer&);
+
+	/** Array of added tokens */
+	TArray<FExpressionToken> Tokens;
+
+	/** Stream that looks at the constructed expression */
+	FTokenStream Stream;
+};
+
+/** 
+ * Typedef that defines a function used to consume tokens
+ * 	Definitions may add FExpressionNodes parsed from the provided consumer's stream, or return an optional error.
+ *	Where a definition performs no mutable operations, subsequent token definitions will be invoked.
+ */
+typedef TOptional<FExpressionError> (FExpressionDefinition)(FExpressionTokenConsumer&);
+
+
+/** A lexeme dictionary defining how to lex an expression. */
+class CORE_API FTokenDefinitions
+{
+public:
+	FTokenDefinitions() : bIgnoreWhitespace(false) {}
+
+	/** Define the grammer to ignore whitespace between tokens, unless explicitly included in a token */
+	void IgnoreWhitespace() { bIgnoreWhitespace = true; }
+
+	/** Define a token by way of a function to be invoked to attempt to parse a token from a stream */
+	void DefineToken(const TFunction<FExpressionDefinition>& Definition);
+
+public:
+
+	/** Check if the grammar ignores whitespace */
+	bool DoesIgnoreWhitespace() { return bIgnoreWhitespace; }
+
+	/** Consume a token for the specified consumer */
+	TOptional<FExpressionError> ConsumeTokens(FExpressionTokenConsumer& Consumer) const;
+
+private:
+
+	TOptional<FExpressionError> ConsumeToken(FExpressionTokenConsumer& Consumer) const;
+
+private:
+
+	bool bIgnoreWhitespace;
+	TArray<TFunction<FExpressionDefinition>> Definitions;
+};
+
+/** A lexical gammer defining how to parse an expression. Clients must define the tokens and operators to be interpreted by the parser. */
+class CORE_API FExpressionGrammar
+{
+public:
+	/** Define a grouping operator from two expression node types */
+	template<typename TStartGroup, typename TEndGroup>
+	void DefineGrouping() { Groupings.Add(TGetExpressionNodeTypeId<TStartGroup>::GetTypeId(), TGetExpressionNodeTypeId<TEndGroup>::GetTypeId()); }
+
+	/** Define a pre-unary operator for the specified symbol */
+	template<typename TExpressionNode>
+	void DefinePreUnaryOperator() { PreUnaryOperators.Add(TGetExpressionNodeTypeId<TExpressionNode>::GetTypeId()); }
+
+	/** Define a post-unary operator for the specified symbol */
+	template<typename TExpressionNode>
+	void DefinePostUnaryOperator() { PostUnaryOperators.Add(TGetExpressionNodeTypeId<TExpressionNode>::GetTypeId()); }
+
+	/** Define a binary operator for the specified symbol, with the specified precedence */
+	template<typename TExpressionNode>
+	void DefineBinaryOperator(int32 InPrecedence) { BinaryOperators.Add(TGetExpressionNodeTypeId<TExpressionNode>::GetTypeId(), InPrecedence); }
+
+public:
+
+	/** Retrieve the corresponding grouping token for the specified open group type, or nullptr if it's not a group token */
+	const FGuid* GetGrouping(const FGuid& TypeId) const;
+
+	/** Check if this grammar defines a pre-unary operator for the specified symbol */
+	bool HasPreUnaryOperator(const FGuid& TypeId) const;
+	
+	/** Check if this grammar defines a post-unary operator for the specified symbol */
+	bool HasPostUnaryOperator(const FGuid& TypeId) const;
+
+	/** Get the binary operator precedence for the specified symbol, if any */
+	const int* GetBinaryOperatorPrecedence(const FGuid& TypeId) const;
+
+private:
+
+	TMap<FGuid, FGuid>	Groupings;
+	TSet<FGuid>			PreUnaryOperators;
+	TSet<FGuid>			PostUnaryOperators;
+	TMap<FGuid, int32>	BinaryOperators;
+};
\ No newline at end of file
diff --git a/Engine/Source/Runtime/Core/Public/Templates/ValueOrError.h b/Engine/Source/Runtime/Core/Public/Templates/ValueOrError.h
new file mode 100644
--- /dev/null
+++ b/Engine/Source/Runtime/Core/Public/Templates/ValueOrError.h
@@ -0,0 +1,74 @@
+// Copyright 1998-2015 Epic Games, Inc. All Rights Reserved.
+
+#pragma once
+
+template<typename ArgType>
+struct TValueOrError_ValueProxy
+{
+	TValueOrError_ValueProxy(ArgType&& InArg) : Arg(InArg) {}
+	ArgType& Arg;
+};
+
+template<typename ArgType>
+struct TValueOrError_ErrorProxy
+{
+	TValueOrError_ErrorProxy(ArgType&& InArg) : Arg(InArg) {}
+	ArgType& Arg;
+};
+
+template <typename A>
+FORCEINLINE TValueOrError_ValueProxy<A> MakeValue(A&& arg)
+{
+    return TValueOrError_ValueProxy<A>(Forward<A>(arg));
+}
+
+template <typename A>
+FORCEINLINE TValueOrError_ErrorProxy<A> MakeError(A&& arg)
+{
+    return TValueOrError_ErrorProxy<A>(Forward<A>(arg));
+}
+
+/** Type used to return either some data, or an error */
+template<typename ValueType, typename ErrorType>
+class TValueOrError
+{
+public:
+	/** Construct the result from a value, or an error (See MakeValue and MakeError) */
+	template<typename A>
+	TValueOrError(TValueOrError_ValueProxy<A>&& Proxy) 		: Value(Forward<A>(Proxy.Arg)) {}
+
+	template<typename A>
+	TValueOrError(TValueOrError_ErrorProxy<A>&& Proxy)		: Error(Forward<A>(Proxy.Arg)) {}
+
+	/** Copy construction/Assignment */
+	TValueOrError(const TValueOrError& In) 					: Error(In.Error), Value(In.Value) {}
+	TValueOrError& operator=(const TValueOrError& In) 		{ Error = In.Error; Value = In.Value; return *this; }
+
+	/** Move construction/assignment */
+	TValueOrError(TValueOrError&& In) 						: Error(MoveTemp(In.Error)), Value(MoveTemp(In.Value)) {}
+	TValueOrError& operator=(TValueOrError&& In) 			{ Error = MoveTemp(In.Error); Value = MoveTemp(In.Value); return *this; }
+
+	/** Check whether this value is valid */
+	bool IsValid() const				{ return Value.IsSet() && !Error.IsSet(); }
+
+	/** Get the error, if set */
+	ErrorType& GetError()				{ return Error.GetValue(); }
+	const ErrorType& GetError() const	{ return Error.GetValue(); }
+	
+	/** Steal this result's error, if set, causing it to become unset */
+	ErrorType StealError()				{ ErrorType Temp = MoveTemp(Error.GetValue()); Error.Reset(); return Temp; }
+
+	/** Access the value contained in this result */
+	ValueType& GetValue()				{ return Value.GetValue(); }
+	const ValueType& GetValue() const	{ return Value.GetValue(); }
+
+	/** Steal this result's value, causing it to become unset */
+	ValueType StealValue()				{ ValueType Temp = MoveTemp(Value.GetValue()); Value.Reset(); return Temp; }
+
+private:
+
+	/** The error reported by the procedure, if any */
+	TOptional<ErrorType> Error;
+	/** Optional value to return as part of the result */
+	TOptional<ValueType> Value;	
+};
\ No newline at end of file
