Added support for unit-based numeric math expressions in unit based numeric input boxes
This addresses UE-13765 and UE-13443

diff --git a/Engine/Source/Editor/EditorStyle/Private/EditorStyleClasses.cpp b/Engine/Source/Editor/EditorStyle/Private/EditorStyleClasses.cpp
--- a/Engine/Source/Editor/EditorStyle/Private/EditorStyleClasses.cpp
+++ b/Engine/Source/Editor/EditorStyle/Private/EditorStyleClasses.cpp
@@ -1,7 +1,6 @@
 // Copyright 1998-2015 Epic Games, Inc. All Rights Reserved.
 
 #include "EditorStylePrivatePCH.h"
-#include "UnitConversion.h"
 
 /* UEditorStyleSettings interface
  *****************************************************************************/
diff --git a/Engine/Source/Runtime/Core/Private/Math/UnitConversion.cpp b/Engine/Source/Runtime/Core/Private/Math/UnitConversion.cpp
--- a/Engine/Source/Runtime/Core/Private/Math/UnitConversion.cpp
+++ b/Engine/Source/Runtime/Core/Private/Math/UnitConversion.cpp
@@ -2,8 +2,18 @@
 
 #include "CorePrivatePCH.h"
 #include "UnitConversion.h"
+#include "BasicMathExpressionEvaluator.h"
 
-FUnitConversion::FParseCandidate FUnitConversion::ParseCandidates[] = {
+#define LOCTEXT_NAMESPACE "UnitConversion"
+
+/** Structure used to match units when parsing */
+struct FParseCandidate
+{
+	const TCHAR* String;
+	EUnit Unit;
+};
+
+FParseCandidate ParseCandidates[] = {
 	
 	{ TEXT("Micrometers"),			EUnit::Micrometers },			{ TEXT("um"),		EUnit::Micrometers }, 			{ TEXT("\u00B5m"),	EUnit::Micrometers },
 	{ TEXT("Millimeters"),			EUnit::Millimeters },			{ TEXT("mm"),		EUnit::Millimeters },
@@ -63,7 +73,8 @@ FUnitConversion::FParseCandidate FUnitConversion::ParseCandidates[] = {
 	{ TEXT("Years"),				EUnit::Years },					{ TEXT("yr"),		EUnit::Years },
 };
 
-const TCHAR* const FUnitConversion::DisplayStrings[] = {
+/** Static array of display strings that directly map to EUnit enumerations */
+const TCHAR* const DisplayStrings[] = {
 	TEXT("\u00B5m"),			TEXT("mm"),					TEXT("cm"),					TEXT("m"),					TEXT("km"),
 	TEXT("in"),					TEXT("ft"),					TEXT("yd"),					TEXT("mi"),
 	TEXT("ly"),
@@ -88,7 +99,7 @@ const TCHAR* const FUnitConversion::DisplayStrings[] = {
 	TEXT("ms"), TEXT("s"), TEXT("min"), TEXT("hr"), TEXT("dy"), TEXT("mth"), TEXT("yr"),
 };
 
-const EUnitType FUnitConversion::Types[] = {
+const EUnitType UnitTypes[] = {
 	EUnitType::Distance,	EUnitType::Distance,	EUnitType::Distance,	EUnitType::Distance,	EUnitType::Distance,
 	EUnitType::Distance,	EUnitType::Distance,	EUnitType::Distance,	EUnitType::Distance,
 	EUnitType::Distance,
@@ -113,6 +124,205 @@ const EUnitType FUnitConversion::Types[] = {
 	EUnitType::Time,		EUnitType::Time,		EUnitType::Time,		EUnitType::Time,		EUnitType::Time,		EUnitType::Time,		EUnitType::Time,
 };
 
+
+
+DEFINE_EXPRESSION_NODE_TYPE(FNumericUnit<double>, 0x3C138BC9, 0x71314F0B, 0xBB469BF7, 0xED47D147)
+
+struct FUnitExpressionParser
+{
+	FUnitExpressionParser(EUnit InDefaultUnit)
+	{
+		using namespace ExpressionParser;
+
+		TokenDefinitions.IgnoreWhitespace();
+			
+		// Defined in order of importance
+		TokenDefinitions.DefineToken(&ConsumeSymbol<FPlus>);
+		TokenDefinitions.DefineToken(&ConsumeSymbol<FMinus>);
+		TokenDefinitions.DefineToken(&ConsumeSymbol<FStar>);
+		TokenDefinitions.DefineToken(&ConsumeSymbol<FForwardSlash>);
+		TokenDefinitions.DefineToken(&ConsumeSymbol<FSubExpressionStart>);
+		TokenDefinitions.DefineToken(&ConsumeSymbol<FSubExpressionEnd>);
+
+		TokenDefinitions.DefineToken([&](FExpressionTokenConsumer& Consumer){ return ConsumeNumberWithUnits(Consumer); });
+
+		Grammar.DefineGrouping<FSubExpressionStart, FSubExpressionEnd>();
+		
+		Grammar.DefinePreUnaryOperator<FPlus>();
+		Grammar.DefinePreUnaryOperator<FMinus>();
+
+		Grammar.DefineBinaryOperator<FPlus>(5);
+		Grammar.DefineBinaryOperator<FMinus>(5);
+		Grammar.DefineBinaryOperator<FStar>(4);
+		Grammar.DefineBinaryOperator<FForwardSlash>(4);
+
+		// Unary operators for numeric units
+		JumpTable.MapPreUnary<FPlus>(	[](const FNumericUnit<double>& N) {	return FNumericUnit<double>(N.Value, N.Units);		});
+		JumpTable.MapPreUnary<FMinus>(	[](const FNumericUnit<double>& N) {	return FNumericUnit<double>(-N.Value, N.Units);		});
+
+		/** Addition for numeric units */
+		JumpTable.MapBinary<FPlus>([InDefaultUnit](const FNumericUnit<double>& A, const FNumericUnit<double>& B) -> FExpressionResult {
+
+			// Have to ensure we're adding the correct units here
+
+			EUnit UnitsLHS = A.Units, UnitsRHS = B.Units;
+
+			if (UnitsLHS == EUnit::Unspecified && UnitsRHS != EUnit::Unspecified)
+			{
+				UnitsLHS = InDefaultUnit;
+			}
+			else if (UnitsLHS != EUnit::Unspecified && UnitsRHS == EUnit::Unspecified)
+			{
+				UnitsRHS = InDefaultUnit;
+			}
+
+			if (FUnitConversion::AreUnitsCompatible(UnitsLHS, UnitsRHS))
+			{
+				if (UnitsLHS != EUnit::Unspecified)
+				{
+					return MakeValue(FNumericUnit<double>(A.Value + FUnitConversion::Convert(B.Value, UnitsRHS, UnitsLHS), UnitsLHS));
+				}
+				else
+				{
+					return MakeValue(FNumericUnit<double>(FUnitConversion::Convert(A.Value, UnitsLHS, UnitsRHS) + B.Value, UnitsRHS));
+				}
+			}
+
+			FFormatOrderedArguments Args;
+			Args.Add(FText::FromString(FUnitConversion::GetUnitDisplayString(B.Units)));
+			Args.Add(FText::FromString(FUnitConversion::GetUnitDisplayString(A.Units)));
+			return MakeError(FText::Format(LOCTEXT("CannotAddErr", "Cannot add {0} to {1}"), Args));
+		});
+
+		/** Subtraction for numeric units */
+		JumpTable.MapBinary<FMinus>([InDefaultUnit](const FNumericUnit<double>& A, const FNumericUnit<double>& B) -> FExpressionResult {
+			
+			// Have to ensure we're adding the correct units here
+
+			EUnit UnitsLHS = A.Units, UnitsRHS = B.Units;
+
+			if (UnitsLHS == EUnit::Unspecified && UnitsRHS != EUnit::Unspecified)
+			{
+				UnitsLHS = InDefaultUnit;
+			}
+			else if (UnitsLHS != EUnit::Unspecified && UnitsRHS == EUnit::Unspecified)
+			{
+				UnitsRHS = InDefaultUnit;
+			}
+
+			if (FUnitConversion::AreUnitsCompatible(UnitsLHS, UnitsRHS))
+			{
+				if (UnitsLHS != EUnit::Unspecified)
+				{
+					return MakeValue(FNumericUnit<double>(A.Value - FUnitConversion::Convert(B.Value, UnitsRHS, UnitsLHS), UnitsLHS));
+				}
+				else
+				{
+					return MakeValue(FNumericUnit<double>(FUnitConversion::Convert(A.Value, UnitsLHS, UnitsRHS) - B.Value, UnitsRHS));
+				}
+			}
+
+			FFormatOrderedArguments Args;
+			Args.Add(FText::FromString(FUnitConversion::GetUnitDisplayString(B.Units)));
+			Args.Add(FText::FromString(FUnitConversion::GetUnitDisplayString(A.Units)));
+			return MakeError(FText::Format(LOCTEXT("CannotSubtractErr", "Cannot subtract {1} from {0}"), Args));
+		});
+
+		/** Multiplication */
+		JumpTable.MapBinary<FStar>([](const FNumericUnit<double>& A, const FNumericUnit<double>& B) -> FExpressionResult {
+			if (A.Units != EUnit::Unspecified && B.Units != EUnit::Unspecified)
+			{
+				return MakeError(LOCTEXT("InvalidMultiplication", "Cannot multiply by numbers with units"));
+			}
+			return MakeValue(FNumericUnit<double>(B.Value*A.Value, A.Units == EUnit::Unspecified ? B.Units : A.Units));
+		});
+
+		/** Division */
+		JumpTable.MapBinary<FForwardSlash>([](const FNumericUnit<double>& A, const FNumericUnit<double>& B) -> FExpressionResult {
+			if (B.Units != EUnit::Unspecified)
+			{
+				return MakeError(LOCTEXT("InvalidDivision", "Cannot divide by numbers with units"));
+			}
+			else if (B.Value == 0)
+			{
+				return MakeError(LOCTEXT("DivideByZero", "DivideByZero"));	
+			}
+			return MakeValue(FNumericUnit<double>(A.Value/B.Value, A.Units));
+		});
+	}
+
+	/** Consume a number from the stream, optionally including units */
+	TOptional<FExpressionError> ConsumeNumberWithUnits(FExpressionTokenConsumer& Consumer)
+	{
+		auto& Stream = Consumer.GetStream();
+
+		if (!FChar::IsDigit(Stream.PeekChar()))
+		{
+			return TOptional<FExpressionError>();
+		}
+
+		TOptional<FStringToken> NumberToken = ExpressionParser::ParseNumber(Stream);
+		
+		if (NumberToken.IsSet())
+		{
+			// Skip over whitespace
+			Stream.ParseToken([](TCHAR InC){ return FChar::IsWhitespace(InC) ? EParseState::Continue : EParseState::StopBefore; }, &NumberToken.GetValue());
+
+			TOptional<EUnit> Unit;
+			TOptional<FStringToken> UnitToken;
+			for (int32 Index = 0; Index < ARRAY_COUNT(ParseCandidates); ++Index)
+			{
+				UnitToken = Stream.ParseTokenIgnoreCase(ParseCandidates[Index].String, &NumberToken.GetValue());
+				if (UnitToken.IsSet())
+				{
+					Unit = ParseCandidates[Index].Unit;
+					break;
+				}
+			}
+
+			const double Value = FCString::Atod(*NumberToken.GetValue().GetString());
+			if (Unit.IsSet())
+			{
+				Consumer.Add(NumberToken.GetValue(), FNumericUnit<double>(Value, Unit.GetValue()));
+			}
+			else
+			{
+				Consumer.Add(NumberToken.GetValue(), FNumericUnit<double>(Value));
+			}
+		}
+
+		return TOptional<FExpressionError>();
+	}
+
+	TValueOrError<FNumericUnit<double>, FExpressionError> Evaluate(const TCHAR* InExpression) const
+	{
+		TValueOrError<FExpressionNode, FExpressionError> Result = ExpressionParser::Evaluate(InExpression, TokenDefinitions, Grammar, JumpTable);
+		if (!Result.IsValid())
+		{
+			return MakeError(Result.GetError());
+		}
+
+		auto& Node = Result.GetValue();
+		if (const auto* Numeric = Node.Cast<double>())
+		{
+			return MakeValue(FNumericUnit<double>(*Numeric, EUnit::Unspecified));
+		}
+		else if (const auto* NumericUnit = Node.Cast<FNumericUnit<double>>())
+		{
+			return MakeValue(*NumericUnit);
+		}
+		else
+		{
+			return MakeError(LOCTEXT("UnrecognizedResult", "Unrecognized result returned from expression"));
+		}
+	}
+	
+private:
+	FTokenDefinitions TokenDefinitions;
+	FExpressionGrammar Grammar;
+	FOperatorJumpTable JumpTable;
+};
+
 FUnitSettings::FUnitSettings()
 	: bGlobalUnitDisplay(true)
 {
@@ -128,7 +338,18 @@ FUnitSettings::FUnitSettings()
 	DisplayUnits[(uint8)EUnitType::Time].Add(EUnit::Seconds);
 }
 
-const TArray<EUnit>& FUnitSettings::GetDisplayUnits(EUnitType InType)
+bool FUnitSettings::ShouldDisplayUnits() const
+{
+	return bGlobalUnitDisplay;
+}
+
+void FUnitSettings::SetShouldDisplayUnits(bool bInGlobalUnitDisplay)
+{
+	bGlobalUnitDisplay = bInGlobalUnitDisplay;
+	SettingChangedEvent.Broadcast();
+}
+
+const TArray<EUnit>& FUnitSettings::GetDisplayUnits(EUnitType InType) const
 {
 	return DisplayUnits[(uint8)InType];
 }
@@ -171,11 +392,33 @@ FUnitSettings& FUnitConversion::Settings()
 	return *Settings;
 }
 
+/** Check whether it is possible to convert a number between the two specified units */
+bool FUnitConversion::AreUnitsCompatible(EUnit From, EUnit To)
+{
+	return From == EUnit::Unspecified || To == EUnit::Unspecified || GetUnitType(From) == GetUnitType(To);
+}
+
+/** Check whether a unit is of the specified type */
+bool FUnitConversion::IsUnitOfType(EUnit Unit, EUnitType Type)
+{
+	return Unit != EUnit::Unspecified && GetUnitType(Unit) == Type;
+}
+
+/** Get the type of the specified unit */
+EUnitType FUnitConversion::GetUnitType(EUnit InUnit)
+{
+	if (ensure(InUnit != EUnit::Unspecified))
+	{
+		return UnitTypes[(uint8)InUnit];
+	}
+	return EUnitType::NumberOf;
+}
+
 /** Get the unit abbreviation the specified unit type */
 const TCHAR* FUnitConversion::GetUnitDisplayString(EUnit Unit)
 {
-	static_assert(ARRAY_COUNT(FUnitConversion::Types) == (uint8)EUnit::Unspecified, "Type array does not match size of unit enum");
-	static_assert(ARRAY_COUNT(FUnitConversion::DisplayStrings) == (uint8)EUnit::Unspecified, "Display String array does not match size of unit enum");
+	static_assert(ARRAY_COUNT(UnitTypes) == (uint8)EUnit::Unspecified, "Type array does not match size of unit enum");
+	static_assert(ARRAY_COUNT(DisplayStrings) == (uint8)EUnit::Unspecified, "Display String array does not match size of unit enum");
 	
 	if (Unit != EUnit::Unspecified)
 	{
@@ -203,212 +446,240 @@ TOptional<EUnit> FUnitConversion::UnitFromString(const TCHAR* UnitString)
 	return TOptional<EUnit>();
 }
 
-double FUnitConversion::DistanceUnificationFactor(EUnit From)
+namespace UnitConversion
 {
-	// Convert to meters
-	double Factor = 1;
-	switch (From)
+
+	double DistanceUnificationFactor(EUnit From)
 	{
-		case EUnit::Micrometers:		return 0.000001;
-		case EUnit::Millimeters:		return 0.001;
-		case EUnit::Centimeters:		return 0.01;
-		case EUnit::Kilometers:			return 1000;
-
-		case EUnit::Lightyears:			return 9.4605284e15;
-
-		case EUnit::Miles:				Factor *= 1760;				// fallthrough
-		case EUnit::Yards:				Factor *= 3;				// fallthrough
-		case EUnit::Feet:				Factor *= 12;				// fallthrough
-		case EUnit::Inches:				Factor /= 39.3700787;		// fallthrough
-		default: 						return Factor;				// return
+		// Convert to meters
+		double Factor = 1;
+		switch (From)
+		{
+			case EUnit::Micrometers:		return 0.000001;
+			case EUnit::Millimeters:		return 0.001;
+			case EUnit::Centimeters:		return 0.01;
+			case EUnit::Kilometers:			return 1000;
+
+			case EUnit::Lightyears:			return 9.4605284e15;
+
+			case EUnit::Miles:				Factor *= 1760;				// fallthrough
+			case EUnit::Yards:				Factor *= 3;				// fallthrough
+			case EUnit::Feet:				Factor *= 12;				// fallthrough
+			case EUnit::Inches:				Factor /= 39.3700787;		// fallthrough
+			default: 						return Factor;				// return
+		}
 	}
-}
 
-double FUnitConversion::AngleUnificationFactor(EUnit From)
-{
-	// Convert to degrees
-	switch (From)
+	double AngleUnificationFactor(EUnit From)
 	{
-		case EUnit::Radians:			return (180 / PI);
-		default: 						return 1;
+		// Convert to degrees
+		switch (From)
+		{
+			case EUnit::Radians:			return (180 / PI);
+			default: 						return 1;
+		}
 	}
-}
 
-double FUnitConversion::SpeedUnificationFactor(EUnit From)
-{
-	// Convert to km/h
-	switch (From)
+	double SpeedUnificationFactor(EUnit From)
 	{
-		case EUnit::MetersPerSecond:	return 3.6;
-		case EUnit::MilesPerHour:		return DistanceUnificationFactor(EUnit::Miles) / 1000;
-		default: 						return 1;
+		// Convert to km/h
+		switch (From)
+		{
+			case EUnit::MetersPerSecond:	return 3.6;
+			case EUnit::MilesPerHour:		return DistanceUnificationFactor(EUnit::Miles) / 1000;
+			default: 						return 1;
+		}
 	}
-}
 
-double FUnitConversion::MassUnificationFactor(EUnit From)
-{
-	double Factor = 1;
-	// Convert to grams
-	switch (From)
+	double MassUnificationFactor(EUnit From)
 	{
-		case EUnit::Micrograms:			return 0.000001;
-		case EUnit::Milligrams:			return 0.001;
-		case EUnit::Kilograms:			return 1000;
-		case EUnit::MetricTons:			return 1000000;
-
-		case EUnit::Stones:				Factor *= 14;		// fallthrough
-		case EUnit::Pounds:				Factor *= 16;		// fallthrough
-		case EUnit::Ounces:				Factor *= 28.3495;	// fallthrough
-		default: 						return Factor;		// return
+		double Factor = 1;
+		// Convert to grams
+		switch (From)
+		{
+			case EUnit::Micrograms:			return 0.000001;
+			case EUnit::Milligrams:			return 0.001;
+			case EUnit::Kilograms:			return 1000;
+			case EUnit::MetricTons:			return 1000000;
+
+			case EUnit::Stones:				Factor *= 14;		// fallthrough
+			case EUnit::Pounds:				Factor *= 16;		// fallthrough
+			case EUnit::Ounces:				Factor *= 28.3495;	// fallthrough
+			default: 						return Factor;		// return
+		}
 	}
-}
 
-double FUnitConversion::ForceUnificationFactor(EUnit From)
-{
-	// Convert to Newtons
-	switch (From)
+	double ForceUnificationFactor(EUnit From)
 	{
-		case EUnit::PoundsForce:		return 4.44822162;
-		case EUnit::KilogramsForce:		return 9.80665;
-		default: 						return 1;
+		// Convert to Newtons
+		switch (From)
+		{
+			case EUnit::PoundsForce:		return 4.44822162;
+			case EUnit::KilogramsForce:		return 9.80665;
+			default: 						return 1;
+		}
 	}
-}
 
-double FUnitConversion::FrequencyUnificationFactor(EUnit From)
-{
-	// Convert to KHz
-	switch (From)
+	double FrequencyUnificationFactor(EUnit From)
 	{
-		case EUnit::Hertz:				return 0.001;
-		case EUnit::Megahertz:			return 1000;
-		case EUnit::Gigahertz:			return 1000000;
+		// Convert to KHz
+		switch (From)
+		{
+			case EUnit::Hertz:				return 0.001;
+			case EUnit::Megahertz:			return 1000;
+			case EUnit::Gigahertz:			return 1000000;
 
-		case EUnit::RevolutionsPerMinute:	return 0.001/60;
+			case EUnit::RevolutionsPerMinute:	return 0.001/60;
 
-		default: 						return 1;
+			default: 						return 1;
+		}
 	}
-}
 
-double FUnitConversion::DataSizeUnificationFactor(EUnit From)
-{
-	// Convert to MB
-	switch (From)
+	double DataSizeUnificationFactor(EUnit From)
 	{
-		case EUnit::Bytes:				return 1.0/(1024*1024);
-		case EUnit::Kilobytes:			return 1.0/1024;
-		case EUnit::Gigabytes:			return 1024;
-		case EUnit::Terabytes:			return 1024*1024;
+		// Convert to MB
+		switch (From)
+		{
+			case EUnit::Bytes:				return 1.0/(1024*1024);
+			case EUnit::Kilobytes:			return 1.0/1024;
+			case EUnit::Gigabytes:			return 1024;
+			case EUnit::Terabytes:			return 1024*1024;
 
-		default: 						return 1;
+			default: 						return 1;
+		}
 	}
-}
 
-double FUnitConversion::TimeUnificationFactor(EUnit From)
-{
-	// Convert to hours
-	double Factor = 1;
-	switch (From)
+	double TimeUnificationFactor(EUnit From)
 	{
-		case EUnit::Months:				return (365.242 * 24) / 12;
+		// Convert to hours
+		double Factor = 1;
+		switch (From)
+		{
+			case EUnit::Months:				return (365.242 * 24) / 12;
 
-		case EUnit::Years:				Factor *= 365.242;	// fallthrough
-		case EUnit::Days:				Factor *= 24;		// fallthrough
-										return Factor;
+			case EUnit::Years:				Factor *= 365.242;	// fallthrough
+			case EUnit::Days:				Factor *= 24;		// fallthrough
+											return Factor;
 
-		case EUnit::Milliseconds:		Factor /= 1000;		// fallthrough
-		case EUnit::Seconds:			Factor /= 60;		// fallthrough
-		case EUnit::Minutes:			Factor /= 60;		// fallthrough
-										return Factor;
+			case EUnit::Milliseconds:		Factor /= 1000;		// fallthrough
+			case EUnit::Seconds:			Factor /= 60;		// fallthrough
+			case EUnit::Minutes:			Factor /= 60;		// fallthrough
+											return Factor;
 
-		default: 						return 1;
+			default: 						return 1;
+		}
 	}
-}
 
-TOptional<const TArray<FUnitConversion::FQuantizationInfo>*> FUnitConversion::GetQuantizationBounds(EUnit Unit)
-{
-	struct FStaticBounds
+	TValueOrError<FNumericUnit<double>, FText> TryParseExpression(const TCHAR* InExpression, EUnit From)
+	{
+		const FUnitExpressionParser Parser(From);
+		auto Result = Parser.Evaluate(InExpression);
+		if (Result.IsValid())
+		{
+			if (Result.GetValue().Units == EUnit::Unspecified)
+			{
+				return MakeValue(FNumericUnit<double>(Result.GetValue().Value, From));
+			}
+			else
+			{
+				return MakeValue(Result.GetValue());
+			}
+		}
+		else
+		{
+			return MakeError(Result.GetError().Text);
+		}
+	}
+
+	TOptional<const TArray<FQuantizationInfo>*> GetQuantizationBounds(EUnit Unit)
 	{
-		TArray<FQuantizationInfo> MetricDistance;
-		TArray<FQuantizationInfo> ImperialDistance;
+		struct FStaticBounds
+		{
+			TArray<FQuantizationInfo> MetricDistance;
+			TArray<FQuantizationInfo> ImperialDistance;
 
-		TArray<FQuantizationInfo> MetricMass;
-		TArray<FQuantizationInfo> ImperialMass;
+			TArray<FQuantizationInfo> MetricMass;
+			TArray<FQuantizationInfo> ImperialMass;
 
-		TArray<FQuantizationInfo> Frequency;
-		TArray<FQuantizationInfo> DataSize;
+			TArray<FQuantizationInfo> Frequency;
+			TArray<FQuantizationInfo> DataSize;
 
-		TArray<FQuantizationInfo> Time;
+			TArray<FQuantizationInfo> Time;
 
-		FStaticBounds()
-		{
-			MetricDistance.Emplace(EUnit::Micrometers,	1000);
-			MetricDistance.Emplace(EUnit::Millimeters,	10);
-			MetricDistance.Emplace(EUnit::Centimeters,	100);
-			MetricDistance.Emplace(EUnit::Meters,			1000);
-			MetricDistance.Emplace(EUnit::Kilometers,		0);
-
-			ImperialDistance.Emplace(EUnit::Inches,	12);
-			ImperialDistance.Emplace(EUnit::Feet,		3);
-			ImperialDistance.Emplace(EUnit::Yards,	1760);
-			ImperialDistance.Emplace(EUnit::Miles,	0);
-
-			MetricMass.Emplace(EUnit::Micrograms,	1000);
-			MetricMass.Emplace(EUnit::Milligrams,	1000);
-			MetricMass.Emplace(EUnit::Grams,		1000);
-			MetricMass.Emplace(EUnit::Kilograms,	1000);
-			MetricMass.Emplace(EUnit::MetricTons,	0);
-
-			ImperialMass.Emplace(EUnit::Ounces,	16);
-			ImperialMass.Emplace(EUnit::Pounds,	14);
-			ImperialMass.Emplace(EUnit::Stones,	0);
-
-			Frequency.Emplace(EUnit::Hertz,		1000);
-			Frequency.Emplace(EUnit::Kilohertz,	1000);
-			Frequency.Emplace(EUnit::Megahertz,	1000);
-			Frequency.Emplace(EUnit::Gigahertz,	0);
-
-			DataSize.Emplace(EUnit::Bytes,		1000);
-			DataSize.Emplace(EUnit::Kilobytes,	1000);
-			DataSize.Emplace(EUnit::Megabytes,	1000);
-			DataSize.Emplace(EUnit::Gigabytes,	1000);
-			DataSize.Emplace(EUnit::Terabytes,	0);
-
-			Time.Emplace(EUnit::Milliseconds,		1000);
-			Time.Emplace(EUnit::Seconds,			60);
-			Time.Emplace(EUnit::Minutes,			60);
-			Time.Emplace(EUnit::Hours,			24);
-			Time.Emplace(EUnit::Days,				365.242f / 12);
-			Time.Emplace(EUnit::Months,			12);
-			Time.Emplace(EUnit::Years,			0);
-		}
-	};
+			FStaticBounds()
+			{
+				MetricDistance.Emplace(EUnit::Micrometers,	1000);
+				MetricDistance.Emplace(EUnit::Millimeters,	10);
+				MetricDistance.Emplace(EUnit::Centimeters,	100);
+				MetricDistance.Emplace(EUnit::Meters,			1000);
+				MetricDistance.Emplace(EUnit::Kilometers,		0);
+
+				ImperialDistance.Emplace(EUnit::Inches,	12);
+				ImperialDistance.Emplace(EUnit::Feet,		3);
+				ImperialDistance.Emplace(EUnit::Yards,	1760);
+				ImperialDistance.Emplace(EUnit::Miles,	0);
+
+				MetricMass.Emplace(EUnit::Micrograms,	1000);
+				MetricMass.Emplace(EUnit::Milligrams,	1000);
+				MetricMass.Emplace(EUnit::Grams,		1000);
+				MetricMass.Emplace(EUnit::Kilograms,	1000);
+				MetricMass.Emplace(EUnit::MetricTons,	0);
+
+				ImperialMass.Emplace(EUnit::Ounces,	16);
+				ImperialMass.Emplace(EUnit::Pounds,	14);
+				ImperialMass.Emplace(EUnit::Stones,	0);
+
+				Frequency.Emplace(EUnit::Hertz,		1000);
+				Frequency.Emplace(EUnit::Kilohertz,	1000);
+				Frequency.Emplace(EUnit::Megahertz,	1000);
+				Frequency.Emplace(EUnit::Gigahertz,	0);
+
+				DataSize.Emplace(EUnit::Bytes,		1000);
+				DataSize.Emplace(EUnit::Kilobytes,	1000);
+				DataSize.Emplace(EUnit::Megabytes,	1000);
+				DataSize.Emplace(EUnit::Gigabytes,	1000);
+				DataSize.Emplace(EUnit::Terabytes,	0);
+
+				Time.Emplace(EUnit::Milliseconds,		1000);
+				Time.Emplace(EUnit::Seconds,			60);
+				Time.Emplace(EUnit::Minutes,			60);
+				Time.Emplace(EUnit::Hours,			24);
+				Time.Emplace(EUnit::Days,				365.242f / 12);
+				Time.Emplace(EUnit::Months,			12);
+				Time.Emplace(EUnit::Years,			0);
+			}
+		};
 
-	static FStaticBounds Bounds;
+		static FStaticBounds Bounds;
 
-	switch (Unit)
-	{
-	case EUnit::Micrometers: case EUnit::Millimeters: case EUnit::Centimeters: case EUnit::Meters: case EUnit::Kilometers:
-		return &Bounds.MetricDistance;
+		switch (Unit)
+		{
+		case EUnit::Micrometers: case EUnit::Millimeters: case EUnit::Centimeters: case EUnit::Meters: case EUnit::Kilometers:
+			return &Bounds.MetricDistance;
 
-	case EUnit::Inches: case EUnit::Feet: case EUnit::Yards: case EUnit::Miles:
-		return &Bounds.ImperialDistance;
+		case EUnit::Inches: case EUnit::Feet: case EUnit::Yards: case EUnit::Miles:
+			return &Bounds.ImperialDistance;
 
-	case EUnit::Micrograms: case EUnit::Milligrams: case EUnit::Grams: case EUnit::Kilograms: case EUnit::MetricTons:
-		return &Bounds.MetricMass;
+		case EUnit::Micrograms: case EUnit::Milligrams: case EUnit::Grams: case EUnit::Kilograms: case EUnit::MetricTons:
+			return &Bounds.MetricMass;
 
-	case EUnit::Ounces: case EUnit::Pounds: case EUnit::Stones:
-		return &Bounds.ImperialMass;
+		case EUnit::Ounces: case EUnit::Pounds: case EUnit::Stones:
+			return &Bounds.ImperialMass;
 
-	case EUnit::Hertz: case EUnit::Kilohertz: case EUnit::Megahertz: case EUnit::Gigahertz: case EUnit::RevolutionsPerMinute:
-		return &Bounds.Frequency;
+		case EUnit::Hertz: case EUnit::Kilohertz: case EUnit::Megahertz: case EUnit::Gigahertz: case EUnit::RevolutionsPerMinute:
+			return &Bounds.Frequency;
 
-	case EUnit::Bytes: case EUnit::Kilobytes: case EUnit::Megabytes: case EUnit::Gigabytes: case EUnit::Terabytes:
-		return &Bounds.DataSize;
+		case EUnit::Bytes: case EUnit::Kilobytes: case EUnit::Megabytes: case EUnit::Gigabytes: case EUnit::Terabytes:
+			return &Bounds.DataSize;
 
-	case EUnit::Milliseconds: case EUnit::Seconds: case EUnit::Minutes: case EUnit::Hours: case EUnit::Days: case EUnit::Months: case EUnit::Years:
-		return &Bounds.Time;
+		case EUnit::Milliseconds: case EUnit::Seconds: case EUnit::Minutes: case EUnit::Hours: case EUnit::Days: case EUnit::Months: case EUnit::Years:
+			return &Bounds.Time;
 
-	default:
-		return TOptional<const TArray<FQuantizationInfo>*>();
+		default:
+			return TOptional<const TArray<FQuantizationInfo>*>();
+		}
 	}
-}
\ No newline at end of file
+
+}	// namespace UnitConversion
+
+#undef LOCTEXT_NAMESPACE
\ No newline at end of file
diff --git a/Engine/Source/Runtime/Core/Public/Math/UnitConversion.h b/Engine/Source/Runtime/Core/Public/Math/UnitConversion.h
--- a/Engine/Source/Runtime/Core/Public/Math/UnitConversion.h
+++ b/Engine/Source/Runtime/Core/Public/Math/UnitConversion.h
@@ -3,7 +3,9 @@
 #pragma once
 
 #include "Optional.h"
-#include "Function.h"
+#include "ValueOrError.h"
+
+#include "ExpressionParser.h"
 
 /** Enum *must* be zero-indexed and sequential. Must be grouped by relevance and ordered by magnitude. */
 /** Enum *must* match the mirrored enum that exists in CoreUObject/Classes/Object.h for the purposes of UObject reflection */
@@ -65,11 +67,11 @@ public:
 	FUnitSettings();
 
 	/** Check whether unit display is globally enabled or disabled */
-	bool ShouldDisplayUnits() { return bGlobalUnitDisplay; }
-	void SetShouldDisplayUnits(bool bInGlobalUnitDisplay) { bGlobalUnitDisplay = bInGlobalUnitDisplay; SettingChangedEvent.Broadcast(); }
+	bool ShouldDisplayUnits() const;
+	void SetShouldDisplayUnits(bool bInGlobalUnitDisplay);
 	
-	/** Get the specific valid units to display the specified type of unit in */
-	const TArray<EUnit>& GetDisplayUnits(EUnitType InType);
+	/** Get/Set the specific valid units to display the specified type of unit in */
+	const TArray<EUnit>& GetDisplayUnits(EUnitType InType) const;
 	void SetDisplayUnits(EUnitType InType, const TArray<EUnit>& Units);
 	void SetDisplayUnits(EUnitType InType, EUnit Units);
 
@@ -95,64 +97,13 @@ struct CORE_API FUnitConversion
 	static FUnitSettings& Settings();
 
 	/** Check whether it is possible to convert a number between the two specified units */
-	static bool AreUnitsCompatible(EUnit From, EUnit To)
-	{
-		return From == EUnit::Unspecified || To == EUnit::Unspecified || Types[(uint8)From] == Types[(uint8)To];
-	}
+	static bool AreUnitsCompatible(EUnit From, EUnit To);
 
 	/** Check whether a unit is of the specified type */
-	static bool IsUnitOfType(EUnit Unit, EUnitType Type)
-	{
-		return Unit != EUnit::Unspecified && Types[(uint8)Unit] == Type;
-	}
+	static bool IsUnitOfType(EUnit Unit, EUnitType Type);
 
-	/** Convert the specified number from one unit to another. Does nothing if the units are incompatible. */
-	template<typename T>
-	static T Convert(T InValue, EUnit From, EUnit To)
-	{
-		if (!AreUnitsCompatible(From, To))
-		{
-			return InValue;
-		}
-		else if (From == EUnit::Unspecified || To == EUnit::Unspecified)
-		{
-			return InValue;
-		}
-
-		switch(Types[(uint8)From])
-		{
-			case EUnitType::Distance:			return InValue * DistanceUnificationFactor(From)		* (1.0 / DistanceUnificationFactor(To));
-			case EUnitType::Angle:				return InValue * AngleUnificationFactor(From) 			* (1.0 / AngleUnificationFactor(To));
-			case EUnitType::Speed:				return InValue * SpeedUnificationFactor(From) 			* (1.0 / SpeedUnificationFactor(To));
-			case EUnitType::Mass:				return InValue * MassUnificationFactor(From) 			* (1.0 / MassUnificationFactor(To));
-			case EUnitType::Force:				return InValue * ForceUnificationFactor(From) 			* (1.0 / ForceUnificationFactor(To));
-			case EUnitType::Frequency:			return InValue * FrequencyUnificationFactor(From) 		* (1.0 / FrequencyUnificationFactor(To));
-			case EUnitType::DataSize:			return InValue * DataSizeUnificationFactor(From) 		* (1.0 / DataSizeUnificationFactor(To));
-			case EUnitType::LuminousFlux:		return InValue;
-			case EUnitType::Time:				return InValue * TimeUnificationFactor(From) 			* (1.0 / TimeUnificationFactor(To));			
-			// Temperature conversion is not just a simple multiplication, so needs special treatment
-			case EUnitType::Temperature:
-			{
-				double NewValue = InValue;
-				// Put it into kelvin
-				switch (From)
-				{
-					case EUnit::Celsius:			NewValue = NewValue + 273.15;					break;
-					case EUnit::Farenheit:			NewValue = (NewValue + 459.67) * 5.f/9.f;		break;
-					default: 																		break;
-				}
-				// And out again
-				switch (To)
-				{
-					case EUnit::Celsius:			return NewValue - 273.15;
-					case EUnit::Farenheit:			return NewValue * 9.f/5.f - 459.67;
-					default: 						return NewValue;
-				}
-			}
-
-			default:							return InValue;
-		}
-	}
+	/** Get the type of the specified unit */
+	static EUnitType GetUnitType(EUnit);
 
 	/** Get the display string for the the specified unit type */
 	static const TCHAR* GetUnitDisplayString(EUnit Unit);
@@ -160,160 +111,20 @@ struct CORE_API FUnitConversion
 	/** Helper function to find a unit from a string (name or display string) */
 	static TOptional<EUnit> UnitFromString(const TCHAR* UnitString);
 
-private:
-
-	/** Find the common quantization factor for the specified distance unit. Quantizes to Meters. */
-	static double DistanceUnificationFactor(EUnit From);
-	/** Find the common quantization factor for the specified angular unit. Quantizes to Degrees. */
-	static double AngleUnificationFactor(EUnit From);
-	/** Find the common quantization factor for the specified speed unit. Quantizes to km/h. */
-	static double SpeedUnificationFactor(EUnit From);
-	/** Find the common quantization factor for the specified temperature unit. Quantizes to Kelvin. */
-	static double TemperatureUnificationFactor(EUnit From);
-	/** Find the common quantization factor for the specified mass unit. Quantizes to Grams. */
-	static double MassUnificationFactor(EUnit From);
-	/** Find the common quantization factor for the specified force unit. Quantizes to Newtons. */
-	static double ForceUnificationFactor(EUnit From);
-	/** Find the common quantization factor for the specified frequency unit. Quantizes to KHz. */
-	static double FrequencyUnificationFactor(EUnit From);
-	/** Find the common quantization factor for the specified data size unit. Quantizes to MB. */
-	static double DataSizeUnificationFactor(EUnit From);
-	/** Find the common quantization factor for the specified time unit. Quantizes to hours. */
-	static double TimeUnificationFactor(EUnit From);
-
 public:
 
-	/** Structure used to define the factor required to get from one unit type to the next. */
-	struct FQuantizationInfo
-	{
-		/** The unit to which this factor applies */
-		EUnit Units;
-		/** The factor by which to multiply to get to the next unit in this range */
-		float Factor;
-
-		/** Constructor */
-		FQuantizationInfo(EUnit InUnit, float InFactor) : Units(InUnit), Factor(InFactor) {}
-	};
-
-	/** Find the quantization bounds for the specified unit, if any */
-	static TOptional<const TArray<FQuantizationInfo>*> GetQuantizationBounds(EUnit Unit);
+	/** Convert the specified number from one unit to another. Does nothing if the units are incompatible. */
+	template<typename T>
+	static T Convert(T InValue, EUnit From, EUnit To);
 
 	/** Quantizes this number to the most appropriate unit for user friendly presentation (e.g. 1000m returns 1km). */
 	template<typename T>
-	static FNumericUnit<T> QuantizeUnitsToBestFit(T Value, EUnit Units)
-	{
-		auto OptionalBounds = FUnitConversion::GetQuantizationBounds(Units);
-		if (!OptionalBounds.IsSet())
-		{
-			return FNumericUnit<T>(Value, Units);
-		}
-
-		const auto& Bounds = *OptionalBounds.GetValue();
-
-		const int32 CurrentUnitIndex = (uint8)Units - (uint8)Bounds[0].Units;
-
-		EUnit NewUnits = Units;
-		double NewValue = Value;
-
-		if (FMath::Abs(NewValue) > 1)
-		{
-			// Large number? Try larger units
-			for (int32 Index = CurrentUnitIndex; Index < Bounds.Num(); ++Index)
-			{
-				if (Bounds[Index].Factor == 0)
-				{
-					break;
-				}
-
-				const auto Tmp = NewValue / Bounds[Index].Factor;
-
-				if (FMath::Abs(Tmp) < 1)
-				{
-					break;
-				}
-
-				NewValue = Tmp;
-				NewUnits = Bounds[Index + 1].Units;
-			}
-		}
-		else if (NewValue != 0)
-		{
-			// Small number? Try smaller units
-			for (int32 Index = CurrentUnitIndex - 1; Index >= 0; --Index)
-			{
-				NewValue *= Bounds[Index].Factor;
-				NewUnits = Bounds[Index].Units;
-
-				if (FMath::Abs(NewValue) > 1)
-				{
-					break;
-				}
-			}
-		}
-
-		return FNumericUnit<T>(NewValue, NewUnits);
-	}
+	static FNumericUnit<T> QuantizeUnitsToBestFit(T Value, EUnit Units);
 
 	/** Quantizes this number to the most appropriate unit for user friendly presentation (e.g. 1000m returns 1km), adhereing to global display settings. */
 	template<typename T>
-	static EUnit CalculateDisplayUnit(T Value, EUnit InUnits)
-	{
-		if (InUnits == EUnit::Unspecified)
-		{
-			return EUnit::Unspecified;
-		}
-
-		const TArray<EUnit>& DisplayUnits = Settings().GetDisplayUnits(Types[(uint8)InUnits]);
-		if (DisplayUnits.Num() == 0)
-		{
-			return QuantizeUnitsToBestFit(Value, InUnits).Units;
-		}
-		else if (DisplayUnits.Num() == 1)
-		{
-			return DisplayUnits[0];
-		}
-
-		// If the value we were given was 0, change it to something we can actually work with
-		if (Value == 0)
-		{
-			Value = 1;
-		}
-
-		int32 BestIndex = 0;
-		for (int32 Index = 0; Index < DisplayUnits.Num() - 1; ++Index)
-		{
-			double This = Convert(Value, InUnits, DisplayUnits[Index]);
-			double Next = Convert(Value, InUnits, DisplayUnits[Index + 1]);
-
-			if (FMath::Abs(FMath::LogX(10.0f, This)) < FMath::Abs(FMath::LogX(10.0f, Next)))
-			{
-				BestIndex = Index;
-			}
-			else
-			{
-				BestIndex = Index + 1;
-			}
-		}
-
-		return DisplayUnits[BestIndex];
-	}
-
-private:
+	static EUnit CalculateDisplayUnit(T Value, EUnit InUnits);
 
-	/** Structure used to match units when parsing */
-	struct FParseCandidate
-	{
-		const TCHAR* String;
-		EUnit Unit;
-	};
-	/** Arbitrary length array used for parsing input strings into EUnit types */
-	static FParseCandidate	 	ParseCandidates[];
-
-	/** Static array of display strings that directly map to EUnit enumerations */
-	static const TCHAR* const 	DisplayStrings[];
-
-	/** Static array of unit types that directly map to EUnit enumerations */
-	static const EUnitType 		Types[];
 };
 
 
@@ -330,111 +141,30 @@ struct FNumericUnit
 	const EUnit Units;
 
 	/** Constructors */
-	FNumericUnit() : Units(EUnit::Unspecified) {}
-	FNumericUnit(const NumericType& InValue, EUnit InUnits = EUnit::Unspecified) : Value(InValue), Units(InUnits) {}
+	FNumericUnit();
+	FNumericUnit(const NumericType& InValue, EUnit InUnits = EUnit::Unspecified);
 
 	/** Copy construction/assignment from the same type */
-	FNumericUnit(const FNumericUnit& Other) : Units(EUnit::Unspecified)							{	(*this) = Other;								}
-	FNumericUnit& operator=(const FNumericUnit& Other)											{ 	CopyValueWithConversion(Other);	return *this; 	}
+	FNumericUnit(const FNumericUnit& Other);
+	FNumericUnit& operator=(const FNumericUnit& Other);
 
 	/** Templated Copy construction/assignment from differing numeric types. Relies on implicit conversion of the two numeric types. */
-	template<typename OtherType> FNumericUnit(const FNumericUnit<OtherType>& Other)				{	(*this) = Other;								}
-	template<typename OtherType> FNumericUnit& operator=(const FNumericUnit<OtherType>& Other)	{	CopyValueWithConversion(Other);	return *this;	}
+	template<typename OtherType> FNumericUnit(const FNumericUnit<OtherType>& Other);
+	template<typename OtherType> FNumericUnit& operator=(const FNumericUnit<OtherType>& Other);
 
 	/** Convert this quantity to a different unit */
-	TOptional<FNumericUnit<NumericType>> ConvertTo(EUnit ToUnits) const
-	{
-		if (Units == EUnit::Unspecified)
-		{
-			return FNumericUnit(Value, ToUnits);
-		}
-		else if (FUnitConversion::AreUnitsCompatible(Units, ToUnits))
-		{
-			return FNumericUnit<NumericType>(FUnitConversion::Convert(Value, Units, ToUnits), ToUnits);
-		}
-
-		return TOptional<FNumericUnit<NumericType>>();
-	}
+	TOptional<FNumericUnit<NumericType>> ConvertTo(EUnit ToUnits) const;
 
 public:
 
 	/** Quantizes this number to the most appropriate unit for user friendly presentation (e.g. 1000m returns 1km). */
-	FNumericUnit<NumericType> QuantizeUnitsToBestFit() const
-	{
-		return FUnitConversion::QuantizeUnitsToBestFit(Value, Units);
-	}
-
-	/** Parse a numeric unit from a string */
-	static TOptional<FNumericUnit<NumericType>> TryParseString(const TCHAR* InSource)
-	{
-		TOptional<FNumericUnit<NumericType>> Result;
-		if (!InSource || !*InSource)
-		{
-			return Result;
-		}
-
-		const TCHAR* NumberEnd = nullptr;
-		if (!ExtractNumberBoundary(InSource, NumberEnd))
-		{
-			return Result;
-		}
-
-		NumericType NewValue;
-		LexicalConversion::FromString(NewValue, InSource);
-
-		// Now parse the units
-		while(FChar::IsWhitespace(*NumberEnd)) ++NumberEnd;
-
-		if (*NumberEnd == '\0')
-		{
-			// No units
-			Result.Emplace(NewValue);
-		}
-		else
-		{
-			// If the string specifies units, they must map to something that exists for this function to succeed
-			auto NewUnits = FUnitConversion::UnitFromString(NumberEnd);
-			if (NewUnits)
-			{
-				Result.Emplace(NewValue, NewUnits.GetValue());
-			}
-		}
-
-		return Result;
-	}
+	FNumericUnit<NumericType> QuantizeUnitsToBestFit() const;
 
-public:
-	/** Global arithmetic operators for number types. Deals with conversion from related units correctly. */
-	template<typename OtherType>
-	friend bool operator==(const FNumericUnit& LHS, const FNumericUnit<OtherType>& RHS)
-	{
-		if (LHS.Units != EUnit::Unspecified && RHS.Units != EUnit::Unspecified)
-		{
-			if (LHS.Units == RHS.Units)
-			{
-				return LHS.Value == RHS.Value;
-			}
-			else if (FUnitConversion::AreUnitsCompatible(LHS.Units, RHS.Units))
-			{
-				return LHS.Value == FUnitConversion::Convert(RHS.Value, RHS.Units, LHS.Units);
-			}
-			else
-			{
-				// Invalid conversion
-				return false;
-			}
-		}
-		else
-		{
-			return LHS.Value == RHS.Value;
-		}
-	}
+	/** Try and parse an expression into a numeric unit */
+	static TValueOrError<FNumericUnit<NumericType>, FText> TryParseExpression(const TCHAR* InExpression, EUnit InDefaultUnit);
 
-	template<typename OtherType>
-	friend bool operator!=(const FNumericUnit& LHS, const OtherType& RHS)
-	{
-		return !(LHS == RHS);
-	}
+	/** Parse a numeric unit from a string */
+	static TOptional<FNumericUnit<NumericType>> TryParseString(const TCHAR* InSource);
 
 private:
 	/** Conversion to the numeric type disabled as coupled with implicit construction from NumericType can easily lead to loss of associated units. */
@@ -442,112 +172,34 @@ private:
 
 	/** Copy another unit into this one, taking account of its units, and applying necessary conversion */
 	template<typename OtherType>
-	void CopyValueWithConversion(const FNumericUnit<OtherType>& Other)
-	{
-		if (Units != EUnit::Unspecified && Other.Units != EUnit::Unspecified)
-		{
-			if (Units == Other.Units)
-			{
-				Value = Other.Value;
-			}
-			else if (FUnitConversion::AreUnitsCompatible(Units, Other.Units))
-			{
-				Value = FUnitConversion::Convert(Other.Value, Other.Units, Units);
-			}
-			else
-			{
-				// Invalid conversion - assignment invalid
-			}
-		}
-		else
-		{
-			// If our units haven't been specified, we take on the units of the rhs
-			if (Units == EUnit::Unspecified)
-			{
-				// This is the only time we ever change units. Const_cast is 'acceptible' here since the units haven't been specified yet.
-				const_cast<EUnit&>(Units) = Other.Units;
-			}
-
-			Value = Other.Value;
-		}
-	}
+	void CopyValueWithConversion(const FNumericUnit<OtherType>& Other);
 
 	/** Given a string, skip past whitespace, then any numeric characters. Set End pointer to the end of the last numeric character. */
-	static bool ExtractNumberBoundary(const TCHAR* Start, const TCHAR*& End)
-	{
-		while(FChar::IsWhitespace(*Start)) ++Start;
-
-		End = Start;
-		if (*End == '-' || *End == '+')
-		{
-			End++;
-		}
-
-		bool bHasDot = false;
-		while (FChar::IsDigit(*End) || *End == '.')
-		{
-			if (*End == '.')
-			{
-				if (bHasDot)
-				{
-					return false;
-				}
-				bHasDot = true;
-			}
-			++End;
-		}
-
-		return true;
-	}
+	static bool ExtractNumberBoundary(const TCHAR* Start, const TCHAR*& End);
 };
 
+/** Global arithmetic operators for number types. Deals with conversion from related units correctly. */
+template<typename NumericType, typename OtherType>
+bool operator==(const FNumericUnit<NumericType>& LHS, const FNumericUnit<OtherType>& RHS);
 
-template <typename NumericType>
-struct TNumericLimits<FNumericUnit<NumericType>> : public TNumericLimits<NumericType>
-{ };
+template<typename NumericType, typename OtherType>
+bool operator!=(const FNumericUnit<NumericType>& LHS, const FNumericUnit<OtherType>& RHS);
 
 namespace LexicalConversion
 {
 	template<typename T>
-	FString ToString(const FNumericUnit<T>& NumericUnit)
-	{
-		FString String = ToString(NumericUnit.Value);
-		String += TEXT(" ");
-		String += FUnitConversion::GetUnitDisplayString(NumericUnit.Units);
-
-		return String;
-	}
+	FString ToString(const FNumericUnit<T>& NumericUnit);
 
 	template<typename T>
-	FString ToSanitizedString(const FNumericUnit<T>& NumericUnit)
-	{
-		FString String = ToSanitizedString(NumericUnit.Value);
-		String += TEXT(" ");
-		String += FUnitConversion::GetUnitDisplayString(NumericUnit.Units);
-
-		return String;
-	}
+	FString ToSanitizedString(const FNumericUnit<T>& NumericUnit);
 
 	template<typename T>
-	void FromString(FNumericUnit<T>& OutValue, const TCHAR* String)
-	{
-		auto Parsed = FNumericUnit<T>::TryParseString(String);
-		if (Parsed)
-		{
-			OutValue = Parsed.GetValue();
-		}
-	}
+	void FromString(FNumericUnit<T>& OutValue, const TCHAR* String);
 	
 	template<typename T>
-	bool TryParseString(FNumericUnit<T>& OutValue, const TCHAR* String)
-	{
-		auto Parsed = FNumericUnit<T>::TryParseString(String);
-		if (Parsed)
-		{
-			OutValue = Parsed.GetValue();
-			return true;
-		}
-
-		return false;
-	}
-}
\ No newline at end of file
+	bool TryParseString(FNumericUnit<T>& OutValue, const TCHAR* String);
+}
+
+
+// Include template definitions
+#include "UnitConversion.inl"
\ No newline at end of file
diff --git a/Engine/Source/Runtime/Core/Public/Math/UnitConversion.inl b/Engine/Source/Runtime/Core/Public/Math/UnitConversion.inl
new file mode 100644
--- /dev/null
+++ b/Engine/Source/Runtime/Core/Public/Math/UnitConversion.inl
@@ -0,0 +1,448 @@
+// Copyright 1998-2015 Epic Games, Inc. All Rights Reserved.
+
+/** Inline file for UnitConversion.h to separate the implementation from the header */
+
+
+namespace UnitConversion
+{
+	/** Find the common quantization factor for the specified distance unit. Quantizes to Meters. */
+	CORE_API double DistanceUnificationFactor(EUnit From);
+	/** Find the common quantization factor for the specified angular unit. Quantizes to Degrees. */
+	CORE_API double AngleUnificationFactor(EUnit From);
+	/** Find the common quantization factor for the specified speed unit. Quantizes to km/h. */
+	CORE_API double SpeedUnificationFactor(EUnit From);
+	/** Find the common quantization factor for the specified temperature unit. Quantizes to Kelvin. */
+	CORE_API double TemperatureUnificationFactor(EUnit From);
+	/** Find the common quantization factor for the specified mass unit. Quantizes to Grams. */
+	CORE_API double MassUnificationFactor(EUnit From);
+	/** Find the common quantization factor for the specified force unit. Quantizes to Newtons. */
+	CORE_API double ForceUnificationFactor(EUnit From);
+	/** Find the common quantization factor for the specified frequency unit. Quantizes to KHz. */
+	CORE_API double FrequencyUnificationFactor(EUnit From);
+	/** Find the common quantization factor for the specified data size unit. Quantizes to MB. */
+	CORE_API double DataSizeUnificationFactor(EUnit From);
+	/** Find the common quantization factor for the specified time unit. Quantizes to hours. */
+	CORE_API double TimeUnificationFactor(EUnit From);
+
+	/** Attempt to parse an expression into a numeric unit */
+	CORE_API TValueOrError<FNumericUnit<double>, FText> TryParseExpression(const TCHAR* InExpression, EUnit DefaultUnit);
+
+	/** Structure used to define the factor required to get from one unit type to the next. */
+	struct FQuantizationInfo
+	{
+		/** The unit to which this factor applies */
+		EUnit Units;
+		/** The factor by which to multiply to get to the next unit in this range */
+		float Factor;
+
+		/** Constructor */
+		FQuantizationInfo(EUnit InUnit, float InFactor) : Units(InUnit), Factor(InFactor) {}
+	};
+
+	/** Find the quantization bounds for the specified unit, if any */
+	CORE_API TOptional<const TArray<FQuantizationInfo>*> GetQuantizationBounds(EUnit Unit);
+
+}	// namespace UnitConversion
+
+/** Convert the specified number from one unit to another. Does nothing if the units are incompatible. */
+template<typename T>
+T FUnitConversion::Convert(T InValue, EUnit From, EUnit To)
+{
+	using namespace UnitConversion;
+
+	if (!AreUnitsCompatible(From, To))
+	{
+		return InValue;
+	}
+	else if (From == EUnit::Unspecified || To == EUnit::Unspecified)
+	{
+		return InValue;
+	}
+
+	switch(FUnitConversion::GetUnitType(From))
+	{
+		case EUnitType::Distance:			return InValue * DistanceUnificationFactor(From)		* (1.0 / DistanceUnificationFactor(To));
+		case EUnitType::Angle:				return InValue * AngleUnificationFactor(From) 			* (1.0 / AngleUnificationFactor(To));
+		case EUnitType::Speed:				return InValue * SpeedUnificationFactor(From) 			* (1.0 / SpeedUnificationFactor(To));
+		case EUnitType::Mass:				return InValue * MassUnificationFactor(From) 			* (1.0 / MassUnificationFactor(To));
+		case EUnitType::Force:				return InValue * ForceUnificationFactor(From) 			* (1.0 / ForceUnificationFactor(To));
+		case EUnitType::Frequency:			return InValue * FrequencyUnificationFactor(From) 		* (1.0 / FrequencyUnificationFactor(To));
+		case EUnitType::DataSize:			return InValue * DataSizeUnificationFactor(From) 		* (1.0 / DataSizeUnificationFactor(To));
+		case EUnitType::LuminousFlux:		return InValue;
+		case EUnitType::Time:				return InValue * TimeUnificationFactor(From) 			* (1.0 / TimeUnificationFactor(To));			
+		// Temperature conversion is not just a simple multiplication, so needs special treatment
+		case EUnitType::Temperature:
+		{
+			double NewValue = InValue;
+			// Put it into kelvin
+			switch (From)
+			{
+				case EUnit::Celsius:			NewValue = NewValue + 273.15;					break;
+				case EUnit::Farenheit:			NewValue = (NewValue + 459.67) * 5.f/9.f;		break;
+				default: 																		break;
+			}
+			// And out again
+			switch (To)
+			{
+				case EUnit::Celsius:			return NewValue - 273.15;
+				case EUnit::Farenheit:			return NewValue * 9.f/5.f - 459.67;
+				default: 						return NewValue;
+			}
+		}
+
+		default:							return InValue;
+	}
+}
+
+template<typename T>
+FNumericUnit<T> FUnitConversion::QuantizeUnitsToBestFit(T Value, EUnit Units)
+{
+	auto OptionalBounds = UnitConversion::GetQuantizationBounds(Units);
+	if (!OptionalBounds.IsSet())
+	{
+		return FNumericUnit<T>(Value, Units);
+	}
+
+	const auto& Bounds = *OptionalBounds.GetValue();
+
+	const int32 CurrentUnitIndex = (uint8)Units - (uint8)Bounds[0].Units;
+
+	EUnit NewUnits = Units;
+	double NewValue = Value;
+
+	if (FMath::Abs(NewValue) > 1)
+	{
+		// Large number? Try larger units
+		for (int32 Index = CurrentUnitIndex; Index < Bounds.Num(); ++Index)
+		{
+			if (Bounds[Index].Factor == 0)
+			{
+				break;
+			}
+
+			const auto Tmp = NewValue / Bounds[Index].Factor;
+
+			if (FMath::Abs(Tmp) < 1)
+			{
+				break;
+			}
+
+			NewValue = Tmp;
+			NewUnits = Bounds[Index + 1].Units;
+		}
+	}
+	else if (NewValue != 0)
+	{
+		// Small number? Try smaller units
+		for (int32 Index = CurrentUnitIndex - 1; Index >= 0; --Index)
+		{
+			NewValue *= Bounds[Index].Factor;
+			NewUnits = Bounds[Index].Units;
+
+			if (FMath::Abs(NewValue) > 1)
+			{
+				break;
+			}
+		}
+	}
+
+	return FNumericUnit<T>(NewValue, NewUnits);
+}
+
+template<typename T>
+EUnit FUnitConversion::CalculateDisplayUnit(T Value, EUnit InUnits)
+{
+	if (InUnits == EUnit::Unspecified)
+	{
+		return EUnit::Unspecified;
+	}
+
+	const TArray<EUnit>& DisplayUnits = Settings().GetDisplayUnits(GetUnitType(InUnits));
+	if (DisplayUnits.Num() == 0)
+	{
+		return QuantizeUnitsToBestFit(Value, InUnits).Units;
+	}
+	else if (DisplayUnits.Num() == 1)
+	{
+		return DisplayUnits[0];
+	}
+
+	// If the value we were given was 0, change it to something we can actually work with
+	if (Value == 0)
+	{
+		Value = 1;
+	}
+
+	int32 BestIndex = 0;
+	for (int32 Index = 0; Index < DisplayUnits.Num() - 1; ++Index)
+	{
+		double This = Convert(Value, InUnits, DisplayUnits[Index]);
+		double Next = Convert(Value, InUnits, DisplayUnits[Index + 1]);
+
+		if (FMath::Abs(FMath::LogX(10.0f, This)) < FMath::Abs(FMath::LogX(10.0f, Next)))
+		{
+			BestIndex = Index;
+		}
+		else
+		{
+			BestIndex = Index + 1;
+		}
+	}
+
+	return DisplayUnits[BestIndex];
+}
+
+template<typename NumericType>
+FNumericUnit<NumericType>::FNumericUnit()
+	: Units(EUnit::Unspecified)
+{}
+
+template<typename NumericType>
+FNumericUnit<NumericType>::FNumericUnit(const NumericType& InValue, EUnit InUnits)
+	: Value(InValue), Units(InUnits)
+{}
+
+template<typename NumericType>
+FNumericUnit<NumericType>::FNumericUnit(const FNumericUnit& Other)
+	: Units(EUnit::Unspecified)
+{
+	(*this) = Other;
+}
+
+template<typename NumericType>
+FNumericUnit<NumericType>& FNumericUnit<NumericType>::operator=(const FNumericUnit& Other)
+{
+	CopyValueWithConversion(Other);
+	return *this;
+}
+
+/** Templated Copy construction/assignment from differing numeric types. Relies on implicit conversion of the two numeric types. */
+template<typename NumericType> template<typename OtherType>
+FNumericUnit<NumericType>::FNumericUnit(const FNumericUnit<OtherType>& Other)
+{
+	(*this) = Other;
+}
+
+template<typename NumericType> template<typename OtherType>
+FNumericUnit<NumericType>& FNumericUnit<NumericType>::operator=(const FNumericUnit<OtherType>& Other)
+{
+	CopyValueWithConversion(Other);
+	return *this;
+}
+
+/** Convert this quantity to a different unit */
+template<typename NumericType>
+TOptional<FNumericUnit<NumericType>> FNumericUnit<NumericType>::ConvertTo(EUnit ToUnits) const
+{
+	if (Units == EUnit::Unspecified)
+	{
+		return FNumericUnit(Value, ToUnits);
+	}
+	else if (FUnitConversion::AreUnitsCompatible(Units, ToUnits))
+	{
+		return FNumericUnit<NumericType>(FUnitConversion::Convert(Value, Units, ToUnits), ToUnits);
+	}
+
+	return TOptional<FNumericUnit<NumericType>>();
+}
+
+template<typename NumericType>
+FNumericUnit<NumericType> FNumericUnit<NumericType>::QuantizeUnitsToBestFit() const
+{
+	return FUnitConversion::QuantizeUnitsToBestFit(Value, Units);
+}
+
+template<typename NumericType>
+TValueOrError<FNumericUnit<NumericType>, FText> FNumericUnit<NumericType>::TryParseExpression(const TCHAR* InExpression, EUnit InDefaultUnit)
+{
+	TValueOrError<FNumericUnit<double>, FText> Result = UnitConversion::TryParseExpression(InExpression, InDefaultUnit);
+	if (Result.IsValid())
+	{
+		const auto& Value = Result.GetValue();
+		return MakeValue(FNumericUnit<NumericType>((NumericType)Value.Value, Value.Units));
+	}
+
+	return MakeError(Result.GetError());
+}
+
+template<typename NumericType>
+TOptional<FNumericUnit<NumericType>> FNumericUnit<NumericType>::TryParseString(const TCHAR* InSource)
+{
+	TOptional<FNumericUnit<NumericType>> Result;
+	if (!InSource || !*InSource)
+	{
+		return Result;
+	}
+
+	const TCHAR* NumberEnd = nullptr;
+	if (!ExtractNumberBoundary(InSource, NumberEnd))
+	{
+		return Result;
+	}
+
+	NumericType NewValue;
+	LexicalConversion::FromString(NewValue, InSource);
+
+	// Now parse the units
+	while(FChar::IsWhitespace(*NumberEnd)) ++NumberEnd;
+
+	if (*NumberEnd == '\0')
+	{
+		// No units
+		Result.Emplace(NewValue);
+	}
+	else
+	{
+		// If the string specifies units, they must map to something that exists for this function to succeed
+		auto NewUnits = FUnitConversion::UnitFromString(NumberEnd);
+		if (NewUnits)
+		{
+			Result.Emplace(NewValue, NewUnits.GetValue());
+		}
+	}
+
+	return Result;
+}
+
+/** Copy another unit into this one, taking account of its units, and applying necessary conversion */
+template<typename NumericType> template<typename OtherType>
+void FNumericUnit<NumericType>::CopyValueWithConversion(const FNumericUnit<OtherType>& Other)
+{
+	if (Units != EUnit::Unspecified && Other.Units != EUnit::Unspecified)
+	{
+		if (Units == Other.Units)
+		{
+			Value = Other.Value;
+		}
+		else if (FUnitConversion::AreUnitsCompatible(Units, Other.Units))
+		{
+			Value = FUnitConversion::Convert(Other.Value, Other.Units, Units);
+		}
+		else
+		{
+			// Invalid conversion - assignment invalid
+		}
+	}
+	else
+	{
+		// If our units haven't been specified, we take on the units of the rhs
+		if (Units == EUnit::Unspecified)
+		{
+			// This is the only time we ever change units. Const_cast is 'acceptible' here since the units haven't been specified yet.
+			const_cast<EUnit&>(Units) = Other.Units;
+		}
+
+		Value = Other.Value;
+	}
+}
+
+template<typename NumericType>
+bool FNumericUnit<NumericType>::ExtractNumberBoundary(const TCHAR* Start, const TCHAR*& End)
+{
+	while(FChar::IsWhitespace(*Start)) ++Start;
+
+	End = Start;
+	if (*End == '-' || *End == '+')
+	{
+		End++;
+	}
+
+	bool bHasDot = false;
+	while (FChar::IsDigit(*End) || *End == '.')
+	{
+		if (*End == '.')
+		{
+			if (bHasDot)
+			{
+				return false;
+			}
+			bHasDot = true;
+		}
+		++End;
+	}
+
+	return true;
+}
+
+/** Global arithmetic operators for number types. Deals with conversion from related units correctly. */
+template<typename NumericType, typename OtherType>
+bool operator==(const FNumericUnit<NumericType>& LHS, const FNumericUnit<OtherType>& RHS)
+{
+	if (LHS.Units != EUnit::Unspecified && RHS.Units != EUnit::Unspecified)
+	{
+		if (LHS.Units == RHS.Units)
+		{
+			return LHS.Value == RHS.Value;
+		}
+		else if (FUnitConversion::AreUnitsCompatible(LHS.Units, RHS.Units))
+		{
+			return LHS.Value == FUnitConversion::Convert(RHS.Value, RHS.Units, LHS.Units);
+		}
+		else
+		{
+			// Invalid conversion
+			return false;
+		}
+	}
+	else
+	{
+		return LHS.Value == RHS.Value;
+	}
+}
+
+template<typename NumericType, typename OtherType>
+bool operator!=(const FNumericUnit<NumericType>& LHS, const FNumericUnit<OtherType>& RHS)
+{
+	return !(LHS == RHS);
+}
+
+
+template <typename NumericType>
+struct TNumericLimits<FNumericUnit<NumericType>> : public TNumericLimits<NumericType>
+{ };
+
+namespace LexicalConversion
+{
+	template<typename T>
+	FString ToString(const FNumericUnit<T>& NumericUnit)
+	{
+		FString String = ToString(NumericUnit.Value);
+		String += TEXT(" ");
+		String += FUnitConversion::GetUnitDisplayString(NumericUnit.Units);
+
+		return String;
+	}
+
+	template<typename T>
+	FString ToSanitizedString(const FNumericUnit<T>& NumericUnit)
+	{
+		FString String = ToSanitizedString(NumericUnit.Value);
+		String += TEXT(" ");
+		String += FUnitConversion::GetUnitDisplayString(NumericUnit.Units);
+
+		return String;
+	}
+
+	template<typename T>
+	void FromString(FNumericUnit<T>& OutValue, const TCHAR* String)
+	{
+		auto Parsed = FNumericUnit<T>::TryParseString(String);
+		if (Parsed)
+		{
+			OutValue = Parsed.GetValue();
+		}
+	}
+	
+	template<typename T>
+	bool TryParseString(FNumericUnit<T>& OutValue, const TCHAR* String)
+	{
+		auto Parsed = FNumericUnit<T>::TryParseString(String);
+		if (Parsed)
+		{
+			OutValue = Parsed.GetValue();
+			return true;
+		}
+
+		return false;
+	}
+}
\ No newline at end of file
diff --git a/Engine/Source/Runtime/Slate/Public/Widgets/Input/NumericUnitTypeInterface.inl b/Engine/Source/Runtime/Slate/Public/Widgets/Input/NumericUnitTypeInterface.inl
--- a/Engine/Source/Runtime/Slate/Public/Widgets/Input/NumericUnitTypeInterface.inl
+++ b/Engine/Source/Runtime/Slate/Public/Widgets/Input/NumericUnitTypeInterface.inl
@@ -34,26 +34,17 @@ TOptional<NumericType> TNumericUnitTypeInterface<NumericType>::FromString(const
 	using namespace LexicalConversion;
 
 	// Always parse in as a double, to allow for input of higher-order units with decimal numerals into integral types (eg, inputting 0.5km as 500m)
-	FNumericUnit<double> NewValue;
-	bool bEvalResult = TryParseString( NewValue, *InString ) && FUnitConversion::AreUnitsCompatible( NewValue.Units, UnderlyingUnits );
-	if (bEvalResult)
+	auto NewValue = FNumericUnit<double>::TryParseExpression(*InString, FixedDisplayUnits.IsSet() ? FixedDisplayUnits.GetValue() : UnderlyingUnits);
+	if (NewValue.IsValid())
 	{
 		// Convert the number into the correct units
-		EUnit SourceUnits = NewValue.Units;
+		EUnit SourceUnits = NewValue.GetValue().Units;
 		if (SourceUnits == EUnit::Unspecified && FixedDisplayUnits.IsSet())
 		{
 			// Use the default supplied input units
 			SourceUnits = FixedDisplayUnits.GetValue();
 		}
-		return FUnitConversion::Convert(NewValue.Value, SourceUnits, UnderlyingUnits);
-	}
-	else
-	{
-		float FloatValue = 0.f;
-		if (FMath::Eval( *InString, FloatValue  ))
-		{
-			return FloatValue;
-		}
+		return FUnitConversion::Convert(NewValue.GetValue().Value, SourceUnits, UnderlyingUnits);
 	}
 
 	return TOptional<NumericType>();
