Fixed expression evaluation not giving the correct error when an operator function returned an error

diff --git a/Engine/Source/Runtime/Core/Private/Misc/ExpressionParser.cpp b/Engine/Source/Runtime/Core/Private/Misc/ExpressionParser.cpp
--- a/Engine/Source/Runtime/Core/Private/Misc/ExpressionParser.cpp
+++ b/Engine/Source/Runtime/Core/Private/Misc/ExpressionParser.cpp
@@ -653,11 +653,7 @@ namespace ExpressionParser
 						}
 						else
 						{
-							FFormatOrderedArguments Args;
-							Args.Add(FText::FromString(Op->WrappedToken.Context.GetString()));
-							Args.Add(FText::FromString(L.Context.GetString()));
-							Args.Add(FText::FromString(R.Context.GetString()));
-							return MakeError(FText::Format(LOCTEXT("InvalidBinaryOp", "Cannot perform operator {0} on '{1}' and '{2}'"), Args));
+							return MakeError(OpResult.GetError());
 						}
 					}
 					else
@@ -686,10 +682,7 @@ namespace ExpressionParser
 						}
 						else
 						{
-							FFormatOrderedArguments Args;
-							Args.Add(FText::FromString(Op->WrappedToken.Context.GetString()));
-							Args.Add(FText::FromString(Operand.Context.GetString()));
-							return MakeError(FText::Format(LOCTEXT("InvalidUnaryOp", "Cannot perform operator {0} on '{1}'"), Args));
+							return MakeError(OpResult.GetError());
 						}			
 					}
 					else
