#ue4 - Perf: Add vectorized VectorSinCos for SSE platforms. Optimize FRotator->FQuat conversion to use the vectorized version, and also vectorize the quaternion construction operations.

diff --git a/Engine/Source/Runtime/Core/Private/Math/UnrealMath.cpp b/Engine/Source/Runtime/Core/Private/Math/UnrealMath.cpp
--- a/Engine/Source/Runtime/Core/Private/Math/UnrealMath.cpp
+++ b/Engine/Source/Runtime/Core/Private/Math/UnrealMath.cpp
@@ -280,37 +280,52 @@ FQuat FRotator::Quaternion() const
 	FQuat RotationMatrix = FQuat( FRotationMatrix( *this ) );
 #endif
 
-#if WITH_DIRECTXMATH	// Currently VectorSinCos() is only vectorized in DirectX implementation. Other platforms become slower if they use this.
-	VectorRegister Angles = MakeVectorRegister(Roll, Pitch, Yaw, 0.0f);
-	VectorRegister HalfAngles = VectorMultiply(Angles, GlobalVectorConstants::DEG_TO_RAD_HALF);
+#if PLATFORM_ENABLE_VECTORINTRINSICS
+	const VectorRegister Angles = MakeVectorRegister(Pitch, Yaw, Roll, 0.0f);
+	const VectorRegister HalfAngles = VectorMultiply(Angles, GlobalVectorConstants::DEG_TO_RAD_HALF);
 
-	union { VectorRegister v; float f[4]; } SinAngles, CosAngles;	
-	VectorSinCos(&SinAngles.v, &CosAngles.v, &HalfAngles);
+	VectorRegister SinAngles, CosAngles;
+	VectorSinCos(&SinAngles, &CosAngles, &HalfAngles);
 
-	const float	SR	= SinAngles.f[0];
-	const float	SP	= SinAngles.f[1];
-	const float	SY	= SinAngles.f[2];
-	const float	CR	= CosAngles.f[0];
-	const float	CP	= CosAngles.f[1];
-	const float	CY	= CosAngles.f[2];
+	// Vectorized conversion, measured 20% faster than using scalar version after VectorSinCos.
+	// Indices within VectorRegister (for shuffles): P=0, Y=1, R=2
+	const VectorRegister SR = VectorReplicate(SinAngles, 2);
+	const VectorRegister CR = VectorReplicate(CosAngles, 2);
 
-#else
-	static const float DEG_TO_RAD = PI/(180.f);
-	static const float DIVIDE_BY_2 = DEG_TO_RAD/2.f;
+	const VectorRegister SY_SY_CY_CY_Temp = VectorShuffle(SinAngles, CosAngles, 1, 1, 1, 1);
+
+	const VectorRegister SP_SP_CP_CP = VectorShuffle(SinAngles, CosAngles, 0, 0, 0, 0);
+	const VectorRegister SY_CY_SY_CY = VectorShuffle(SY_SY_CY_CY_Temp, SY_SY_CY_CY_Temp, 0, 2, 0, 2);
+
+	const VectorRegister CP_CP_SP_SP = VectorShuffle(CosAngles, SinAngles, 0, 0, 0, 0);
+	const VectorRegister CY_SY_CY_SY = VectorShuffle(SY_SY_CY_CY_Temp, SY_SY_CY_CY_Temp, 2, 0, 2, 0);
+
+	const uint32 Neg = uint32(1 << 31);
+	const uint32 Pos = uint32(0);
+	const VectorRegister SignBitsLeft  = MakeVectorRegister(Pos, Neg, Pos, Pos);
+	const VectorRegister SignBitsRight = MakeVectorRegister(Neg, Neg, Neg, Pos);
+	const VectorRegister LeftTerm  = VectorBitwiseXor(SignBitsLeft , VectorMultiply(CR, VectorMultiply(SP_SP_CP_CP, SY_CY_SY_CY)));
+	const VectorRegister RightTerm = VectorBitwiseXor(SignBitsRight, VectorMultiply(SR, VectorMultiply(CP_CP_SP_SP, CY_SY_CY_SY)));
 
+	FQuat RotationQuat;
+	const VectorRegister Result = VectorAdd(LeftTerm, RightTerm);	
+	VectorStoreAligned(Result, &RotationQuat);
+#else
+	const float DEG_TO_RAD = PI/(180.f);
+	const float DIVIDE_BY_2 = DEG_TO_RAD/2.f;
 	float SP, SY, SR;
 	float CP, CY, CR;
 
 	FMath::SinCos(&SP, &CP, Pitch*DIVIDE_BY_2);
 	FMath::SinCos(&SY, &CY, Yaw*DIVIDE_BY_2);
 	FMath::SinCos(&SR, &CR, Roll*DIVIDE_BY_2);
-#endif
 
 	FQuat RotationQuat;
-	RotationQuat.W = CR*CP*CY + SR*SP*SY;
-	RotationQuat.X = CR*SP*SY - SR*CP*CY;
+	RotationQuat.X =  CR*SP*SY - SR*CP*CY;
 	RotationQuat.Y = -CR*SP*CY - SR*CP*SY;
-	RotationQuat.Z = CR*CP*SY - SR*SP*CY;
+	RotationQuat.Z =  CR*CP*SY - SR*SP*CY;
+	RotationQuat.W =  CR*CP*CY + SR*SP*SY;
+#endif // PLATFORM_ENABLE_VECTORINTRINSICS
 
 #if USE_MATRIX_ROTATOR 
 	if (!RotationMatrix.Equals(RotationQuat, KINDA_SMALL_NUMBER) && !RotationMatrix.Equals(RotationQuat*-1.f, KINDA_SMALL_NUMBER))
diff --git a/Engine/Source/Runtime/Core/Public/Math/RotationTranslationMatrix.h b/Engine/Source/Runtime/Core/Public/Math/RotationTranslationMatrix.h
--- a/Engine/Source/Runtime/Core/Public/Math/RotationTranslationMatrix.h
+++ b/Engine/Source/Runtime/Core/Public/Math/RotationTranslationMatrix.h
@@ -27,20 +27,20 @@ public:
 
 FORCEINLINE FRotationTranslationMatrix::FRotationTranslationMatrix(const FRotator& Rot, const FVector& Origin)
 {
-#if WITH_DIRECTXMATH	// Currently VectorSinCos() is only vectorized in DirectX implementation. Other platforms become slower if they use this.
+#if PLATFORM_ENABLE_VECTORINTRINSICS
 
-	VectorRegister Angles = MakeVectorRegister(Rot.Roll, Rot.Pitch, Rot.Yaw, 0.0f);
-	VectorRegister HalfAngles = VectorMultiply(Angles, GlobalVectorConstants::DEG_TO_RAD);
+	const VectorRegister Angles = MakeVectorRegister(Rot.Pitch, Rot.Yaw, Rot.Roll, 0.0f);
+	const VectorRegister HalfAngles = VectorMultiply(Angles, GlobalVectorConstants::DEG_TO_RAD);
 
-	union { VectorRegister v; float f[4]; } SinAngles, CosAngles;	
+	union { VectorRegister v; float f[4]; } SinAngles, CosAngles;
 	VectorSinCos(&SinAngles.v, &CosAngles.v, &HalfAngles);
 
-	const float	SR	= SinAngles.f[0];
-	const float	SP	= SinAngles.f[1];
-	const float	SY	= SinAngles.f[2];
-	const float	CR	= CosAngles.f[0];
-	const float	CP	= CosAngles.f[1];
-	const float	CY	= CosAngles.f[2];
+	const float	SP	= SinAngles.f[0];
+	const float	SY	= SinAngles.f[1];
+	const float	SR	= SinAngles.f[2];
+	const float	CP	= CosAngles.f[0];
+	const float	CY	= CosAngles.f[1];
+	const float	CR	= CosAngles.f[2];
 
 #else
 	
@@ -50,7 +50,7 @@ FORCEINLINE FRotationTranslationMatrix::FRotationTranslationMatrix(const FRotato
 	FMath::SinCos(&SY, &CY, FMath::DegreesToRadians(Rot.Yaw));
 	FMath::SinCos(&SR, &CR, FMath::DegreesToRadians(Rot.Roll));
 
-#endif
+#endif // PLATFORM_ENABLE_VECTORINTRINSICS
 
 	M[0][0]	= CP * CY;
 	M[0][1]	= CP * SY;
diff --git a/Engine/Source/Runtime/Core/Public/Math/UnrealMathDirectX.h b/Engine/Source/Runtime/Core/Public/Math/UnrealMathDirectX.h
--- a/Engine/Source/Runtime/Core/Public/Math/UnrealMathDirectX.h
+++ b/Engine/Source/Runtime/Core/Public/Math/UnrealMathDirectX.h
@@ -727,7 +727,7 @@ FORCEINLINE void VectorQuaternionVector3InverseRotate( FVector *Result, const FV
 * @param VCosAngles	VectorRegister Pointer to where the Cos result should be stored
 * @param VAngles VectorRegister Pointer to the input angles 
 */
-FORCEINLINE void VectorSinCos(  VectorRegister* VSinAngles, VectorRegister* VCosAngles, const VectorRegister* VAngles )
+FORCEINLINE void VectorSinCos(VectorRegister* RESTRICT VSinAngles, VectorRegister* RESTRICT VCosAngles, const VectorRegister* RESTRICT VAngles)
 {
 	using namespace DirectX;
 	// Force the value within the bounds of pi	
diff --git a/Engine/Source/Runtime/Core/Public/Math/UnrealMathSSE.h b/Engine/Source/Runtime/Core/Public/Math/UnrealMathSSE.h
--- a/Engine/Source/Runtime/Core/Public/Math/UnrealMathSSE.h
+++ b/Engine/Source/Runtime/Core/Public/Math/UnrealMathSSE.h
@@ -964,28 +964,6 @@ FORCEINLINE void VectorQuaternionMultiply( void* RESTRICT Result, const void* RE
 	*((VectorRegister *)Result) = VectorQuaternionMultiply2(*((const VectorRegister *)Quat1), *((const VectorRegister *)Quat2));
 }
 
-
-/**
-* Computes the sine and cosine of each component of a Vector.
-*
-* @param VSinAngles	VectorRegister Pointer to where the Sin result should be stored
-* @param VCosAngles	VectorRegister Pointer to where the Cos result should be stored
-* @param VAngles VectorRegister Pointer to the input angles 
-*/
-FORCEINLINE void VectorSinCos(  VectorRegister* VSinAngles, VectorRegister* VCosAngles, const VectorRegister* VAngles )
-{	
-	union { VectorRegister v; float f[4]; } VecSin, VecCos, VecAngles;
-	VecAngles.v = *VAngles;
-
-	FMath::SinCos(&VecSin.f[0], &VecCos.f[0], VecAngles.f[0]);
-	FMath::SinCos(&VecSin.f[1], &VecCos.f[1], VecAngles.f[1]);
-	FMath::SinCos(&VecSin.f[2], &VecCos.f[2], VecAngles.f[2]);
-	FMath::SinCos(&VecSin.f[3], &VecCos.f[3], VecAngles.f[3]);
-
-	*VSinAngles = VecSin.v;
-	*VCosAngles = VecCos.v;
-}
-
 // Returns true if the vector contains a component that is either NAN or +/-infinite.
 inline bool VectorContainsNaNOrInfinite(const VectorRegister& Vec)
 {
@@ -1095,6 +1073,60 @@ FORCEINLINE VectorRegister VectorCos(const VectorRegister& X)
 	return VectorSin(VectorAdd(X, GlobalVectorConstants::PiByTwo));
 }
 
+
+/**
+* Computes the sine and cosine of each component of a Vector.
+*
+* @param VSinAngles	VectorRegister Pointer to where the Sin result should be stored
+* @param VCosAngles	VectorRegister Pointer to where the Cos result should be stored
+* @param VAngles VectorRegister Pointer to the input angles
+*/
+FORCEINLINE void VectorSinCos(VectorRegister* RESTRICT VSinAngles, VectorRegister* RESTRICT VCosAngles, const VectorRegister* RESTRICT VAngles)
+{
+	// Map to [-pi, pi]
+	VectorRegister X = VectorFloor(VectorMultiplyAdd(*VAngles, GlobalVectorConstants::OneOverTwoPi, GlobalVectorConstants::FloatOneHalf));
+	X = VectorSubtract(*VAngles, VectorMultiply(GlobalVectorConstants::TwoPi, X));
+
+	// Map in [-pi/2,pi/2]
+	VectorRegister sign = VectorBitwiseAnd(X, GlobalVectorConstants::SignBit);
+	VectorRegister c = VectorBitwiseOr(GlobalVectorConstants::Pi, sign);  // pi when x >= 0, -pi when x < 0
+	VectorRegister absx = VectorAbs(X);
+	VectorRegister rflx = VectorSubtract(c, X);
+	VectorRegister comp = VectorCompareGT(absx, GlobalVectorConstants::PiByTwo);
+	X = VectorSelect(comp, rflx, X);
+	sign = VectorSelect(comp, GlobalVectorConstants::FloatMinusOne, GlobalVectorConstants::FloatOne);
+
+	const VectorRegister XSquared = VectorMultiply(X, X);
+
+	// 11-degree minimax approximation
+	//*ScalarSin = (((((-2.3889859e-08f * y2 + 2.7525562e-06f) * y2 - 0.00019840874f) * y2 + 0.0083333310f) * y2 - 0.16666667f) * y2 + 1.0f) * y;
+	const VectorRegister SinCoeff0 = MakeVectorRegister(1.0f, -0.16666667f, 0.0083333310f, -0.00019840874f);
+	const VectorRegister SinCoeff1 = MakeVectorRegister(2.7525562e-06f, -2.3889859e-08f, /*unused*/ 0.f, /*unused*/ 0.f);
+
+	VectorRegister S;
+	S = VectorReplicate(SinCoeff1, 1);
+	S = VectorMultiplyAdd(XSquared, S, VectorReplicate(SinCoeff1, 0));
+	S = VectorMultiplyAdd(XSquared, S, VectorReplicate(SinCoeff0, 3));
+	S = VectorMultiplyAdd(XSquared, S, VectorReplicate(SinCoeff0, 2));
+	S = VectorMultiplyAdd(XSquared, S, VectorReplicate(SinCoeff0, 1));
+	S = VectorMultiplyAdd(XSquared, S, VectorReplicate(SinCoeff0, 0));
+	*VSinAngles = VectorMultiply(S, X);
+
+	// 10-degree minimax approximation
+	//*ScalarCos = sign * (((((-2.6051615e-07f * y2 + 2.4760495e-05f) * y2 - 0.0013888378f) * y2 + 0.041666638f) * y2 - 0.5f) * y2 + 1.0f);
+	const VectorRegister CosCoeff0 = MakeVectorRegister(1.0f, -0.5f, 0.041666638f, -0.0013888378f);
+	const VectorRegister CosCoeff1 = MakeVectorRegister(2.4760495e-05f, -2.6051615e-07f, /*unused*/ 0.f, /*unused*/ 0.f);
+
+	VectorRegister C;
+	C = VectorReplicate(CosCoeff1, 1);
+	C = VectorMultiplyAdd(XSquared, C, VectorReplicate(CosCoeff1, 0));
+	C = VectorMultiplyAdd(XSquared, C, VectorReplicate(CosCoeff0, 3));
+	C = VectorMultiplyAdd(XSquared, C, VectorReplicate(CosCoeff0, 2));
+	C = VectorMultiplyAdd(XSquared, C, VectorReplicate(CosCoeff0, 1));
+	C = VectorMultiplyAdd(XSquared, C, VectorReplicate(CosCoeff0, 0));
+	*VCosAngles = VectorMultiply(C, sign);
+}
+
 //TODO: Vectorize
 FORCEINLINE VectorRegister VectorTan(const VectorRegister& X)
 {
