#ue4 - Perf: Optimize FQuat comparison when vector intrinsics are available.

diff --git a/Engine/Source/Runtime/Core/Public/Math/Quat.h b/Engine/Source/Runtime/Core/Public/Math/Quat.h
--- a/Engine/Source/Runtime/Core/Public/Math/Quat.h
+++ b/Engine/Source/Runtime/Core/Public/Math/Quat.h
@@ -681,8 +681,18 @@ FORCEINLINE FQuat FQuat::operator-( const FQuat& Q ) const
 
 FORCEINLINE bool FQuat::Equals(const FQuat& Q, float Tolerance) const
 {
+#if PLATFORM_ENABLE_VECTORINTRINSICS
+	const VectorRegister ToleranceV = VectorLoadFloat1(&Tolerance);
+	const VectorRegister A = VectorLoadAligned(this);
+	const VectorRegister B = VectorLoadAligned(&Q);
+
+	const VectorRegister RotationSub = VectorAbs(VectorSubtract(A, B));
+	const VectorRegister RotationAdd = VectorAbs(VectorAdd(A, B));
+	return !VectorAnyGreaterThan(RotationSub, ToleranceV) || !VectorAnyGreaterThan(RotationAdd, ToleranceV);
+#else
 	return (FMath::Abs(X - Q.X) <= Tolerance && FMath::Abs(Y - Q.Y) <= Tolerance && FMath::Abs(Z - Q.Z) <= Tolerance && FMath::Abs(W - Q.W) <= Tolerance)
 		|| (FMath::Abs(X + Q.X) <= Tolerance && FMath::Abs(Y + Q.Y) <= Tolerance && FMath::Abs(Z + Q.Z) <= Tolerance && FMath::Abs(W + Q.W) <= Tolerance);
+#endif // PLATFORM_ENABLE_VECTORINTRINSICS
 }
 
 
@@ -769,7 +779,13 @@ FORCEINLINE FQuat FQuat::operator/(const float Scale) const
 
 FORCEINLINE bool FQuat::operator==( const FQuat& Q ) const
 {
+#if PLATFORM_ENABLE_VECTORINTRINSICS
+	const VectorRegister A = VectorLoadAligned(this);
+	const VectorRegister B = VectorLoadAligned(&Q);
+	return VectorMaskBits(VectorCompareEQ(A, B)) == 0x0F;
+#else
 	return X == Q.X && Y == Q.Y && Z == Q.Z && W == Q.W;
+#endif // PLATFORM_ENABLE_VECTORINTRINSICS
 }
 
 
