#ue4 - Add better logging for some FQuat and FVector comparisons in start-up math validation in UnrealMathTest.

- These will now include the correct data in the log message if they aren't equal.

diff --git a/Engine/Source/Runtime/Core/Private/Tests/Math/UnrealMathTest.cpp b/Engine/Source/Runtime/Core/Private/Tests/Math/UnrealMathTest.cpp
--- a/Engine/Source/Runtime/Core/Private/Tests/Math/UnrealMathTest.cpp
+++ b/Engine/Source/Runtime/Core/Private/Tests/Math/UnrealMathTest.cpp
@@ -107,6 +107,23 @@ bool TestVectorsEqual3_NoVec( const FVector& Vec0, const FVector& Vec1, float To
 	return GSum <= Tolerance;
 }
 
+bool TestQuatsEqual(const FQuat& Q0, const FQuat& Q1, float Tolerance)
+{
+	GScratch[0] = Q0.X;
+	GScratch[1] = Q0.Y;
+	GScratch[2] = Q0.Z;
+	GScratch[3] = Q0.W;
+	GScratch[4] = Q1.X;
+	GScratch[5] = Q1.Y;
+	GScratch[6] = Q1.Z;
+	GScratch[7] = Q1.W;
+	GSum = 0.f;
+
+	const bool bEqual = Q0.Equals(Q1, Tolerance);
+	GPassing = GPassing && bEqual;
+	return bEqual;
+}
+
 /**
  * Tests if two matrices (4x4 xyzw) are equal within an optional tolerance
  *
@@ -897,33 +914,33 @@ bool FVectorRegisterAbstractionTest::RunTest(const FString& Parameters)
 		Rotator0 = FRotator(30.0f, -45.0f, 90.0f);
 		Q0 = Rotator0.Quaternion();
 		Q1 = TestRotatorToQuaternion(Rotator0);
-		LogTest( TEXT("TestRotatorToQuaternion"), Q0.Equals( Q1, 1e-6f ) );
+		LogTest( TEXT("TestRotatorToQuaternion"), TestQuatsEqual(Q0, Q1, 1e-6f));
 
 		FVector FV0, FV1;
 		FV0 = Rotator0.Vector();
 		FV1 = FRotationMatrix( Rotator0 ).GetScaledAxis( EAxis::X );
-		LogTest( TEXT("Test0 Rotator::Vector()"), FV0.Equals( FV1, 1e-6f ) );
+		LogTest( TEXT("Test0 Rotator::Vector()"), TestVectorsEqual3_NoVec(FV0, FV1, 1e-6f));
 		
 		FV0 = FRotationMatrix( Rotator0 ).GetScaledAxis( EAxis::X );
 		FV1 = FQuatRotationMatrix( Q0 ).GetScaledAxis( EAxis::X );
-		LogTest( TEXT("Test0 FQuatRotationMatrix"), FV0.Equals( FV1, 1e-5f ) );
+		LogTest( TEXT("Test0 FQuatRotationMatrix"), TestVectorsEqual3_NoVec(FV0, FV1, 1e-5f));
 
 		Rotator0 = FRotator(45.0f,  60.0f, 120.0f);
 		Q0 = Rotator0.Quaternion();
 		Q1 = TestRotatorToQuaternion(Rotator0);
-		LogTest( TEXT("TestRotatorToQuaternion"), Q0.Equals( Q1, 1e-6f ) );
+		LogTest( TEXT("TestRotatorToQuaternion"), TestQuatsEqual(Q0, Q1, 1e-6f));
 
 		FV0 = Rotator0.Vector();
 		FV1 = FRotationMatrix( Rotator0 ).GetScaledAxis( EAxis::X );
-		LogTest( TEXT("Test1 Rotator::Vector()"), FV0.Equals( FV1, 1e-6f ) );
+		LogTest( TEXT("Test1 Rotator::Vector()"), TestVectorsEqual3_NoVec(FV0, FV1, 1e-6f));
 
 		FV0 = FRotationMatrix( Rotator0 ).GetScaledAxis( EAxis::X );
 		FV1 = FQuatRotationMatrix( Q0 ).GetScaledAxis( EAxis::X );
-		LogTest(TEXT("Test1 FQuatRotationMatrix"), FV0.Equals(FV1, 1e-5f));
+		LogTest(TEXT("Test1 FQuatRotationMatrix"), TestVectorsEqual3_NoVec(FV0, FV1, 1e-5f));
 
 		FV0 = FRotationMatrix(FRotator::ZeroRotator).GetScaledAxis(EAxis::X);
 		FV1 = FQuatRotationMatrix(FQuat::Identity).GetScaledAxis(EAxis::X);
-		LogTest(TEXT("Test2 FQuatRotationMatrix"), FV0.Equals(FV1, 1e-6f));
+		LogTest(TEXT("Test2 FQuatRotationMatrix"), TestVectorsEqual3_NoVec(FV0, FV1, 1e-6f));
 	}
 
 	{
@@ -991,7 +1008,7 @@ bool FVectorRegisterAbstractionTest::RunTest(const FString& Parameters)
 	Q1 = FQuat(FRotator(45.0f,  60.0f, 120.0f));
 	VectorQuaternionMultiply(&Q2, &Q0, &Q1);
 	TestVectorQuaternionMultiply(&Q3, &Q0, &Q1);
-	LogTest( TEXT("VectorQuaternionMultiply"), Q2.Equals( Q3, 1e-6f ) );
+	LogTest( TEXT("VectorQuaternionMultiply"), TestQuatsEqual(Q2, Q3, 1e-6f));
 	V0 = VectorLoadAligned(&Q0);
 	V1 = VectorLoadAligned(&Q1);
 	V2 = VectorQuaternionMultiply2(V0, V1);
@@ -1002,7 +1019,7 @@ bool FVectorRegisterAbstractionTest::RunTest(const FString& Parameters)
 	Q1 = FQuat(FRotator(-120.0f, -90.0f, 0.0f));
 	VectorQuaternionMultiply(&Q2, &Q0, &Q1);
 	TestVectorQuaternionMultiply(&Q3, &Q0, &Q1);
-	LogTest( TEXT("VectorMatrixInverse"), Q2.Equals( Q3, 1e-6f ) );
+	LogTest( TEXT("VectorMatrixInverse"), TestQuatsEqual(Q2, Q3, 1e-6f) );
 	V0 = VectorLoadAligned(&Q0);
 	V1 = VectorLoadAligned(&Q1);
 	V2 = VectorQuaternionMultiply2(V0, V1);
