Core: Added method to check whether a type is registered in a type container

diff --git a/Engine/Source/Runtime/Core/Public/Misc/TypeContainer.h b/Engine/Source/Runtime/Core/Public/Misc/TypeContainer.h
--- a/Engine/Source/Runtime/Core/Public/Misc/TypeContainer.h
+++ b/Engine/Source/Runtime/Core/Public/Misc/TypeContainer.h
@@ -159,7 +159,7 @@ public:
 	/**
 	 * Gets a shared pointer to an instance of the specified class.
 	 *
-	 * @param T The type of class to get an instance for.
+	 * @param R The type of class to get an instance for.
 	 * @param A shared reference to the instance.
 	 * @see RegisterClass, RegisterDelegate, RegisterFactory, RegisterInstance
 	 */
@@ -191,6 +191,21 @@ public:
 		return GetInstance<R>().ToSharedRef();
 	}
 
+	/**
+	 * Check whether the specified class has been registered.
+	 *
+	 * @param R The type of registered class to check.
+	 * @return true if the class was registered, false otherwise.
+	 */
+	template<class R>
+	bool IsRegistered()
+	{
+		FScopeLock Lock(&CriticalSection);
+		{
+			return Providers.Contains(TNameOf<R>::GetName());
+		}
+	}
+
 	/**
 	 * Registers a class for instances of the specified class.
 	 *
