Core: const correctness for type container

diff --git a/Engine/Source/Runtime/Core/Public/Misc/TypeContainer.h b/Engine/Source/Runtime/Core/Public/Misc/TypeContainer.h
--- a/Engine/Source/Runtime/Core/Public/Misc/TypeContainer.h
+++ b/Engine/Source/Runtime/Core/Public/Misc/TypeContainer.h
@@ -164,7 +164,7 @@ public:
 	 * @see RegisterClass, RegisterDelegate, RegisterFactory, RegisterInstance
 	 */
 	template<class R>
-	TSharedRef<R> GetInstance()
+	TSharedRef<R> GetInstance() const
 	{
 		FScopeLock Lock(&CriticalSection);
 		{
@@ -186,7 +186,7 @@ public:
 	 * @see GetInstance, RegisterClass, RegisterInstance
 	 */
 	template<class R>
-	TSharedRef<R> GetInstanceRef()
+	TSharedRef<R> GetInstanceRef() const
 	{
 		return GetInstance<R>().ToSharedRef();
 	}
@@ -198,7 +198,7 @@ public:
 	 * @return true if the class was registered, false otherwise.
 	 */
 	template<class R>
-	bool IsRegistered()
+	bool IsRegistered() const
 	{
 		FScopeLock Lock(&CriticalSection);
 		{
@@ -401,7 +401,7 @@ protected:
 private:
 
 	/** Critical section for synchronizing access. */
-	FCriticalSection CriticalSection;
+	mutable FCriticalSection CriticalSection;
 
 	/** Maps class names to instance providers. */
 	TMap<FString, TSharedPtr<IInstanceProvider>> Providers;
