Disabling automation test while I investigate a failure

diff --git a/Engine/Source/Runtime/Core/Private/Math/BasicMathExpressionEvaluator.cpp b/Engine/Source/Runtime/Core/Private/Math/BasicMathExpressionEvaluator.cpp
--- a/Engine/Source/Runtime/Core/Private/Math/BasicMathExpressionEvaluator.cpp
+++ b/Engine/Source/Runtime/Core/Private/Math/BasicMathExpressionEvaluator.cpp
@@ -196,7 +196,7 @@ bool TestInvalidExpression(FAutomationTestBase* Test, const TCHAR* Expression)
 	return true;
 }
 
-IMPLEMENT_SIMPLE_AUTOMATION_TEST(FBasicMathExpressionEvaluatorTest, "System.Core.Math.Expressions", EAutomationTestFlags::ATF_SmokeTest)
+IMPLEMENT_SIMPLE_AUTOMATION_TEST(FBasicMathExpressionEvaluatorTest, "System.Core.Math.Expressions", EAutomationTestFlags::ATF_Editor)
 bool FBasicMathExpressionEvaluatorTest::RunTest( const FString& Parameters )
 {
 	bool bResult = true;
