Core - PanicFlushThreadedLogs also flushes the related output devices

diff --git a/Engine/Source/Runtime/Core/Private/HAL/OutputDevices.cpp b/Engine/Source/Runtime/Core/Private/HAL/OutputDevices.cpp
--- a/Engine/Source/Runtime/Core/Private/HAL/OutputDevices.cpp
+++ b/Engine/Source/Runtime/Core/Private/HAL/OutputDevices.cpp
@@ -771,6 +771,16 @@ void FOutputDeviceRedirector::PanicFlushThreadedLogs()
 	// Flush threaded logs, but use the safe version.
 	UnsynchronizedFlushThreadedLogs( false );
 
+	// Flush devices.
+	for (int32 OutputDeviceIndex = 0; OutputDeviceIndex<OutputDevices.Num(); OutputDeviceIndex++)
+	{
+		FOutputDevice* OutputDevice = OutputDevices[OutputDeviceIndex];
+		if (OutputDevice->CanBeUsedOnAnyThread())
+		{
+			OutputDevice->Flush();
+		}
+	}
+
 	BufferedLines.Empty();
 }
 
