#ue4 - Perf: Faster FQuat::operator!=

diff --git a/Engine/Source/Runtime/Core/Public/Math/Quat.h b/Engine/Source/Runtime/Core/Public/Math/Quat.h
--- a/Engine/Source/Runtime/Core/Public/Math/Quat.h
+++ b/Engine/Source/Runtime/Core/Public/Math/Quat.h
@@ -791,7 +791,13 @@ FORCEINLINE bool FQuat::operator==( const FQuat& Q ) const
 
 FORCEINLINE bool FQuat::operator!=( const FQuat& Q ) const
 {
+#if PLATFORM_ENABLE_VECTORINTRINSICS
+	const VectorRegister A = VectorLoadAligned(this);
+	const VectorRegister B = VectorLoadAligned(&Q);
+	return VectorMaskBits(VectorCompareNE(A, B)) != 0x00;
+#else
 	return X != Q.X || Y != Q.Y || Z != Q.Z || W != Q.W;
+#endif
 }
 
 
