#ue4 - Remove static initialization in VectorNormalizeRotator, and move constants to GlobalVectorConstants.

diff --git a/Engine/Source/Runtime/Core/Public/Math/UnrealMathVectorCommon.h b/Engine/Source/Runtime/Core/Public/Math/UnrealMathVectorCommon.h
--- a/Engine/Source/Runtime/Core/Public/Math/UnrealMathVectorCommon.h
+++ b/Engine/Source/Runtime/Core/Public/Math/UnrealMathVectorCommon.h
@@ -112,17 +112,14 @@ FORCEINLINE VectorRegister VectorNormalizeQuaternion(const VectorRegister& Unnor
 // Normalize Rotator
 FORCEINLINE VectorRegister VectorNormalizeRotator(const VectorRegister& UnnormalizedRotator)
 {
-	static const VectorRegister VEC_360 = MakeVectorRegister(360.f, 360.f, 360.f, 360.f);
-	static const VectorRegister VEC_180 = MakeVectorRegister(180.f, 180.f, 180.f, 180.f);
-
 	// shift in the range [-360,360]
-	VectorRegister V0	= VectorMod( UnnormalizedRotator, VEC_360);
-	VectorRegister V1	= VectorAdd( V0, VEC_360 );
+	VectorRegister V0	= VectorMod( UnnormalizedRotator, GlobalVectorConstants::Float360);
+	VectorRegister V1	= VectorAdd( V0, GlobalVectorConstants::Float360 );
 	VectorRegister V2	= VectorSelect(VectorCompareGE(V0, VectorZero()), V0, V1);
 
 	// shift to [-180,180]
-	VectorRegister V3	= VectorSubtract( V2, VEC_360 );
-	VectorRegister V4	= VectorSelect(VectorCompareGT(V2, VEC_180), V3, V2);
+	VectorRegister V3	= VectorSubtract( V2, GlobalVectorConstants::Float360 );
+	VectorRegister V4	= VectorSelect(VectorCompareGT(V2, GlobalVectorConstants::Float180), V3, V2);
 
 	return  V4;
 }
diff --git a/Engine/Source/Runtime/Core/Public/Math/UnrealMathVectorConstants.h b/Engine/Source/Runtime/Core/Public/Math/UnrealMathVectorConstants.h
--- a/Engine/Source/Runtime/Core/Public/Math/UnrealMathVectorConstants.h
+++ b/Engine/Source/Runtime/Core/Public/Math/UnrealMathVectorConstants.h
@@ -50,4 +50,6 @@ namespace GlobalVectorConstants
 	static const VectorRegister OneOverTwoPi = MakeVectorRegister(1.0f / (2.0f*PI), 1.0f / (2.0f*PI), 1.0f / (2.0f*PI), 1.0f / (2.0f*PI));
 
 	static const VectorRegister Float255 = MakeVectorRegister(255.0f, 255.0f, 255.0f, 255.0f);
+	static const VectorRegister Float360 = MakeVectorRegister(360.f, 360.f, 360.f, 360.f);
+	static const VectorRegister Float180 = MakeVectorRegister(180.f, 180.f, 180.f, 180.f);
 }
