#ue4 - More accurate VectorReciprocalSqrtAccurate on SSE platforms.

- Matches results for SSE platforms with other platforms that perform 2 Newton-Raphson iterations on the hardware estimate (Xbox, iOS), or non-SIMD that uses the fully accurate sqrt.

- Allows NormalizeRotation (0,0,0,1) to return the same vector rather than (0,0,0,0.99998...).

diff --git a/Engine/Source/Runtime/Core/Public/Math/UnrealMathSSE.h b/Engine/Source/Runtime/Core/Public/Math/UnrealMathSSE.h
--- a/Engine/Source/Runtime/Core/Public/Math/UnrealMathSSE.h
+++ b/Engine/Source/Runtime/Core/Public/Math/UnrealMathSSE.h
@@ -468,7 +468,7 @@ FORCEINLINE VectorRegister VectorReciprocalLen( const VectorRegister& Vector )
 */
 FORCEINLINE VectorRegister VectorReciprocalSqrtAccurate(const VectorRegister& Vec)
 {
-	// Perform a single pass of Newton-Raphson iteration on the hardware estimate
+	// Perform two passes of Newton-Raphson iteration on the hardware estimate
 	//    v^-0.5 = x
 	// => x^2 = v^-1
 	// => 1/(x^2) = v
@@ -491,7 +491,12 @@ FORCEINLINE VectorRegister VectorReciprocalSqrtAccurate(const VectorRegister& Ve
 	const VectorRegister InnerTerm0 = VectorSubtract(OneHalf, VectorMultiply(VecDivBy2, x0Squared));
 	const VectorRegister x1 = VectorMultiplyAdd(x0, InnerTerm0, x0);
 
-	return x1;
+	// Second iteration
+	const VectorRegister x1Squared = VectorMultiply(x1, x1);
+	const VectorRegister InnerTerm1 = VectorSubtract(OneHalf, VectorMultiply(VecDivBy2, x1Squared));
+	const VectorRegister x2 = VectorMultiplyAdd(x1, InnerTerm1, x1);
+
+	return x2;
 }
 
 /**
