Core - More reliable crash reporting on Windows

diff --git a/Engine/Source/Runtime/Core/Private/Windows/WindowsPlatformCrashContext.cpp b/Engine/Source/Runtime/Core/Private/Windows/WindowsPlatformCrashContext.cpp
--- a/Engine/Source/Runtime/Core/Private/Windows/WindowsPlatformCrashContext.cpp
+++ b/Engine/Source/Runtime/Core/Private/Windows/WindowsPlatformCrashContext.cpp
@@ -364,6 +364,7 @@ int32 ReportCrashUsingCrashReportClient(EXCEPTION_POINTERS* ExceptionInfo, const
 static FCriticalSection EnsureLock;
 static bool bReentranceGuard = false;
 
+// #YRX_Crash: 2015-05-28 This should be named EngineEnsureHandler
 /** 
  * Report an ensure to the crash reporting system
  */
@@ -455,6 +456,7 @@ void CreateExceptionInfoString(EXCEPTION_RECORD* ExceptionRecord)
 }
 #include "HideWindowsPlatformTypes.h"
 
+// #YRX_Crash: 2015-05-28 THis should be named EngineCrashHandler
 int32 ReportCrash( LPEXCEPTION_POINTERS ExceptionInfo )
 {
 	// Only create a minidump the first time this function is called.
@@ -464,12 +466,32 @@ int32 ReportCrash( LPEXCEPTION_POINTERS ExceptionInfo )
 		return EXCEPTION_EXECUTE_HANDLER;
 	}
 
+	GLog->PanicFlushThreadedLogs();
+
+	// First launch the crash reporter client.
+#if WINVER > 0x502	// Windows Error Reporting is not supported on Windows XP
+	if (GUseCrashReportClient)
+	{
+		ReportCrashUsingCrashReportClient( ExceptionInfo, GErrorMessage, EErrorReportUI::ShowDialog );
+	}
+	else
+#endif		// WINVER
+	{
+		WriteMinidump( MiniDumpFilenameW, ExceptionInfo );
+
+#if UE_BUILD_SHIPPING && WITH_EDITOR
+		uint32 dwOpt = 0;
+		EFaultRepRetVal repret = ReportFault( ExceptionInfo, dwOpt );
+#endif
+	}
+
+	// Then try run time crash processing and broadcast information about a crash.
 	FCoreDelegates::OnHandleSystemError.Broadcast();
 
 	const SIZE_T StackTraceSize = 65535;
 	ANSICHAR* StackTrace = (ANSICHAR*) GMalloc->Malloc( StackTraceSize );
 	StackTrace[0] = 0;
-	// Walk the stack and dump it to the allocated memory.
+	// Walk the stack and dump it to the allocated memory. This process usually allocates a lot of memory.
 	FPlatformStackWalk::StackWalkAndDump( StackTrace, StackTraceSize, 0, ExceptionInfo->ContextRecord );
 
 	// @TODO yrx 2014-08-20 Make this constant.
@@ -484,22 +506,6 @@ int32 ReportCrash( LPEXCEPTION_POINTERS ExceptionInfo )
 
 	GMalloc->Free( StackTrace );
 
-#if WINVER > 0x502	// Windows Error Reporting is not supported on Windows XP
-	if (GUseCrashReportClient)
-	{
-		ReportCrashUsingCrashReportClient(ExceptionInfo, GErrorMessage, EErrorReportUI::ShowDialog);
-	}
-	else
-#endif		// WINVER
-	{
-		WriteMinidump(MiniDumpFilenameW, ExceptionInfo);
-
-#if UE_BUILD_SHIPPING && WITH_EDITOR
-		uint32 dwOpt = 0;
-		EFaultRepRetVal repret = ReportFault( ExceptionInfo, dwOpt);
-#endif
-	}
-
 	return EXCEPTION_EXECUTE_HANDLER;
 }
 
diff --git a/Engine/Source/Runtime/Core/Private/Windows/WindowsRunnableThread.cpp b/Engine/Source/Runtime/Core/Private/Windows/WindowsRunnableThread.cpp
--- a/Engine/Source/Runtime/Core/Private/Windows/WindowsRunnableThread.cpp
+++ b/Engine/Source/Runtime/Core/Private/Windows/WindowsRunnableThread.cpp
@@ -10,40 +10,43 @@ DEFINE_LOG_CATEGORY_STATIC(LogThreadingWindows, Log, All);
 
 uint32 FRunnableThreadWin::GuardedRun()
 {
-	uint32 ExitCode = 1;
+	uint32 ExitCode = 0;
 
 	FPlatformProcess::SetThreadAffinityMask(ThreadAffinityMask);
 
-#if !PLATFORM_SEH_EXCEPTIONS_DISABLED
-	if( !FPlatformMisc::IsDebuggerPresent() || GAlwaysReportCrash )
+#if UE_BUILD_DEBUG
+	if (true && !GAlwaysReportCrash)
+#else
+	if (FPlatformMisc::IsDebuggerPresent() && !GAlwaysReportCrash)
+#endif // UE_BUILD_DEBUG
 	{
+		ExitCode = Run();
+	}
+	else
+	{
+#if !PLATFORM_SEH_EXCEPTIONS_DISABLED
 		__try
+#endif // !PLATFORM_SEH_EXCEPTIONS_DISABLED
 		{
 			ExitCode = Run();
 		}
-		__except( ReportCrash( GetExceptionInformation() ) )
+#if !PLATFORM_SEH_EXCEPTIONS_DISABLED
+		__except (ReportCrash( GetExceptionInformation() ))
 		{
 			// Make sure the information which thread crashed makes it into the log.
-			UE_LOG(LogThreadingWindows, Error, TEXT("Runnable thread %s crashed."), *ThreadName);
+			UE_LOG( LogThreadingWindows, Error, TEXT( "Runnable thread %s crashed." ), *ThreadName );
 			GWarn->Flush();
 
 			// Append the thread name at the end of the error report.
 			FCString::Strncat( GErrorHist, LINE_TERMINATOR TEXT( "Crash in runnable thread " ), ARRAY_COUNT( GErrorHist ) );
 			FCString::Strncat( GErrorHist, *ThreadName, ARRAY_COUNT( GErrorHist ) );
 
-			ExitCode = 1;
-			// Generate status report.				
+			// Crashed.
+			ExitCode = 1;		
 			GError->HandleError();
-			// Throw an error so that main thread shuts down too (otherwise task graph stalls forever).
-			UE_LOG(LogThreadingWindows, Fatal, TEXT("Runnable thread %s crashed."), *ThreadName);
-
 			FPlatformMisc::RequestExit( true );
 		}
-	}
-	else
-#endif
-	{
-		ExitCode = Run();
+#endif // !PLATFORM_SEH_EXCEPTIONS_DISABLED
 	}
 
 	return ExitCode;
