Fixed expression parser jump table execution erroneously returning the proxy, rather than a valid result.
This caused the proxy to potentially outlive its content, which caused the automation tests to fail.

Also cleaned up some MoveTemp usage and removed some redundant move operators.
Re-enabled the math expression evaluator tests as smoke tests

diff --git a/Engine/Source/Runtime/Core/Private/Math/BasicMathExpressionEvaluator.cpp b/Engine/Source/Runtime/Core/Private/Math/BasicMathExpressionEvaluator.cpp
--- a/Engine/Source/Runtime/Core/Private/Math/BasicMathExpressionEvaluator.cpp
+++ b/Engine/Source/Runtime/Core/Private/Math/BasicMathExpressionEvaluator.cpp
@@ -196,7 +196,7 @@ bool TestInvalidExpression(FAutomationTestBase* Test, const TCHAR* Expression)
 	return true;
 }
 
-IMPLEMENT_SIMPLE_AUTOMATION_TEST(FBasicMathExpressionEvaluatorTest, "System.Core.Math.Expressions", EAutomationTestFlags::ATF_Editor)
+IMPLEMENT_SIMPLE_AUTOMATION_TEST(FBasicMathExpressionEvaluatorTest, "System.Core.Math.Expressions", EAutomationTestFlags::ATF_SmokeTest)
 bool FBasicMathExpressionEvaluatorTest::RunTest( const FString& Parameters )
 {
 	bool bResult = true;
diff --git a/Engine/Source/Runtime/Core/Private/Misc/ExpressionParser.cpp b/Engine/Source/Runtime/Core/Private/Misc/ExpressionParser.cpp
--- a/Engine/Source/Runtime/Core/Private/Misc/ExpressionParser.cpp
+++ b/Engine/Source/Runtime/Core/Private/Misc/ExpressionParser.cpp
@@ -307,13 +307,13 @@ TArray<FExpressionToken> FExpressionTokenConsumer::Extract()
 {
 	TArray<FExpressionToken> Swapped;
 	Swap(Swapped, Tokens);
-	return MoveTemp(Swapped);
+	return Swapped;
 }
 
 void FExpressionTokenConsumer::Add(const FStringToken& SourceToken, FExpressionNode Node)
 {
 	Stream.SetReadPos(SourceToken);
-	Tokens.Add(FExpressionToken(SourceToken, MoveTemp(Node)));
+	Tokens.Add(FExpressionToken(SourceToken, Node));
 }
 
 void FTokenDefinitions::DefineToken(const TFunction<FExpressionDefinition>& Definition)
diff --git a/Engine/Source/Runtime/Core/Public/Misc/ExpressionParserTypes.h b/Engine/Source/Runtime/Core/Public/Misc/ExpressionParserTypes.h
--- a/Engine/Source/Runtime/Core/Public/Misc/ExpressionParserTypes.h
+++ b/Engine/Source/Runtime/Core/Public/Misc/ExpressionParserTypes.h
@@ -171,7 +171,7 @@ public:
 		}
 		else
 		{
-			BigData = MakeShareable(new T(MoveTemp(In)));
+			BigData = MakeShareable(new T(In));
 		}
 	}
 
@@ -228,15 +228,12 @@ private:
 class FExpressionToken
 {
 public:
-	FExpressionToken(const FStringToken& InContext, FExpressionNode InNode)
-		: Node(MoveTemp(InNode))
+	FExpressionToken(const FStringToken& InContext, const FExpressionNode& InNode)
+		: Node(InNode)
 		, Context(InContext)
 	{
 	}
 
-	FExpressionToken(FExpressionToken&& In) : Node(MoveTemp(In.Node)), Context(In.Context) {}
-	FExpressionToken& operator=(FExpressionToken&& In) { Node = MoveTemp(In.Node); Context = MoveTemp(In.Context); return *this; }
-
 	FExpressionToken(const FExpressionToken& In) : Node(In.Node), Context(In.Context) {}
 	FExpressionToken& operator=(const FExpressionToken& In) { Node = In.Node; Context = In.Context; return *this; }
 
@@ -247,7 +244,7 @@ public:
 /** Built-in wrapper tokens used for compilation */
 struct FGroupMarker
 {
-	FGroupMarker(FExpressionToken InWrappedToken) : WrappedToken(MoveTemp(InWrappedToken)), NumTokens(0) {}
+	FGroupMarker(const FExpressionToken& InWrappedToken) : WrappedToken(InWrappedToken), NumTokens(0) {}
 
 	template<typename T>	bool 		IsA() const 	{ return WrappedToken.Node.GetTypeId() == TGetExpressionNodeTypeId<T>::GetTypeId(); }
 	template<typename T>	const T* 	Cast() const	{ return WrappedToken.Node.Cast<T>(); }
@@ -260,7 +257,7 @@ DEFINE_EXPRESSION_NODE_TYPE(FGroupMarker, 0xEBF6DBA6, 0xE81B4684, 0x938EF46A, 0x
 struct FWrappedOperator
 {
 	enum EType { PreUnary, PostUnary, Binary };
-	FWrappedOperator(EType InType, FExpressionToken InWrappedToken, int32 InPrec = 0) : Type(InType), WrappedToken(MoveTemp(InWrappedToken)), Precedence(InPrec) {}
+	FWrappedOperator(EType InType, const FExpressionToken& InWrappedToken, int32 InPrec = 0) : Type(InType), WrappedToken(InWrappedToken), Precedence(InPrec) {}
 
 	template<typename T>	bool 		IsA() const 	{ return WrappedToken.Node.GetTypeId() == TGetExpressionNodeTypeId<T>::GetTypeId(); }
 	template<typename T>	const T* 	Cast() const	{ return WrappedToken.Node.Cast<T>(); }
@@ -322,7 +319,7 @@ struct FOperatorJumpTable
 			TGetExpressionNodeTypeId<OperandType>::GetTypeId()
 		};
 
-		PreUnaryOps.Add(ID, [=](const FExpressionNode& InOperand) -> FExpressionResult {
+		PreUnaryOps.Add(ID, [=](const FExpressionNode& InOperand) {
 			return InFunc(*InOperand.Cast<OperandType>());
 		});
 	}
@@ -346,6 +343,7 @@ struct FOperatorJumpTable
 			TGetExpressionNodeTypeId<OperandType>::GetTypeId()
 		};
 
+		// Explicit return type is important here, to ensure that the proxy returned from MakeValue does not outlive the value it's proxying
 		PreUnaryOps.Add(ID, [=](const FExpressionNode& InOperand) -> FExpressionResult {
 			return MakeValue(InFunc(*InOperand.Cast<OperandType>()));
 		});
@@ -369,7 +367,7 @@ struct FOperatorJumpTable
 			FGuid()
 		};
 
-		PostUnaryOps.Add(ID, [=](const FExpressionNode& InOperand) -> FExpressionResult {
+		PostUnaryOps.Add(ID, [=](const FExpressionNode& InOperand) {
 			return InFunc(*InOperand.Cast<OperandType>());
 		});
 	}
@@ -393,6 +391,7 @@ struct FOperatorJumpTable
 			FGuid()
 		};
 
+		// Explicit return type is important here, to ensure that the proxy returned from MakeValue does not outlive the value it's proxying
 		PostUnaryOps.Add(ID, [=](const FExpressionNode& InOperand) -> FExpressionResult {
 			return MakeValue(InFunc(*InOperand.Cast<OperandType>()));
 		});
@@ -448,7 +447,8 @@ struct FOperatorJumpTable
 			TGetExpressionNodeTypeId<RightOperandType>::GetTypeId()
 		};
 
-		BinaryOps.Add(ID, [=](const FExpressionNode& InLeftOperand, const FExpressionNode& InRightOperand){
+		// Explicit return type is important here, to ensure that the proxy returned from MakeValue does not outlive the value it's proxying
+		BinaryOps.Add(ID, [=](const FExpressionNode& InLeftOperand, const FExpressionNode& InRightOperand) -> FExpressionResult {
 			return MakeValue(InFunc(*InLeftOperand.Cast<LeftOperandType>(), *InRightOperand.Cast<RightOperandType>()));
 		});
 	}
