#ue4 - Remove some temp VectorRegisters from SSE VectorReciprocalSqrtAccurate() implementation.

diff --git a/Engine/Source/Runtime/Core/Public/Math/UnrealMathSSE.h b/Engine/Source/Runtime/Core/Public/Math/UnrealMathSSE.h
--- a/Engine/Source/Runtime/Core/Public/Math/UnrealMathSSE.h
+++ b/Engine/Source/Runtime/Core/Public/Math/UnrealMathSSE.h
@@ -487,14 +487,14 @@ FORCEINLINE VectorRegister VectorReciprocalSqrtAccurate(const VectorRegister& Ve
 	const VectorRegister x0 = VectorReciprocalSqrt(Vec);
 
 	// First iteration
-	const VectorRegister x0Squared = VectorMultiply(x0, x0);
-	const VectorRegister InnerTerm0 = VectorSubtract(OneHalf, VectorMultiply(VecDivBy2, x0Squared));
-	const VectorRegister x1 = VectorMultiplyAdd(x0, InnerTerm0, x0);
+	VectorRegister x1 = VectorMultiply(x0, x0);
+	x1 = VectorSubtract(OneHalf, VectorMultiply(VecDivBy2, x1));
+	x1 = VectorMultiplyAdd(x0, x1, x0);
 
 	// Second iteration
-	const VectorRegister x1Squared = VectorMultiply(x1, x1);
-	const VectorRegister InnerTerm1 = VectorSubtract(OneHalf, VectorMultiply(VecDivBy2, x1Squared));
-	const VectorRegister x2 = VectorMultiplyAdd(x1, InnerTerm1, x1);
+	VectorRegister x2 = VectorMultiply(x1, x1);
+	x2 = VectorSubtract(OneHalf, VectorMultiply(VecDivBy2, x2));
+	x2 = VectorMultiplyAdd(x1, x2, x1);
 
 	return x2;
 }
