#ue4 - Remove USE_MATRIX_ROTATOR from FQuat. It's not maintained and makes the code less readable.

diff --git a/Engine/Source/Runtime/Core/Private/Math/UnrealMath.cpp b/Engine/Source/Runtime/Core/Private/Math/UnrealMath.cpp
--- a/Engine/Source/Runtime/Core/Private/Math/UnrealMath.cpp
+++ b/Engine/Source/Runtime/Core/Private/Math/UnrealMath.cpp
@@ -276,10 +276,6 @@ FQuat FRotator::Quaternion() const
 {
 	SCOPE_CYCLE_COUNTER(STAT_MathConvertRotatorToQuat);
 
-#if USE_MATRIX_ROTATOR 
-	FQuat RotationMatrix = FQuat( FRotationMatrix( *this ) );
-#endif
-
 #if PLATFORM_ENABLE_VECTORINTRINSICS
 	const VectorRegister Angles = MakeVectorRegister(Pitch, Yaw, Roll, 0.0f);
 	const VectorRegister HalfAngles = VectorMultiply(Angles, GlobalVectorConstants::DEG_TO_RAD_HALF);
@@ -327,16 +323,7 @@ FQuat FRotator::Quaternion() const
 	RotationQuat.W =  CR*CP*CY + SR*SP*SY;
 #endif // PLATFORM_ENABLE_VECTORINTRINSICS
 
-#if USE_MATRIX_ROTATOR 
-	if (!RotationMatrix.Equals(RotationQuat, KINDA_SMALL_NUMBER) && !RotationMatrix.Equals(RotationQuat*-1.f, KINDA_SMALL_NUMBER))
-	{
-		UE_LOG(LogUnrealMath, Log, TEXT("RotationMatrix (%s), RotationQuat (%s)"), *RotationMatrix.ToString(), *RotationQuat.ToString());
-	}
-	return RotationMatrix;
-
-#else
 	return RotationQuat;
-#endif
 }
 
 FVector FRotator::Euler() const
@@ -447,12 +434,6 @@ FRotator FQuat::Rotator() const
 {
 	SCOPE_CYCLE_COUNTER(STAT_MathConvertQuatToRotator);
 
-#if USE_MATRIX_ROTATOR 
-	// if you think this function is problem, you can undo previous matrix rotator by returning RotatorFromMatrix
-	FRotator RotatorFromMatrix = FQuatRotationTranslationMatrix(*this, FVector::ZeroVector).Rotator();
-	checkSlow(IsNormalized());
-#endif
-
 	const float SingularityTest = Z*X-W*Y;
 	const float YawY = 2.f*(W*Z+X*Y);
 	const float YawX = (1.f-2.f*(FMath::Square(Y) + FMath::Square(Z)));
@@ -487,19 +468,7 @@ FRotator FQuat::Rotator() const
 		RotatorFromQuat.Roll = FMath::Atan2(-2.f*(W*X+Y*Z), (1.f-2.f*(FMath::Square(X) + FMath::Square(Y)))) * RAD_TO_DEG;
 	}
 
-#if USE_MATRIX_ROTATOR
-	RotatorFromMatrix = RotatorFromMatrix.Clamp();
-
-	// this Euler is degree, so less 1 is negligible
-	if (!DebugRotatorEquals(RotatorFromQuat, RotatorFromMatrix, 0.1f))
-	{
-		FPlatformMisc::LowLevelOutputDebugStringf(TEXT("WRONG: (Singularity: %.9f) RotationMatrix (%s), RotationQuat(%s)"), SingularityTest, *RotatorFromMatrix.ToString(), *RotatorFromQuat.ToString());
-	}
-
-	return RotatorFromMatrix;
-#else
 	return RotatorFromQuat;
-#endif
 }
 
 FQuat FQuat::MakeFromEuler(const FVector& Euler)
diff --git a/Engine/Source/Runtime/Core/Public/Math/Quat.h b/Engine/Source/Runtime/Core/Public/Math/Quat.h
--- a/Engine/Source/Runtime/Core/Public/Math/Quat.h
+++ b/Engine/Source/Runtime/Core/Public/Math/Quat.h
@@ -569,31 +569,6 @@ inline FMatrix FQuat::operator*( const FMatrix& M ) const
 }
 
 
-/** now we directly convert from Rotator to Quaternion and vice versa. If you see issue with rotator, you can contact @LH or  
- *  you can undo this to see if this fixes issue. Feel free to enable this and see if that fixes the issue.
- */
-#define USE_MATRIX_ROTATOR 0
-
-#if USE_MATRIX_ROTATOR
-
-#include "QuatRotationTranslationMatrix.h"
-
-/** 
- * this is not right representation of if both rotation is equal or not
- * if you really like to equal, use matrix form. 
- * @warning DO NOT USE THIS To VERIFY if rotation is same. I'm only using this for the debug purpose
- */
-FORCEINLINE bool DebugRotatorEquals(const FRotator& R1, const FRotator& R2, float Tolerance=1.f) 
-{
-	// also 0 and 360 should be considered same 
-	return ( FMath::Abs(R1.Pitch-R2.Pitch) <= Tolerance || FMath::Abs(R1.Pitch+R2.Pitch-360) <= Tolerance )
-		&& ( FMath::Abs(R1.Yaw-R2.Yaw) <= Tolerance || FMath::Abs(R1.Yaw+R2.Yaw-360) <= Tolerance )
-		&& ( FMath::Abs(R1.Roll-R2.Roll) <= Tolerance || FMath::Abs(R1.Roll+R2.Roll-360) <= Tolerance );
-}
-#endif
-
-
-
 /* FQuat inline functions
  *****************************************************************************/
 
