#ue4 - More accurate InvSqrt() and InvSqrtEst() on some platforms.

- Matches results for SSE platforms with other platforms that perform 2 Newton-Raphson iterations on the hardware estimate (iOS), or non-SSE that uses the fully accurate sqrt.
- Scalar functions now match the vectorized versions (VectorReciprocalSqrtAccurate) so that comparisons within 1e-8f are possible for normalizing normalized values (important for FQuat which uses both scalar and vector).
- Allows normalize vector (1,0,0) to return the same vector rather than (0.9999998xx, 0, 0).
- Fixed InvSqrtEst() on some platforms actually using higher precision (slower) FGenericPlatformMath InvSqrt().

diff --git a/Engine/Source/Runtime/Core/Private/Tests/Math/UnrealMathTest.cpp b/Engine/Source/Runtime/Core/Private/Tests/Math/UnrealMathTest.cpp
--- a/Engine/Source/Runtime/Core/Private/Tests/Math/UnrealMathTest.cpp
+++ b/Engine/Source/Runtime/Core/Private/Tests/Math/UnrealMathTest.cpp
@@ -477,6 +477,29 @@ void LogRotatorTest(bool bExpected, const TCHAR* TestName, const FRotator& A, co
 }
 
 
+// Normalize tests
+
+MATHTEST_INLINE VectorRegister TestVectorNormalize_Sqrt(const VectorRegister& V)
+{
+	const VectorRegister Len = VectorDot4(V, V);
+	const float rlen = 1.0f / FMath::Sqrt(VectorGetComponent(Len, 0));
+	return VectorMultiply(V, VectorLoadFloat1(&rlen));
+}
+
+MATHTEST_INLINE VectorRegister TestVectorNormalize_InvSqrt(const VectorRegister& V)
+{
+	const VectorRegister Len = VectorDot4(V, V);
+	const float rlen = FMath::InvSqrt(VectorGetComponent(Len, 0));
+	return VectorMultiply(V, VectorLoadFloat1(&rlen));
+}
+
+MATHTEST_INLINE VectorRegister TestVectorNormalize_InvSqrtEst(const VectorRegister& V)
+{
+	const VectorRegister Len = VectorDot4(V, V);
+	const float rlen = FMath::InvSqrtEst(VectorGetComponent(Len, 0));
+	return VectorMultiply(V, VectorLoadFloat1(&rlen));
+}
+
 
 /**
  * Helper debugf function to print out success or failure information for a test
@@ -677,6 +700,26 @@ bool FVectorRegisterAbstractionTest::RunTest(const FString& Parameters)
 	V0 = MakeVectorRegister( 0.5f, -0.5f, 0.5f, -0.5f );
 	LogTest( TEXT("VectorNormalize"), TestVectorsEqual( V0, V1, 0.001f ) );
 
+	V0 = MakeVectorRegister(2.0f, -2.0f, 2.0f, -2.0f);
+	V1 = VectorNormalizeAccurate(V0);
+	V0 = MakeVectorRegister(0.5f, -0.5f, 0.5f, -0.5f);
+	LogTest(TEXT("VectorNormalizeAccurate"), TestVectorsEqual(V0, V1, 1e-8f));
+
+	V0 = MakeVectorRegister(2.0f, -2.0f, 2.0f, -2.0f);
+	V1 = TestVectorNormalize_Sqrt(V0);
+	V0 = MakeVectorRegister(0.5f, -0.5f, 0.5f, -0.5f);
+	LogTest(TEXT("TestVectorNormalize_Sqrt"), TestVectorsEqual(V0, V1, 1e-8f));
+
+	V0 = MakeVectorRegister(2.0f, -2.0f, 2.0f, -2.0f);
+	V1 = TestVectorNormalize_InvSqrt(V0);
+	V0 = MakeVectorRegister(0.5f, -0.5f, 0.5f, -0.5f);
+	LogTest(TEXT("TestVectorNormalize_InvSqrt"), TestVectorsEqual(V0, V1, 1e-8f));
+
+	V0 = MakeVectorRegister(2.0f, -2.0f, 2.0f, -2.0f);
+	V1 = TestVectorNormalize_InvSqrtEst(V0);
+	V0 = MakeVectorRegister(0.5f, -0.5f, 0.5f, -0.5f);
+	LogTest(TEXT("TestVectorNormalize_InvSqrtEst"), TestVectorsEqual(V0, V1, 1e-6f));
+
 	V0 = MakeVectorRegister( 2.0f, -2.0f, 2.0f, -2.0f );
 	V1 = VectorSet_W0( V0 );
 	V0 = MakeVectorRegister( 2.0f, -2.0f, 2.0f, 0.0f );
diff --git a/Engine/Source/Runtime/Core/Public/Linux/LinuxPlatformMath.h b/Engine/Source/Runtime/Core/Public/Linux/LinuxPlatformMath.h
--- a/Engine/Source/Runtime/Core/Public/Linux/LinuxPlatformMath.h
+++ b/Engine/Source/Runtime/Core/Public/Linux/LinuxPlatformMath.h
@@ -7,6 +7,7 @@
 #pragma once
 #include "GenericPlatform/GenericPlatformMath.h"
 #include "Linux/LinuxSystemIncludes.h"
+#include "Math/UnrealPlatformMathSSE.h"
 
 /**
 * Linux implementation of the Math OS functions
@@ -60,28 +61,14 @@ struct FLinuxPlatformMath : public FGenericPlatformMath
 	static FORCEINLINE bool IsNaN( float A ) { return isnan(A) != 0; }
 	static FORCEINLINE bool IsFinite( float A ) { return isfinite(A); }
 
-	//
-	// MSM: Fast float inverse square root using SSE.
-	// Accurate to within 1 LSB.
-	//
 	static FORCEINLINE float InvSqrt( float F )
 	{
-		static const __m128 fThree = _mm_set_ss( 3.0f );
-		static const __m128 fOneHalf = _mm_set_ss( 0.5f );
-		__m128 Y0, X0, Temp;
-		float temp;
-
-		Y0 = _mm_set_ss( F );
-		X0 = _mm_rsqrt_ss( Y0 );	// 1/sqrt estimate (12 bits)
-
-		// Newton-Raphson iteration (X1 = 0.5*X0*(3-(Y*X0)*X0))
-		Temp = _mm_mul_ss( _mm_mul_ss(Y0, X0), X0 );	// (Y*X0)*X0
-		Temp = _mm_sub_ss( fThree, Temp );				// (3-(Y*X0)*X0)
-		Temp = _mm_mul_ss( X0, Temp );					// X0*(3-(Y*X0)*X0)
-		Temp = _mm_mul_ss( fOneHalf, Temp );			// 0.5*X0*(3-(Y*X0)*X0)
-		_mm_store_ss( &temp, Temp );
-
-		return temp;
+		return UnrealPlatformMathSSE::InvSqrt(F);
+	}
+
+	static FORCEINLINE float InvSqrtEst( float F )
+	{
+		return UnrealPlatformMathSSE::InvSqrtEst(F);
 	}
 #endif
 
diff --git a/Engine/Source/Runtime/Core/Public/Mac/MacPlatformMath.h b/Engine/Source/Runtime/Core/Public/Mac/MacPlatformMath.h
--- a/Engine/Source/Runtime/Core/Public/Mac/MacPlatformMath.h
+++ b/Engine/Source/Runtime/Core/Public/Mac/MacPlatformMath.h
@@ -7,6 +7,7 @@
 #pragma once
 #include "GenericPlatform/GenericPlatformMath.h"
 #include "Mac/MacSystemIncludes.h"
+#include "Math/UnrealPlatformMathSSE.h"
 
 /**
 * Mac implementation of the Math OS functions
@@ -58,29 +59,15 @@ struct FMacPlatformMath : public FGenericPlatformMath
 	static FORCEINLINE bool IsNaN( float A ) { return isnan(A) != 0; }
 	static FORCEINLINE bool IsFinite( float A ) { return isfinite(A); }
 
-	//
-	// MSM: Fast float inverse square root using SSE.
-	// Accurate to within 1 LSB.
-	//
 #if PLATFORM_ENABLE_VECTORINTRINSICS
-	static FORCEINLINE float InvSqrt( float F )
+	static FORCEINLINE float InvSqrt(float F)
 	{
-		static const __m128 fThree = _mm_set_ss( 3.0f );
-		static const __m128 fOneHalf = _mm_set_ss( 0.5f );
-		__m128 Y0, X0, Temp;
-		float temp;
-
-		Y0 = _mm_set_ss( F );
-		X0 = _mm_rsqrt_ss( Y0 );	// 1/sqrt estimate (12 bits)
-
-		// Newton-Raphson iteration (X1 = 0.5*X0*(3-(Y*X0)*X0))
-		Temp = _mm_mul_ss( _mm_mul_ss(Y0, X0), X0 );	// (Y*X0)*X0
-		Temp = _mm_sub_ss( fThree, Temp );				// (3-(Y*X0)*X0)
-		Temp = _mm_mul_ss( X0, Temp );					// X0*(3-(Y*X0)*X0)
-		Temp = _mm_mul_ss( fOneHalf, Temp );			// 0.5*X0*(3-(Y*X0)*X0)
-		_mm_store_ss( &temp, Temp );
-
-		return temp;
+		return UnrealPlatformMathSSE::InvSqrt(F);
+	}
+
+	static FORCEINLINE float InvSqrtEst( float F )
+	{
+		return UnrealPlatformMathSSE::InvSqrtEst(F);
 	}
 #endif
 
diff --git a/Engine/Source/Runtime/Core/Public/Math/UnrealPlatformMathSSE.h b/Engine/Source/Runtime/Core/Public/Math/UnrealPlatformMathSSE.h
new file mode 100644
--- /dev/null
+++ b/Engine/Source/Runtime/Core/Public/Math/UnrealPlatformMathSSE.h
@@ -0,0 +1,68 @@
+// Copyright 1998-2015 Epic Games, Inc. All Rights Reserved.
+
+#pragma once
+
+// Code including this header is responsible for including the correct platform-specific header for SSE intrinsics.
+
+
+namespace UnrealPlatformMathSSE
+{
+	static FORCEINLINE float InvSqrt(float F)
+	{
+		// Performs two passes of Newton-Raphson iteration on the hardware estimate
+		//    v^-0.5 = x
+		// => x^2 = v^-1
+		// => 1/(x^2) = v
+		// => F(x) = x^-2 - v
+		//    F'(x) = -2x^-3
+
+		//    x1 = x0 - F(x0)/F'(x0)
+		// => x1 = x0 + 0.5 * (x0^-2 - Vec) * x0^3
+		// => x1 = x0 + 0.5 * (x0 - Vec * x0^3)
+		// => x1 = x0 + x0 * (0.5 - 0.5 * Vec * x0^2)
+		//
+		// This final form has one more operation than the legacy factorization (X1 = 0.5*X0*(3-(Y*X0)*X0)
+		// but retains better accuracy (namely InvSqrt(1) = 1 exactly).
+
+		const __m128 fOneHalf = _mm_set_ss(0.5f);
+		__m128 Y0, X0, X1, X2, FOver2;
+		float temp;
+
+		Y0 = _mm_set_ss(F);
+		X0 = _mm_rsqrt_ss(Y0);	// 1/sqrt estimate (12 bits)
+		FOver2 = _mm_mul_ss(Y0, fOneHalf);
+
+		// 1st Newton-Raphson iteration
+		X1 = _mm_mul_ss(X0, X0);
+		X1 = _mm_sub_ss(fOneHalf, _mm_mul_ss(FOver2, X1));
+		X1 = _mm_add_ss(X0, _mm_mul_ss(X0, X1));
+
+		// 2nd Newton-Raphson iteration
+		X2 = _mm_mul_ss(X1, X1);
+		X2 = _mm_sub_ss(fOneHalf, _mm_mul_ss(FOver2, X2));
+		X2 = _mm_add_ss(X1, _mm_mul_ss(X1, X2));
+
+		_mm_store_ss(&temp, X2);
+		return temp;
+	}
+
+	static FORCEINLINE float InvSqrtEst(float F)
+	{
+		// Performs one pass of Newton-Raphson iteration on the hardware estimate
+		const __m128 fOneHalf = _mm_set_ss(0.5f);
+		__m128 Y0, X0, X1, FOver2;
+		float temp;
+
+		Y0 = _mm_set_ss(F);
+		X0 = _mm_rsqrt_ss(Y0);	// 1/sqrt estimate (12 bits)
+		FOver2 = _mm_mul_ss(Y0, fOneHalf);
+
+		// 1st Newton-Raphson iteration
+		X1 = _mm_mul_ss(X0, X0);
+		X1 = _mm_sub_ss(fOneHalf, _mm_mul_ss(FOver2, X1));
+		X1 = _mm_add_ss(X0, _mm_mul_ss(X0, X1));
+
+		_mm_store_ss(&temp, X1);
+		return temp;
+	}
+}
diff --git a/Engine/Source/Runtime/Core/Public/Windows/WindowsPlatformMath.h b/Engine/Source/Runtime/Core/Public/Windows/WindowsPlatformMath.h
--- a/Engine/Source/Runtime/Core/Public/Windows/WindowsPlatformMath.h
+++ b/Engine/Source/Runtime/Core/Public/Windows/WindowsPlatformMath.h
@@ -11,6 +11,8 @@
 #include "XboxOne/XboxOneSystemIncludes.h"
 #endif
 
+#include "Math/UnrealPlatformMathSSE.h"
+
 /**
 * Windows implementation of the Math OS functions
 **/
@@ -63,34 +65,16 @@ struct FWindowsPlatformMath : public FGenericPlatformMath
 	static FORCEINLINE bool IsNaN( float A ) { return _isnan(A) != 0; }
 	static FORCEINLINE bool IsFinite( float A ) { return _finite(A) != 0; }
 
-	//
-	// MSM: Fast float inverse square root using SSE.
-	// Accurate to within 1 LSB.
-	//
-	static FORCEINLINE float InvSqrt( float F )
+	static FORCEINLINE float InvSqrt(float F)
 	{
-		static const __m128 fThree = _mm_set_ss( 3.0f );
-		static const __m128 fOneHalf = _mm_set_ss( 0.5f );
-		__m128 Y0, X0, Temp;
-		float temp;
-
-		Y0 = _mm_set_ss( F );
-		X0 = _mm_rsqrt_ss( Y0 );	// 1/sqrt estimate (12 bits)
-
-		// Newton-Raphson iteration (X1 = 0.5*X0*(3-(Y*X0)*X0))
-		Temp = _mm_mul_ss( _mm_mul_ss(Y0, X0), X0 );	// (Y*X0)*X0
-		Temp = _mm_sub_ss( fThree, Temp );				// (3-(Y*X0)*X0)
-		Temp = _mm_mul_ss( X0, Temp );					// X0*(3-(Y*X0)*X0)
-		Temp = _mm_mul_ss( fOneHalf, Temp );			// 0.5*X0*(3-(Y*X0)*X0)
-		_mm_store_ss( &temp, Temp );
-
-		return temp;
+		return UnrealPlatformMathSSE::InvSqrt(F);
 	}
 
 	static FORCEINLINE float InvSqrtEst( float F )
 	{
-		return InvSqrt( F );
+		return UnrealPlatformMathSSE::InvSqrtEst(F);
 	}
+		
 
 	#pragma intrinsic( _BitScanReverse )
 	static FORCEINLINE uint32 FloorLog2(uint32 Value) 
