#ue4 - Add DirectX implementations of VectorCeil, VectorFloor, VectorTruncate, VectorFractional.

diff --git a/Engine/Source/Runtime/Core/Public/Math/UnrealMathDirectX.h b/Engine/Source/Runtime/Core/Public/Math/UnrealMathDirectX.h
--- a/Engine/Source/Runtime/Core/Public/Math/UnrealMathDirectX.h
+++ b/Engine/Source/Runtime/Core/Public/Math/UnrealMathDirectX.h
@@ -876,25 +876,21 @@ FORCEINLINE VectorRegister VectorATan2(const VectorRegister& X, const VectorRegi
 		(float)atan2(VectorGetComponent(X, 3), VectorGetComponent(Y, 3)));
 }
 
-//TODO: Vectorize
 FORCEINLINE VectorRegister VectorCeil(const VectorRegister& X)
 {
-	return MakeVectorRegister((float)ceil(VectorGetComponent(X, 0)), (float)ceil(VectorGetComponent(X, 1)), (float)ceil(VectorGetComponent(X, 2)), (float)ceil(VectorGetComponent(X, 3)));
+	return DirectX::XMVectorCeiling(X);
 }
 
-//TODO: Vectorize
 FORCEINLINE VectorRegister VectorFloor(const VectorRegister& X)
 {
-	return MakeVectorRegister((float)floor(VectorGetComponent(X, 0)), (float)floor(VectorGetComponent(X, 1)), (float)floor(VectorGetComponent(X, 2)), (float)floor(VectorGetComponent(X, 3)));
+	return DirectX::XMVectorFloor(X);
 }
 
-//TODO: Vectorize
 FORCEINLINE VectorRegister VectorTruncate(const VectorRegister& X)
 {
-	return MakeVectorRegister((float)(int32)VectorGetComponent(X, 0), (float)(int32)VectorGetComponent(X, 1), (float)(int32)VectorGetComponent(X, 2), (float)(int32)VectorGetComponent(X, 3));
+	return DirectX::XMVectorTruncate(X);
 }
 
-//TODO: Vectorize
 FORCEINLINE VectorRegister VectorFractional(const VectorRegister& X)
 {
 	return VectorSubtract(X, VectorTruncate(X));
