Expression parser improvements

Refactored the guts of FExpressionNode to better afford holding moveable types on the stack. Large types are now also stored in the inline bytes for simplicity's sake, but allocated on the heap, and managed with a TUniquePtr.
Added the concept of an evaluation context that can be passed in to the Evaluate function. Jump tables are now templated on a given context type, which is wrapped up with a supplied context for evaluation purposes into a TOperatorEvaluationEnvironment (so that the evaluation logic need not be concerned with the context itself)

diff --git a/Engine/Source/Runtime/Core/Private/Math/BasicMathExpressionEvaluator.cpp b/Engine/Source/Runtime/Core/Private/Math/BasicMathExpressionEvaluator.cpp
--- a/Engine/Source/Runtime/Core/Private/Math/BasicMathExpressionEvaluator.cpp
+++ b/Engine/Source/Runtime/Core/Private/Math/BasicMathExpressionEvaluator.cpp
@@ -160,6 +160,8 @@ TValueOrError<double, FExpressionError> FBasicMathExpressionEvaluator::Evaluate(
 	}
 }
 
+#if !(UE_BUILD_SHIPPING || UE_BUILD_TEST)
+
 bool TestExpression(FAutomationTestBase* Test, const TCHAR* Expression, double Expected)
 {
 	FBasicMathExpressionEvaluator Parser;
@@ -233,4 +235,6 @@ bool FBasicMathExpressionEvaluatorTest::RunTest( const FString& Parameters )
 	return bResult;
 }
 
+#endif //!(UE_BUILD_SHIPPING || UE_BUILD_TEST)
+
 #undef LOCTEXT_NAMESPACE
\ No newline at end of file
diff --git a/Engine/Source/Runtime/Core/Private/Misc/ExpressionParser.cpp b/Engine/Source/Runtime/Core/Private/Misc/ExpressionParser.cpp
--- a/Engine/Source/Runtime/Core/Private/Misc/ExpressionParser.cpp
+++ b/Engine/Source/Runtime/Core/Private/Misc/ExpressionParser.cpp
@@ -1,15 +1,17 @@
+// Copyright 1998-2015 Epic Games, Inc. All Rights Reserved.
 
 #include "CorePrivatePCH.h"
 #include "ExpressionParser.h"
 
 #define LOCTEXT_NAMESPACE "ExpressionParser"
 
+
+
 FTokenStream::FTokenStream(const TCHAR* In)
 	: Start(In)
 	, End(Start + FCString::Strlen(Start))
 	, ReadPos(In)
 {
-
 }
 
 bool FTokenStream::IsReadPosValid(const TCHAR* InPos, int32 MinNumChars) const
@@ -256,49 +258,6 @@ void FTokenStream::SetReadPos(const FStringToken& Token)
 	}
 }
 
-FExpressionResult FOperatorJumpTable::ExecBinary(const FExpressionToken& Operator, const FExpressionToken& L, const FExpressionToken& R) const
-{
-	FOperatorFunctionID ID = { Operator.Node.GetTypeId(), L.Node.GetTypeId(), R.Node.GetTypeId() };
-	if (const auto* Func = BinaryOps.Find(ID))
-	{
-		return (*Func)(L.Node, R.Node);
-	}
-
-	FFormatOrderedArguments Args;
-	Args.Add(FText::FromString(Operator.Context.GetString()));
-	Args.Add(FText::FromString(L.Context.GetString()));
-	Args.Add(FText::FromString(R.Context.GetString()));
-	return MakeError(FText::Format(LOCTEXT("UnaryExecutionError", "Binary operator {0} cannot operate on {1} and {2}"), Args));
-}
-
-FExpressionResult FOperatorJumpTable::ExecPreUnary(const FExpressionToken& Operator, const FExpressionToken& R) const
-{
-	FOperatorFunctionID ID = { Operator.Node.GetTypeId(), FGuid(), R.Node.GetTypeId() };
-	if (const auto* Func = PreUnaryOps.Find(ID))
-	{
-		return (*Func)(R.Node);
-	}
-
-	FFormatOrderedArguments Args;
-	Args.Add(FText::FromString(Operator.Context.GetString()));
-	Args.Add(FText::FromString(R.Context.GetString()));
-	return MakeError(FText::Format(LOCTEXT("UnaryExecutionError", "Pre-unary operator {0} cannot operate on {1}"), Args));
-}
-
-FExpressionResult FOperatorJumpTable::ExecPostUnary(const FExpressionToken& Operator, const FExpressionToken& L) const
-{
-	FOperatorFunctionID ID = { Operator.Node.GetTypeId(), L.Node.GetTypeId(), FGuid() };
-	if (const auto* Func = PostUnaryOps.Find(ID))
-	{
-		return (*Func)(L.Node);
-	}
-
-	FFormatOrderedArguments Args;
-	Args.Add(FText::FromString(Operator.Context.GetString()));
-	Args.Add(FText::FromString(L.Context.GetString()));
-	return MakeError(FText::Format(LOCTEXT("UnaryExecutionError", "Post-unary operator {0} cannot operate on {1}"), Args));
-}
-
 FExpressionTokenConsumer::FExpressionTokenConsumer(const TCHAR* InExpression)
 	: Stream(InExpression)
 {}
@@ -313,7 +272,7 @@ TArray<FExpressionToken> FExpressionTokenConsumer::Extract()
 void FExpressionTokenConsumer::Add(const FStringToken& SourceToken, FExpressionNode Node)
 {
 	Stream.SetReadPos(SourceToken);
-	Tokens.Add(FExpressionToken(SourceToken, Node));
+	Tokens.Add(FExpressionToken(SourceToken, MoveTemp(Node)));
 }
 
 void FTokenDefinitions::DefineToken(const TFunction<FExpressionDefinition>& Definition)
@@ -381,6 +340,58 @@ TOptional<FExpressionError> FTokenDefinitions::ConsumeTokens(FExpressionTokenCon
 	return TOptional<FExpressionError>();
 }
 
+FExpressionNode::~FExpressionNode()
+{
+	if (auto* Data = GetData())
+	{
+		Data->~IExpressionNodeStorage();
+	}
+}
+
+FExpressionNode::FExpressionNode(FExpressionNode&& In)
+{
+	*this = MoveTemp(In);
+}
+
+FExpressionNode& FExpressionNode::operator=(FExpressionNode&& In)
+{
+	if (TypeId == In.TypeId && TypeId.IsValid())
+	{
+		// If we have the same types, we can move-assign properly
+		In.GetData()->MoveAssign(InlineBytes);
+	}
+	else
+	{
+		// Otherwise we have to destroy what we have, and reseat the RHS
+		if (auto* ThisData = GetData())
+		{
+			ThisData->~IExpressionNodeStorage();
+		}
+
+		TypeId = In.TypeId;
+		if (auto* SrcData = In.GetData())
+		{
+			SrcData->Reseat(InlineBytes);
+
+			// Empty the RHS
+			In.TypeId = FGuid();
+			SrcData->~IExpressionNodeStorage();
+		}
+	}
+
+	return *this;
+}
+
+const FGuid& FExpressionNode::GetTypeId() const
+{
+	return TypeId;
+}
+
+Impl::IExpressionNodeStorage* FExpressionNode::GetData()
+{
+	return TypeId.IsValid() ? reinterpret_cast<Impl::IExpressionNodeStorage*>(InlineBytes) : nullptr;
+}
+
 const FGuid* FExpressionGrammar::GetGrouping(const FGuid& TypeId) const
 {
 	return Groupings.Find(TypeId);
@@ -403,14 +414,14 @@ const int* FExpressionGrammar::GetBinaryOperatorPrecedence(const FGuid& InTypeId
 
 struct FExpressionCompiler
 {
-	FExpressionCompiler(const FExpressionGrammar& InGrammar, const TArray<FExpressionToken>& InTokens)
+	FExpressionCompiler(const FExpressionGrammar& InGrammar, TArray<FExpressionToken>& InTokens)
 		: Grammar(InGrammar), Tokens(InTokens)
 	{
 		CurrentTokenIndex = 0;
 		Commands.Reserve(Tokens.Num());
 	}
 
-	TValueOrError<TArray<FExpressionToken>, FExpressionError> Compile()
+	TValueOrError<TArray<FCompiledToken>, FExpressionError> Compile()
 	{
 		auto Error = CompileGroup(nullptr, nullptr);
 		if (Error.IsSet())
@@ -420,6 +431,21 @@ struct FExpressionCompiler
 		return MakeValue(MoveTemp(Commands));
 	}
 
+	struct FWrappedOperator : FNoncopyable
+	{
+		FWrappedOperator(FCompiledToken InToken, int32 InPrecedence = 0)
+			: Token(MoveTemp(InToken)), Precedence(InPrecedence)
+		{}
+
+		FWrappedOperator(FWrappedOperator&& In) : Token(MoveTemp(In.Token)), Precedence(In.Precedence) {}
+		FWrappedOperator& operator=(FWrappedOperator&& In) { Token = MoveTemp(In.Token); Precedence = In.Precedence; return *this; }
+
+		FCompiledToken Steal() { return MoveTemp(Token); }
+
+		FCompiledToken Token;
+		int32 Precedence;
+	};
+
 	TOptional<FExpressionError> CompileGroup(const FExpressionToken* GroupStart, const FGuid* StopAt)
 	{
 		enum class EState { PreUnary, PostUnary, Binary };
@@ -441,18 +467,8 @@ struct FExpressionCompiler
 				// Ignore this token
 				CurrentTokenIndex++;
 
-				const auto MarkerIndex = Commands.Num();
-
-				FGroupMarker Marker(Token);
-
-				// Add the group marker
-				Commands.Add(FExpressionToken(Token.Context, Marker));
-
 				// Start of group - recurse
 				auto Error = CompileGroup(&Token, GroupingEnd);
-				
-				Marker.NumTokens = Commands.Num() - MarkerIndex;
-				Commands[MarkerIndex].Node = Marker;
 
 				if (Error.IsSet())
 				{
@@ -472,7 +488,7 @@ struct FExpressionCompiler
 				if (Grammar.HasPreUnaryOperator(TypeId))
 				{
 					// Make this a unary op
-					OperatorStack.Add(FWrappedOperator(FWrappedOperator::PreUnary, Token));
+					OperatorStack.Emplace(FCompiledToken(FCompiledToken::PreUnaryOperator, MoveTemp(Token)));
 				}
 				else if (Grammar.GetBinaryOperatorPrecedence(TypeId))
 				{
@@ -486,17 +502,16 @@ struct FExpressionCompiler
 					// Pop off any pending unary operators
 					while (OperatorStack.Num() > 0 && OperatorStack.Last().Precedence <= 0)
 					{
-						auto Operator = OperatorStack.Pop(false);
-						Commands.Add(FExpressionToken(Operator.WrappedToken.Context, Operator));
+						Commands.Add(OperatorStack.Pop(false).Steal());
 					}
 
 					// Make this a post-unary op
-					OperatorStack.Add(FWrappedOperator(FWrappedOperator::PostUnary, Token));
+					OperatorStack.Emplace(FCompiledToken(FCompiledToken::PostUnaryOperator, MoveTemp(Token)));
 				}
 				else
 				{
 					// Not an operator, so treat it as an ordinary token
-					Commands.Add(Token);
+					Commands.Add(FCompiledToken(FCompiledToken::Operand, MoveTemp(Token)));
 					State = EState::PostUnary;
 				}
 			}
@@ -507,12 +522,11 @@ struct FExpressionCompiler
 					// Pop off any pending unary operators
 					while (OperatorStack.Num() > 0 && OperatorStack.Last().Precedence <= 0)
 					{
-						auto Operator = OperatorStack.Pop(false);
-						Commands.Add(FExpressionToken(Operator.WrappedToken.Context, Operator));
+						Commands.Add(OperatorStack.Pop(false).Steal());
 					}
 
 					// Make this a post-unary op
-					OperatorStack.Add(FWrappedOperator(FWrappedOperator::PostUnary, Token));
+					OperatorStack.Emplace(FCompiledToken(FCompiledToken::PostUnaryOperator, MoveTemp(Token)));
 				}
 				else
 				{
@@ -522,12 +536,11 @@ struct FExpressionCompiler
 						// Pop off anything of higher precedence than this one onto the command stack
 						while (OperatorStack.Num() > 0 && OperatorStack.Last().Precedence < *Prec)
 						{
-							auto Operator = OperatorStack.Pop(false);
-							Commands.Add(FExpressionToken(Operator.WrappedToken.Context, Operator));
+							Commands.Add(OperatorStack.Pop(false).Steal());
 						}
 
 						// Add the operator itself to the op stack
-						OperatorStack.Add(FWrappedOperator(FWrappedOperator::Binary, Token, *Prec));
+						OperatorStack.Emplace(FCompiledToken(FCompiledToken::BinaryOperator, MoveTemp(Token)), *Prec);
 
 						// Check for a unary op again
 						State = EState::PreUnary;
@@ -536,7 +549,7 @@ struct FExpressionCompiler
 					{
 						// Just add the token. It's possible that this is a syntax error (there's no binary operator specified between two tokens),
 						// But we don't have enough information at this point to say whether or not it is an error
-						Commands.Add(Token);
+						Commands.Add(FCompiledToken(FCompiledToken::Operand, MoveTemp(Token)));
 						State = EState::PreUnary;
 					}
 				}
@@ -555,8 +568,7 @@ struct FExpressionCompiler
 		// Pop everything off the operator stack, onto the command stack
 		while (OperatorStack.Num() > 0)
 		{
-			auto Operator = OperatorStack.Pop(false);
-			Commands.Add(FExpressionToken(Operator.WrappedToken.Context, Operator));
+			Commands.Add(OperatorStack.Pop(false).Token);
 		}
 
 		return TOptional<FExpressionError>();
@@ -568,17 +580,17 @@ private:
 	int32 CurrentTokenIndex;
 
 	/** Working structures */
-	TArray<FExpressionToken> Commands;
+	TArray<FCompiledToken> Commands;
 
 private:
 	/** Const data provided by the parser */
-	const FExpressionGrammar& 			Grammar;
-	const TArray<FExpressionToken>& Tokens;
+	const FExpressionGrammar& Grammar;
+	TArray<FExpressionToken>& Tokens;
 };
 
 namespace ExpressionParser
 {
-	TValueOrError<TArray<FExpressionToken>, FExpressionError> Lex(const TCHAR* InExpression, const FTokenDefinitions& TokenDefinitions)
+	LexResultType Lex(const TCHAR* InExpression, const FTokenDefinitions& TokenDefinitions)
 	{
 		FExpressionTokenConsumer TokenConsumer(InExpression);
 		
@@ -593,116 +605,112 @@ namespace ExpressionParser
 		}
 	}
 
-	TValueOrError<TArray<FExpressionToken>, FExpressionError> Compile(const TCHAR* InExpression, const FTokenDefinitions& TokenDefinitions, const FExpressionGrammar& InGrammar)
+	CompileResultType Compile(const TCHAR* InExpression, const FTokenDefinitions& InTokenDefinitions, const FExpressionGrammar& InGrammar)
 	{
-		TValueOrError<TArray<FExpressionToken>, FExpressionError> Result = Lex(InExpression, TokenDefinitions);
+		TValueOrError<TArray<FExpressionToken>, FExpressionError> Result = Lex(InExpression, InTokenDefinitions);
 
 		if (!Result.IsValid())
 		{
 			return MakeError(Result.GetError());
 		}
 
-		return Compile(Result.GetValue(), InGrammar);
+		return Compile(MoveTemp(Result.GetValue()), InGrammar);
 	}
 
-	ResultType Compile(const TArray<FExpressionToken>& InTokens, const FExpressionGrammar& InGrammar)
+	CompileResultType Compile(TArray<FExpressionToken> InTokens, const FExpressionGrammar& InGrammar)
 	{
 		return FExpressionCompiler(InGrammar, InTokens).Compile();
 	}
 
-	FExpressionResult Evaluate(const TCHAR* InExpression, const FTokenDefinitions& TokenDefinitions, const FExpressionGrammar& InGrammar, const FOperatorJumpTable& InJumpTable)
+	FExpressionResult Evaluate(const TCHAR* InExpression, const FTokenDefinitions& InTokenDefinitions, const FExpressionGrammar& InGrammar, const IOperatorEvaluationEnvironment& InEnvironment)
 	{
-		TValueOrError<TArray<FExpressionToken>, FExpressionError> CompilationResult = Compile(InExpression, TokenDefinitions, InGrammar);
+		TValueOrError<TArray<FCompiledToken>, FExpressionError> CompilationResult = Compile(InExpression, InTokenDefinitions, InGrammar);
 
 		if (!CompilationResult.IsValid())
 		{
 			return MakeError(CompilationResult.GetError());
 		}
 
-		auto& Tokens = CompilationResult.GetValue();
+		return Evaluate(MoveTemp(CompilationResult.GetValue()), InEnvironment);
+	}
 
-		TArray<FExpressionToken> Stack;
-		TArray<FExpressionToken> Popped;
+	FExpressionResult Evaluate(TArray<FCompiledToken> CompiledTokens, const IOperatorEvaluationEnvironment& InEnvironment)
+	{
+		TArray<FExpressionToken> OperandStack;
 
-		for (int32 Index = 0; Index < Tokens.Num(); ++Index)
+		for (int32 Index = 0; Index < CompiledTokens.Num(); ++Index)
 		{
-			auto& Token = Tokens[Index];
+			auto& Token = CompiledTokens[Index];
 
-			// We can ignore group markers for evaluation
-			if (Token.Node.Cast<FGroupMarker>())
+			switch(Token.Type)
 			{
+			case FCompiledToken::Benign:
 				continue;
-			}
 
-			if (const auto* Op = Token.Node.Cast<FWrappedOperator>())
-			{
-				switch(Op->Type)
+			case FCompiledToken::Operand:
+				OperandStack.Push(MoveTemp(Token));
+				continue;
+
+			case FCompiledToken::BinaryOperator:
+				if (OperandStack.Num() >= 2)
 				{
-				case FWrappedOperator::Binary:
-					if (Stack.Num() >= 2)
-					{
-						// Binary
-						auto R = Stack.Pop();
-						auto L = Stack.Pop();
+					// Binary
+					auto R = OperandStack.Pop();
+					auto L = OperandStack.Pop();
 
-						auto OpResult = InJumpTable.ExecBinary(Op->WrappedToken, L, R);
-						if (OpResult.IsValid())
-						{
-							// Inherit the LHS context
-							Stack.Push(FExpressionToken(L.Context, OpResult.GetValue()));
-						}
-						else
-						{
-							return MakeError(OpResult.GetError());
-						}
+					auto OpResult = InEnvironment.ExecBinary(Token, L, R);
+					if (OpResult.IsValid())
+					{
+						// Inherit the LHS context
+						OperandStack.Push(FExpressionToken(L.Context, MoveTemp(OpResult.GetValue())));
 					}
 					else
 					{
-						FFormatOrderedArguments Args;
-						Args.Add(FText::FromString(Op->WrappedToken.Context.GetString()));
-						return MakeError(FText::Format(LOCTEXT("SyntaxError_NoUnaryOperand", "Not enough operands for binary operator {0}"), Args));
+						return MakeError(OpResult.GetError());
 					}
-					break;
-				
-				case FWrappedOperator::PostUnary:
-				case FWrappedOperator::PreUnary:
+				}
+				else
+				{
+					FFormatOrderedArguments Args;
+					Args.Add(FText::FromString(Token.Context.GetString()));
+					return MakeError(FText::Format(LOCTEXT("SyntaxError_NoUnaryOperand", "Not enough operands for binary operator {0}"), Args));
+				}
+				break;
+			
+			case FCompiledToken::PostUnaryOperator:
+			case FCompiledToken::PreUnaryOperator:
 
-					if (Stack.Num() >= 1)
-					{
-						auto Operand = Stack.Pop();
+				if (OperandStack.Num() >= 1)
+				{
+					auto Operand = OperandStack.Pop();
 
-						FExpressionResult OpResult = (Op->Type == FWrappedOperator::PreUnary) ?
-							InJumpTable.ExecPreUnary(Op->WrappedToken, Operand) :
-							InJumpTable.ExecPostUnary(Op->WrappedToken, Operand);
+					FExpressionResult OpResult = (Token.Type == FCompiledToken::PreUnaryOperator) ?
+						InEnvironment.ExecPreUnary(Token, Operand) :
+						InEnvironment.ExecPostUnary(Token, Operand);
 
-						if (OpResult.IsValid())
-						{
-							// Inherit the LHS context
-							Stack.Push(FExpressionToken(Op->WrappedToken.Context, OpResult.GetValue()));
-						}
-						else
-						{
-							return MakeError(OpResult.GetError());
-						}			
+					if (OpResult.IsValid())
+					{
+						// Inherit the LHS context
+						OperandStack.Push(FExpressionToken(Token.Context, MoveTemp(OpResult.GetValue())));
 					}
 					else
 					{
-						FFormatOrderedArguments Args;
-						Args.Add(FText::FromString(Op->WrappedToken.Context.GetString()));
-						return MakeError(FText::Format(LOCTEXT("SyntaxError_NoUnaryOperand", "No operand for unary operator {0}"), Args));
-					}
-					break;
+						return MakeError(OpResult.GetError());
+					}			
 				}
-			}
-			else
-			{
-				Stack.Push(Token);
+				else
+				{
+					FFormatOrderedArguments Args;
+					Args.Add(FText::FromString(Token.Context.GetString()));
+					return MakeError(FText::Format(LOCTEXT("SyntaxError_NoUnaryOperand", "No operand for unary operator {0}"), Args));
+				}
+				break;
 			}
 		}
 
-		if (Stack.Num() == 1)
+		if (OperandStack.Num() == 1)
 		{
-			return MakeValue(Stack[0].Node);
+			return MakeValue(MoveTemp(OperandStack[0].Node));
 		}
 		else
 		{
@@ -711,4 +719,152 @@ namespace ExpressionParser
 	}
 }
 
+#if !(UE_BUILD_SHIPPING || UE_BUILD_TEST)
+
+namespace Tests
+{
+PRAGMA_DISABLE_OPTIMIZATION
+	struct FOperator {};
+
+	struct FMoveableType
+	{
+		static int32* LeakCount;
+
+		FMoveableType(int32 InId)
+			: Id(InId), bOwnsLeak(true)
+		{
+			++*LeakCount;
+		}
+
+		FMoveableType(FMoveableType&& In) : Id(-1), bOwnsLeak(false) { *this = MoveTemp(In); }
+		FMoveableType& operator=(FMoveableType&& In)
+		{
+			if (bOwnsLeak)
+			{
+				bOwnsLeak = false;
+				--*LeakCount;
+			}
+
+			Id = In.Id;
+			In.Id = -1;
+			
+			bOwnsLeak = In.bOwnsLeak;
+
+			In.bOwnsLeak = false;
+			return *this;
+		}
+
+		virtual ~FMoveableType()
+		{
+			if (bOwnsLeak)
+			{
+				--*LeakCount;
+			}
+		}
+
+		int32 Id;
+		bool bOwnsLeak;
+	};
+	
+	int32* FMoveableType::LeakCount = nullptr;
+
+	template<typename T>
+	bool TestWithType(FAutomationTestBase* Test)
+	{
+		int32 NumLeaks = 0;
+		
+		// Test that move-assigning the expression node correctly assigns the data, and calls the destructors successfully
+		{
+			TGuardValue<int32*> LeakCounter(T::LeakCount, &NumLeaks);
+			
+			FExpressionNode Original(T(1));
+			FExpressionNode New = MoveTemp(Original);
+			
+			int32 ResultingId = New.Cast<T>()->Id;
+			if (ResultingId != 1)
+			{
+				Test->AddError(FString::Printf(TEXT("Expression node move operator did not operate correctly. Expected moved-to state to be 1, it's actually %d."), ResultingId));
+				return false;
+			}
+
+			// Try assigning it over the top again
+			Original = FExpressionNode(T(1));
+			New = MoveTemp(Original);
+
+			ResultingId = New.Cast<T>()->Id;
+			if (ResultingId != 1)
+			{
+				Test->AddError(FString::Printf(TEXT("Expression node move operator did not operate correctly. Expected moved-to state to be 1, it's actually %d."), ResultingId));
+				return false;
+			}
+
+			// Now try running it all through a parser
+			FTokenDefinitions TokenDefs;
+			FExpressionGrammar Grammar;
+			FOperatorJumpTable JumpTable;
+
+			// Only valid tokens are a, b, and +
+			TokenDefs.DefineToken([](FExpressionTokenConsumer& Consumer){
+				auto Token = Consumer.GetStream().GenerateToken(1);
+				if (Token.IsSet())
+				{
+					switch(Consumer.GetStream().PeekChar())
+					{
+					case 'a': Consumer.Add(Token.GetValue(), T(1)); break;
+					case '+': Consumer.Add(Token.GetValue(), FOperator()); break;
+					}
+				}
+				return TOptional<FExpressionError>();
+			});
+
+			Grammar.DefinePreUnaryOperator<FOperator>();
+			Grammar.DefineBinaryOperator<FOperator>(1);
+
+			JumpTable.MapPreUnary<FOperator>([](const T& A)					{ return T(A.Id); });
+			JumpTable.MapBinary<FOperator>([](const T& A, const T& B)		{ return T(A.Id); });
+
+			ExpressionParser::Evaluate(TEXT("+a"), TokenDefs, Grammar, JumpTable);
+			ExpressionParser::Evaluate(TEXT("a+a"), TokenDefs, Grammar, JumpTable);
+			ExpressionParser::Evaluate(TEXT("+a++a"), TokenDefs, Grammar, JumpTable);
+		}
+
+		if (NumLeaks != 0)
+		{
+			Test->AddError(FString::Printf(TEXT("Expression node did not call wrapped type's destructors correctly. Potentially resulted in %d leaks."), NumLeaks));
+			return false;
+		}
+
+		return true;
+	}
+
+	IMPLEMENT_SIMPLE_AUTOMATION_TEST(FExpressionParserMoveableTypes, "System.Core.Expression Parser.Moveable Types", EAutomationTestFlags::ATF_SmokeTest)
+	bool FExpressionParserMoveableTypes::RunTest( const FString& Parameters )
+	{
+		return TestWithType<FMoveableType>(this);
+	}
+
+	struct FHugeType : FMoveableType
+	{
+		FHugeType(int32 InId) : FMoveableType(InId) {}
+		FHugeType(FHugeType&& In) : FMoveableType(MoveTemp(In)) {}
+
+		uint8 Padding[1024];
+	};
+
+	IMPLEMENT_SIMPLE_AUTOMATION_TEST(FExpressionParserAllocatedTypes, "System.Core.Expression Parser.Allocated Types", EAutomationTestFlags::ATF_SmokeTest)
+	bool FExpressionParserAllocatedTypes::RunTest( const FString& Parameters )
+	{
+		return TestWithType<FHugeType>(this);
+	}
+
+}
+
+DEFINE_EXPRESSION_NODE_TYPE(Tests::FMoveableType, 0xB7F3F127, 0xD5E74833, 0x9EAB754E, 0x6CF3AAC1)
+DEFINE_EXPRESSION_NODE_TYPE(Tests::FHugeType, 0x4A329D81, 0x102343A8, 0xAB95BF45, 0x6578EE54)
+DEFINE_EXPRESSION_NODE_TYPE(Tests::FOperator, 0xC777A5D7, 0x6895456C, 0x9854BFA0, 0xB71B5A8D)
+
+PRAGMA_ENABLE_OPTIMIZATION
+
+#endif
+
 #undef LOCTEXT_NAMESPACE
\ No newline at end of file
diff --git a/Engine/Source/Runtime/Core/Private/Misc/ExpressionParserExamples.cpp b/Engine/Source/Runtime/Core/Private/Misc/ExpressionParserExamples.cpp
--- a/Engine/Source/Runtime/Core/Private/Misc/ExpressionParserExamples.cpp
+++ b/Engine/Source/Runtime/Core/Private/Misc/ExpressionParserExamples.cpp
@@ -1,3 +1,4 @@
+// Copyright 1998-2015 Epic Games, Inc. All Rights Reserved.
 
 // Don't include this code in shipping or perf test builds
 #if !(UE_BUILD_SHIPPING || UE_BUILD_TEST)
diff --git a/Engine/Source/Runtime/Core/Public/Misc/ExpressionParser.h b/Engine/Source/Runtime/Core/Public/Misc/ExpressionParser.h
--- a/Engine/Source/Runtime/Core/Public/Misc/ExpressionParser.h
+++ b/Engine/Source/Runtime/Core/Public/Misc/ExpressionParser.h
@@ -14,17 +14,36 @@
 
 namespace ExpressionParser
 {
-	typedef TValueOrError< TArray<FExpressionToken>, FExpressionError > ResultType;
+	typedef TValueOrError< TArray<FExpressionToken>, FExpressionError > LexResultType;
+	typedef TValueOrError< TArray<FCompiledToken>, FExpressionError > CompileResultType;
 
 	/** Lex the specified string, using the specified grammar */
-	CORE_API ResultType Lex(const TCHAR* InExpression, const FTokenDefinitions& TokenDefinitions);
+	CORE_API LexResultType Lex(const TCHAR* InExpression, const FTokenDefinitions& TokenDefinitions);
 	
 	/** Compile the specified expression into an array of Reverse-Polish order nodes for evaluation, according to our grammar definition */
-	CORE_API ResultType Compile(const TCHAR* InExpression, const FTokenDefinitions& TokenDefinitions, const FExpressionGrammar& InGrammar);
+	CORE_API CompileResultType Compile(const TCHAR* InExpression, const FTokenDefinitions& TokenDefinitions, const FExpressionGrammar& InGrammar);
 
 	/** Compile the specified tokens into an array of Reverse-Polish order nodes for evaluation, according to our grammar definition */
-	CORE_API ResultType Compile(const TArray<FExpressionToken>& InTokens, const FExpressionGrammar& InGrammar);
+	CORE_API CompileResultType Compile(TArray<FExpressionToken> InTokens, const FExpressionGrammar& InGrammar);
 
-	/** Evaluate the specific expression using our grammar definition, and return the result */
-	CORE_API FExpressionResult Evaluate(const TCHAR* Expression, const FTokenDefinitions& TokenDefinitions, const FExpressionGrammar& InGrammar, const FOperatorJumpTable& InJumpTable);
+	/** Evaluate the specified expression using the specified token definitions, grammar definition, and evaluation environment */
+	CORE_API FExpressionResult Evaluate(const TCHAR* InExpression, const FTokenDefinitions& InTokenDefinitions, const FExpressionGrammar& InGrammar, const IOperatorEvaluationEnvironment& InEnvironment);
+
+	/** Evaluate the specified pre-compiled tokens using an evaluation environment */
+	CORE_API FExpressionResult Evaluate(TArray<FCompiledToken> CompiledTokens, const IOperatorEvaluationEnvironment& InEnvironment);
+
+	/** Templated versions of evaluation functions used when passing a specific jump table and context */
+	template<typename ContextType>
+	FExpressionResult Evaluate(const TCHAR* InExpression, const FTokenDefinitions& InTokenDefinitions, const FExpressionGrammar& InGrammar,	const TOperatorJumpTable<ContextType>& InJumpTable, const ContextType* InContext = nullptr)
+	{
+		TOperatorEvaluationEnvironment<ContextType> Env(InJumpTable, InContext);
+		return Evaluate(InExpression, InTokenDefinitions, InGrammar, Env);
+	}
+
+	template<typename ContextType>
+	FExpressionResult Evaluate(TArray<FCompiledToken>& CompiledTokens, const TOperatorJumpTable<ContextType>& InJumpTable, const ContextType* InContext = nullptr)
+	{
+		TOperatorEvaluationEnvironment<ContextType> Env(InJumpTable, InContext);
+		return Evaluate(CompiledTokens, Env);
+	}
 }
\ No newline at end of file
diff --git a/Engine/Source/Runtime/Core/Public/Misc/ExpressionParserTypes.h b/Engine/Source/Runtime/Core/Public/Misc/ExpressionParserTypes.h
--- a/Engine/Source/Runtime/Core/Public/Misc/ExpressionParserTypes.h
+++ b/Engine/Source/Runtime/Core/Public/Misc/ExpressionParserTypes.h
@@ -1,3 +1,4 @@
+// Copyright 1998-2015 Epic Games, Inc. All Rights Reserved.
 
 #pragma once
 
@@ -6,6 +7,11 @@
 class FExpressionNode;
 struct FExpressionError;
 
+namespace Impl
+{
+	struct IExpressionNodeStorage;
+}
+
 typedef TValueOrError<FExpressionNode, FExpressionError> FExpressionResult;
 
 /** Simple error structure used for reporting parse errors */
@@ -154,306 +160,157 @@ template<typename T> struct TGetExpressionNodeTypeId;
  * 	Can be constructed from any C++ type that has a corresponding DEFINE_EXPRESSION_NODE_TYPE.
  * 	Evaluation behaviour (unary/binary operator etc) is defined in the expression grammar, rather than the type itself.
  */
-class FExpressionNode
+class CORE_API FExpressionNode : FNoncopyable
 {
 public:
 	/** Construction from client expression data type */
 	template<typename T>
-	FExpressionNode(const T& In,
+	FExpressionNode(T In,
 		/** @todo: make this a default function template parameter when VS2012 support goes */
 		typename TEnableIf<!TPointerIsConvertibleFromTo<T, FExpressionNode>::Value>::Type* = nullptr
-		)
-		: TypeId(TGetExpressionNodeTypeId<T>::GetTypeId())
-	{
-		if (sizeof(T) <= SizeOfInlineBytes)
-		{
-			FMemory::Memcpy(InlineBytes, &In, sizeof(T));
-		}
-		else
-		{
-			BigData = MakeShareable(new T(In));
-		}
-	}
+		);
 
-	/** Copy construction/assignment */
-	FExpressionNode(const FExpressionNode& In) { *this = In; }
-	FExpressionNode& operator=(const FExpressionNode& In)
-	{
-		TypeId = In.TypeId;
-		if (In.BigData.IsValid())
-		{
-			BigData = In.BigData;
-		}
-		else
-		{
-			BigData = nullptr;
-			FMemory::Memcpy(InlineBytes, In.InlineBytes, SizeOfInlineBytes);
-		}
-		return *this;
-	}
+	~FExpressionNode();
+
+	/** Move construction/assignment */
+	FExpressionNode(FExpressionNode&& In);
+	FExpressionNode& operator=(FExpressionNode&& In);
 
 	/** Get the type identifier of this node */
-	const FGuid& GetTypeId() const { return TypeId; }
+	const FGuid& GetTypeId() const;
 
-	/** Cast this node to the specified type. Will asset if the types do not match. */
+	/** Cast this node to the specified type. Will return nullptr if the types do not match. */
 	template<typename T>
-	const T* Cast() const
-	{
-		if (TypeId == TGetExpressionNodeTypeId<T>::GetTypeId())
-		{
-			if (BigData.IsValid())
-			{
-				return reinterpret_cast<const T*>(BigData.Get());
-			}
-			else
-			{
-				return reinterpret_cast<const T*>(InlineBytes);
-			}
-		}
-		return nullptr;
-	}
+	const T* Cast() const;
 
 private:
 
-	static const int32 SizeOfInlineBytes = 128 - sizeof(FGuid) - sizeof(TSharedPtr<void>);
+	/** The maximum size of type we will allow allocation on the stack (for efficiency). Anything larger will be allocated on the heap. */
+	static const int32 MaxStackAllocationSize = 64 - sizeof(FGuid);
+
+	/** Helper accessor to the data interface. Returns null for empty containers. */
+	Impl::IExpressionNodeStorage* GetData();
 
 	/** TypeID - 16 bytes */
 	FGuid TypeId;
-	/** SharedPtr for large types - 2*sizeof(void*) */
-	TSharedPtr<void> BigData;
-	uint8 InlineBytes[SizeOfInlineBytes];
+	uint8 InlineBytes[MaxStackAllocationSize];
 };
 
 /** A specific token in a stream. Comprises an expression node, and the stream token it was created from */
-class FExpressionToken
+class FExpressionToken : FNoncopyable
 {
 public:
-	FExpressionToken(const FStringToken& InContext, const FExpressionNode& InNode)
-		: Node(InNode)
+	FExpressionToken(const FStringToken& InContext, FExpressionNode InNode)
+		: Node(MoveTemp(InNode))
 		, Context(InContext)
 	{
 	}
 
-	FExpressionToken(const FExpressionToken& In) : Node(In.Node), Context(In.Context) {}
-	FExpressionToken& operator=(const FExpressionToken& In) { Node = In.Node; Context = In.Context; return *this; }
+	FExpressionToken(FExpressionToken&& In) : Node(MoveTemp(In.Node)), Context(In.Context) {}
+	FExpressionToken& operator=(FExpressionToken&& In) { Node = MoveTemp(In.Node); Context = In.Context; return *this; }
 
 	FExpressionNode Node;
 	FStringToken Context;
 };
 
-/** Built-in wrapper tokens used for compilation */
-struct FGroupMarker
+/** A compiled token, holding the token itself, and any compiler information required to evaluate it */
+struct FCompiledToken : FExpressionToken
 {
-	FGroupMarker(const FExpressionToken& InWrappedToken) : WrappedToken(InWrappedToken), NumTokens(0) {}
-
-	template<typename T>	bool 		IsA() const 	{ return WrappedToken.Node.GetTypeId() == TGetExpressionNodeTypeId<T>::GetTypeId(); }
-	template<typename T>	const T* 	Cast() const	{ return WrappedToken.Node.Cast<T>(); }
+	// Todo: add callable types here?
+	enum EType { Operand, PreUnaryOperator, PostUnaryOperator, BinaryOperator, Benign };
 
-	FExpressionToken WrappedToken;
-	int32 NumTokens;
-};
-DEFINE_EXPRESSION_NODE_TYPE(FGroupMarker, 0xEBF6DBA6, 0xE81B4684, 0x938EF46A, 0xD2FC0052)
-
-struct FWrappedOperator
-{
-	enum EType { PreUnary, PostUnary, Binary };
-	FWrappedOperator(EType InType, const FExpressionToken& InWrappedToken, int32 InPrec = 0) : Type(InType), WrappedToken(InWrappedToken), Precedence(InPrec) {}
+	FCompiledToken(EType InType, FExpressionToken InToken)
+		: FExpressionToken(MoveTemp(InToken)), Type(InType)
+	{}
 
-	template<typename T>	bool 		IsA() const 	{ return WrappedToken.Node.GetTypeId() == TGetExpressionNodeTypeId<T>::GetTypeId(); }
-	template<typename T>	const T* 	Cast() const	{ return WrappedToken.Node.Cast<T>(); }
+	FCompiledToken(FCompiledToken&& In) : FExpressionToken(MoveTemp(In)), Type(In.Type) {}
+	FCompiledToken& operator=(FCompiledToken&& In) { FExpressionToken::operator=(MoveTemp(In)); Type = In.Type; return *this; }
 
 	EType Type;
-	FExpressionToken WrappedToken;
-	int32 Precedence;
 };
-DEFINE_EXPRESSION_NODE_TYPE(FWrappedOperator, 0x449B352B, 0x17B14A76, 0xBD264E7B, 0x9669D8F4)
-
-namespace impl
-{
-	template <typename> struct TParams;
-	template <typename> struct TParamsImpl;
-
-	template <typename Ret_, typename T, typename Arg1_>
-	struct TParamsImpl<Ret_ (T::*)(Arg1_)> { typedef Ret_ Ret; typedef Arg1_ Arg1; };
-	template <typename Ret_, typename T, typename Arg1_>
-	struct TParamsImpl<Ret_ (T::*)(Arg1_) const> { typedef Ret_ Ret; typedef Arg1_ Arg1; };
-
-	template <typename Ret_, typename T, typename Arg1_, typename Arg2_>
-	struct TParamsImpl<Ret_ (T::*)(Arg1_, Arg2_)> { typedef Ret_ Ret; typedef Arg1_ Arg1; typedef Arg2_ Arg2; };
-	template <typename Ret_, typename T, typename Arg1_, typename Arg2_>
-	struct TParamsImpl<Ret_ (T::*)(Arg1_, Arg2_) const> { typedef Ret_ Ret; typedef Arg1_ Arg1; typedef Arg2_ Arg2; };
-
-	/** @todo: decltype(&T::operator()) can go directly in the specialization below if it weren't for VS2012 support */
-	template<typename T>
-	struct TGetOperatorCallPtr { typedef decltype(&T::operator()) Type; };
-
-	template <typename T>
-	struct TParams : TParamsImpl<typename TGetOperatorCallPtr<T>::Type> {};
-}
 
 /** Jump table specifying how to execute an operator with different types */
-struct FOperatorJumpTable
+template<typename ContextType=void>
+struct TOperatorJumpTable
 {
 	/** Execute the specified token as a unary operator, if such an overload exists */
-	FExpressionResult ExecPreUnary(const FExpressionToken& Operator, const FExpressionToken& R) const;
+	FExpressionResult ExecPreUnary(const FExpressionToken& Operator, const FExpressionToken& R, const ContextType* Context) const;
 	/** Execute the specified token as a unary operator, if such an overload exists */
-	FExpressionResult ExecPostUnary(const FExpressionToken& Operator, const FExpressionToken& L) const;
+	FExpressionResult ExecPostUnary(const FExpressionToken& Operator, const FExpressionToken& L, const ContextType* Context) const;
 	/** Execute the specified token as a binary operator, if such an overload exists */
-	FExpressionResult ExecBinary(const FExpressionToken& Operator, const FExpressionToken& L, const FExpressionToken& R) const;
+	FExpressionResult ExecBinary(const FExpressionToken& Operator, const FExpressionToken& L, const FExpressionToken& R, const ContextType* Context) const;
 
 	/**
-	 * Map an expression node to a unary operator with the specified implementation.
-	 * This overload accepts callable types that return an FExpressionResult (and thus support returning an error)
-	 * Example usage that binds a '!' token to a function that attempts to do a boolean 'not'
-	 *		JumpTable.MapUnary<FExclamation>([](bool A) -> FExpressionResult { return MakeValue(!A); });
+	 * Map an expression node to a pre-unary operator with the specified implementation.
+	 *
+	 * The callable type must match the declaration Ret(Operand[, Context]), where:
+	 *	 	Ret 	= Any DEFINE_EXPRESSION_NODE_TYPE type, OR FExpressionResult
+	 *	 	Operand = Any DEFINE_EXPRESSION_NODE_TYPE type
+	 *	 	Context = (optional) const ptr to user-supplied arbitrary context
+	 *
+	 * Examples that binds a '!' token to a function that attempts to do a boolean 'not':
+	 *		JumpTable.MapPreUnary<FExclamation>([](bool A){ return !A; });
+	 *		JumpTable.MapPreUnary<FExclamation>([](bool A, FMyContext* Ctxt){ if (Ctxt->IsBooleanNotOpEnabled()) { return !A; } else { return A; } });
+	 *		JumpTable.MapPreUnary<FExclamation>([](bool A, const FMyContext* Ctxt) -> FExpressionResult {
+
+	 *			if (Ctxt->IsBooleanNotOpEnabled())
+	 *			{
+	 *				return MakeValue(!A);
+	 *			}
+	 *			return MakeError(FExpressionError(LOCTEXT("NotNotEnabled", "Boolean not is not enabled.")));
+	 *		});
 	 */
 	template<typename OperatorType, typename FuncType>
-	typename TEnableIf<TIsSame<FExpressionResult, typename impl::TParams<FuncType>::Ret>::Value>::Type
-		MapPreUnary(FuncType InFunc)
-	{
-		typedef typename TRemoveConst<typename TRemoveReference<typename impl::TParams<FuncType>::Arg1>::Type>::Type OperandType;
-
-		FOperatorFunctionID ID = {
-			TGetExpressionNodeTypeId<OperatorType>::GetTypeId(),
-			FGuid(),
-			TGetExpressionNodeTypeId<OperandType>::GetTypeId()
-		};
-
-		PreUnaryOps.Add(ID, [=](const FExpressionNode& InOperand) {
-			return InFunc(*InOperand.Cast<OperandType>());
-		});
-	}
+	void MapPreUnary(FuncType InFunc);
 
 	/**
-	 * Map an expression node to a unary operator with the specified implementation.
-	 * This overload accepts callable types that return anything other than an FExpressionResult. Such functions cannot return errors.
-	 * The returned type is passed into an FExpressionNode, and as such must be a valid expression node type (see DEFINE_EXPRESSION_NODE_TYPE).
-	 * Example usage that binds a 'subtract' token to a function that negates the integer:
-	 *		JumpTable.MapUnary<FSubtractToken>([](int32 A){ return -A; });
+	 * Map an expression node to a post-unary operator with the specified implementation.
+	 * The same function signature rules apply here as with MapPreUnary.
 	 */
 	template<typename OperatorType, typename FuncType>
-	typename TEnableIf<!TIsSame<FExpressionResult, typename impl::TParams<FuncType>::Ret>::Value>::Type
-		MapPreUnary(FuncType InFunc)
-	{
-		typedef typename TRemoveConst<typename TRemoveReference<typename impl::TParams<FuncType>::Arg1>::Type>::Type OperandType;
-
-		FOperatorFunctionID ID = {
-			TGetExpressionNodeTypeId<OperatorType>::GetTypeId(),
-			FGuid(),
-			TGetExpressionNodeTypeId<OperandType>::GetTypeId()
-		};
-
-		// Explicit return type is important here, to ensure that the proxy returned from MakeValue does not outlive the value it's proxying
-		PreUnaryOps.Add(ID, [=](const FExpressionNode& InOperand) -> FExpressionResult {
-			return MakeValue(InFunc(*InOperand.Cast<OperandType>()));
-		});
-	}
+	void MapPostUnary(FuncType InFunc);
 
 	/**
-	 * Map an expression node to a unary operator with the specified implementation.
-	 * This overload accepts callable types that return an FExpressionResult (and thus support returning an error)
-	 * Example usage that binds a '!' token to a function that attempts to do a boolean 'not'
-	 *		JumpTable.MapUnary<FExclamation>([](bool A) -> FExpressionResult { return MakeValue(!A); });
-	 */
-	template<typename OperatorType, typename FuncType>
-	typename TEnableIf<TIsSame<FExpressionResult, typename impl::TParams<FuncType>::Ret>::Value>::Type
-		MapPostUnary(FuncType InFunc)
-	{
-		typedef typename TRemoveConst<typename TRemoveReference<typename impl::TParams<FuncType>::Arg1>::Type>::Type OperandType;
-
-		FOperatorFunctionID ID = {
-			TGetExpressionNodeTypeId<OperatorType>::GetTypeId(),
-			TGetExpressionNodeTypeId<OperandType>::GetTypeId(),
-			FGuid()
-		};
-
-		PostUnaryOps.Add(ID, [=](const FExpressionNode& InOperand) {
-			return InFunc(*InOperand.Cast<OperandType>());
-		});
-	}
-
-	/**
-	 * Map an expression node to a unary operator with the specified implementation.
-	 * This overload accepts callable types that return anything other than an FExpressionResult. Such functions cannot return errors.
-	 * The returned type is passed into an FExpressionNode, and as such must be a valid expression node type (see DEFINE_EXPRESSION_NODE_TYPE).
-	 * Example usage that binds a 'subtract' token to a function that negates the integer:
-	 *		JumpTable.MapUnary<FSubtractToken>([](int32 A){ return -A; });
+	 * Map an expression node to a binary operator with the specified implementation.
+	 *
+	 * The callable type must match the declaration Ret(OperandL, OperandR, [, Context]), where:
+	 *	 	Ret 		= Any DEFINE_EXPRESSION_NODE_TYPE type, OR FExpressionResult
+	 *	 	OperandL	= Any DEFINE_EXPRESSION_NODE_TYPE type
+	 *	 	OperandR	= Any DEFINE_EXPRESSION_NODE_TYPE type
+	 *	 	Context 	= (optional) const ptr to user-supplied arbitrary context
+	 *
+	 * Examples that binds a '/' token to a function that attempts to do a division:
+	 *		JumpTable.MapUnary<FForwardSlash>([](double A, double B){ return A / B; }); // Runtime exception on div/0 
+	 *		JumpTable.MapUnary<FForwardSlash>([](double A, double B, FMyContext* Ctxt){
+	 *			if (!Ctxt->IsMathEnabled())
+	 *			{
+	 *				return A;
+	 *			}
+	 *			return A / B; // Runtime exception on div/0 
+	 *		});
+	 *		JumpTable.MapUnary<FForwardSlash>([](double A, double B, const FMyContext* Ctxt) -> FExpressionResult {
+	 *			if (!Ctxt->IsMathEnabled())
+	 *			{
+	 *				return MakeError(FExpressionError(LOCTEXT("MathNotEnabled", "Math is not enabled.")));
+	 *			}
+	 *			else if (B == 0)
+	 *			{
+	 *				return MakeError(FExpressionError(LOCTEXT("DivisionByZero", "Division by zero.")));	
+	 *			}
+	 *
+	 *			return MakeValue(!A);
+	 *		});
 	 */
 	template<typename OperatorType, typename FuncType>
-	typename TEnableIf<!TIsSame<FExpressionResult, typename impl::TParams<FuncType>::Ret>::Value>::Type
-		MapPostUnary(FuncType InFunc)
-	{
-		typedef typename TRemoveConst<typename TRemoveReference<typename impl::TParams<FuncType>::Arg1>::Type>::Type OperandType;
-
-		FOperatorFunctionID ID = {
-			TGetExpressionNodeTypeId<OperatorType>::GetTypeId(),
-			TGetExpressionNodeTypeId<OperandType>::GetTypeId(),
-			FGuid()
-		};
-
-		// Explicit return type is important here, to ensure that the proxy returned from MakeValue does not outlive the value it's proxying
-		PostUnaryOps.Add(ID, [=](const FExpressionNode& InOperand) -> FExpressionResult {
-			return MakeValue(InFunc(*InOperand.Cast<OperandType>()));
-		});
-	}
+	void MapBinary(FuncType InFunc);
 
-	/**
-	 * Map an expression node to a unary operator with the specified implementation.
-	 * This overload accepts callable types that return an FExpressionResult (and thus support returning an error)
-	 * Example usage that binds a 'slash' token to a function that subtracts two integers:
-	 *		JumpTable.MapBinary<FSlashToken>([](double A, double B){
-	 			if (B == 0) {
-	 				return MakeError(LOCTEXT("DivisionByZero", "Division by zero"));
-	 			} else {
-	 				return MakeValue(A / B);
-	 			}
-	 		});
-	 */
-	template<typename OperatorType, typename FuncType>
-	typename TEnableIf<TIsSame<FExpressionResult, typename impl::TParams<FuncType>::Ret>::Value>::Type
-		MapBinary(FuncType InFunc)
-	{
-		typedef typename TRemoveConst<typename TRemoveReference<typename impl::TParams<FuncType>::Arg1>::Type>::Type LeftOperandType;
-		typedef typename TRemoveConst<typename TRemoveReference<typename impl::TParams<FuncType>::Arg2>::Type>::Type RightOperandType;
-
-		FOperatorFunctionID ID = {
-			TGetExpressionNodeTypeId<OperatorType>::GetTypeId(),
-			TGetExpressionNodeTypeId<LeftOperandType>::GetTypeId(),
-			TGetExpressionNodeTypeId<RightOperandType>::GetTypeId()
-		};
-
-		BinaryOps.Add(ID, [=](const FExpressionNode& InLeftOperand, const FExpressionNode& InRightOperand){
-			return InFunc(*InLeftOperand.Cast<LeftOperandType>(), *InRightOperand.Cast<RightOperandType>());
-		});
-	}
+public:
 
-	/**
-	 * Map an expression node to a binary operator with the specified implementation.
-	 * This overload accepts callable types that return anything other than an FExpressionResult. Such functions cannot return errors.
-	 * The returned type is passed into an FExpressionNode, and as such must be a valid expression node type (see DEFINE_EXPRESSION_NODE_TYPE).
-	 * Example usage that binds a 'subtract' token to a function that subtracts two integers:
-	 *		JumpTable.MapBinary<FSubtractToken>([](int32 A, int32 B){ return A - B; });
-	 */
-	template<typename OperatorType, typename FuncType>
-	typename TEnableIf<!TIsSame<FExpressionResult, typename impl::TParams<FuncType>::Ret>::Value>::Type
-		MapBinary(FuncType InFunc)
-	{
-		typedef typename TRemoveConst<typename TRemoveReference<typename impl::TParams<FuncType>::Arg1>::Type>::Type LeftOperandType;
-		typedef typename TRemoveConst<typename TRemoveReference<typename impl::TParams<FuncType>::Arg2>::Type>::Type RightOperandType;
-
-		FOperatorFunctionID ID = {
-			TGetExpressionNodeTypeId<OperatorType>::GetTypeId(),
-			TGetExpressionNodeTypeId<LeftOperandType>::GetTypeId(),
-			TGetExpressionNodeTypeId<RightOperandType>::GetTypeId()
-		};
-
-		// Explicit return type is important here, to ensure that the proxy returned from MakeValue does not outlive the value it's proxying
-		BinaryOps.Add(ID, [=](const FExpressionNode& InLeftOperand, const FExpressionNode& InRightOperand) -> FExpressionResult {
-			return MakeValue(InFunc(*InLeftOperand.Cast<LeftOperandType>(), *InRightOperand.Cast<RightOperandType>()));
-		});
-	}
+	typedef TFunction<FExpressionResult(const FExpressionNode&, const ContextType* Context)> FUnaryFunction;
+	typedef TFunction<FExpressionResult(const FExpressionNode&, const FExpressionNode&, const ContextType* Context)> FBinaryFunction;
 
 private:
+
 	/** Struct used to identify a function for a specific operator overload */
 	struct FOperatorFunctionID
 	{
@@ -476,11 +333,52 @@ private:
 	};
 
 	/** Maps of unary/binary operators */
-	TMap<FOperatorFunctionID, TFunction<FExpressionResult(const FExpressionNode&)>> PreUnaryOps;
-	TMap<FOperatorFunctionID, TFunction<FExpressionResult(const FExpressionNode&)>> PostUnaryOps;
-	TMap<FOperatorFunctionID, TFunction<FExpressionResult(const FExpressionNode&, const FExpressionNode&)>> BinaryOps;
+	TMap<FOperatorFunctionID, FUnaryFunction> PreUnaryOps;
+	TMap<FOperatorFunctionID, FUnaryFunction> PostUnaryOps;
+	TMap<FOperatorFunctionID, FBinaryFunction> BinaryOps;
 };
 
+typedef TOperatorJumpTable<> FOperatorJumpTable;
+
+/** Structures used for managing the evaluation environment for operators in an expression. This class manages the evaluation context
+ * to avoid templating the whole evaluation code on a context type
+ */
+struct IOperatorEvaluationEnvironment
+{
+	/** Execute the specified token as a unary operator, if such an overload exists */
+	virtual FExpressionResult ExecPreUnary(const FExpressionToken& Operator, const FExpressionToken& R) const = 0;
+	/** Execute the specified token as a unary operator, if such an overload exists */
+	virtual FExpressionResult ExecPostUnary(const FExpressionToken& Operator, const FExpressionToken& L) const = 0;
+	/** Execute the specified token as a binary operator, if such an overload exists */
+	virtual FExpressionResult ExecBinary(const FExpressionToken& Operator, const FExpressionToken& L, const FExpressionToken& R) const = 0;
+};
+template<typename ContextType = void>
+struct TOperatorEvaluationEnvironment : IOperatorEvaluationEnvironment
+{
+	TOperatorEvaluationEnvironment(const TOperatorJumpTable<ContextType>& InOperators, const ContextType* InContext)
+		: Operators(InOperators), Context(InContext)
+	{}
+
+	virtual FExpressionResult ExecPreUnary(const FExpressionToken& Operator, const FExpressionToken& R) const override
+	{
+		return Operators.ExecPreUnary(Operator, R, Context);
+	}
+	virtual FExpressionResult ExecPostUnary(const FExpressionToken& Operator, const FExpressionToken& L) const
+	{
+		return Operators.ExecPostUnary(Operator, L, Context);
+	}
+	virtual FExpressionResult ExecBinary(const FExpressionToken& Operator, const FExpressionToken& L, const FExpressionToken& R) const override
+	{
+		return Operators.ExecBinary(Operator, L, R, Context);
+	}
+
+private:
+	const TOperatorJumpTable<ContextType>& Operators;
+	const ContextType* Context;
+};
+
+typedef TOperatorEvaluationEnvironment<> FOperatorEvaluationEnvironment;
+
 /** Class used to consume tokens from a string */
 class CORE_API FExpressionTokenConsumer
 {
@@ -588,4 +486,6 @@ private:
 	TSet<FGuid>			PreUnaryOperators;
 	TSet<FGuid>			PostUnaryOperators;
 	TMap<FGuid, int32>	BinaryOperators;
-};
\ No newline at end of file
+};
+
+#include "ExpressionParserTypes.inl"
\ No newline at end of file
diff --git a/Engine/Source/Runtime/Core/Public/Misc/ExpressionParserTypes.inl b/Engine/Source/Runtime/Core/Public/Misc/ExpressionParserTypes.inl
new file mode 100644
--- /dev/null
+++ b/Engine/Source/Runtime/Core/Public/Misc/ExpressionParserTypes.inl
@@ -0,0 +1,249 @@
+// Copyright 1998-2015 Epic Games, Inc. All Rights Reserved.
+
+#pragma once
+
+#define LOCTEXT_NAMESPACE "ExpressionParser"
+
+namespace Impl
+{
+	/** Interface for a wrapper utility for any moveable/copyable data */
+	struct IExpressionNodeStorage
+	{
+		virtual ~IExpressionNodeStorage() {}
+		/** Move this type into another unallocated buffer (move-construct a new type from our wrapped value) */
+		virtual void Reseat(uint8* Dst) = 0;
+		/** Move this type to a buffer already allocated to the same type (uses type-defined move-assignment) */
+		virtual void MoveAssign(uint8* Dst) = 0;
+	};
+
+	/** Implementation of the wrapper utility for any moveable/copyable data, that allows us to virtually move/copy/destruct the data */
+	/** Data is stored inline in this implementation, for efficiency */
+ 	template<typename T>
+	struct FInlineDataStorage : IExpressionNodeStorage
+	{
+		/** The data itself, allocated on the stack */
+		T Value;
+
+		/** Constructor/destructor */
+		FInlineDataStorage(T InValue) : Value(MoveTemp(InValue)) {}
+		~FInlineDataStorage(){}
+		
+		const T* Access() const { return &Value; }
+
+		virtual void Reseat(uint8* Dst) override 			{ new(Dst) FInlineDataStorage(MoveTemp(Value)); }
+		virtual void MoveAssign(uint8* Dst) override 		{ reinterpret_cast<FInlineDataStorage*>(Dst)->Value = MoveTemp(Value); }
+	};
+
+	/** Data is stored on the heap in this implementation */
+	template<typename T>
+	struct FHeapDataStorage : IExpressionNodeStorage
+	{
+		/** The data itself, allocated on the heap */
+		TUniquePtr<T> Value;
+		
+		/** Constructor/destructor */
+		FHeapDataStorage(T InValue) : Value(new T(MoveTemp(InValue))) {}
+		FHeapDataStorage(FHeapDataStorage&& In) : Value(MoveTemp(In.Value)) {}
+		~FHeapDataStorage(){}
+		
+		const T* Access() const { return Value.Get(); }
+
+		virtual void Reseat(uint8* Dst) override 			{ new(Dst) FHeapDataStorage(MoveTemp(*this)); }
+		virtual void MoveAssign(uint8* Dst) override 		{ reinterpret_cast<FHeapDataStorage*>(Dst)->Value = MoveTemp(Value); }
+	};
+
+	template<typename T, int32 MaxStackAllocationSize, typename Enabled=void> struct TStorageTypeDeduction;
+	template<typename T, int32 MaxStackAllocationSize> struct TStorageTypeDeduction<T, MaxStackAllocationSize, typename TEnableIf<(sizeof(FInlineDataStorage<T>) <= MaxStackAllocationSize)>::Type>
+		{ typedef FInlineDataStorage<T> Type; };
+	template<typename T, int32 MaxStackAllocationSize> struct TStorageTypeDeduction<T, MaxStackAllocationSize, typename TEnableIf<(sizeof(FInlineDataStorage<T>) >  MaxStackAllocationSize)>::Type>
+		{ typedef FHeapDataStorage<T> 	Type; };
+
+
+	/** Machinery required for operator mapping */
+	template <typename> struct TCallableInfo;
+	template <typename> struct TCallableInfoImpl;
+
+	template <typename Ret_, typename T, typename Arg1_>
+	struct TCallableInfoImpl<Ret_ (T::*)(Arg1_)> { typedef Ret_ Ret; typedef Arg1_ Arg1; enum { NumArgs = 1 }; };
+	template <typename Ret_, typename T, typename Arg1_>
+	struct TCallableInfoImpl<Ret_ (T::*)(Arg1_) const> { typedef Ret_ Ret; typedef Arg1_ Arg1; enum { NumArgs = 1 };  };
+
+	template <typename Ret_, typename T, typename Arg1_, typename Arg2_>
+	struct TCallableInfoImpl<Ret_ (T::*)(Arg1_, Arg2_)> { typedef Ret_ Ret; typedef Arg1_ Arg1; typedef Arg2_ Arg2; enum { NumArgs = 2 };  };
+	template <typename Ret_, typename T, typename Arg1_, typename Arg2_>
+	struct TCallableInfoImpl<Ret_ (T::*)(Arg1_, Arg2_) const> { typedef Ret_ Ret; typedef Arg1_ Arg1; typedef Arg2_ Arg2; enum { NumArgs = 2 };  };
+	
+	template <typename Ret_, typename T, typename Arg1_, typename Arg2_, typename Arg3_>
+	struct TCallableInfoImpl<Ret_ (T::*)(Arg1_, Arg2_, Arg3_)> { typedef Ret_ Ret; typedef Arg1_ Arg1; typedef Arg2_ Arg2; typedef Arg3_ Arg3; enum { NumArgs = 3 }; };
+	template <typename Ret_, typename T, typename Arg1_, typename Arg2_, typename Arg3_>
+	struct TCallableInfoImpl<Ret_ (T::*)(Arg1_, Arg2_, Arg3_) const> { typedef Ret_ Ret; typedef Arg1_ Arg1; typedef Arg2_ Arg2; typedef Arg3_ Arg3; enum { NumArgs = 3 }; };
+
+	/** @todo: decltype(&T::operator()) can go directly in the specialization below if it weren't for VS2012 support */
+	template<typename T>
+	struct TGetOperatorCallPtr { typedef decltype(&T::operator()) Type; };
+
+	template <typename T>
+	struct TCallableInfo : TCallableInfoImpl<typename TGetOperatorCallPtr<T>::Type> {};
+
+	/** Overloaded function that returns an FExpressionResult, regardless of what is passed in */
+	template<typename T>
+	static FExpressionResult ForwardReturnType(T&& Result) 					{ return MakeValue(MoveTemp(Result)); }
+	static FExpressionResult ForwardReturnType(FExpressionResult&& Result) 	{ return MoveTemp(Result); }
+
+	/** Wrapper function for supplied functions of the signature T(A) */
+	template<typename OperandType, typename ContextType, typename FuncType>
+	static typename TEnableIf<Impl::TCallableInfo<FuncType>::NumArgs == 1, typename TOperatorJumpTable<ContextType>::FUnaryFunction>::Type WrapUnaryFunction(FuncType In)
+	{
+		// Ignore the context
+		return [=](const FExpressionNode& InOperand, const ContextType* Context) {
+			return ForwardReturnType(In(*InOperand.Cast<OperandType>()));
+		};
+	}
+
+	/** Wrapper function for supplied functions of the signature T(A, const ContextType* Context) */
+	template<typename OperandType, typename ContextType, typename FuncType>
+	static typename TEnableIf<Impl::TCallableInfo<FuncType>::NumArgs == 2, typename TOperatorJumpTable<ContextType>::FUnaryFunction>::Type WrapUnaryFunction(FuncType In)
+	{
+		// Ignore the context
+		return [=](const FExpressionNode& InOperand, const ContextType* Context) {
+			return ForwardReturnType(In(*InOperand.Cast<OperandType>(), Context));
+		};
+	}
+
+	/** Wrapper function for supplied functions of the signature T(A, B) */
+	template<typename LeftOperandType, typename RightOperandType, typename ContextType, typename FuncType>
+	static typename TEnableIf<Impl::TCallableInfo<FuncType>::NumArgs == 2, typename TOperatorJumpTable<ContextType>::FBinaryFunction>::Type WrapBinaryFunction(FuncType In)
+	{
+		// Ignore the context
+		return [=](const FExpressionNode& InLeftOperand, const FExpressionNode& InRightOperand, const ContextType* Context) {
+			return ForwardReturnType(In(*InLeftOperand.Cast<LeftOperandType>(), *InRightOperand.Cast<RightOperandType>()));
+		};
+	}
+
+	/** Wrapper function for supplied functions of the signature T(A, B, const ContextType* Context) */
+	template<typename LeftOperandType, typename RightOperandType, typename ContextType, typename FuncType>
+	static typename TEnableIf<Impl::TCallableInfo<FuncType>::NumArgs == 3, typename TOperatorJumpTable<ContextType>::FBinaryFunction>::Type WrapBinaryFunction(FuncType In)
+	{
+		// Ignore the context
+		return [=](const FExpressionNode& InLeftOperand, const FExpressionNode& InRightOperand, const ContextType* Context) {
+			return ForwardReturnType(In(*InLeftOperand.Cast<LeftOperandType>(), *InRightOperand.Cast<RightOperandType>(), Context));
+		};
+	}
+
+}
+
+template<typename T>
+FExpressionNode::FExpressionNode(T In, typename TEnableIf<!TPointerIsConvertibleFromTo<T,FExpressionNode>::Value>::Type*)
+	: TypeId(TGetExpressionNodeTypeId<T>::GetTypeId())
+{
+	// Choose the relevant allocation strategy based on the size of the type
+	new(InlineBytes) typename Impl::TStorageTypeDeduction<T, MaxStackAllocationSize>::Type(MoveTemp(In));
+}
+
+template<typename T>
+const T* FExpressionNode::Cast() const
+{
+	if (TypeId == TGetExpressionNodeTypeId<T>::GetTypeId())
+	{
+		return reinterpret_cast<const Impl::TStorageTypeDeduction<T, MaxStackAllocationSize>::Type*>(InlineBytes)->Access();
+	}
+	return nullptr;
+}
+
+// End FExpresionNode
+
+template<typename ContextType>
+FExpressionResult TOperatorJumpTable<ContextType>::ExecBinary(const FExpressionToken& Operator, const FExpressionToken& L, const FExpressionToken& R, const ContextType* Context) const
+{
+	FOperatorFunctionID ID = { Operator.Node.GetTypeId(), L.Node.GetTypeId(), R.Node.GetTypeId() };
+	if (const auto* Func = BinaryOps.Find(ID))
+	{
+		return (*Func)(L.Node, R.Node, Context);
+	}
+
+	FFormatOrderedArguments Args;
+	Args.Add(FText::FromString(Operator.Context.GetString()));
+	Args.Add(FText::FromString(L.Context.GetString()));
+	Args.Add(FText::FromString(R.Context.GetString()));
+	return MakeError(FText::Format(LOCTEXT("UnaryExecutionError", "Binary operator {0} cannot operate on {1} and {2}"), Args));
+}
+
+template<typename ContextType>
+FExpressionResult TOperatorJumpTable<ContextType>::ExecPreUnary(const FExpressionToken& Operator, const FExpressionToken& R, const ContextType* Context) const
+{
+	FOperatorFunctionID ID = { Operator.Node.GetTypeId(), FGuid(), R.Node.GetTypeId() };
+	if (const auto* Func = PreUnaryOps.Find(ID))
+	{
+		return (*Func)(R.Node, Context);
+	}
+
+	FFormatOrderedArguments Args;
+	Args.Add(FText::FromString(Operator.Context.GetString()));
+	Args.Add(FText::FromString(R.Context.GetString()));
+	return MakeError(FText::Format(LOCTEXT("UnaryExecutionError", "Pre-unary operator {0} cannot operate on {1}"), Args));
+}
+
+template<typename ContextType>
+FExpressionResult TOperatorJumpTable<ContextType>::ExecPostUnary(const FExpressionToken& Operator, const FExpressionToken& L, const ContextType* Context) const
+{
+	FOperatorFunctionID ID = { Operator.Node.GetTypeId(), L.Node.GetTypeId(), FGuid() };
+	if (const auto* Func = PostUnaryOps.Find(ID))
+	{
+		return (*Func)(L.Node, Context);
+	}
+
+	FFormatOrderedArguments Args;
+	Args.Add(FText::FromString(Operator.Context.GetString()));
+	Args.Add(FText::FromString(L.Context.GetString()));
+	return MakeError(FText::Format(LOCTEXT("UnaryExecutionError", "Post-unary operator {0} cannot operate on {1}"), Args));
+}
+
+template<typename ContextType>
+template<typename OperatorType, typename FuncType>
+void TOperatorJumpTable<ContextType>::MapPreUnary(FuncType InFunc)
+{
+	typedef typename TRemoveConst<typename TRemoveReference<typename Impl::TCallableInfo<FuncType>::Arg1>::Type>::Type OperandType;
+
+	FOperatorFunctionID ID = {
+		TGetExpressionNodeTypeId<OperatorType>::GetTypeId(),
+		FGuid(),
+		TGetExpressionNodeTypeId<OperandType>::GetTypeId()
+	};
+
+	PreUnaryOps.Add(ID, Impl::WrapUnaryFunction<OperandType, ContextType>(InFunc));
+}
+
+template<typename ContextType>
+template<typename OperatorType, typename FuncType>
+void TOperatorJumpTable<ContextType>::MapPostUnary(FuncType InFunc)
+{
+	typedef typename TRemoveConst<typename TRemoveReference<typename Impl::TCallableInfo<FuncType>::Arg1>::Type>::Type OperandType;
+
+	FOperatorFunctionID ID = {
+		TGetExpressionNodeTypeId<OperatorType>::GetTypeId(),
+		TGetExpressionNodeTypeId<OperandType>::GetTypeId(),
+		FGuid()
+	};
+
+	PostUnaryOps.Add(ID, Impl::WrapUnaryFunction<OperandType, ContextType>(InFunc));
+}
+
+template<typename ContextType>
+template<typename OperatorType, typename FuncType>
+void TOperatorJumpTable<ContextType>::MapBinary(FuncType InFunc)
+{
+	typedef typename TRemoveConst<typename TRemoveReference<typename Impl::TCallableInfo<FuncType>::Arg1>::Type>::Type LeftOperandType;
+	typedef typename TRemoveConst<typename TRemoveReference<typename Impl::TCallableInfo<FuncType>::Arg2>::Type>::Type RightOperandType;
+
+	FOperatorFunctionID ID = {
+		TGetExpressionNodeTypeId<OperatorType>::GetTypeId(),
+		TGetExpressionNodeTypeId<LeftOperandType>::GetTypeId(),
+		TGetExpressionNodeTypeId<RightOperandType>::GetTypeId()
+	};
+
+	BinaryOps.Add(ID, Impl::WrapBinaryFunction<LeftOperandType, RightOperandType, ContextType>(InFunc));
+}
+
+typedef TOperatorEvaluationEnvironment<> FOperatorEvaluationEnvironment;
+
+#undef LOCTEXT_NAMESPACE
\ No newline at end of file
diff --git a/Engine/Source/Runtime/Core/Public/Templates/ValueOrError.h b/Engine/Source/Runtime/Core/Public/Templates/ValueOrError.h
--- a/Engine/Source/Runtime/Core/Public/Templates/ValueOrError.h
+++ b/Engine/Source/Runtime/Core/Public/Templates/ValueOrError.h
@@ -17,15 +17,15 @@ struct TValueOrError_ErrorProxy
 };
 
 template <typename A>
-FORCEINLINE TValueOrError_ValueProxy<A> MakeValue(A&& arg)
+FORCEINLINE TValueOrError_ValueProxy<A> MakeValue(A&& Arg)
 {
-    return TValueOrError_ValueProxy<A>(Forward<A>(arg));
+    return TValueOrError_ValueProxy<A>(Forward<A>(Arg));
 }
 
 template <typename A>
-FORCEINLINE TValueOrError_ErrorProxy<A> MakeError(A&& arg)
+FORCEINLINE TValueOrError_ErrorProxy<A> MakeError(A&& Arg)
 {
-    return TValueOrError_ErrorProxy<A>(Forward<A>(arg));
+    return TValueOrError_ErrorProxy<A>(Forward<A>(Arg));
 }
 
 /** Type used to return either some data, or an error */
@@ -35,14 +35,10 @@ class TValueOrError
 public:
 	/** Construct the result from a value, or an error (See MakeValue and MakeError) */
 	template<typename A>
-	TValueOrError(TValueOrError_ValueProxy<A>&& Proxy) 		: Value(Forward<A>(Proxy.Arg)) {}
+	TValueOrError(TValueOrError_ValueProxy<A>&& Proxy) 		: Value(MoveTemp(Proxy.Arg)) {}
 
 	template<typename A>
-	TValueOrError(TValueOrError_ErrorProxy<A>&& Proxy)		: Error(Forward<A>(Proxy.Arg)) {}
-
-	/** Copy construction/Assignment */
-	TValueOrError(const TValueOrError& In) 					: Error(In.Error), Value(In.Value) {}
-	TValueOrError& operator=(const TValueOrError& In) 		{ Error = In.Error; Value = In.Value; return *this; }
+	TValueOrError(TValueOrError_ErrorProxy<A>&& Proxy)		: Error(MoveTemp(Proxy.Arg)) {}
 
 	/** Move construction/assignment */
 	TValueOrError(TValueOrError&& In) 						: Error(MoveTemp(In.Error)), Value(MoveTemp(In.Value)) {}
@@ -70,5 +66,9 @@ private:
 	/** The error reported by the procedure, if any */
 	TOptional<ErrorType> Error;
 	/** Optional value to return as part of the result */
-	TOptional<ValueType> Value;	
+	TOptional<ValueType> Value;
+
+	/** Copy construction/Assignment disallowed */
+	TValueOrError(const TValueOrError&);
+	TValueOrError& operator=(const TValueOrError&);
 };
\ No newline at end of file
