Added missing typename

diff --git a/Engine/Source/Runtime/Core/Public/Misc/ExpressionParserTypes.inl b/Engine/Source/Runtime/Core/Public/Misc/ExpressionParserTypes.inl
--- a/Engine/Source/Runtime/Core/Public/Misc/ExpressionParserTypes.inl
+++ b/Engine/Source/Runtime/Core/Public/Misc/ExpressionParserTypes.inl
@@ -145,7 +145,7 @@ const T* FExpressionNode::Cast() const
 {
 	if (TypeId == TGetExpressionNodeTypeId<T>::GetTypeId())
 	{
-		return reinterpret_cast<const Impl::TStorageTypeDeduction<T, MaxStackAllocationSize>::Type*>(InlineBytes)->Access();
+		return reinterpret_cast<const typename Impl::TStorageTypeDeduction<T, MaxStackAllocationSize>::Type*>(InlineBytes)->Access();
 	}
 	return nullptr;
 }
