Fixed signed/unsigned mismatch

diff --git a/Engine/Source/Runtime/Core/Public/Misc/ExpressionParserTypes.h b/Engine/Source/Runtime/Core/Public/Misc/ExpressionParserTypes.h
--- a/Engine/Source/Runtime/Core/Public/Misc/ExpressionParserTypes.h
+++ b/Engine/Source/Runtime/Core/Public/Misc/ExpressionParserTypes.h
@@ -186,7 +186,7 @@ public:
 private:
 
 	/** The maximum size of type we will allow allocation on the stack (for efficiency). Anything larger will be allocated on the heap. */
-	static const int32 MaxStackAllocationSize = 64 - sizeof(FGuid);
+	static const uint32 MaxStackAllocationSize = 64 - sizeof(FGuid);
 
 	/** Helper accessor to the data interface. Returns null for empty containers. */
 	Impl::IExpressionNodeStorage* GetData();
diff --git a/Engine/Source/Runtime/Core/Public/Misc/ExpressionParserTypes.inl b/Engine/Source/Runtime/Core/Public/Misc/ExpressionParserTypes.inl
--- a/Engine/Source/Runtime/Core/Public/Misc/ExpressionParserTypes.inl
+++ b/Engine/Source/Runtime/Core/Public/Misc/ExpressionParserTypes.inl
@@ -52,10 +52,10 @@ namespace Impl
 		virtual void MoveAssign(uint8* Dst) override 		{ reinterpret_cast<FHeapDataStorage*>(Dst)->Value = MoveTemp(Value); }
 	};
 
-	template<typename T, int32 MaxStackAllocationSize, typename Enabled=void> struct TStorageTypeDeduction;
-	template<typename T, int32 MaxStackAllocationSize> struct TStorageTypeDeduction<T, MaxStackAllocationSize, typename TEnableIf<(sizeof(FInlineDataStorage<T>) <= MaxStackAllocationSize)>::Type>
+	template<typename T, uint32 MaxStackAllocationSize, typename Enabled=void> struct TStorageTypeDeduction;
+	template<typename T, uint32 MaxStackAllocationSize> struct TStorageTypeDeduction<T, MaxStackAllocationSize, typename TEnableIf<(sizeof(FInlineDataStorage<T>) <= MaxStackAllocationSize)>::Type>
 		{ typedef FInlineDataStorage<T> Type; };
-	template<typename T, int32 MaxStackAllocationSize> struct TStorageTypeDeduction<T, MaxStackAllocationSize, typename TEnableIf<(sizeof(FInlineDataStorage<T>) >  MaxStackAllocationSize)>::Type>
+	template<typename T, uint32 MaxStackAllocationSize> struct TStorageTypeDeduction<T, MaxStackAllocationSize, typename TEnableIf<(sizeof(FInlineDataStorage<T>) >  MaxStackAllocationSize)>::Type>
 		{ typedef FHeapDataStorage<T> 	Type; };
 
 
