Expression evaluation now operates on a const array of pre-compiled tokens.

This required the addition of an explicit expression node copy so that evaluation can work on an immutable array.

diff --git a/Engine/Source/Runtime/Core/Private/Misc/ExpressionParser.cpp b/Engine/Source/Runtime/Core/Private/Misc/ExpressionParser.cpp
--- a/Engine/Source/Runtime/Core/Private/Misc/ExpressionParser.cpp
+++ b/Engine/Source/Runtime/Core/Private/Misc/ExpressionParser.cpp
@@ -392,6 +392,20 @@ Impl::IExpressionNodeStorage* FExpressionNode::GetData()
 	return TypeId.IsValid() ? reinterpret_cast<Impl::IExpressionNodeStorage*>(InlineBytes) : nullptr;
 }
 
+const Impl::IExpressionNodeStorage* FExpressionNode::GetData() const
+{
+	return TypeId.IsValid() ? reinterpret_cast<const Impl::IExpressionNodeStorage*>(InlineBytes) : nullptr;
+}
+
+FExpressionNode FExpressionNode::Copy() const
+{
+	if (const auto* Data = GetData())
+	{
+		return Data->Copy();
+	}
+	return FExpressionNode();
+}
+
 const FGuid* FExpressionGrammar::GetGrouping(const FGuid& TypeId) const
 {
 	return Groupings.Find(TypeId);
@@ -631,16 +645,40 @@ namespace ExpressionParser
 			return MakeError(CompilationResult.GetError());
 		}
 
-		return Evaluate(MoveTemp(CompilationResult.GetValue()), InEnvironment);
+		return Evaluate(CompilationResult.GetValue(), InEnvironment);
 	}
 
-	FExpressionResult Evaluate(TArray<FCompiledToken> CompiledTokens, const IOperatorEvaluationEnvironment& InEnvironment)
+	FExpressionResult Evaluate(const TArray<FCompiledToken>& CompiledTokens, const IOperatorEvaluationEnvironment& InEnvironment)
 	{
-		TArray<FExpressionToken> OperandStack;
+		// Evaluation strategy: the supplied compiled tokens are const. To avoid copying the whole array, we store a separate array of
+		// any tokens that are generated at runtime by the evaluator. The operand stack will consist of indices into either the CompiledTokens
+		// array, or the RuntimeGeneratedTokens (where Index >= CompiledTokens.Num())
+		TArray<FExpressionToken> RuntimeGeneratedTokens;
+		TArray<int32> OperandStack;
+
+		/** Get the token pertaining to the specified operand index */
+		auto GetToken = [&](int32 Index) -> const FExpressionToken& {
+			if (Index < CompiledTokens.Num())
+			{
+				return CompiledTokens[Index];
+			}
+			else
+			{
+				return RuntimeGeneratedTokens[Index - CompiledTokens.Num()];
+			}
+		};
+
+		/** Add a new token to the runtime generated array */
+		auto AddToken = [&](FExpressionToken&& In) -> int32 {
+			auto Index = CompiledTokens.Num() + RuntimeGeneratedTokens.Num();
+			RuntimeGeneratedTokens.Emplace(MoveTemp(In));
+			return Index;
+		};
+
 
 		for (int32 Index = 0; Index < CompiledTokens.Num(); ++Index)
 		{
-			auto& Token = CompiledTokens[Index];
+			const auto& Token = CompiledTokens[Index];
 
 			switch(Token.Type)
 			{
@@ -648,21 +686,21 @@ namespace ExpressionParser
 				continue;
 
 			case FCompiledToken::Operand:
-				OperandStack.Push(MoveTemp(Token));
+				OperandStack.Push(Index);
 				continue;
 
 			case FCompiledToken::BinaryOperator:
 				if (OperandStack.Num() >= 2)
 				{
 					// Binary
-					auto R = OperandStack.Pop();
-					auto L = OperandStack.Pop();
+					const auto& R = GetToken(OperandStack.Pop());
+					const auto& L = GetToken(OperandStack.Pop());
 
 					auto OpResult = InEnvironment.ExecBinary(Token, L, R);
 					if (OpResult.IsValid())
 					{
 						// Inherit the LHS context
-						OperandStack.Push(FExpressionToken(L.Context, MoveTemp(OpResult.GetValue())));
+						OperandStack.Push(AddToken(FExpressionToken(L.Context, MoveTemp(OpResult.GetValue()))));
 					}
 					else
 					{
@@ -682,7 +720,7 @@ namespace ExpressionParser
 
 				if (OperandStack.Num() >= 1)
 				{
-					auto Operand = OperandStack.Pop();
+					const auto& Operand = GetToken(OperandStack.Pop());
 
 					FExpressionResult OpResult = (Token.Type == FCompiledToken::PreUnaryOperator) ?
 						InEnvironment.ExecPreUnary(Token, Operand) :
@@ -691,7 +729,7 @@ namespace ExpressionParser
 					if (OpResult.IsValid())
 					{
 						// Inherit the LHS context
-						OperandStack.Push(FExpressionToken(Token.Context, MoveTemp(OpResult.GetValue())));
+						OperandStack.Push(AddToken(FExpressionToken(Operand.Context, MoveTemp(OpResult.GetValue()))));
 					}
 					else
 					{
@@ -710,7 +748,7 @@ namespace ExpressionParser
 
 		if (OperandStack.Num() == 1)
 		{
-			return MakeValue(MoveTemp(OperandStack[0].Node));
+			return MakeValue(GetToken(OperandStack[0]).Node.Copy());
 		}
 		else
 		{
@@ -754,6 +792,23 @@ PRAGMA_DISABLE_OPTIMIZATION
 			return *this;
 		}
 
+		FMoveableType(const FMoveableType& In) : Id(-1), bOwnsLeak(false) { *this = In; }
+		const FMoveableType& operator=(const FMoveableType& In)
+		{
+			const bool bDidOwnLeak = bOwnsLeak;
+			bOwnsLeak = In.bOwnsLeak;
+
+			if (bOwnsLeak && !bDidOwnLeak)
+			{
+				++*LeakCount;
+			}
+			else if (!bOwnsLeak && bDidOwnLeak)
+			{
+				--*LeakCount;
+			}
+			return *this;
+		}
+
 		virtual ~FMoveableType()
 		{
 			if (bOwnsLeak)
@@ -847,7 +902,8 @@ PRAGMA_DISABLE_OPTIMIZATION
 	{
 		FHugeType(int32 InId) : FMoveableType(InId) {}
 		FHugeType(FHugeType&& In) : FMoveableType(MoveTemp(In)) {}
-
+		FHugeType(const FHugeType& In) : FMoveableType(In) {}
+		
 		uint8 Padding[1024];
 	};
 
diff --git a/Engine/Source/Runtime/Core/Public/Misc/ExpressionParser.h b/Engine/Source/Runtime/Core/Public/Misc/ExpressionParser.h
--- a/Engine/Source/Runtime/Core/Public/Misc/ExpressionParser.h
+++ b/Engine/Source/Runtime/Core/Public/Misc/ExpressionParser.h
@@ -30,7 +30,7 @@ namespace ExpressionParser
 	CORE_API FExpressionResult Evaluate(const TCHAR* InExpression, const FTokenDefinitions& InTokenDefinitions, const FExpressionGrammar& InGrammar, const IOperatorEvaluationEnvironment& InEnvironment);
 
 	/** Evaluate the specified pre-compiled tokens using an evaluation environment */
-	CORE_API FExpressionResult Evaluate(TArray<FCompiledToken> CompiledTokens, const IOperatorEvaluationEnvironment& InEnvironment);
+	CORE_API FExpressionResult Evaluate(const TArray<FCompiledToken>& CompiledTokens, const IOperatorEvaluationEnvironment& InEnvironment);
 
 	/** Templated versions of evaluation functions used when passing a specific jump table and context */
 	template<typename ContextType>
@@ -41,7 +41,7 @@ namespace ExpressionParser
 	}
 
 	template<typename ContextType>
-	FExpressionResult Evaluate(TArray<FCompiledToken>& CompiledTokens, const TOperatorJumpTable<ContextType>& InJumpTable, const ContextType* InContext = nullptr)
+	FExpressionResult Evaluate(const TArray<FCompiledToken>& CompiledTokens, const TOperatorJumpTable<ContextType>& InJumpTable, const ContextType* InContext = nullptr)
 	{
 		TOperatorEvaluationEnvironment<ContextType> Env(InJumpTable, InContext);
 		return Evaluate(CompiledTokens, Env);
diff --git a/Engine/Source/Runtime/Core/Public/Misc/ExpressionParserTypes.h b/Engine/Source/Runtime/Core/Public/Misc/ExpressionParserTypes.h
--- a/Engine/Source/Runtime/Core/Public/Misc/ExpressionParserTypes.h
+++ b/Engine/Source/Runtime/Core/Public/Misc/ExpressionParserTypes.h
@@ -163,6 +163,10 @@ template<typename T> struct TGetExpressionNodeTypeId;
 class CORE_API FExpressionNode : FNoncopyable
 {
 public:
+
+	/** Default constructor */
+	FExpressionNode() : TypeId() {}
+
 	/** Construction from client expression data type */
 	template<typename T>
 	FExpressionNode(T In,
@@ -183,6 +187,9 @@ public:
 	template<typename T>
 	const T* Cast() const;
 
+	/** Copy this node and its wrapped data */
+	FExpressionNode Copy() const;
+
 private:
 
 	/** The maximum size of type we will allow allocation on the stack (for efficiency). Anything larger will be allocated on the heap. */
@@ -190,6 +197,7 @@ private:
 
 	/** Helper accessor to the data interface. Returns null for empty containers. */
 	Impl::IExpressionNodeStorage* GetData();
+	const Impl::IExpressionNodeStorage* GetData() const;
 
 	/** TypeID - 16 bytes */
 	FGuid TypeId;
diff --git a/Engine/Source/Runtime/Core/Public/Misc/ExpressionParserTypes.inl b/Engine/Source/Runtime/Core/Public/Misc/ExpressionParserTypes.inl
--- a/Engine/Source/Runtime/Core/Public/Misc/ExpressionParserTypes.inl
+++ b/Engine/Source/Runtime/Core/Public/Misc/ExpressionParserTypes.inl
@@ -14,6 +14,8 @@ namespace Impl
 		virtual void Reseat(uint8* Dst) = 0;
 		/** Move this type to a buffer already allocated to the same type (uses type-defined move-assignment) */
 		virtual void MoveAssign(uint8* Dst) = 0;
+		/** Copy this data */
+		virtual FExpressionNode Copy() const = 0;
 	};
 
 	/** Implementation of the wrapper utility for any moveable/copyable data, that allows us to virtually move/copy/destruct the data */
@@ -32,6 +34,7 @@ namespace Impl
 
 		virtual void Reseat(uint8* Dst) override 			{ new(Dst) FInlineDataStorage(MoveTemp(Value)); }
 		virtual void MoveAssign(uint8* Dst) override 		{ reinterpret_cast<FInlineDataStorage*>(Dst)->Value = MoveTemp(Value); }
+		virtual FExpressionNode Copy() const override		{ return Value; }
 	};
 
 	/** Data is stored on the heap in this implementation */
@@ -50,6 +53,7 @@ namespace Impl
 
 		virtual void Reseat(uint8* Dst) override 			{ new(Dst) FHeapDataStorage(MoveTemp(*this)); }
 		virtual void MoveAssign(uint8* Dst) override 		{ reinterpret_cast<FHeapDataStorage*>(Dst)->Value = MoveTemp(Value); }
+		virtual FExpressionNode Copy() const override		{ return *Value; }
 	};
 
 	template<typename T, uint32 MaxStackAllocationSize, typename Enabled=void> struct TStorageTypeDeduction;
