Added advanced search query syntax to the Content Browser

UETOOL-332 - Collections 2.0
UETOOL-366 - Ability to search and filter assets that are only in specific collections

This allows you to perform complex queries against any of the meta-data from the asset registry (using standard comparison operators), as well as being able to test other standard keys, such as the collections that contain a given asset, its name, path, or class.

The syntax supports sub-expression grouping, and also partial (unquoted), exact (quoted) starts with (s...), or ends with (...s) string matching, as well as standard logic operators (OR, AND, NOT).

diff --git a/Engine/Source/Editor/ContentBrowser/Private/FrontendFilters.cpp b/Engine/Source/Editor/ContentBrowser/Private/FrontendFilters.cpp
--- a/Engine/Source/Editor/ContentBrowser/Private/FrontendFilters.cpp
+++ b/Engine/Source/Editor/ContentBrowser/Private/FrontendFilters.cpp
@@ -5,6 +5,7 @@
 #include "ISourceControlModule.h"
 #include "ObjectTools.h"
 #include "AssetToolsModule.h"
+#include "CollectionManagerModule.h"
 
 
 /////////////////////////////////////////
@@ -53,9 +54,97 @@ void AssetDataToClassAndNameStrings(FAssetFilterType Asset, OUT TArray< FString
 	AssetDataToNameString(Asset, Array);
 }
 
+bool AssetDataTestComplexExpression(FAssetFilterType Asset, const FName& InKey, const FString& InValue, ETextFilterComparisonOperation InComparisonOperation, ETextFilterTextComparisonMode InTextComparisonMode)
+{
+	static const FName NameKeyName = "Name";
+	static const FName PathKeyName = "Path";
+	static const FName ClassKeyName = "Class";
+	static const FName TypeKeyName = "Type";
+	static const FName CollectionKeyName = "Collection";
+	static const FName TagKeyName = "Tag";
+
+	// Special case for the asset name, as this isn't contained within the asset registry meta-data
+	if (InKey == NameKeyName)
+	{
+		// Names can only work with Equal or NotEqual type tests
+		if (InComparisonOperation != ETextFilterComparisonOperation::Equal && InComparisonOperation != ETextFilterComparisonOperation::NotEqual)
+		{
+			return false;
+		}
+
+		const bool bIsMatch = TextFilterUtils::TestBasicStringExpression(Asset.AssetName.ToString(), InValue, InTextComparisonMode);
+		return (InComparisonOperation == ETextFilterComparisonOperation::Equal) ? bIsMatch : !bIsMatch;
+	}
+
+	// Special case for the asset path, as this isn't contained within the asset registry meta-data
+	if (InKey == PathKeyName)
+	{
+		// Paths can only work with Equal or NotEqual type tests
+		if (InComparisonOperation != ETextFilterComparisonOperation::Equal && InComparisonOperation != ETextFilterComparisonOperation::NotEqual)
+		{
+			return false;
+		}
+
+		const bool bIsMatch = TextFilterUtils::TestBasicStringExpression(Asset.PackagePath.ToString(), InValue, InTextComparisonMode);
+		return (InComparisonOperation == ETextFilterComparisonOperation::Equal) ? bIsMatch : !bIsMatch;
+	}
+
+	// Special case for the asset type, as this isn't contained within the asset registry meta-data
+	if (InKey == ClassKeyName || InKey == TypeKeyName)
+	{
+		// Class names can only work with Equal or NotEqual type tests
+		if (InComparisonOperation != ETextFilterComparisonOperation::Equal && InComparisonOperation != ETextFilterComparisonOperation::NotEqual)
+		{
+			return false;
+		}
+
+		const bool bIsMatch = TextFilterUtils::TestBasicStringExpression(Asset.AssetClass.ToString(), InValue, InTextComparisonMode);
+		return (InComparisonOperation == ETextFilterComparisonOperation::Equal) ? bIsMatch : !bIsMatch;
+	}
+
+	// Special case for collections, as these aren't contained within the asset registry meta-data
+	if ((InKey == CollectionKeyName || InKey == TagKeyName) && FCollectionManagerModule::IsModuleAvailable())
+	{
+		// Collections can only work with Equal or NotEqual type tests
+		if (InComparisonOperation != ETextFilterComparisonOperation::Equal && InComparisonOperation != ETextFilterComparisonOperation::NotEqual)
+		{
+			return false;
+		}
+
+		FCollectionManagerModule& CollectionManagerModule = FCollectionManagerModule::GetModule();
+
+		TArray<FName> AssetCollectionNames;
+		CollectionManagerModule.Get().GetCollectionsContainingObject(Asset.ObjectPath, ECollectionShareType::CST_All, AssetCollectionNames);
+
+		bool bFoundMatch = false;
+		for (const FName& AssetCollectionName : AssetCollectionNames)
+		{
+			if (TextFilterUtils::TestBasicStringExpression(AssetCollectionName.ToString(), InValue, InTextComparisonMode))
+			{
+				bFoundMatch = true;
+				break;
+			}
+		}
+
+		return (InComparisonOperation == ETextFilterComparisonOperation::Equal) ? bFoundMatch : !bFoundMatch;
+	}
+
+	// Generic handling for anything in the asset meta-data
+	const FString* const MetaDataValue = Asset.TagsAndValues.Find(InKey);
+	if (MetaDataValue)
+	{
+		return TextFilterUtils::TestComplexExpression(*MetaDataValue, InValue, InComparisonOperation, InTextComparisonMode);
+	}
+
+	return false;
+}
+
 FFrontendFilter_Text::FFrontendFilter_Text()
 	: FFrontendFilter(nullptr)
-	, TextFilter(TTextFilter<FAssetFilterType>::FItemToStringArray::CreateStatic(&AssetDataToClassAndNameStrings))
+	, TextFilter(
+		TTextFilter<FAssetFilterType>::FItemToStringArray::CreateStatic(&AssetDataToClassAndNameStrings), 
+		TTextFilter<FAssetFilterType>::FItemTestComplexExpression::CreateStatic(&AssetDataTestComplexExpression)
+		)
 {
 	SetIncludeClassName(true);
 }
@@ -80,6 +169,11 @@ void FFrontendFilter_Text::SetRawFilterText(const FText& InFilterText)
 	return TextFilter.SetRawFilterText(InFilterText);
 }
 
+FText FFrontendFilter_Text::GetFilterErrorText() const
+{
+	return TextFilter.GetFilterErrorText();
+}
+
 void FFrontendFilter_Text::SetIncludeClassName(bool bIncludeClassName)
 {
 	// Preserve the existing text. The filter is going to get recreated here.
@@ -87,11 +181,17 @@ void FFrontendFilter_Text::SetIncludeClassName(bool bIncludeClassName)
 
 	if ( bIncludeClassName )
 	{
-		TextFilter = TTextFilter<FAssetFilterType>::FItemToStringArray::CreateStatic(&AssetDataToClassAndNameStrings);
+		TextFilter = TTextFilter<FAssetFilterType>(
+			TTextFilter<FAssetFilterType>::FItemToStringArray::CreateStatic(&AssetDataToClassAndNameStrings),
+			TTextFilter<FAssetFilterType>::FItemTestComplexExpression::CreateStatic(&AssetDataTestComplexExpression)
+			);
 	}
 	else
 	{
-		TextFilter = TTextFilter<FAssetFilterType>::FItemToStringArray::CreateStatic(&AssetDataToNameString);
+		TextFilter = TTextFilter<FAssetFilterType>(
+			TTextFilter<FAssetFilterType>::FItemToStringArray::CreateStatic(&AssetDataToNameString),
+			TTextFilter<FAssetFilterType>::FItemTestComplexExpression::CreateStatic(&AssetDataTestComplexExpression)
+			);
 	}
 
 	// Apply the existing text before we re-assign the delegate since we want the text preservation to be opaque.
@@ -234,47 +334,7 @@ bool FFrontendFilter_ArbitraryComparisonOperation::PassesFilter(FAssetFilterType
 {
 	if (const FString* pTagValue = InItem.TagsAndValues.Find(TagName))
 	{
-		const FString& TestTagValue = *pTagValue;
-
-		const bool bNumericComparison = TestTagValue.IsNumeric() && TargetTagValue.IsNumeric();
-
-		int32 ComparisonSign = 0;
-		if (bNumericComparison)
-		{
-			const double Difference = FCString::Atod(*TestTagValue) - FCString::Atod(*TargetTagValue);
-			ComparisonSign = (int32)FMath::Sign(Difference);
-		}
-		else
-		{
-			ComparisonSign = TestTagValue.Compare(TargetTagValue, ESearchCase::IgnoreCase);
-		}
-
-		bool bComparisonPassed = false;
-		switch (ComparisonOp)
-		{
-		case ETextFilterComparisonOperation::Equal:
-			bComparisonPassed = ComparisonSign == 0;
-			break;
-		case ETextFilterComparisonOperation::NotEqual:
-			bComparisonPassed = ComparisonSign != 0;
-			break;
-		case ETextFilterComparisonOperation::Less:
-			bComparisonPassed = ComparisonSign < 0;
-			break;
-		case ETextFilterComparisonOperation::LessOrEqual:
-			bComparisonPassed = ComparisonSign <= 0;
-			break;
-		case ETextFilterComparisonOperation::Greater:
-			bComparisonPassed = ComparisonSign > 0;
-			break;
-		case ETextFilterComparisonOperation::GreaterOrEqual:
-			bComparisonPassed = ComparisonSign >= 0;
-			break;
-		default:
-			check(false);
-		};
-
-		return bComparisonPassed;
+		return TextFilterUtils::TestComplexExpression(*pTagValue, TargetTagValue, ComparisonOp, ETextFilterTextComparisonMode::Exact);
 	}
 	else
 	{
diff --git a/Engine/Source/Editor/ContentBrowser/Private/FrontendFilters.h b/Engine/Source/Editor/ContentBrowser/Private/FrontendFilters.h
--- a/Engine/Source/Editor/ContentBrowser/Private/FrontendFilters.h
+++ b/Engine/Source/Editor/ContentBrowser/Private/FrontendFilters.h
@@ -28,6 +28,9 @@ public:
 	/** Set the Text to be used as the Filter's restrictions */
 	void SetRawFilterText(const FText& InFilterText);
 
+	/** Get the last error returned from lexing or compiling the current filter text */
+	FText GetFilterErrorText() const;
+
 	/** If bIncludeClassName is true, the text filter will include an asset's class name in the search */
 	void SetIncludeClassName(bool bIncludeClassName);
 
diff --git a/Engine/Source/Editor/ContentBrowser/Private/SContentBrowser.cpp b/Engine/Source/Editor/ContentBrowser/Private/SContentBrowser.cpp
--- a/Engine/Source/Editor/ContentBrowser/Private/SContentBrowser.cpp
+++ b/Engine/Source/Editor/ContentBrowser/Private/SContentBrowser.cpp
@@ -829,7 +829,8 @@ void SContentBrowser::SyncToAssets( const TArray<FAssetData>& AssetDataList, con
 	// Disable the filter search (reset the filter, then clear the search text)
 	// Note: we have to remove the filter immediately, we can't wait for OnSearchBoxChanged to hit
 	SetSearchBoxText(FText::GetEmpty());
-	SearchBoxPtr->SetText(FText::GetEmpty());	
+	SearchBoxPtr->SetText(FText::GetEmpty());
+	SearchBoxPtr->SetError(FText::GetEmpty());
 
 	// Tell the sources view first so the asset view will be up to date by the time we request the sync
 	PathViewPtr->SyncToAssets(AssetDataList, bAllowImplicitSync);
@@ -1242,6 +1243,7 @@ void SContentBrowser::SetSearchBoxText(const FText& InSearchText)
 	if (!InSearchText.EqualToCaseIgnored(TextFilter->GetRawFilterText()))
 	{
 		TextFilter->SetRawFilterText( InSearchText );
+		SearchBoxPtr->SetError( TextFilter->GetFilterErrorText() );
 		if(InSearchText.IsEmpty())
 		{
 			FrontendFilters->Remove(TextFilter);
diff --git a/Engine/Source/Editor/EditorWidgets/Private/SAssetSearchBox.cpp b/Engine/Source/Editor/EditorWidgets/Private/SAssetSearchBox.cpp
--- a/Engine/Source/Editor/EditorWidgets/Private/SAssetSearchBox.cpp
+++ b/Engine/Source/Editor/EditorWidgets/Private/SAssetSearchBox.cpp
@@ -53,6 +53,16 @@ void SAssetSearchBox::SetText(const TAttribute< FText >& InNewText)
 	PreCommittedText = InNewText.Get();
 }
 
+void SAssetSearchBox::SetError( const FText& InError )
+{
+	InputText->SetError(InError);
+}
+
+void SAssetSearchBox::SetError( const FString& InError )
+{
+	InputText->SetError(InError);
+}
+
 FReply SAssetSearchBox::OnPreviewKeyDown( const FGeometry& MyGeometry, const FKeyEvent& InKeyEvent )
 {
 	if ( SuggestionBox->IsOpen() && InKeyEvent.GetKey() == EKeys::Escape )
diff --git a/Engine/Source/Editor/EditorWidgets/Public/SAssetSearchBox.h b/Engine/Source/Editor/EditorWidgets/Public/SAssetSearchBox.h
--- a/Engine/Source/Editor/EditorWidgets/Public/SAssetSearchBox.h
+++ b/Engine/Source/Editor/EditorWidgets/Public/SAssetSearchBox.h
@@ -52,6 +52,10 @@ public:
 	/** Sets the text string currently being edited */
 	void SetText(const TAttribute< FText >& InNewText);
 
+	/** Set or clear the current error reporting information for this search box */
+	void SetError( const FText& InError );
+	void SetError( const FString& InError );
+
 	// SWidget implementation
 	virtual FReply OnPreviewKeyDown( const FGeometry& MyGeometry, const FKeyEvent& InKeyEvent ) override;
 	virtual FReply OnKeyDown( const FGeometry& MyGeometry, const FKeyEvent& InKeyEvent ) override;
diff --git a/Engine/Source/Runtime/Core/Private/Misc/TextFilterExpressionEvaluator.cpp b/Engine/Source/Runtime/Core/Private/Misc/TextFilterExpressionEvaluator.cpp
new file mode 100644
--- /dev/null
+++ b/Engine/Source/Runtime/Core/Private/Misc/TextFilterExpressionEvaluator.cpp
@@ -0,0 +1,536 @@
+// Copyright 1998-2015 Epic Games, Inc. All Rights Reserved.
+
+#include "CorePrivatePCH.h"
+#include "TextFilterExpressionEvaluator.h"
+#include "BasicMathExpressionEvaluator.h"
+
+namespace TextFilterExpressionParser
+{
+	class FTextToken
+	{
+	public:
+		FTextToken(FString InString, ETextFilterTextComparisonMode InTextComparisonMode)
+			: String(MoveTemp(InString))
+			, TextComparisonMode(MoveTemp(InTextComparisonMode))
+		{
+		}
+
+		FTextToken(const FTextToken& Other)
+			: String(Other.String)
+			, TextComparisonMode(Other.TextComparisonMode)
+		{
+		}
+
+		FTextToken(FTextToken&& Other)
+			: String(MoveTemp(Other.String))
+			, TextComparisonMode(MoveTemp(Other.TextComparisonMode))
+		{
+		}
+
+		FTextToken& operator=(const FTextToken& Other)
+		{
+			String = Other.String;
+			TextComparisonMode = Other.TextComparisonMode;
+			return *this;
+		}
+
+		FTextToken& operator=(FTextToken&& Other)
+		{
+			String = MoveTemp(Other.String);
+			TextComparisonMode = MoveTemp(Other.TextComparisonMode);
+			return *this;
+		}
+
+		const FString& GetString() const
+		{
+			return String;
+		}
+
+		ETextFilterTextComparisonMode GetTextComparisonMode() const
+		{
+			return TextComparisonMode;
+		}
+
+		void SetTextComparisonMode(ETextFilterTextComparisonMode InTextComparisonMode)
+		{
+			TextComparisonMode = InTextComparisonMode;
+		}
+
+	private:
+		FString String;
+		ETextFilterTextComparisonMode TextComparisonMode;
+	};
+}
+
+#define DEFINE_TEXT_EXPRESSION_OPERATOR_NODE(TYPE, MONIKER_COUNT, ...)	\
+	namespace TextFilterExpressionParser								\
+	{																	\
+		struct TYPE														\
+		{																\
+			static const int32 MonikerCount = MONIKER_COUNT;			\
+			static const TCHAR* Monikers[MonikerCount];					\
+		};																\
+	}																	\
+	DEFINE_EXPRESSION_NODE_TYPE(TextFilterExpressionParser::TYPE, __VA_ARGS__)
+
+DEFINE_TEXT_EXPRESSION_OPERATOR_NODE(FSubExpressionStart,		1,	0x0E7E1BC9, 0xF0B94D5D, 0x80F44D45, 0x85A082AA)
+DEFINE_TEXT_EXPRESSION_OPERATOR_NODE(FSubExpressionEnd,			1,	0x5E10956D, 0x2E17411F, 0xB6469E22, 0xB5E56E6C)
+
+DEFINE_TEXT_EXPRESSION_OPERATOR_NODE(FEqual,					3,	0x32457EFC, 0x4928406F, 0xBD78D943, 0x633797D1)
+DEFINE_TEXT_EXPRESSION_OPERATOR_NODE(FNotEqual,					2,	0x8F445A19, 0xF33443D9, 0x90D6DC85, 0xB0C5D071)
+DEFINE_TEXT_EXPRESSION_OPERATOR_NODE(FLess,						1,	0xB85E222B, 0x47E24E1F, 0xBC5A384D, 0x2FF471E1)
+DEFINE_TEXT_EXPRESSION_OPERATOR_NODE(FLessOrEqual,				2,	0x8C0A46B0, 0x8DAA4E2B, 0xA7FE4A23, 0xEF215918)
+DEFINE_TEXT_EXPRESSION_OPERATOR_NODE(FGreater,					1,	0x6A6247F4, 0xFD78467F, 0xA6AC1244, 0x0A31E0A5)
+DEFINE_TEXT_EXPRESSION_OPERATOR_NODE(FGreaterOrEqual,			2,	0x09D75C5E, 0xA29A4194, 0x8E8E5278, 0xC84365FD)
+
+DEFINE_TEXT_EXPRESSION_OPERATOR_NODE(FOr,						3,	0xF4778B51, 0xF535414D, 0x9C0EB5F2, 0x2F2B0FD4)
+DEFINE_TEXT_EXPRESSION_OPERATOR_NODE(FAnd,						3,	0x7511397A, 0x02D24DC2, 0x86729800, 0xF454C320)
+DEFINE_TEXT_EXPRESSION_OPERATOR_NODE(FNot,						3,	0x03D78990, 0x41D04E26, 0x8E98AD2F, 0x74667868)
+
+DEFINE_TEXT_EXPRESSION_OPERATOR_NODE(FTextCmpAnchor,			1,	0xC70749AD, 0x954D4A92, 0x95061222, 0xC8AC6D6D)
+
+DEFINE_EXPRESSION_NODE_TYPE(TextFilterExpressionParser::FTextToken,	0x09E49538, 0x633545E3, 0x84B5644F, 0x1F11628F);
+DEFINE_EXPRESSION_NODE_TYPE(bool,									0xC1CD5DCF, 0x2AB44958, 0xB3FF4F8F, 0xE665D121);
+
+namespace TextFilterExpressionParser
+{
+	/** This contains all the symbols that can define breaking points between text and an operator */
+	static const TCHAR BasicTextBreakingCharacters[]	= { '(', ')', '!', '-', '&', '|', '.', '"', ' ' };						// ETextFilterExpressionEvaluatorMode::BasicString
+	static const TCHAR ComplexTextBreakingCharacters[]	= { '(', ')', '=', ':', '<', '>', '!', '-', '&', '|', '.', '"', ' ' };	// ETextFilterExpressionEvaluatorMode::Complex
+
+	const TCHAR* FSubExpressionStart::Monikers[]		= { TEXT("(") };
+	const TCHAR* FSubExpressionEnd::Monikers[]			= { TEXT(")") };
+
+	const TCHAR* FEqual::Monikers[]						= { TEXT("=="), TEXT("="), TEXT(":") };
+	const TCHAR* FNotEqual::Monikers[]					= { TEXT("!="), TEXT("!:") };
+	const TCHAR* FLess::Monikers[]						= { TEXT("<") };
+	const TCHAR* FLessOrEqual::Monikers[]				= { TEXT("<="), TEXT("<:") };
+	const TCHAR* FGreater::Monikers[]					= { TEXT(">") };
+	const TCHAR* FGreaterOrEqual::Monikers[]			= { TEXT(">="), TEXT(">:") };
+
+	const TCHAR* FOr::Monikers[]						= { TEXT("OR"), TEXT("||"), TEXT("|") };
+	const TCHAR* FAnd::Monikers[]						= { TEXT("AND"), TEXT("&&"), TEXT("&") };
+	const TCHAR* FNot::Monikers[]						= { TEXT("NOT"), TEXT("!"), TEXT("-") };
+
+	const TCHAR* FTextCmpAnchor::Monikers[]				= { TEXT("...") };
+
+	enum class ETextNodeType : uint8
+	{
+		Text,
+		SubExpressionStart,
+		SubExpressionEnd,
+		TextCmpAnchor,
+		Other,
+	};
+
+	/** Get the type of the given node */
+	ETextNodeType GetTextNodeType(const FExpressionNode& InNode)
+	{
+		if (InNode.Cast<FTextToken>())			return ETextNodeType::Text;
+		if (InNode.Cast<FSubExpressionStart>()) return ETextNodeType::SubExpressionStart;
+		if (InNode.Cast<FSubExpressionEnd>())	return ETextNodeType::SubExpressionEnd;
+		if (InNode.Cast<FTextCmpAnchor>())		return ETextNodeType::TextCmpAnchor;
+		return ETextNodeType::Other;
+	}
+
+	/** Check the given node to see if we consider it to form part of a complex expression */
+	bool IsComplexNode(const FExpressionNode& InNode)
+	{
+		// Key->value pairs require some form of comparison operator, so if we've found once of those, consider this a complex expression
+		return (InNode.Cast<FEqual>() || InNode.Cast<FNotEqual>()|| InNode.Cast<FLess>()|| InNode.Cast<FLessOrEqual>()|| InNode.Cast<FGreater>()|| InNode.Cast<FGreaterOrEqual>());
+	}
+
+	/** Consume an operator from the specified consumer's stream, if one exists at the current read position */
+	template<typename TSymbol>
+	TOptional<FExpressionError> ConsumeOperator(FExpressionTokenConsumer& Consumer)
+	{
+		auto& Stream = Consumer.GetStream();
+
+		for (const TCHAR* Moniker : TSymbol::Monikers)
+		{
+			TOptional<FStringToken> OperatorToken = Stream.ParseToken(Moniker);
+			if (OperatorToken.IsSet())
+			{
+				Consumer.Add(OperatorToken.GetValue(), TSymbol());
+			}
+		}
+
+		return TOptional<FExpressionError>();
+	}
+
+	/** Consume a number from the specified consumer's stream, if one exists at the current read position */
+	TOptional<FExpressionError> ConsumeNumber(FExpressionTokenConsumer& Consumer)
+	{
+		auto& Stream = Consumer.GetStream();
+
+		TOptional<FStringToken> NumberToken = ExpressionParser::ParseNumber(Stream);
+		
+		if (NumberToken.IsSet())
+		{
+			Consumer.Add(NumberToken.GetValue(), FExpressionNode(FTextToken(NumberToken.GetValue().GetString(), ETextFilterTextComparisonMode::Partial)));
+		}
+
+		return TOptional<FExpressionError>();
+	}
+
+	/** Consume quoted text from the specified consumer's stream, if one exists at the current read position */
+	TOptional<FExpressionError> ConsumeQuotedText(FExpressionTokenConsumer& Consumer)
+	{
+		auto& Stream = Consumer.GetStream();
+
+		const TCHAR QuoteChar = Stream.PeekChar();
+		
+		if (QuoteChar != '"' && QuoteChar != '\'')
+		{
+			return TOptional<FExpressionError>();
+		}
+
+		FStringToken TextTokenWithQuotes = Stream.ParseSymbol().GetValue();
+
+		int32 NumConsecutiveSlashes = 0;
+		TOptional<FStringToken> TextToken = Stream.ParseToken([&](TCHAR InC)
+		{
+			if (InC == QuoteChar)
+			{
+				if (NumConsecutiveSlashes%2 == 0)
+				{
+					// Stop consuming, excluding this closing quote
+					return EParseState::StopBefore;
+				}
+			}
+			else if (InC == '\\')
+			{
+				NumConsecutiveSlashes++;
+			}
+			else
+			{
+				NumConsecutiveSlashes = 0;
+			}
+			
+			return EParseState::Continue;
+		}, &TextTokenWithQuotes);
+
+		if (!Stream.ParseSymbol(QuoteChar, &TextTokenWithQuotes).IsSet())
+		{
+			return FExpressionError(NSLOCTEXT("TextFilterExpressionParser", "Error_UnboundedString", "Reached the end of the expression before a closing quote was found."));
+		}
+		else if (TextToken.IsSet())
+		{
+			const TCHAR EscapedQuote[] = { '\\', QuoteChar, 0 };
+			const TCHAR UnescapedQuote[] = { QuoteChar, 0 };
+
+			// Unescape any literal quotes within the string
+			FString FinalString = TextToken.GetValue().GetString();
+			FinalString.ReplaceInline(EscapedQuote, UnescapedQuote);
+			FinalString.ReplaceInline(TEXT("\\\\"), TEXT("\\"));
+			Consumer.Add(TextTokenWithQuotes, FExpressionNode(FTextToken(MoveTemp(FinalString), ETextFilterTextComparisonMode::Exact)));
+		}
+		else
+		{
+			Consumer.Add(TextTokenWithQuotes, FExpressionNode(FTextToken(FString(), ETextFilterTextComparisonMode::Exact)));
+		}
+
+		return TOptional<FExpressionError>();
+	}
+
+	/** Consume the text from the specified consumer's stream */
+	TOptional<FExpressionError> ConsumeTextImpl(FExpressionTokenConsumer& Consumer, const TFunctionRef<bool(TCHAR)> IsBreakingCharacter)
+	{
+		auto& Stream = Consumer.GetStream();
+
+		TOptional<FStringToken> TextToken = Stream.ParseToken([&IsBreakingCharacter](TCHAR InC)
+		{
+			// This is a non-quoted string, consume until we hit a breaking character
+			return (IsBreakingCharacter(InC)) ? EParseState::StopBefore : EParseState::Continue;
+		});
+
+		if (TextToken.IsSet())
+		{
+			Consumer.Add(TextToken.GetValue(), FExpressionNode(FTextToken(TextToken.GetValue().GetString(), ETextFilterTextComparisonMode::Partial)));
+		}
+
+		return TOptional<FExpressionError>();
+	}
+
+	/** Consume the basic text from the specified consumer's stream */
+	TOptional<FExpressionError> ConsumeBasicText(FExpressionTokenConsumer& Consumer)
+	{
+		return ConsumeTextImpl(Consumer, [](TCHAR InC) -> bool
+		{
+			for (const TCHAR BreakingCharacter : BasicTextBreakingCharacters)
+			{
+				if (InC == BreakingCharacter)
+				{
+					return true;
+				}
+			}
+			return false;
+		});
+	}
+
+	/** Consume the basic text from the specified consumer's stream */
+	TOptional<FExpressionError> ConsumeComplexText(FExpressionTokenConsumer& Consumer)
+	{
+		return ConsumeTextImpl(Consumer, [](TCHAR InC) -> bool
+		{
+			for (const TCHAR BreakingCharacter : ComplexTextBreakingCharacters)
+			{
+				if (InC == BreakingCharacter)
+				{
+					return true;
+				}
+			}
+			return false;
+		});
+	}
+}
+
+#undef DEFINE_TEXT_EXPRESSION_OPERATOR_NODE
+
+FTextFilterExpressionEvaluator::FTextFilterExpressionEvaluator(const ETextFilterExpressionEvaluatorMode InMode)
+	: ExpressionEvaluatorMode(InMode)
+	, FilterType(ETextFilterExpressionType::Empty)
+{
+	ConstructExpressionParser();
+}
+
+FTextFilterExpressionEvaluator::FTextFilterExpressionEvaluator(const FTextFilterExpressionEvaluator& Other)
+	: ExpressionEvaluatorMode(Other.ExpressionEvaluatorMode)
+	, FilterType(ETextFilterExpressionType::Empty)
+{
+	ConstructExpressionParser();
+	SetFilterText(Other.FilterText);
+}
+
+FTextFilterExpressionEvaluator& FTextFilterExpressionEvaluator::operator=(const FTextFilterExpressionEvaluator& Other)
+{
+	FilterType = ETextFilterExpressionType::Empty;
+	FilterText = FText::GetEmpty();
+	FilterErrorText = FText::GetEmpty();
+	CompiledFilter.Reset();
+
+	// If our source object has a different evaluation mode, we need to re-construct the expression parser
+	if (ExpressionEvaluatorMode != Other.ExpressionEvaluatorMode)
+	{
+		ExpressionEvaluatorMode = Other.ExpressionEvaluatorMode;
+
+		TokenDefinitions = FTokenDefinitions();
+		Grammar = FExpressionGrammar();
+		JumpTable = TOperatorJumpTable<ITextFilterExpressionContext>();
+
+		ConstructExpressionParser();
+	}
+
+	// Re-compile the filter from the other object
+	SetFilterText(Other.FilterText);
+
+	return *this;
+}
+
+ETextFilterExpressionType FTextFilterExpressionEvaluator::GetFilterType() const
+{
+	return FilterType;
+}
+
+FText FTextFilterExpressionEvaluator::GetFilterText() const
+{
+	return FilterText;
+}
+
+bool FTextFilterExpressionEvaluator::SetFilterText(const FText& InFilterText)
+{
+	using namespace TextFilterExpressionParser;
+	using TextFilterExpressionParser::FTextToken;
+
+	// Sanitize the new string
+	const FText NewText = FText::TrimPrecedingAndTrailing(InFilterText);
+
+	// Has anything changed?
+	if (NewText.EqualTo(FilterText))
+	{
+		return false;
+	}
+
+	FilterType = ETextFilterExpressionType::Invalid;
+	FilterText = NewText;
+
+	if (FilterText.IsEmpty())
+	{
+		// If the filter is empty, just clear out the compiled filter so we match everything
+		FilterType = ETextFilterExpressionType::Empty;
+		FilterErrorText = FText::GetEmpty();
+		CompiledFilter.Reset();
+	}
+	else
+	{
+		// Otherwise, re-parse the filter text into a compiled expression
+		auto LexResult = ExpressionParser::Lex(*NewText.ToString(), TokenDefinitions);
+		if (LexResult.IsValid())
+		{
+			auto& TmpLexTokens = LexResult.GetValue();
+
+			TArray<FExpressionToken> FinalTokens;
+			FinalTokens.Reserve(TmpLexTokens.Num());
+
+			bool bIsComplexExpression = false;
+
+			ETextNodeType LastTokenType = ETextNodeType::Other;
+			for (FExpressionToken& TmpLexToken : TmpLexTokens)
+			{
+				const ETextNodeType CurrentTokenType = GetTextNodeType(TmpLexToken.Node);
+
+				// We need to inject an OR operator between any consecutive text tokens
+				// This emulates the old behavior of the search filter, where a space between terms implied an OR operator
+				// We also do this between sub-expressions that aren't separated by an operator to ensure that we don't get multiple root nodes in our expression
+				if ((LastTokenType == ETextNodeType::Text || LastTokenType == ETextNodeType::SubExpressionEnd) &&
+					(CurrentTokenType == ETextNodeType::Text || CurrentTokenType == ETextNodeType::SubExpressionStart)
+					)
+				{
+					FinalTokens.Add(FExpressionToken(TmpLexToken.Context, FOr()));
+				}
+
+				// We process TextCmpAnchor tokens during post-compilation, as they only ever operate on constants so don't need to be constantly re-evaluated for each item
+				if (LastTokenType == ETextNodeType::TextCmpAnchor && CurrentTokenType == ETextNodeType::Text)
+				{
+					// Pre-unary TextCmpAnchor - modify the current text token to performs an "ends with" string compare
+					const FTextToken* TextToken = TmpLexToken.Node.Cast<FTextToken>();
+					check(TextToken);
+					const_cast<FTextToken*>(TextToken)->SetTextComparisonMode(ETextFilterTextComparisonMode::EndsWith);
+				}
+				else if (LastTokenType == ETextNodeType::Text && CurrentTokenType == ETextNodeType::TextCmpAnchor)
+				{
+					// Post-unary TextCmpAnchor - modify the previous text token to perform a "starts with" string compare
+					const FTextToken* TextToken = FinalTokens.Last().Node.Cast<FTextToken>();
+					check(TextToken);
+					const_cast<FTextToken*>(TextToken)->SetTextComparisonMode(ETextFilterTextComparisonMode::StartsWith);
+				}
+
+				// Key->value pairs require some form of comparison operator, so if we've found once of those, consider this a complex expression
+				if (!bIsComplexExpression && IsComplexNode(TmpLexToken.Node))
+				{
+					bIsComplexExpression = true;
+				}
+
+				// We strip out the TextCmpAnchor tokens, as we process them as part of this loop
+				if (CurrentTokenType != ETextNodeType::TextCmpAnchor)
+				{
+					FinalTokens.Add(MoveTemp(TmpLexToken));
+				}
+
+				LastTokenType = CurrentTokenType;
+			}
+
+			CompiledFilter = ExpressionParser::Compile(MoveTemp(FinalTokens), Grammar);
+
+			auto& CompiledResult = CompiledFilter.GetValue();
+			if (CompiledResult.IsValid())
+			{
+				FilterType = (bIsComplexExpression) ? ETextFilterExpressionType::Complex : ETextFilterExpressionType::BasicString;
+				FilterErrorText = FText::GetEmpty();
+			}
+			else
+			{
+				FilterErrorText = CompiledResult.GetError().Text;
+			}
+		}
+		else
+		{
+			FilterErrorText = LexResult.GetError().Text;
+		}
+	}
+
+	return true;
+}
+
+FText FTextFilterExpressionEvaluator::GetFilterErrorText() const
+{
+	return FilterErrorText;
+}
+
+bool FTextFilterExpressionEvaluator::TestTextFilter(const ITextFilterExpressionContext& InContext) const
+{
+	using namespace TextFilterExpressionParser;
+	using TextFilterExpressionParser::FTextToken;
+
+	if (CompiledFilter.IsSet())
+	{
+		auto& CompiledResult = CompiledFilter.GetValue();
+		if (CompiledResult.IsValid())
+		{
+			auto EvalResult = ExpressionParser::Evaluate(CompiledResult.GetValue(), JumpTable, &InContext);
+			if (EvalResult.IsValid())
+			{
+				if (const bool* BoolResult = EvalResult.GetValue().Cast<bool>())
+				{
+					return *BoolResult;
+				}
+				else if (const FTextToken* TextResult = EvalResult.GetValue().Cast<FTextToken>())
+				{
+					return InContext.TestBasicStringExpression(TextResult->GetString(), TextResult->GetTextComparisonMode());
+				}
+			}
+		}
+	}
+
+	return false;
+}
+
+void FTextFilterExpressionEvaluator::ConstructExpressionParser()
+{
+	using namespace TextFilterExpressionParser;
+	using TextFilterExpressionParser::FTextToken;
+
+	TokenDefinitions.IgnoreWhitespace();
+	TokenDefinitions.DefineToken(&ConsumeOperator<FSubExpressionStart>);
+	TokenDefinitions.DefineToken(&ConsumeOperator<FSubExpressionEnd>);
+	if (ExpressionEvaluatorMode == ETextFilterExpressionEvaluatorMode::Complex)
+	{
+		// Comparisons are only available for complex expressions
+		TokenDefinitions.DefineToken(&ConsumeOperator<FLessOrEqual>);
+		TokenDefinitions.DefineToken(&ConsumeOperator<FLess>);
+		TokenDefinitions.DefineToken(&ConsumeOperator<FGreaterOrEqual>);
+		TokenDefinitions.DefineToken(&ConsumeOperator<FGreater>);
+		TokenDefinitions.DefineToken(&ConsumeOperator<FNotEqual>);
+		TokenDefinitions.DefineToken(&ConsumeOperator<FEqual>);
+	}
+	TokenDefinitions.DefineToken(&ConsumeOperator<FOr>);
+	TokenDefinitions.DefineToken(&ConsumeOperator<FAnd>);
+	TokenDefinitions.DefineToken(&ConsumeOperator<FNot>);
+	TokenDefinitions.DefineToken(&ConsumeOperator<FTextCmpAnchor>);
+	TokenDefinitions.DefineToken(&ConsumeNumber);
+	TokenDefinitions.DefineToken(&ConsumeQuotedText);
+	TokenDefinitions.DefineToken((ExpressionEvaluatorMode == ETextFilterExpressionEvaluatorMode::Complex) ? &ConsumeComplexText : &ConsumeBasicText);
+
+	Grammar.DefineGrouping<FSubExpressionStart, FSubExpressionEnd>();
+	Grammar.DefineBinaryOperator<FLessOrEqual>(1);
+	Grammar.DefineBinaryOperator<FLess>(1);
+	Grammar.DefineBinaryOperator<FGreaterOrEqual>(1);
+	Grammar.DefineBinaryOperator<FGreater>(1);
+	Grammar.DefineBinaryOperator<FNotEqual>(1);
+	Grammar.DefineBinaryOperator<FEqual>(1);
+	Grammar.DefineBinaryOperator<FOr>(2);
+	Grammar.DefineBinaryOperator<FAnd>(2);
+	Grammar.DefinePreUnaryOperator<FNot>();
+	Grammar.DefinePreUnaryOperator<FTextCmpAnchor>();
+	Grammar.DefinePostUnaryOperator<FTextCmpAnchor>();
+
+	JumpTable.MapBinary<FLessOrEqual>([](const FTextToken& A, const FTextToken& B, const ITextFilterExpressionContext* InContext)		{ return InContext->TestComplexExpression(FName(*A.GetString()), B.GetString(), ETextFilterComparisonOperation::LessOrEqual, B.GetTextComparisonMode()); });
+	JumpTable.MapBinary<FLess>([](const FTextToken& A, const FTextToken& B, const ITextFilterExpressionContext* InContext)				{ return InContext->TestComplexExpression(FName(*A.GetString()), B.GetString(), ETextFilterComparisonOperation::Less, B.GetTextComparisonMode()); });
+	JumpTable.MapBinary<FGreaterOrEqual>([](const FTextToken& A, const FTextToken& B, const ITextFilterExpressionContext* InContext)	{ return InContext->TestComplexExpression(FName(*A.GetString()), B.GetString(), ETextFilterComparisonOperation::GreaterOrEqual, B.GetTextComparisonMode()); });
+	JumpTable.MapBinary<FGreater>([](const FTextToken& A, const FTextToken& B, const ITextFilterExpressionContext* InContext)			{ return InContext->TestComplexExpression(FName(*A.GetString()), B.GetString(), ETextFilterComparisonOperation::Greater, B.GetTextComparisonMode()); });
+	JumpTable.MapBinary<FNotEqual>([](const FTextToken& A, const FTextToken& B, const ITextFilterExpressionContext* InContext)			{ return InContext->TestComplexExpression(FName(*A.GetString()), B.GetString(), ETextFilterComparisonOperation::NotEqual, B.GetTextComparisonMode()); });
+	JumpTable.MapBinary<FEqual>([](const FTextToken& A, const FTextToken& B, const ITextFilterExpressionContext* InContext)				{ return InContext->TestComplexExpression(FName(*A.GetString()), B.GetString(), ETextFilterComparisonOperation::Equal, B.GetTextComparisonMode()); });
+
+	JumpTable.MapBinary<FOr>([](const FTextToken& A, const FTextToken& B, const ITextFilterExpressionContext* InContext)				{ return InContext->TestBasicStringExpression(A.GetString(), A.GetTextComparisonMode()) || InContext->TestBasicStringExpression(B.GetString(), B.GetTextComparisonMode()); });
+	JumpTable.MapBinary<FOr>([](const FTextToken& A, bool B, const ITextFilterExpressionContext* InContext)								{ return InContext->TestBasicStringExpression(A.GetString(), A.GetTextComparisonMode()) || B; });
+	JumpTable.MapBinary<FOr>([](bool A, const FTextToken& B, const ITextFilterExpressionContext* InContext)								{ return A || InContext->TestBasicStringExpression(B.GetString(), B.GetTextComparisonMode()); });
+	JumpTable.MapBinary<FOr>([](bool A, bool B, const ITextFilterExpressionContext* InContext)											{ return A || B; });
+
+	JumpTable.MapBinary<FAnd>([](const FTextToken& A, const FTextToken& B, const ITextFilterExpressionContext* InContext)				{ return InContext->TestBasicStringExpression(A.GetString(), A.GetTextComparisonMode()) && InContext->TestBasicStringExpression(B.GetString(), B.GetTextComparisonMode()); });
+	JumpTable.MapBinary<FAnd>([](const FTextToken& A, bool B, const ITextFilterExpressionContext* InContext)							{ return InContext->TestBasicStringExpression(A.GetString(), A.GetTextComparisonMode()) && B; });
+	JumpTable.MapBinary<FAnd>([](bool A, const FTextToken& B, const ITextFilterExpressionContext* InContext)							{ return A && InContext->TestBasicStringExpression(B.GetString(), B.GetTextComparisonMode()); });
+	JumpTable.MapBinary<FAnd>([](bool A, bool B, const ITextFilterExpressionContext* InContext)											{ return A && B; });
+
+	JumpTable.MapPreUnary<FNot>([](const FTextToken& V, const ITextFilterExpressionContext* InContext)									{ return !InContext->TestBasicStringExpression(V.GetString(), V.GetTextComparisonMode()); });
+	JumpTable.MapPreUnary<FNot>([](bool V, const ITextFilterExpressionContext* InContext)												{ return !V; });
+}
diff --git a/Engine/Source/Runtime/Core/Private/Misc/TextFilterTests.cpp b/Engine/Source/Runtime/Core/Private/Misc/TextFilterTests.cpp
new file mode 100644
--- /dev/null
+++ b/Engine/Source/Runtime/Core/Private/Misc/TextFilterTests.cpp
@@ -0,0 +1,244 @@
+// Copyright 1998-2015 Epic Games, Inc. All Rights Reserved.
+
+#include "CorePrivatePCH.h"
+#include "TextFilter.h"
+
+#if !(UE_BUILD_SHIPPING || UE_BUILD_TEST)
+
+namespace TextFilterTests
+{
+	struct FTestFilterItem : public ITextFilterExpressionContext
+	{
+		FTestFilterItem()
+		{
+			BasicStrings.Add(TEXT("Wooble"));
+			BasicStrings.Add(TEXT("Flibble"));
+			BasicStrings.Add(TEXT("Funky<String>"));
+
+			KeyValuePairs.Add("StringKey", TEXT("Test"));
+			KeyValuePairs.Add("IntKey", TEXT("123"));
+			KeyValuePairs.Add("FloatKey", TEXT("456.789"));
+		}
+
+		// ITextFilterExpressionContext API - used for testing FTextFilterExpressionEvaluator
+		virtual bool TestBasicStringExpression(const FString& InValue, const ETextFilterTextComparisonMode InTextComparisonMode) const override
+		{
+			for (const FString& BasicString : BasicStrings)
+			{
+				if (TextFilterUtils::TestBasicStringExpression(BasicString, InValue, InTextComparisonMode))
+				{
+					return true;
+				}
+			}
+			return false;
+		}
+		virtual bool TestComplexExpression(const FName& InKey, const FString& InValue, const ETextFilterComparisonOperation InComparisonOperation, const ETextFilterTextComparisonMode InTextComparisonMode) const override
+		{
+			const FString* ItemValue = KeyValuePairs.Find(InKey);
+			if (ItemValue)
+			{
+				return TextFilterUtils::TestComplexExpression(InValue, *ItemValue, InComparisonOperation, InTextComparisonMode);
+			}
+			return false;
+		}
+
+		// TTextFilter API - used for testing TTextFilter
+		static void ExtractItemStrings(const FTestFilterItem* InItem, TArray<FString>& OutStrings)
+		{
+			OutStrings = InItem->BasicStrings;
+		}
+		static bool TestItemComplexExpression(const FTestFilterItem* InItem, const FName& InKey, const FString& InValue, ETextFilterComparisonOperation InComparisonOperation, ETextFilterTextComparisonMode InTextComparisonMode)
+		{
+			return InItem->TestComplexExpression(InKey, InValue, InComparisonOperation, InTextComparisonMode);
+		}
+
+		TArray<FString> BasicStrings;
+		TMap<FName, FString> KeyValuePairs;
+	};
+
+	struct ITestFilterExpression
+	{
+		virtual ~ITestFilterExpression() {}
+		virtual bool TestFilterExpression(const TCHAR* InFilterExpression, const bool bExpected) = 0;
+	};
+
+	struct FTestFilterExpression : public ITestFilterExpression
+	{
+		FTestFilterExpression(FAutomationTestBase* InTest, const FTestFilterItem* InTestItem)
+			: Test(InTest)
+			, TestItem(InTestItem)
+		{
+		}
+
+		virtual bool TestFilterExpression(const TCHAR* InFilterExpression, const bool bExpected) override
+		{
+			FText FilterErrorText;
+			const bool bActual = TestFilterExpressionImpl(InFilterExpression, TestItem, FilterErrorText);
+
+			bool bResult = true;
+
+			if (!FilterErrorText.IsEmpty())
+			{
+				Test->AddError(FString::Printf(TEXT("Filter expression '%s' reported an error: %s"), InFilterExpression, *FilterErrorText.ToString()));
+				bResult = false;
+			}
+
+			if (bActual != bExpected)
+			{
+				Test->AddError(FString::Printf(TEXT("Filter expression '%s' evaluated incorrectly: %d != %d"), InFilterExpression, bActual, bExpected));
+				bResult = false;
+			}
+
+			return bResult;
+		}
+
+		virtual bool TestFilterExpressionImpl(const TCHAR* InFilterExpression, const FTestFilterItem* InTestItem, FText& OutErrorText) = 0;
+
+		FAutomationTestBase* Test;
+		const FTestFilterItem* TestItem;
+	};
+
+	struct FTestFilterExpression_TextFilterExpressionEvaluator : public FTestFilterExpression
+	{
+		FTestFilterExpression_TextFilterExpressionEvaluator(FAutomationTestBase* InTest, const FTestFilterItem* InTestItem, FTextFilterExpressionEvaluator* InTextFilterExpressionEvaluator)
+			: FTestFilterExpression(InTest, InTestItem)
+			, TextFilterExpressionEvaluator(InTextFilterExpressionEvaluator)
+		{
+		}
+
+		virtual bool TestFilterExpressionImpl(const TCHAR* InFilterExpression, const FTestFilterItem* InTestItem, FText& OutErrorText) override
+		{
+			TextFilterExpressionEvaluator->SetFilterText(FText::FromString(InFilterExpression));
+			OutErrorText = TextFilterExpressionEvaluator->GetFilterErrorText();
+			return TextFilterExpressionEvaluator->TestTextFilter(*InTestItem);
+		}
+
+		FTextFilterExpressionEvaluator* TextFilterExpressionEvaluator;
+	};
+
+	struct FTestFilterExpression_TextFilter : public FTestFilterExpression
+	{
+		FTestFilterExpression_TextFilter(FAutomationTestBase* InTest, const FTestFilterItem* InTestItem, TTextFilter<const FTestFilterItem*>* InTextFilter)
+			: FTestFilterExpression(InTest, InTestItem)
+			, TextFilter(InTextFilter)
+		{
+		}
+
+		virtual bool TestFilterExpressionImpl(const TCHAR* InFilterExpression, const FTestFilterItem* InTestItem, FText& OutErrorText) override
+		{
+			TextFilter->SetRawFilterText(FText::FromString(InFilterExpression));
+			OutErrorText = TextFilter->GetFilterErrorText();
+			return TextFilter->PassesFilter(InTestItem);
+		}
+
+		TTextFilter<const FTestFilterItem*>* TextFilter;
+	};
+
+	bool TestAllCommonFilterExpressions(ITestFilterExpression& InTestPayload)
+	{
+		bool bResult = true;
+
+		bResult &= InTestPayload.TestFilterExpression(TEXT("Wooble"), true);
+		bResult &= InTestPayload.TestFilterExpression(TEXT("Woo..."), true);
+		bResult &= InTestPayload.TestFilterExpression(TEXT("...ble"), true);
+		bResult &= InTestPayload.TestFilterExpression(TEXT("Wo... AND ...le"), true);
+		bResult &= InTestPayload.TestFilterExpression(TEXT("Wooble2"), false);
+		bResult &= InTestPayload.TestFilterExpression(TEXT("Wooble2 Flibble"), true);
+		bResult &= InTestPayload.TestFilterExpression(TEXT("Wooble2 OR Flibble"), true);
+		bResult &= InTestPayload.TestFilterExpression(TEXT("Wooble2 AND Flibble"), false);
+		bResult &= InTestPayload.TestFilterExpression(TEXT("Wooble && !Flibble"), false);
+		bResult &= InTestPayload.TestFilterExpression(TEXT("\"Funky<String>\""), true);
+
+		return bResult;
+	}
+
+	bool TestAllBasicFilterExpressions(ITestFilterExpression& InTestPayload)
+	{
+		bool bResult = true;
+
+		bResult &= TestAllCommonFilterExpressions(InTestPayload);
+		bResult &= InTestPayload.TestFilterExpression(TEXT("Funky<String>"), true);
+
+		return bResult;
+	}
+
+	bool TestAllComplexFilterExpressions(ITestFilterExpression& InTestPayload)
+	{
+		bool bResult = true;
+
+		bResult &= TestAllCommonFilterExpressions(InTestPayload);
+		bResult &= InTestPayload.TestFilterExpression(TEXT("Funky<String>"), false);
+		bResult &= InTestPayload.TestFilterExpression(TEXT("StringKey=Test"), true);
+		bResult &= InTestPayload.TestFilterExpression(TEXT("StringKey!=Test"), false);
+		bResult &= InTestPayload.TestFilterExpression(TEXT("IntKey=123"), true);
+		bResult &= InTestPayload.TestFilterExpression(TEXT("IntKey>123"), false);
+		bResult &= InTestPayload.TestFilterExpression(TEXT("IntKey<123"), false);
+		bResult &= InTestPayload.TestFilterExpression(TEXT("FloatKey=456.789"), true);
+		bResult &= InTestPayload.TestFilterExpression(TEXT("FloatKey>456.789"), false);
+		bResult &= InTestPayload.TestFilterExpression(TEXT("FloatKey<456.789"), false);
+		bResult &= InTestPayload.TestFilterExpression(TEXT("IntKey==300 || FloatKey==456.789"), true);
+		bResult &= InTestPayload.TestFilterExpression(TEXT("IntKey==300 && FloatKey==456.789"), false);
+		bResult &= InTestPayload.TestFilterExpression(TEXT("(IntKey==300 && FloatKey==456.789) OR StringKey==Test"), true);
+
+		return bResult;
+	}
+};
+
+IMPLEMENT_SIMPLE_AUTOMATION_TEST(FTextFilterTests_TextFilterExpressionEvaluator, "System.Core.Text.TextFilterExpressionEvaluator", EAutomationTestFlags::ATF_ApplicationMask)
+bool FTextFilterTests_TextFilterExpressionEvaluator::RunTest(const FString& Parameters)
+{
+	using namespace TextFilterTests;
+
+	FTestFilterItem TestItem;
+
+	bool bResult = true;
+
+	// Test basic filtering
+	{
+		FTextFilterExpressionEvaluator TextFilterExpressionEvaluator(ETextFilterExpressionEvaluatorMode::BasicString);
+		FTestFilterExpression_TextFilterExpressionEvaluator TestFilterExpressionContext(this, &TestItem, &TextFilterExpressionEvaluator);
+		bResult &= TestAllBasicFilterExpressions(TestFilterExpressionContext);
+	}
+
+	// Test complex filtering
+	{
+		FTextFilterExpressionEvaluator TextFilterExpressionEvaluator(ETextFilterExpressionEvaluatorMode::Complex);
+		FTestFilterExpression_TextFilterExpressionEvaluator TestFilterExpressionContext(this, &TestItem, &TextFilterExpressionEvaluator);
+		bResult &= TestAllComplexFilterExpressions(TestFilterExpressionContext);
+	}
+
+	return bResult;
+}
+
+IMPLEMENT_SIMPLE_AUTOMATION_TEST(FTextFilterTests_TextFilter, "System.Core.Text.TextFilter", EAutomationTestFlags::ATF_ApplicationMask)
+bool FTextFilterTests_TextFilter::RunTest(const FString& Parameters)
+{
+	using namespace TextFilterTests;
+
+	FTestFilterItem TestItem;
+
+	bool bResult = true;
+
+	// Test basic filtering
+	{
+		TTextFilter<const FTestFilterItem*> TextFilter(
+			TTextFilter<const FTestFilterItem*>::FItemToStringArray::CreateStatic(&FTestFilterItem::ExtractItemStrings)
+			);
+		FTestFilterExpression_TextFilter TestFilterExpressionContext(this, &TestItem, &TextFilter);
+		bResult &= TestAllBasicFilterExpressions(TestFilterExpressionContext);
+	}
+
+	// Test complex filtering
+	{
+		TTextFilter<const FTestFilterItem*> TextFilter(
+			TTextFilter<const FTestFilterItem*>::FItemToStringArray::CreateStatic(&FTestFilterItem::ExtractItemStrings), 
+			TTextFilter<const FTestFilterItem*>::FItemTestComplexExpression::CreateStatic(&FTestFilterItem::TestItemComplexExpression)
+			);
+		FTestFilterExpression_TextFilter TestFilterExpressionContext(this, &TestItem, &TextFilter);
+		bResult &= TestAllComplexFilterExpressions(TestFilterExpressionContext);
+	}
+
+	return bResult;
+}
+
+#endif //!(UE_BUILD_SHIPPING || UE_BUILD_TEST)
diff --git a/Engine/Source/Runtime/Core/Private/Misc/TextFilterUtils.cpp b/Engine/Source/Runtime/Core/Private/Misc/TextFilterUtils.cpp
new file mode 100644
--- /dev/null
+++ b/Engine/Source/Runtime/Core/Private/Misc/TextFilterUtils.cpp
@@ -0,0 +1,69 @@
+// Copyright 1998-2015 Epic Games, Inc. All Rights Reserved.
+
+#include "CorePrivatePCH.h"
+#include "TextFilterUtils.h"
+
+bool TextFilterUtils::TestBasicStringExpression(const FString& InValue1, const FString& InValue2, const ETextFilterTextComparisonMode InTextComparisonMode)
+{
+	switch(InTextComparisonMode)
+	{
+	case ETextFilterTextComparisonMode::Exact:
+		return InValue1.Equals(InValue2, ESearchCase::IgnoreCase);
+	case ETextFilterTextComparisonMode::Partial:
+		return InValue1.Contains(InValue2, ESearchCase::IgnoreCase);
+	case ETextFilterTextComparisonMode::StartsWith:
+		return InValue1.StartsWith(InValue2, ESearchCase::IgnoreCase);
+	case ETextFilterTextComparisonMode::EndsWith:
+		return InValue1.EndsWith(InValue2, ESearchCase::IgnoreCase);
+	default:
+		break;
+	}
+	return false;
+}
+
+bool TextFilterUtils::TestComplexExpression(const FString& InValue1, const FString& InValue2, const ETextFilterComparisonOperation InComparisonOperation, const ETextFilterTextComparisonMode InTextComparisonMode)
+{
+	const bool bNumericComparison = InValue1.IsNumeric() && InValue2.IsNumeric();
+
+	int32 ComparisonSign = 0;
+	if (bNumericComparison)
+	{
+		const double Difference = FCString::Atod(*InValue1) - FCString::Atod(*InValue2);
+		ComparisonSign = (int32)FMath::Sign(Difference);
+	}
+	else if (InTextComparisonMode == ETextFilterTextComparisonMode::Exact)
+	{
+		ComparisonSign = InValue1.Compare(InValue2, ESearchCase::IgnoreCase);
+	}
+	else
+	{
+		ComparisonSign = (InValue1.Contains(InValue2, ESearchCase::IgnoreCase)) ? 0 : 1;
+	}
+
+	bool bComparisonPassed = false;
+	switch (InComparisonOperation)
+	{
+	case ETextFilterComparisonOperation::Equal:
+		bComparisonPassed = ComparisonSign == 0;
+		break;
+	case ETextFilterComparisonOperation::NotEqual:
+		bComparisonPassed = ComparisonSign != 0;
+		break;
+	case ETextFilterComparisonOperation::Less:
+		bComparisonPassed = ComparisonSign < 0;
+		break;
+	case ETextFilterComparisonOperation::LessOrEqual:
+		bComparisonPassed = ComparisonSign <= 0;
+		break;
+	case ETextFilterComparisonOperation::Greater:
+		bComparisonPassed = ComparisonSign > 0;
+		break;
+	case ETextFilterComparisonOperation::GreaterOrEqual:
+		bComparisonPassed = ComparisonSign >= 0;
+		break;
+	default:
+		check(false);
+	};
+
+	return bComparisonPassed;
+}
diff --git a/Engine/Source/Runtime/Core/Public/Misc/TextFilter.h b/Engine/Source/Runtime/Core/Public/Misc/TextFilter.h
--- a/Engine/Source/Runtime/Core/Public/Misc/TextFilter.h
+++ b/Engine/Source/Runtime/Core/Public/Misc/TextFilter.h
@@ -2,16 +2,7 @@
 #pragma once
 
 #include "IFilter.h"
-
-enum class ETextFilterComparisonOperation : uint8
-{
-	Equal,
-	NotEqual,
-	Less,
-	LessOrEqual,
-	Greater,
-	GreaterOrEqual,
-};
+#include "TextFilterExpressionEvaluator.h"
 
 /**
  *	A generic filter specialized for text restrictions
@@ -24,15 +15,27 @@ public:
 	/**	Defines a function signature for functions used to transform an Item into an array of FStrings */
 	DECLARE_DELEGATE_TwoParams( FItemToStringArray, ItemType, OUT TArray< FString >& );
 
+	/** Defines a function signature used to test a complex expression for an Item */
+	DECLARE_DELEGATE_RetVal_FiveParams( bool, FItemTestComplexExpression, ItemType, const FName& /*InKey*/, const FString& /*InValue*/, ETextFilterComparisonOperation /*InComparisonOperation*/, ETextFilterTextComparisonMode /*InTextComparisonMode*/ );
+
 	/** 
 	 *	TTextFilter Constructor
 	 *	
 	 * @param	InTransform		A required delegate used to populate an array of strings based on an Item
 	 */
-	TTextFilter( FItemToStringArray InTransform ) 
-		: TransformArray( InTransform )
+	TTextFilter( FItemToStringArray InTransformDelegate )
+		: TextFilterExpressionContext( InTransformDelegate, FItemTestComplexExpression() )
+		, TextFilterExpressionEvaluator( ETextFilterExpressionEvaluatorMode::BasicString )
+	{
+		check( InTransformDelegate.IsBound() );
+	}
+
+	TTextFilter( FItemToStringArray InTransformDelegate, FItemTestComplexExpression InTestComplexExpressionDelegate )
+		: TextFilterExpressionContext( InTransformDelegate, InTestComplexExpressionDelegate )
+		, TextFilterExpressionEvaluator( ETextFilterExpressionEvaluatorMode::Complex )
 	{
-		check( TransformArray.IsBound() );
+		check( InTransformDelegate.IsBound() );
+		check( InTestComplexExpressionDelegate.IsBound() );
 	}
 
 	DECLARE_DERIVED_EVENT(TTextFilter, IFilter<ItemType>::FChangedEvent, FChangedEvent);
@@ -46,195 +49,104 @@ public:
 	 */
 	virtual bool PassesFilter( ItemType InItem ) const override
 	{
-		if( FilterSearchTerms.Num() == 0 )
+		if (TextFilterExpressionEvaluator.GetFilterType() == ETextFilterExpressionType::Empty)
 		{
 			return true;
 		}
 		
-		ItemTextCollection.Empty();
-		TransformArray.Execute( InItem, OUT ItemTextCollection );
-
-		bool bPassesAnyTermSets = false;
-		for( auto TermSetIndex = 0; !bPassesAnyTermSets && TermSetIndex < FilterSearchTerms.Num(); ++TermSetIndex )
-		{
-			bool bPassesFilter = true;
-			for( auto CurSearchTermIndex = 0; bPassesFilter && CurSearchTermIndex < FilterSearchTerms[TermSetIndex].Num(); ++CurSearchTermIndex )
-			{
-				const FString* CurSearchTerm = &FilterSearchTerms[ TermSetIndex ][ CurSearchTermIndex ];
-				FString TempSearchString;
-
-				if( CurSearchTerm->IsEmpty() ) 
-				{
-					continue;
-				}
-
-				// If the search text starts with a plus, then we'll force an exact match search (whole word)
-				const auto bExactMatch = CurSearchTerm->GetCharArray()[ 0 ] == TCHAR( '+' );
-				if( bExactMatch )
-				{
-					// Strip off the leading plus
-					TempSearchString = CurSearchTerm->Right( CurSearchTerm->Len() - 1 );
-					CurSearchTerm = &TempSearchString;	// Storing pointer to stack variable to avoid copies (safe here)
-				}
-
-				if( CurSearchTerm->IsEmpty() )
-				{
-					continue;
-				}
-
-				// If the search text starts with a hyphen, then we'll explicitly exclude results using this search term
-				const auto bExclude = CurSearchTerm->GetCharArray()[ 0 ] == TCHAR( '-' );
-				if( bExclude )
-				{
-					//Strip off the leading hyphen
-					TempSearchString = CurSearchTerm->Right( CurSearchTerm->Len() - 1 );
-					CurSearchTerm = &TempSearchString;	// Storing pointer to stack variable to avoid copies (safe here)
-				}
-
-				if( CurSearchTerm->IsEmpty() )
-				{
-					continue;
-				}
-
-				bPassesFilter = false;
-				for( int TextIndex = 0; !bPassesFilter && TextIndex < ItemTextCollection.Num(); ++TextIndex )
-				{
-					// Matches?
-					const FString& ItemText = ItemTextCollection[ TextIndex ];
-					bPassesFilter = bExactMatch ? ( ItemText == *CurSearchTerm ) : ( ItemText.Contains( **CurSearchTerm) );
-				}
-
-				// If excluding, then we'll reverse the meaning of our test
-				if( bExclude )
-				{
-					bPassesFilter = !bPassesFilter;
-				}
-			}
-
-			if ( bPassesFilter )
-			{
-				bPassesAnyTermSets = true;
-			}
-		}
-
-		return bPassesAnyTermSets;
+		TextFilterExpressionContext.SetItem(&InItem);
+		const bool bResult = TextFilterExpressionEvaluator.TestTextFilter(TextFilterExpressionContext);
+		TextFilterExpressionContext.ClearItem();
+		return bResult;
 	}
 
 	/** Returns the unsanitized and unsplit filter terms */
 	FText GetRawFilterText() const
 	{
-		return FilterText;
+		return TextFilterExpressionEvaluator.GetFilterText();
 	}
 
 	/** Set the Text to be used as the Filter's restrictions */
 	void SetRawFilterText( const FText& InFilterText )
 	{
-		// Sanitize the new string
-		FText NewText = FText::TrimPrecedingAndTrailing(InFilterText);
-		
-		// If the text differs, re-parse and split the filter text into search terms
-		if ( !NewText.EqualTo(FilterText) )
+		if (TextFilterExpressionEvaluator.SetFilterText(InFilterText))
 		{
-			FilterText = NewText;
-			FilterSearchTerms.Reset();
-			FilterSearchTerms.Add( TArray< FString >() );
-			int32 SearchTermIndex = 0;
-
-			bool bIsWithinQuotedSection = false;
-			FString NewSearchTerm;
-			int32 StringLength = FilterText.ToString().Len();
-			for( auto CurCharIndex = 0; CurCharIndex < StringLength; ++CurCharIndex )
-			{
-				const auto CurChar = FilterText.ToString().GetCharArray()[ CurCharIndex ];
+			ChangedEvent.Broadcast();
+		}
+	}
 
-				// Keep an eye out for double-quotes.  We want to retain whitespace within a search term if
-				// it has double-quotes around it
-				if( CurChar == TCHAR('\"') )
-				{
-					// Toggle whether we're within a quoted section, but don't bother adding the quote character
-					bIsWithinQuotedSection = !bIsWithinQuotedSection;
-				}
-				else if ( CurChar == TCHAR('|') )
-				{
-					if( NewSearchTerm.Len() > 0 )
-					{
-						FilterSearchTerms[SearchTermIndex].Add( NewSearchTerm );
-					}
-
-					++SearchTermIndex;
-					FilterSearchTerms.Add( TArray< FString >() );
-					NewSearchTerm = FString();
-				}
-				else if( bIsWithinQuotedSection || !FChar::IsWhitespace( CurChar ) )
-				{
-					// Add the character!
-					NewSearchTerm.AppendChar( CurChar );
-				}
-				else
-				{
-					// Encountered whitespace, so add the search term up to here
-					if( NewSearchTerm.Len() > 0 )
-					{
-						if ( NewSearchTerm == TEXT("OR") )
-						{
-							++SearchTermIndex;
-							FilterSearchTerms.Add( TArray< FString >() );
-						}
-						else
-						{
-							FilterSearchTerms[SearchTermIndex].Add( NewSearchTerm );
-						}
-
-						NewSearchTerm = FString();
-					}
-				}
-			}
+	/** Get the last error returned from lexing or compiling the current filter text */
+	FText GetFilterErrorText() const
+	{
+		return TextFilterExpressionEvaluator.GetFilterErrorText();
+	}
 
-			// Encountered EOL, so add the search term up to here
-			if( NewSearchTerm.Len() > 0 )
-			{
-				if ( NewSearchTerm == TEXT("OR") )
-				{
-					++SearchTermIndex;
-					FilterSearchTerms.Add( TArray< FString >() );
-				}
-				else
-				{
-					FilterSearchTerms[SearchTermIndex].Add( NewSearchTerm );
-				}
+private:
 
-				NewSearchTerm = FString();
-			}
+	class FTextFilterExpressionContext : public ITextFilterExpressionContext
+	{
+	public:
+		typedef typename TRemoveReference<ItemType>::Type* ItemTypePtr;
+
+		FTextFilterExpressionContext(const FItemToStringArray& InTransformArrayDelegate, const FItemTestComplexExpression& InTestComplexExpressionDelegate)
+			: TransformArrayDelegate(InTransformArrayDelegate)
+			, TestComplexExpressionDelegate(InTestComplexExpressionDelegate)
+			, ItemPtr(nullptr)
+		{
+		}
+
+		void SetItem(ItemTypePtr InItem)
+		{
+			ItemPtr = InItem;
+			TransformArrayDelegate.Execute(*ItemPtr, ItemBasicStrings);
+		}
+
+		void ClearItem()
+		{
+			ItemPtr = nullptr;
+			ItemBasicStrings.Reset();
+		}
 
-			if( bIsWithinQuotedSection )
+		virtual bool TestBasicStringExpression(const FString& InValue, const ETextFilterTextComparisonMode InTextComparisonMode) const override
+		{
+			for (const FString& BasicString : ItemBasicStrings)
 			{
-				// User forgot to terminate their double-quoted section.  No big deal.
+				if (TextFilterUtils::TestBasicStringExpression(BasicString, InValue, InTextComparisonMode))
+				{
+					return true;
+				}
 			}
+			return false;
+		}
 
-			//Remove the last array of search terms if it is empty
-			if ( FilterSearchTerms[SearchTermIndex].Num() == 0 )
+		virtual bool TestComplexExpression(const FName& InKey, const FString& InValue, const ETextFilterComparisonOperation InComparisonOperation, const ETextFilterTextComparisonMode InTextComparisonMode) const override
+		{
+			if (ItemPtr && TestComplexExpressionDelegate.IsBound())
 			{
-				FilterSearchTerms.RemoveAt( SearchTermIndex );
+				return TestComplexExpressionDelegate.Execute(*ItemPtr, InKey, InValue, InComparisonOperation, InTextComparisonMode);
 			}
-
-			ChangedEvent.Broadcast();
+			return false;
 		}
-	}
 
+	private:
+		/** The delegate used to transform an item into an array of strings */
+		FItemToStringArray TransformArrayDelegate;
 
-private:
+		/** The delegate used to test a complex expression for an item */
+		FItemTestComplexExpression TestComplexExpressionDelegate;
 
-	/** The delegate used to transform an item into an array of strings */
-	FItemToStringArray TransformArray;
+		/** Pointer to the item we're currently filtering */
+		ItemTypePtr ItemPtr;
 
-	/** An array containing all the strings returned by the last Transform delegate invoked */
-	mutable TArray< FString > ItemTextCollection;
+		/** The strings extracted from the item we're currently filtering */
+		TArray<FString> ItemBasicStrings;
+	};
 
-	/** Split and sanitized list of filter search terms.  Cached for faster search performance. */
-	TArray< TArray< FString > > FilterSearchTerms;
+	/** Transient context data, used when calling PassesFilter. Kept around to minimize re-allocations between multiple calls to PassesFilter */
+	mutable FTextFilterExpressionContext TextFilterExpressionContext;
 
-	/** The original unsplit and unsanitized filter terms */
-	FText FilterText;
+	/** Expression evaluator that can be used to perform complex text filter queries */
+	FTextFilterExpressionEvaluator TextFilterExpressionEvaluator;
 
 	/**	The event that fires whenever new search terms are provided */
 	FChangedEvent ChangedEvent;
diff --git a/Engine/Source/Runtime/Core/Public/Misc/TextFilterExpressionEvaluator.h b/Engine/Source/Runtime/Core/Public/Misc/TextFilterExpressionEvaluator.h
new file mode 100644
--- /dev/null
+++ b/Engine/Source/Runtime/Core/Public/Misc/TextFilterExpressionEvaluator.h
@@ -0,0 +1,84 @@
+// Copyright 1998-2015 Epic Games, Inc. All Rights Reserved.
+#pragma once
+
+#include "ExpressionParser.h"
+#include "TextFilterUtils.h"
+
+/** Defines the complexity of the current filter terms. Complex in this case means that the expression will perform key->value comparisons */
+enum class ETextFilterExpressionType : uint8
+{
+	Invalid,
+	Empty,
+	BasicString,
+	Complex,
+};
+
+/** Defines whether or not the expression parser can evaluate complex expressions */
+enum class ETextFilterExpressionEvaluatorMode : uint8
+{
+	BasicString,
+	Complex,
+};
+
+/** Interface to implement to allow FTextFilterExpressionEvaluator to perform its comparison tests in TestTextFilter */
+class ITextFilterExpressionContext
+{
+public:
+	/** Test the given value against the strings extracted from the current item */
+	virtual bool TestBasicStringExpression(const FString& InValue, const ETextFilterTextComparisonMode InTextComparisonMode) const = 0;
+		
+	/** Perform a complex expression test for the current item */
+	virtual bool TestComplexExpression(const FName& InKey, const FString& InValue, const ETextFilterComparisonOperation InComparisonOperation, const ETextFilterTextComparisonMode InTextComparisonMode) const = 0;
+
+protected:
+	virtual ~ITextFilterExpressionContext() {}
+};
+
+/** Defines an expression evaluator that can be used to perform complex text filter queries */
+class CORE_API FTextFilterExpressionEvaluator
+{
+public:
+	/** Construction and assignment */
+	FTextFilterExpressionEvaluator(const ETextFilterExpressionEvaluatorMode InMode);
+	FTextFilterExpressionEvaluator(const FTextFilterExpressionEvaluator& Other);
+	FTextFilterExpressionEvaluator& operator=(const FTextFilterExpressionEvaluator& Other);
+
+	/** Get the complexity of the current filter terms */
+	ETextFilterExpressionType GetFilterType() const;
+
+	/** Get the filter terms that we're currently using */
+	FText GetFilterText() const;
+
+	/** Set the filter terms to be compiled for evaluation later. Returns true if the filter was actually changed */
+	bool SetFilterText(const FText& InFilterText);
+
+	/** Get the last error returned from lexing or compiling the current filter text */
+	FText GetFilterErrorText() const;
+
+	/** Test our compiled filter using the given context */
+	bool TestTextFilter(const ITextFilterExpressionContext& InContext) const;
+
+private:
+	/** Common function to construct the expression parser */
+	void ConstructExpressionParser();
+
+	/** Defines whether or not the expression parser can evaluate complex expressions */
+	ETextFilterExpressionEvaluatorMode ExpressionEvaluatorMode;
+
+	/** The cached complexity of the current filter terms */
+	ETextFilterExpressionType FilterType;
+
+	/** The the filter terms that we're currently using (compiled into CompiledFilter) */
+	FText FilterText;
+
+	/** The last error returned from lexing or compiling the current filter text */
+	FText FilterErrorText;
+
+	/** The compiled filter created from the current filter text (compiled from FilterText) */
+	TOptional<ExpressionParser::CompileResultType> CompiledFilter;
+
+	/** Expression parser */
+	FTokenDefinitions TokenDefinitions;
+	FExpressionGrammar Grammar;
+	TOperatorJumpTable<ITextFilterExpressionContext> JumpTable;
+};
diff --git a/Engine/Source/Runtime/Core/Public/Misc/TextFilterUtils.h b/Engine/Source/Runtime/Core/Public/Misc/TextFilterUtils.h
new file mode 100644
--- /dev/null
+++ b/Engine/Source/Runtime/Core/Public/Misc/TextFilterUtils.h
@@ -0,0 +1,31 @@
+// Copyright 1998-2015 Epic Games, Inc. All Rights Reserved.
+#pragma once
+
+/** Defines the comparison operators that can be used for a complex (key->value) comparison */
+enum class ETextFilterComparisonOperation : uint8
+{
+	Equal,
+	NotEqual,
+	Less,
+	LessOrEqual,
+	Greater,
+	GreaterOrEqual,
+};
+
+/** Defines the different ways that a string can be compared while evaluating the expression */
+enum class ETextFilterTextComparisonMode : uint8
+{
+	Exact,
+	Partial,
+	StartsWith,
+	EndsWith,
+};
+
+namespace TextFilterUtils
+{
+	/** Utility function to perform a basic string test for the given values */
+	CORE_API bool TestBasicStringExpression(const FString& InValue1, const FString& InValue2, const ETextFilterTextComparisonMode InTextComparisonMode);
+
+	/** Utility function to perform a complex expression test for the given values */
+	CORE_API bool TestComplexExpression(const FString& InValue1, const FString& InValue2, const ETextFilterComparisonOperation InComparisonOperation, const ETextFilterTextComparisonMode InTextComparisonMode);
+}
