addition of Push Notification code
PR-981
#ios

diff --git a/Engine/Source/Runtime/Core/Private/GenericPlatform/GenericPlatformMisc.cpp b/Engine/Source/Runtime/Core/Private/GenericPlatform/GenericPlatformMisc.cpp
--- a/Engine/Source/Runtime/Core/Private/GenericPlatform/GenericPlatformMisc.cpp
+++ b/Engine/Source/Runtime/Core/Private/GenericPlatform/GenericPlatformMisc.cpp
@@ -876,3 +876,8 @@ FString FGenericPlatformMisc::GetOperatingSystemId()
 	// not implemented by default.
 	return FString();
 }
+
+void FGenericPlatformMisc::RegisterForRemoteNotifications()
+{
+	// not implemented by default
+}
\ No newline at end of file
diff --git a/Engine/Source/Runtime/Core/Private/IOS/IOSAppDelegate.cpp b/Engine/Source/Runtime/Core/Private/IOS/IOSAppDelegate.cpp
--- a/Engine/Source/Runtime/Core/Private/IOS/IOSAppDelegate.cpp
+++ b/Engine/Source/Runtime/Core/Private/IOS/IOSAppDelegate.cpp
@@ -624,6 +624,43 @@ void InstallSignalHandlers()
 	FPlatformMisc::HandleLowMemoryWarning();
 }
 
+#ifdef __IPHONE_8_0
+- (void)application:(UIApplication *)application didRegisterUserNotificationSettings:(UIUserNotificationSettings *)notificationSettings
+{
+	[application registerForRemoteNotifications];
+}
+#endif
+
+- (void)application:(UIApplication *)application didRegisterForRemoteNotificationsWithDeviceToken:(NSData *)deviceToken
+{
+	TArray<uint8> Token;
+	Token.AddUninitialized([deviceToken length]);
+	memcpy(Token.GetData(), [deviceToken bytes], [deviceToken length]);
+
+	FCoreDelegates::ApplicationRegisteredForRemoteNotificationsDelegate.Broadcast(Token);
+}
+
+-(void)application:(UIApplication *)application didFailtoRegisterForRemoteNotificationsWithError:(NSError *)error
+{
+	FCoreDelegates::ApplicationFailedToRegisterForRemoteNotificationsDelegate.Broadcast(FString([error description]));
+}
+
+-(void)application : (UIApplication *)application didReceiveRemoteNotification:(NSDictionary *)userInfo fetchCompletionHandler:(void(^)(UIBackgroundFetchResult result))handler
+{
+	NSString* JsonString = @"{}";
+	NSError* JsonError;
+	NSData* JsonData = [NSJSONSerialization dataWithJSONObject : userInfo
+					options : 0
+					error : &JsonError];
+
+	if (JsonData)
+	{
+		JsonString = [[[NSString alloc] initWithData:JsonData encoding : NSUTF8StringEncoding] autorelease];
+	}
+
+	FCoreDelegates::ApplicationReceivedRemoteNotificationDelegate.Broadcast(FString(JsonString));
+}
+
 /**
  * Shows the given Game Center supplied controller on the screen
  *
diff --git a/Engine/Source/Runtime/Core/Private/IOS/IOSPlatformMisc.cpp b/Engine/Source/Runtime/Core/Private/IOS/IOSPlatformMisc.cpp
--- a/Engine/Source/Runtime/Core/Private/IOS/IOSPlatformMisc.cpp
+++ b/Engine/Source/Runtime/Core/Private/IOS/IOSPlatformMisc.cpp
@@ -740,3 +740,19 @@ TArray<uint8> FIOSPlatformMisc::GetSystemFontBytes()
 	return FontBytes;
 }
 
+void FIOSPlatformMisc::RegisterForRemoteNotifications()
+{
+	UIApplication* application = [UIApplication sharedApplication];
+	if ([application respondsToSelector : @selector(registerUserNotifcationSettings:)])
+	{
+#ifdef __IPHONE_8_0
+		UIUserNotificationSettings * settings = [UIUserNotificationSettings settingsForTypes : (UIRemoteNotificationTypeBadge | UIRemoteNotificationTypeSound | UIRemoteNotificationTypeAlert) categories:nil];
+		[application registerUserNotificationSettings : settings];
+#endif
+	}
+	else
+	{
+		UIRemoteNotificationType myTypes = UIRemoteNotificationTypeBadge | UIRemoteNotificationTypeAlert | UIRemoteNotificationTypeSound;
+		[application registerForRemoteNotificationTypes : myTypes];
+	}
+}
diff --git a/Engine/Source/Runtime/Core/Private/Misc/CallbackDevice.cpp b/Engine/Source/Runtime/Core/Private/Misc/CallbackDevice.cpp
--- a/Engine/Source/Runtime/Core/Private/Misc/CallbackDevice.cpp
+++ b/Engine/Source/Runtime/Core/Private/Misc/CallbackDevice.cpp
@@ -62,6 +62,10 @@ FCoreDelegates::FApplicationLifetimeDelegate FCoreDelegates::ApplicationWillEnte
 FCoreDelegates::FApplicationLifetimeDelegate FCoreDelegates::ApplicationHasEnteredForegroundDelegate;
 FCoreDelegates::FApplicationLifetimeDelegate FCoreDelegates::ApplicationWillTerminateDelegate; 
 
+FCoreDelegates::FApplicationRegisteredForRemoteNotificationsDelegate FCoreDelegates::ApplicationRegisteredForRemoteNotificationsDelegate;
+FCoreDelegates::FApplicationFailedToRegisterForRemoteNotificationsDelegate FCoreDelegates::ApplicationFailedToRegisterForRemoteNotificationsDelegate;
+FCoreDelegates::FApplicationReceivedRemoteNotificationDelegate FCoreDelegates::ApplicationReceivedRemoteNotificationDelegate;
+
 FCoreDelegates::FStatCheckEnabled FCoreDelegates::StatCheckEnabled;
 FCoreDelegates::FStatEnabled FCoreDelegates::StatEnabled;
 FCoreDelegates::FStatDisabled FCoreDelegates::StatDisabled;
diff --git a/Engine/Source/Runtime/Core/Public/GenericPlatform/GenericPlatformMisc.h b/Engine/Source/Runtime/Core/Public/GenericPlatform/GenericPlatformMisc.h
--- a/Engine/Source/Runtime/Core/Public/GenericPlatform/GenericPlatformMisc.h
+++ b/Engine/Source/Runtime/Core/Public/GenericPlatform/GenericPlatformMisc.h
@@ -829,6 +829,11 @@ public:
 	 */
 	static const TCHAR* GetEngineMode();
 
+	/**
+	 * Requests permission to send remote notifications to the user's device.
+	 */
+	static void RegisterForRemoteNotifications();
+
 	/**
 	* Execute plaform dependent pre load map actions
 	*/
diff --git a/Engine/Source/Runtime/Core/Public/IOS/IOSPlatformMisc.h b/Engine/Source/Runtime/Core/Public/IOS/IOSPlatformMisc.h
--- a/Engine/Source/Runtime/Core/Public/IOS/IOSPlatformMisc.h
+++ b/Engine/Source/Runtime/Core/Public/IOS/IOSPlatformMisc.h
@@ -95,6 +95,8 @@ struct CORE_API FIOSPlatformMisc : public FGenericPlatformMisc
 	static bool GetStoredValue(const FString& InStoreId, const FString& InSectionName, const FString& InKeyName, FString& OutValue);
 	static TArray<uint8> GetSystemFontBytes();
 
+	static void RegisterForRemoteNotifications();
+
 	static class IPlatformChunkInstall* GetPlatformChunkInstall();
     
 	//////// Platform specific
diff --git a/Engine/Source/Runtime/Core/Public/Misc/CallbackDevice.h b/Engine/Source/Runtime/Core/Public/Misc/CallbackDevice.h
--- a/Engine/Source/Runtime/Core/Public/Misc/CallbackDevice.h
+++ b/Engine/Source/Runtime/Core/Public/Misc/CallbackDevice.h
@@ -178,6 +178,20 @@ public:
 	// save state when ApplicationWillEnterBackgroundDelegate is called instead.
 	static FApplicationLifetimeDelegate ApplicationWillTerminateDelegate;
 
+	/** IOS-style push notification delegates */
+	DECLARE_MULTICAST_DELEGATE_OneParam(FApplicationRegisteredForRemoteNotificationsDelegate, TArray<uint8>);
+	DECLARE_MULTICAST_DELEGATE_OneParam(FApplicationFailedToRegisterForRemoteNotificationsDelegate, FString);
+	DECLARE_MULTICAST_DELEGATE_OneParam(FApplicationReceivedRemoteNotificationDelegate, FString);
+
+	// called when the user grants permission to register for remote notifications
+	static FApplicationRegisteredForRemoteNotificationsDelegate ApplicationRegisteredForRemoteNotificationsDelegate;
+
+	// called when the application fails to register for remote notifications
+	static FApplicationFailedToRegisterForRemoteNotificationsDelegate ApplicationFailedToRegisterForRemoteNotificationsDelegate;
+
+	// called when the application receives a remote notification
+	static FApplicationReceivedRemoteNotificationDelegate ApplicationReceivedRemoteNotificationDelegate;
+
 	/** Checks to see if the stat is already enabled */
 	DECLARE_MULTICAST_DELEGATE_ThreeParams(FStatCheckEnabled, const TCHAR*, bool&, bool&);
 	static FStatCheckEnabled StatCheckEnabled;
diff --git a/Engine/Source/Runtime/Engine/Classes/Kismet/KismetSystemLibrary.h b/Engine/Source/Runtime/Engine/Classes/Kismet/KismetSystemLibrary.h
--- a/Engine/Source/Runtime/Engine/Classes/Kismet/KismetSystemLibrary.h
+++ b/Engine/Source/Runtime/Engine/Classes/Kismet/KismetSystemLibrary.h
@@ -1330,4 +1330,11 @@ class ENGINE_API UKismetSystemLibrary : public UBlueprintFunctionLibrary
 	 */
 	UFUNCTION(BlueprintCallable, Category = "Utilities", meta = (HidePin = "WorldContextObject", DefaultToSelf = "WorldContextObject"))
 	static void SetSupressViewportTransitionMessage(UObject* WorldContextObject, bool bState);
+
+	/**
+	 * Requests permission to send remote notifications to the user's device.
+	 * (iOS only)
+	 */
+	UFUNCTION(BlueprintCallable, Category = "Utilities|Platform")
+	static void RegisterForRemoteNotifications();
 };
diff --git a/Engine/Source/Runtime/Engine/Private/KismetSystemLibrary.cpp b/Engine/Source/Runtime/Engine/Private/KismetSystemLibrary.cpp
--- a/Engine/Source/Runtime/Engine/Private/KismetSystemLibrary.cpp
+++ b/Engine/Source/Runtime/Engine/Private/KismetSystemLibrary.cpp
@@ -3097,4 +3097,9 @@ void UKismetSystemLibrary::LoadAssetClass(UObject* WorldContextObject, const TAs
 			LatentManager.AddNewAction(LatentInfo.CallbackTarget, LatentInfo.UUID, NewAction);
 		}
 	}
+}
+
+void UKismetSystemLibrary::RegisterForRemoteNotifications()
+{
+	FPlatformMisc::RegisterForRemoteNotifications();
 }
\ No newline at end of file
