UE4- simple profiler for the task graph since the stats system used the task graph

diff --git a/Engine/Source/Runtime/Core/Private/Async/TaskGraph.cpp b/Engine/Source/Runtime/Core/Private/Async/TaskGraph.cpp
--- a/Engine/Source/Runtime/Core/Private/Async/TaskGraph.cpp
+++ b/Engine/Source/Runtime/Core/Private/Async/TaskGraph.cpp
@@ -17,6 +17,74 @@ namespace ENamedThreads
 	CORE_API Type RenderThread_Local = ENamedThreads::GameThread_Local; // defaults to game local and is set and reset by the render thread itself
 }
 
+#define PROFILE_TASKGRAPH (0)
+#if PROFILE_TASKGRAPH
+	struct FProfileRec
+	{
+		const TCHAR* Name;
+		FThreadSafeCounter NumSamplesStarted;
+		FThreadSafeCounter NumSamplesFinished;
+		uint32 Samples[1000];
+
+		FProfileRec()
+		{
+			Name = nullptr;
+		}
+	};
+	static FThreadSafeCounter NumProfileSamples;
+	static void DumpProfile();
+	struct FProfileRecScope
+	{
+		FProfileRec* Target;
+		int32 SampleIndex;
+		uint32 StartCycles;
+		FProfileRecScope(FProfileRec* InTarget, const TCHAR* InName)
+			: Target(InTarget)
+			, SampleIndex(InTarget->NumSamplesStarted.Increment() - 1)
+			, StartCycles(FPlatformTime::Cycles())
+		{
+			if (SampleIndex == 0 && !Target->Name)
+			{
+				Target->Name = InName;
+			}
+		}
+		~FProfileRecScope()
+		{
+			if (SampleIndex < 1000)
+			{
+				Target->Samples[SampleIndex] = FPlatformTime::Cycles() - StartCycles;
+				if (Target->NumSamplesFinished.Increment() == 1000)
+				{
+					Target->NumSamplesFinished.Reset();
+					FPlatformMisc::MemoryBarrier();
+					uint64 Total = 0;
+					for (int32 Index = 0; Index < 1000; Index++)
+					{
+						Total += Target->Samples[Index];
+					}
+					float MsPer = FPlatformTime::GetSecondsPerCycle() * double(Total);
+					UE_LOG(LogTemp, Display, TEXT("%6.4f ms / scope %s"),MsPer, Target->Name);
+
+					Target->NumSamplesStarted.Reset();
+				}
+			}
+		}
+	};
+	static FProfileRec ProfileRecs[10];
+	static void DumpProfile()
+	{
+
+	}
+
+	#define TASKGRAPH_SCOPE_CYCLE_COUNTER(Index, Name) \
+		FProfileRecScope ProfileRecScope##Index(&ProfileRecs[Index], TEXT(#Name));
+
+
+#else
+	#define TASKGRAPH_SCOPE_CYCLE_COUNTER(Index, Name)
+#endif
+
+
 
 /** 
  *	Pointer to the task graph implementation singleton.
@@ -509,9 +577,14 @@ public:
 	{
 		TestRandomizedThreads();
 		checkThreadGraph(Queue(QueueIndex).StallRestartEvent); // make sure we are started up
-		bool bWasReopenedByMe = Queue(QueueIndex).IncomingQueue.ReopenIfClosedAndPush(Task);
+		bool bWasReopenedByMe;
+		{
+			TASKGRAPH_SCOPE_CYCLE_COUNTER(0, STAT_TaskGraph_EnqueueFromOtherThread_ReopenIfClosedAndPush);
+			bWasReopenedByMe = Queue(QueueIndex).IncomingQueue.ReopenIfClosedAndPush(Task);
+		}
 		if (bWasReopenedByMe)
 		{
+			TASKGRAPH_SCOPE_CYCLE_COUNTER(1, STAT_TaskGraph_EnqueueFromOtherThread_Trigger);
 			Queue(QueueIndex).StallRestartEvent->Trigger();
 		}
 		return bWasReopenedByMe;
@@ -840,6 +913,8 @@ public:
 	**/
 	virtual void QueueTask(FBaseGraphTask* Task, ENamedThreads::Type ThreadToExecuteOn, ENamedThreads::Type CurrentThreadIfKnown = ENamedThreads::AnyThread) override
 	{
+		TASKGRAPH_SCOPE_CYCLE_COUNTER(2, STAT_TaskGraph_QueueTask);
+
 		TestRandomizedThreads();
 		checkThreadGraph(NextUnnamedThreadMod);
 		if (CurrentThreadIfKnown == ENamedThreads::AnyThread)
@@ -853,10 +928,18 @@ public:
 		}
 		if (ThreadToExecuteOn == ENamedThreads::AnyThread)
 		{
+			TASKGRAPH_SCOPE_CYCLE_COUNTER(3, STAT_TaskGraph_QueueTask_AnyThread);
 			if (FPlatformProcess::SupportsMultithreading())
 			{
-				IncomingAnyThreadTasks.Push(Task);
-				FTaskThread* TempTarget = StalledUnnamedThreads.Pop(); //@todo it is possible that a thread is in the process of stalling and we just missed it, non-fatal, but we could lose a whole task of potential parallelism.
+				{
+					TASKGRAPH_SCOPE_CYCLE_COUNTER(4, STAT_TaskGraph_QueueTask_IncomingAnyThreadTasks_Push);
+					IncomingAnyThreadTasks.Push(Task);
+				}
+				FTaskThread* TempTarget;
+				{
+					TASKGRAPH_SCOPE_CYCLE_COUNTER(5, STAT_TaskGraph_QueueTask_StalledUnnamedThreads_Pop);
+					TempTarget = StalledUnnamedThreads.Pop(); //@todo it is possible that a thread is in the process of stalling and we just missed it, non-fatal, but we could lose a whole task of potential parallelism.
+				}
 				if (TempTarget)
 				{
 					ThreadToExecuteOn = TempTarget->GetThreadId();
