UE4 - Three digits of precision on stats

diff --git a/Engine/Source/Runtime/Core/Private/Stats/StatsData.cpp b/Engine/Source/Runtime/Core/Private/Stats/StatsData.cpp
--- a/Engine/Source/Runtime/Core/Private/Stats/StatsData.cpp
+++ b/Engine/Source/Runtime/Core/Private/Stats/StatsData.cpp
@@ -1469,11 +1469,11 @@ FString FStatsUtils::DebugPrint(FStatMessage const& Item)
 	case EStatDataType::ST_int64:
 		if (Item.NameAndInfo.GetFlag(EStatMetaFlags::IsPackedCCAndDuration))
 		{
-			Result = FString::Printf(TEXT("%.2fms (%4d)"), FPlatformTime::ToMilliseconds(FromPackedCallCountDuration_Duration(Item.GetValue_int64())), FromPackedCallCountDuration_CallCount(Item.GetValue_int64()));
+			Result = FString::Printf(TEXT("%.3fms (%4d)"), FPlatformTime::ToMilliseconds(FromPackedCallCountDuration_Duration(Item.GetValue_int64())), FromPackedCallCountDuration_CallCount(Item.GetValue_int64()));
 		}
 		else if (Item.NameAndInfo.GetFlag(EStatMetaFlags::IsCycle))
 		{
-			Result = FString::Printf(TEXT("%.2fms"), FPlatformTime::ToMilliseconds(Item.GetValue_int64()));
+			Result = FString::Printf(TEXT("%.3fms"), FPlatformTime::ToMilliseconds(Item.GetValue_int64()));
 		}
 		else
 		{
