Added counter to stats system that can total time spent in a reentrant function. Useful for load time profiling of linker and blueprint compiler.

diff --git a/Engine/Source/Developer/Profiler/Private/ProfilerSession.cpp b/Engine/Source/Developer/Profiler/Private/ProfilerSession.cpp
--- a/Engine/Source/Developer/Profiler/Private/ProfilerSession.cpp
+++ b/Engine/Source/Developer/Profiler/Private/ProfilerSession.cpp
@@ -264,7 +264,7 @@ bool FProfilerSession::HandleTicker( float DeltaTime )
 			break;
 		}
 
-		FScopeSecondsCounter SecondsCounter(Seconds);
+		FSimpleScopeSecondsCounter SecondsCounter(Seconds);
 
 		const uint32 FrameIndex = FrameToProcess[0];
 		FrameToProcess.RemoveAt( 0 );
diff --git a/Engine/Source/Runtime/Core/Public/Stats/StatsMisc.h b/Engine/Source/Runtime/Core/Public/Stats/StatsMisc.h
--- a/Engine/Source/Runtime/Core/Public/Stats/StatsMisc.h
+++ b/Engine/Source/Runtime/Core/Public/Stats/StatsMisc.h
@@ -6,18 +6,18 @@
 #if STATS
 /**
  * Utility class to capture time passed in seconds, adding delta time to passed
- * in variable.
+ * in variable. Not useful for reentrant functions
  */
-class FScopeSecondsCounter
+class FSimpleScopeSecondsCounter
 {
 public:
 	/** Ctor, capturing start time. */
-	FScopeSecondsCounter( double& InSeconds )
+	FSimpleScopeSecondsCounter(double& InSeconds)
 	:	StartTime(FPlatformTime::Seconds())
 	,	Seconds(InSeconds)
 	{}
 	/** Dtor, updating seconds with time delta. */
-	~FScopeSecondsCounter()
+	~FSimpleScopeSecondsCounter()
 	{
 		Seconds += FPlatformTime::Seconds() - StartTime;
 	}
@@ -30,12 +30,80 @@ private:
 
 /** Macro for updating a seconds counter without creating new scope. */
 #define SCOPE_SECONDS_COUNTER(Seconds) \
-	FScopeSecondsCounter SecondsCount_##Seconds(Seconds);
+	FSimpleScopeSecondsCounter SecondsCount_##Seconds(Seconds);
 
 #else
 #define SCOPE_SECONDS_COUNTER(Seconds)
 #endif 
 
+
+/**
+* Utility class to store a counter and a time value in seconds. Implementation will be stripped out in 
+* STATS disabled builds, although it will waste a small amount of memory unless stripped by the linker.
+*
+* Useful when timing reentrant functions.
+*/
+struct FSecondsCounterData
+{
+#if STATS
+	FSecondsCounterData()
+		: Time(0.0)
+		, ScopeCounter(0)
+	{
+	}
+
+	double GetTime() const { return Time; }
+	void ClearTime() { check(ScopeCounter == 0);  Time = 0.0; }
+private:
+	double Time;
+	int32 ScopeCounter;
+
+	friend class FSecondsCounterScope;
+#else //STATS
+public:
+	static double GetTime() { return 0.0; }
+	static void ClearTime() {}
+#endif //STATS
+};
+
+/**
+* Utility class to update a FSecondsCounterData. Does nothing in STATS disabled builds.
+*/
+class FSecondsCounterScope
+{
+#if STATS
+public:
+	FSecondsCounterScope(FSecondsCounterData& InData)
+		: Data(InData)
+		, StartTime(-1.f)
+	{
+		if (Data.ScopeCounter == 0)
+		{
+			StartTime = FPlatformTime::Seconds();
+		}
+		++InData.ScopeCounter;
+	}
+
+	~FSecondsCounterScope()
+	{
+		--Data.ScopeCounter;
+		if (Data.ScopeCounter == 0)
+		{
+			checkf(StartTime >= 0.f, TEXT("Counter is corrupt! Thinks it started before epoch"));
+			Data.Time += (FPlatformTime::Seconds() - StartTime);
+		}
+	}
+
+private:
+	FSecondsCounterData& Data;
+	double StartTime;
+#else //STATS
+public:
+	FSecondsCounterScope(FSecondsCounterData& InData) {}
+#endif //STATS
+};
+
+
 /** Key contains total time, value contains total count. */
 typedef TKeyValuePair<double, uint32> FTotalTimeAndCount;
 
