Improved handling of quotes within an unquoted string

If the text filter expression parser comes across a quote within a quoted string, it will now switch into "quoted string" parsing mode, and consume all characters (even special characters) until it finds a corresponding closing quote. The resultant string will also include the quote characters (unlike a normal quoted string).

This allows our asset reference strings to be parsed correctly, eg) Type'/Path/To/Asset.Asset'

This also improves the error message when a quoted string is unbound. It now reports the character it was looking for, as well as the line and character index of the opening quote.

diff --git a/Engine/Source/Runtime/Core/Private/Misc/TextFilterExpressionEvaluator.cpp b/Engine/Source/Runtime/Core/Private/Misc/TextFilterExpressionEvaluator.cpp
--- a/Engine/Source/Runtime/Core/Private/Misc/TextFilterExpressionEvaluator.cpp
+++ b/Engine/Source/Runtime/Core/Private/Misc/TextFilterExpressionEvaluator.cpp
@@ -95,8 +95,8 @@ DEFINE_EXPRESSION_NODE_TYPE(bool,									0xC1CD5DCF, 0x2AB44958, 0xB3FF4F8F, 0x
 namespace TextFilterExpressionParser
 {
 	/** This contains all the symbols that can define breaking points between text and an operator */
-	static const TCHAR BasicTextBreakingCharacters[]	= { '(', ')', '!', '-', '&', '|', '.', '"', ' ' };						// ETextFilterExpressionEvaluatorMode::BasicString
-	static const TCHAR ComplexTextBreakingCharacters[]	= { '(', ')', '=', ':', '<', '>', '!', '-', '&', '|', '.', '"', ' ' };	// ETextFilterExpressionEvaluatorMode::Complex
+	static const TCHAR BasicTextBreakingCharacters[]	= { '(', ')', '!', '-', '&', '|', '.', ' ' };						// ETextFilterExpressionEvaluatorMode::BasicString
+	static const TCHAR ComplexTextBreakingCharacters[]	= { '(', ')', '=', ':', '<', '>', '!', '-', '&', '|', '.', ' ' };	// ETextFilterExpressionEvaluatorMode::Complex
 
 	const TCHAR* FSubExpressionStart::Monikers[]		= { TEXT("(") };
 	const TCHAR* FSubExpressionEnd::Monikers[]			= { TEXT(")") };
@@ -173,6 +173,17 @@ namespace TextFilterExpressionParser
 		return TOptional<FExpressionError>();
 	}
 
+	/** Transform the given string to remove any escape character sequences found in a quoted string */
+	void UnescapeQuotedString(FString& Str, const TCHAR InQuoteChar)
+	{
+		const TCHAR EscapedQuote[] = { '\\', InQuoteChar, 0 };
+		const TCHAR UnescapedQuote[] = { InQuoteChar, 0 };
+
+		// Unescape any literal quotes within the string
+		Str.ReplaceInline(EscapedQuote, UnescapedQuote);
+		Str.ReplaceInline(TEXT("\\\\"), TEXT("\\"));
+	}
+
 	/** Consume quoted text from the specified consumer's stream, if one exists at the current read position */
 	TOptional<FExpressionError> ConsumeQuotedText(FExpressionTokenConsumer& Consumer)
 	{
@@ -212,17 +223,19 @@ namespace TextFilterExpressionParser
 
 		if (!Stream.ParseSymbol(QuoteChar, &TextTokenWithQuotes).IsSet())
 		{
-			return FExpressionError(NSLOCTEXT("TextFilterExpressionParser", "Error_UnboundedString", "Reached the end of the expression before a closing quote was found."));
+			FString QuoteCharStr;
+			QuoteCharStr.AppendChar(QuoteChar);
+
+			return FExpressionError(FText::Format(NSLOCTEXT("TextFilterExpressionParser", "SyntaxError_UnmatchedQuote", "Syntax error: Reached end of expression before matching closing quote for '{0}' at line {1}:{2}"),
+				FText::FromString(QuoteCharStr),
+				FText::AsNumber(TextTokenWithQuotes.GetLineNumber()),
+				FText::AsNumber(TextTokenWithQuotes.GetCharacterIndex())
+			));
 		}
 		else if (TextToken.IsSet())
 		{
-			const TCHAR EscapedQuote[] = { '\\', QuoteChar, 0 };
-			const TCHAR UnescapedQuote[] = { QuoteChar, 0 };
-
-			// Unescape any literal quotes within the string
 			FString FinalString = TextToken.GetValue().GetString();
-			FinalString.ReplaceInline(EscapedQuote, UnescapedQuote);
-			FinalString.ReplaceInline(TEXT("\\\\"), TEXT("\\"));
+			UnescapeQuotedString(FinalString, QuoteChar);
 			Consumer.Add(TextTokenWithQuotes, FExpressionNode(FTextToken(MoveTemp(FinalString), ETextFilterTextComparisonMode::Exact)));
 		}
 		else
@@ -238,15 +251,63 @@ namespace TextFilterExpressionParser
 	{
 		auto& Stream = Consumer.GetStream();
 
-		TOptional<FStringToken> TextToken = Stream.ParseToken([&IsBreakingCharacter](TCHAR InC)
+		FString FinalString;
+		FString CurrentQuotedString;
+
+		TCHAR QuoteChar = 0;
+		int32 NumConsecutiveSlashes = 0;
+		TOptional<FStringToken> TextToken = Stream.ParseToken([&](TCHAR InC)
 		{
-			// This is a non-quoted string, consume until we hit a breaking character
-			return (IsBreakingCharacter(InC)) ? EParseState::StopBefore : EParseState::Continue;
+			if (QuoteChar == 0) // Parsing a non-quoted string...
+			{
+				// Are we starting a quoted sub-string?
+				if (InC == '"' || InC == '\'')
+				{
+					CurrentQuotedString.AppendChar(InC);
+					QuoteChar = InC;
+					NumConsecutiveSlashes = 0;
+				}
+				else
+				{
+					// Consume until we hit a breaking character
+					if (IsBreakingCharacter(InC))
+					{
+						return EParseState::StopBefore;
+					}
+
+					FinalString.AppendChar(InC);
+				}
+			}
+			else // Parsing a quoted sub-string...
+			{
+				CurrentQuotedString.AppendChar(InC);
+
+				// Are we ending a quoted sub-string?
+				if (InC == QuoteChar && NumConsecutiveSlashes%2 == 0)
+				{
+					UnescapeQuotedString(CurrentQuotedString, QuoteChar);
+					FinalString.Append(CurrentQuotedString);
+
+					CurrentQuotedString.Reset();
+					QuoteChar = 0;
+				}
+
+				if (InC == '\\')
+				{
+					NumConsecutiveSlashes++;
+				}
+				else
+				{
+					NumConsecutiveSlashes = 0;
+				}
+			}
+
+			return EParseState::Continue;
 		});
 
 		if (TextToken.IsSet())
 		{
-			Consumer.Add(TextToken.GetValue(), FExpressionNode(FTextToken(TextToken.GetValue().GetString(), ETextFilterTextComparisonMode::Partial)));
+			Consumer.Add(TextToken.GetValue(), FExpressionNode(FTextToken(MoveTemp(FinalString), ETextFilterTextComparisonMode::Partial)));
 		}
 
 		return TOptional<FExpressionError>();
diff --git a/Engine/Source/Runtime/Core/Private/Misc/TextFilterTests.cpp b/Engine/Source/Runtime/Core/Private/Misc/TextFilterTests.cpp
--- a/Engine/Source/Runtime/Core/Private/Misc/TextFilterTests.cpp
+++ b/Engine/Source/Runtime/Core/Private/Misc/TextFilterTests.cpp
@@ -13,6 +13,8 @@ namespace TextFilterTests
 		{
 			BasicStrings.Add(TEXT("Wooble"));
 			BasicStrings.Add(TEXT("Flibble"));
+			BasicStrings.Add(TEXT("Type'/Path/To/Asset.Asset'"));
+			BasicStrings.Add(TEXT("Other'/Path/To/Asset.Asset'FollowingText"));
 			BasicStrings.Add(TEXT("Funky<String>"));
 
 			KeyValuePairs.Add("StringKey", TEXT("Test"));
@@ -147,6 +149,8 @@ namespace TextFilterTests
 		bResult &= InTestPayload.TestFilterExpression(TEXT("Wooble2 OR Flibble"), true);
 		bResult &= InTestPayload.TestFilterExpression(TEXT("Wooble2 AND Flibble"), false);
 		bResult &= InTestPayload.TestFilterExpression(TEXT("Wooble && !Flibble"), false);
+		bResult &= InTestPayload.TestFilterExpression(TEXT("Type'/Path/To/Asset.Asset'"), true);
+		bResult &= InTestPayload.TestFilterExpression(TEXT("Other'/Path/To/Asset.Asset'FollowingText"), true);
 		bResult &= InTestPayload.TestFilterExpression(TEXT("\"Funky<String>\""), true);
 
 		return bResult;
