Fixed an issue where the invalid expressoins evaluation test was printing out a message that had the word "error" in it causing the automation test to fail on the build machines.
Also refactored the expression evaluation test.
- Broke it up to four test implementations based on what is being tests.  This helps QA figure out what is wrong sooner.
- Changed some of the error messages.
- Made the test use the automations "TestTrue" function instead of using a bool to hold the test results to the very end.
- Changed the pretty name so that the tests are grouped together in the UFE automation tab UI.

diff --git a/Engine/Source/Runtime/Core/Private/Math/BasicMathExpressionEvaluator.cpp b/Engine/Source/Runtime/Core/Private/Math/BasicMathExpressionEvaluator.cpp
--- a/Engine/Source/Runtime/Core/Private/Math/BasicMathExpressionEvaluator.cpp
+++ b/Engine/Source/Runtime/Core/Private/Math/BasicMathExpressionEvaluator.cpp
@@ -174,7 +174,7 @@ bool TestExpression(FAutomationTestBase* Test, const TCHAR* Expression, double E
 	}
 	else if (Result.GetValue() != Expected)
 	{
-		Test->AddError(FString::Printf(TEXT("Expression '%s' evaluated incorrectly: %.f != %.f"), Expression, Result.GetValue(), Expected));
+		Test->AddError(FString::Printf(TEXT("'%s' evaluation results: %.f != %.f"), Expression, Result.GetValue(), Expected));
 		return false;
 	}
 	return true;
@@ -187,52 +187,67 @@ bool TestInvalidExpression(FAutomationTestBase* Test, const TCHAR* Expression)
 	TValueOrError<double, FExpressionError> Result = Parser.Evaluate(Expression);
 	if (Result.IsValid())
 	{
-		Test->AddError(FString::Printf(TEXT("Invalid expression '%s' did not report an error"), Expression));
 		return false;
 	}
-	else
-	{
-		Test->AddLogItem(Result.GetError().Text.ToString());
-	}
 
 	return true;
 }
 
-IMPLEMENT_SIMPLE_AUTOMATION_TEST(FBasicMathExpressionEvaluatorTest, "System.Core.Math.Expressions", EAutomationTestFlags::ATF_SmokeTest)
-bool FBasicMathExpressionEvaluatorTest::RunTest( const FString& Parameters )
+
+IMPLEMENT_SIMPLE_AUTOMATION_TEST(FBasicMathExpressionEvaluatorTest, "System.Core.Math.Evaluate.Valid Expressions", EAutomationTestFlags::ATF_SmokeTest)
+// Evaluates valid math expressions.
+bool FBasicMathExpressionEvaluatorTest::RunTest(const FString& Parameters)
+{
+	TestTrue(TEXT("Valid expression, '+2', evaluated incorrectly."), TestExpression(this, TEXT("+1"), 1));
+	TestTrue(TEXT("Valid expression, '-20', evaluated incorrectly."), TestExpression(this, TEXT("-20"), -20));
+	TestTrue(TEXT("Valid expression, '-+-2', evaluated incorrectly."), TestExpression(this, TEXT("-+-2"), 2));
+	TestTrue(TEXT("Valid expression, '1 + 2', evaluated incorrectly."), TestExpression(this, TEXT("1 + 2"), 3));
+	TestTrue(TEXT("Valid expression, '1+2*3', evaluated incorrectly."), TestExpression(this, TEXT("1+2*3"), 7));
+	TestTrue(TEXT("Valid expression, '1+2*3*4+1', evaluated incorrectly."), TestExpression(this, TEXT("1+2*3*4+1"), 1 + 2 * 3 * 4 + 1));
+	TestTrue(TEXT("Valid expression, '1*2+3', evaluated incorrectly."), TestExpression(this, TEXT("1*2+3"), 1 * 2 + 3));
+	TestTrue(TEXT("Valid expression, '1+2*3*4+1', evaluated incorrectly."), TestExpression(this, TEXT("1+2*3*4+1"), 1 + 2 * 3 * 4 + 1));
+	
+	TestTrue(TEXT("Valid expression, '2^2', evaluated incorrectly."), TestExpression(this, TEXT("2^2"), 4));
+	TestTrue(TEXT("Valid expression, 'sqrt(4)', evaluated incorrectly."), TestExpression(this, TEXT("sqrt(4)"), 2));
+	TestTrue(TEXT("Valid expression, '4*sqrt(4)+10', evaluated incorrectly."), TestExpression(this, TEXT("4*sqrt(4)+10"), 18));
+	TestTrue(TEXT("Valid expression, '8%6', evaluated incorrectly."), TestExpression(this, TEXT("8%6"), 2));
+
+	return true;
+}
+
+IMPLEMENT_SIMPLE_AUTOMATION_TEST(FBasicMathExpressionEvaluatorWhitespaceExpressionsTest, "System.Core.Math.Evaluate.Valid Expressions With Whitespaces", EAutomationTestFlags::ATF_SmokeTest)
+// Evaluates a valid math expression with leading and trailing white spaces.
+bool FBasicMathExpressionEvaluatorWhitespaceExpressionsTest::RunTest(const FString& Parameters)
+{
+	TestTrue(TEXT("Expression with leading and trailing whitespaces was not evaluated correctly."), TestExpression(this, TEXT(" 1+2 "), 1 + 2));
+
+	return true;
+}
+
+IMPLEMENT_SIMPLE_AUTOMATION_TEST(FBasicMathExpressionEvaluatorGroupedExpressionsTest, "System.Core.Math.Evaluate.Valid Grouped Expressions", EAutomationTestFlags::ATF_SmokeTest)
+// Evaluates valid math expressions that are grouped
+bool FBasicMathExpressionEvaluatorGroupedExpressionsTest::RunTest(const FString& Parameters)
 {
-	bool bResult = true;
-
-	bResult = bResult && TestExpression(this, TEXT("+1"), 1);
-	bResult = bResult && TestExpression(this, TEXT("-20"), -20);
-	bResult = bResult && TestExpression(this, TEXT("-+-2"), 2);
-	bResult = bResult && TestExpression(this, TEXT("1 + 2"), 3);
-	bResult = bResult && TestExpression(this, TEXT("1+2*3"), 7);
-	bResult = bResult && TestExpression(this, TEXT("1+2*3*4+1"), 1+2*3*4+1);
-	bResult = bResult && TestExpression(this, TEXT("1*2+3"), 1*2+3);
-	bResult = bResult && TestExpression(this, TEXT("1+2*3*4+1"), 1+2*3*4+1);
-
-	bResult = bResult && TestExpression(this, TEXT("2^2"), 4);
-	bResult = bResult && TestExpression(this, TEXT("sqrt(4)"), 2);
-	bResult = bResult && TestExpression(this, TEXT("4*sqrt(4)+10"), 18);
-	bResult = bResult && TestExpression(this, TEXT("8%6"), 2);
-
-	// Leading/Trailing whitespace
-	bResult = bResult && TestExpression(this, TEXT(" 1+2 "), 1+2);
-
-	// Test grouping
-	bResult = bResult && TestExpression(this, TEXT("(1+2)*3*4+1"), (1+2)*3*4+1);
-	bResult = bResult && TestExpression(this, TEXT("(1+2)*3*(4+1)"), (1+2)*3*(4+1));
-	bResult = bResult && TestExpression(this, TEXT("((1+2) / (3+1) + 2) * 3"), ((1.0+2) / (3+1) + 2) * 3);
-
-	// Test that some invalid expressions report errors and don't crash etc
-	bResult = bResult && TestInvalidExpression(this, TEXT("gobbledegook"));
-	bResult = bResult && TestInvalidExpression(this, TEXT("50**10"));
-	bResult = bResult && TestInvalidExpression(this, TEXT("*1"));
-	bResult = bResult && TestInvalidExpression(this, TEXT("+"));
-	bResult = bResult && TestInvalidExpression(this, TEXT("{+}"));
-
-	return bResult;
+	TestTrue(TEXT("Valid grouped expression, '(1+2)*3*4+1', evaluated incorrectly."), TestExpression(this, TEXT("(1+2)*3*4+1"), (1 + 2) * 3 * 4 + 1));
+	TestTrue(TEXT("Valid grouped expression, '(1+2)*3*(4+1)', evaluated incorrectly."), TestExpression(this, TEXT("(1+2)*3*(4+1)"), (1 + 2) * 3 * (4 + 1)));
+	TestTrue(TEXT("Valid grouped expression, '((1+2) / (3+1) + 2) * 3', evaluated incorrectly."), TestExpression(this, TEXT("((1+2) / (3+1) + 2) * 3"), ((1.0 + 2) / (3 + 1) + 2) * 3));
+
+	return true;
+}
+
+IMPLEMENT_SIMPLE_AUTOMATION_TEST(FBasicMathExpressionEvaluatorInvalidExpressionTest, "System.Core.Math.Evaluate.Invalid Expressions", EAutomationTestFlags::ATF_SmokeTest)
+
+// Evaluates invalid expressions.
+// Invalid expressions will report errors and not crash.
+bool FBasicMathExpressionEvaluatorInvalidExpressionTest::RunTest(const FString& Parameters)
+{
+	TestTrue(TEXT("The invalid math expression, 'gobbledegook', did not report an error."), TestInvalidExpression(this, TEXT("gobbledegook")));
+	TestTrue(TEXT("The invalid math expression, '50**10', did not report an error."), TestInvalidExpression(this, TEXT("50**10")));
+	TestTrue(TEXT("The invalid math expression, '*1', did not report an error."), TestInvalidExpression(this, TEXT("*1")));
+	TestTrue(TEXT("The invalid math expression, '+', did not report an error."), TestInvalidExpression(this, TEXT("+")));
+	TestTrue(TEXT("The invalid math expression, '{+}', did not report an error."), TestInvalidExpression(this, TEXT("+")));
+
+	return true;
 }
 
 #endif //!(UE_BUILD_SHIPPING || UE_BUILD_TEST)
