PR #1178: Hot reload fixes in Linux (UE-16090).

- Contributed by [REDACTED_PERSON].
- A hack, proper resolution tracked as UEPLAT-891.

diff --git a/Engine/Source/Runtime/Core/Private/Linux/LinuxPlatformProcess.cpp b/Engine/Source/Runtime/Core/Private/Linux/LinuxPlatformProcess.cpp
--- a/Engine/Source/Runtime/Core/Private/Linux/LinuxPlatformProcess.cpp
+++ b/Engine/Source/Runtime/Core/Private/Linux/LinuxPlatformProcess.cpp
@@ -13,7 +13,18 @@ void* FLinuxPlatformProcess::GetDllHandle( const TCHAR* Filename )
 {
 	check( Filename );
 	FString AbsolutePath = FPaths::ConvertRelativePathToFull(Filename);
-	void *Handle = dlopen( TCHAR_TO_ANSI(*AbsolutePath), RTLD_LAZY | RTLD_GLOBAL );
+
+	int DlOpenMode = RTLD_LAZY;
+	if (AbsolutePath.EndsWith(TEXT("libsteam_api.so")))
+	{
+		DlOpenMode |= RTLD_GLOBAL; //Global symbol resolution when loading shared objects - Needed for Steam to work when its library is loaded by a plugin
+	}
+	else
+	{
+		DlOpenMode |= RTLD_LOCAL; //Local symbol resolution when loading shared objects - Needed for Hot-Reload
+	}
+
+	void *Handle = dlopen( TCHAR_TO_ANSI(*AbsolutePath), DlOpenMode );
 	if (!Handle)
 	{
 		UE_LOG(LogLinux, Warning, TEXT("dlopen failed: %s"), ANSI_TO_TCHAR(dlerror()) );
