#ue4 - Add FTransform static methods which on vectorized platforms can avoid the cost of converting to FVector first before the operations.

- New methods are: AreRotationsEqual, AreTranslationsEqual, AreScale3DsEqual, AddTranslations, SubtractTranslations.

diff --git a/Engine/Source/Runtime/Core/Public/Math/Transform.h b/Engine/Source/Runtime/Core/Public/Math/Transform.h
--- a/Engine/Source/Runtime/Core/Public/Math/Transform.h
+++ b/Engine/Source/Runtime/Core/Public/Math/Transform.h
@@ -610,14 +610,54 @@ public:
 	}
 	*/
 
+private:
+
+	FORCEINLINE bool RotationEquals(const FQuat& InRotation, const float Tolerance = KINDA_SMALL_NUMBER) const
+	{
+		return Rotation.Equals(InRotation, Tolerance);
+	}
+
+	FORCEINLINE bool TranslationEquals(const FVector& InTranslation, const float Tolerance = KINDA_SMALL_NUMBER) const
+	{
+		return Translation.Equals(InTranslation, Tolerance);
+	}
+
+	FORCEINLINE bool Scale3DEquals(const FVector& InScale3D, const float Tolerance = KINDA_SMALL_NUMBER) const
+	{
+		return (Scale3D-InScale3D).SizeSquared() <= FMath::Square(Tolerance);
+	}
+
+public:
+
+	// Test if A's rotation equals B's rotation, within a tolerance. Preferred over "A.GetRotation().Equals(B.GetRotation())" because it is faster on some platforms.
+	FORCEINLINE static bool AreRotationsEqual(const FTransform& A, const FTransform& B, float Tolerance=KINDA_SMALL_NUMBER)
+	{
+		return A.RotationEquals(B.Rotation, Tolerance);
+	}
+
+	// Test if A's translation equals B's translation, within a tolerance. Preferred over "A.GetTranslation().Equals(B.GetTranslation())" because it is faster on some platforms.
+	FORCEINLINE static bool AreTranslationsEqual(const FTransform& A, const FTransform& B, float Tolerance=KINDA_SMALL_NUMBER)
+	{
+		return A.TranslationEquals(B.Translation, Tolerance);
+	}
+
+	// Test if A's scale equals B's scale, within a tolerance. Preferred over "A.GetScale3D().Equals(B.GetScale3D())" because it is faster on some platforms.
+	FORCEINLINE static bool AreScale3DsEqual(const FTransform& A, const FTransform& B, float Tolerance=KINDA_SMALL_NUMBER)
+	{
+		return A.Scale3DEquals(B.Scale3D, Tolerance);
+	}
+
+
+	// Test if all components of the transforms are equal, within a tolerance.
 	inline bool Equals(const FTransform& Other, float Tolerance=KINDA_SMALL_NUMBER) const
 	{
-		return Rotation.Equals(Other.Rotation, Tolerance) && Translation.Equals(Other.Translation, Tolerance) && ((Scale3D-Other.Scale3D).SizeSquared() <= FMath::Square(Tolerance));
+		return RotationEquals(Other.Rotation, Tolerance) && TranslationEquals(Other.Translation, Tolerance) && Scale3DEquals(Other.Scale3D, Tolerance);
 	}
 
+	// Test if rotation and translation components of the transforms are equal, within a tolerance.
 	inline bool EqualsNoScale(const FTransform& Other, float Tolerance=KINDA_SMALL_NUMBER) const
 	{
-		return Rotation.Equals(Other.Rotation, Tolerance) && Translation.Equals(Other.Translation, Tolerance);
+		return RotationEquals(Other.Rotation, Tolerance) && TranslationEquals(Other.Translation, Tolerance);
 	}
 
 	/**
@@ -705,6 +745,24 @@ public:
 		DiagnosticCheckNaN_Translate();
 	}
 
+	/**
+	 * Add the translations from two FTransforms and return the result.
+	 * @return A.Translation + B.Translation
+	 */
+	FORCEINLINE static FVector AddTranslations(const FTransform& A, const FTransform& B)
+	{
+		return A.Translation + B.Translation;
+	}
+
+	/**
+	 * Subtract translations from two FTransforms and return the difference.
+	 * @return A.Translation - B.Translation.
+	 */
+	FORCEINLINE static FVector SubtractTranslations(const FTransform& A, const FTransform& B)
+	{
+		return A.Translation - B.Translation;
+	}
+
 	/**
 	 * Sets the rotation component
 	 * @param NewRotation The new value for the rotation component
diff --git a/Engine/Source/Runtime/Core/Public/Math/TransformVectorized.h b/Engine/Source/Runtime/Core/Public/Math/TransformVectorized.h
--- a/Engine/Source/Runtime/Core/Public/Math/TransformVectorized.h
+++ b/Engine/Source/Runtime/Core/Public/Math/TransformVectorized.h
@@ -637,6 +637,8 @@ public:
 	}
 	*/
 
+private:
+
 	inline bool RotationEquals( const VectorRegister& InRotation, const ScalarRegister& Tolerance = ScalarRegister(GlobalVectorConstants::KindaSmallNumber)) const
 	{			
 		// !( (FMath::Abs(X-Q.X) > Tolerance) || (FMath::Abs(Y-Q.Y) > Tolerance) || (FMath::Abs(Z-Q.Z) > Tolerance) || (FMath::Abs(W-Q.W) > Tolerance) )
@@ -663,12 +665,35 @@ public:
 		return !VectorAnyGreaterThan(SizeSquared, VectorToleranceSqaured);
 	}
 
+public:
+
+	// Test if A's rotation equals B's rotation, within a tolerance. Preferred over "A.GetRotation().Equals(B.GetRotation())" because it is faster on some platforms.
+	FORCEINLINE static bool AreRotationsEqual(const FTransform& A, const FTransform& B, float Tolerance=KINDA_SMALL_NUMBER)
+	{
+		return A.RotationEquals(B.Rotation, ScalarRegister(Tolerance));
+	}
+
+	// Test if A's translation equals B's translation, within a tolerance. Preferred over "A.GetTranslation().Equals(B.GetTranslation())" because it avoids VectorRegister->FVector conversion.
+	FORCEINLINE static bool AreTranslationsEqual(const FTransform& A, const FTransform& B, float Tolerance=KINDA_SMALL_NUMBER)
+	{
+		return A.TranslationEquals(B.Translation, ScalarRegister(Tolerance));
+	}
+
+	// Test if A's scale equals B's scale, within a tolerance. Preferred over "A.GetScale3D().Equals(B.GetScale3D())" because it avoids VectorRegister->FVector conversion.
+	FORCEINLINE static bool AreScale3DsEqual(const FTransform& A, const FTransform& B, float Tolerance=KINDA_SMALL_NUMBER)
+	{
+		return A.Scale3DEquals(B.Scale3D, ScalarRegister(Tolerance));
+	}
+
+
+	// Test if all components of the transforms are equal, within a tolerance.
 	inline bool Equals(const FTransform& Other, float Tolerance=KINDA_SMALL_NUMBER) const
 	{
 		const ScalarRegister ToleranceRegister(Tolerance);
 		return RotationEquals(Other.Rotation, ToleranceRegister) && TranslationEquals(Other.Translation, ToleranceRegister) && Scale3DEquals(Other.Scale3D, ToleranceRegister);
 	}
 
+	// Test if rotation and translation components of the transforms are equal, within a tolerance.
 	inline bool EqualsNoScale(const FTransform& Other, float Tolerance=KINDA_SMALL_NUMBER) const
 	{
 		const ScalarRegister ToleranceRegister(Tolerance);
@@ -754,6 +779,28 @@ public:
 		DiagnosticCheckNaN_Translate();
 	}
 
+	/**
+	 * Add the translations from two FTransforms and return the result.
+	 * @return A.Translation + B.Translation
+	 */
+	FORCEINLINE static FVector AddTranslations(const FTransform& A, const FTransform& B)
+	{
+		FVector Result;
+		VectorStoreFloat3(VectorAdd(A.Translation, B.Translation), &Result);
+		return Result;
+	}
+
+	/**
+	 * Subtract translations from two FTransforms and return the difference.
+	 * @return A.Translation - B.Translation.
+	 */
+	FORCEINLINE static FVector SubtractTranslations(const FTransform& A, const FTransform& B)
+	{
+		FVector Result;
+		VectorStoreFloat3(VectorSubtract(A.Translation, B.Translation), &Result);
+		return Result;
+	}
+
 	/**
 	 * Sets the rotation component
 	 * @param NewRotation The new value for the rotation component
diff --git a/Engine/Source/Runtime/Engine/Private/Components/SceneComponent.cpp b/Engine/Source/Runtime/Engine/Private/Components/SceneComponent.cpp
--- a/Engine/Source/Runtime/Engine/Private/Components/SceneComponent.cpp
+++ b/Engine/Source/Runtime/Engine/Private/Components/SceneComponent.cpp
@@ -728,7 +728,7 @@ void USceneComponent::SetRelativeLocationAndRotation(FVector NewLocation, const
 
 	const FTransform DesiredRelTransform(NewRotation, NewLocation);
 	const FTransform DesiredWorldTransform = CalcNewComponentToWorld(DesiredRelTransform);
-	const FVector DesiredDelta = DesiredWorldTransform.GetTranslation() - ComponentToWorld.GetTranslation();
+	const FVector DesiredDelta = FTransform::SubtractTranslations(DesiredWorldTransform, ComponentToWorld);
 
 	MoveComponent(DesiredDelta, DesiredWorldTransform.GetRotation(), bSweep, OutSweepHitResult);
 }
@@ -790,7 +790,7 @@ void USceneComponent::AddWorldRotation(const FQuat& DeltaRotation, bool bSweep,
 void USceneComponent::AddWorldTransform(const FTransform& DeltaTransform, bool bSweep, FHitResult* OutSweepHitResult)
 {
 	const FQuat NewWorldRotation = DeltaTransform.GetRotation() * ComponentToWorld.GetRotation();
-	const FVector NewWorldLocation = DeltaTransform.GetTranslation() + ComponentToWorld.GetTranslation();
+	const FVector NewWorldLocation = FTransform::AddTranslations(DeltaTransform, ComponentToWorld);
 	SetWorldTransform(FTransform(NewWorldRotation, NewWorldLocation, FVector(1,1,1)));
 }
 
@@ -2450,7 +2450,7 @@ const TArray<FOverlapInfo>* FScopedMovementUpdate::GetOverlapsAtEnd(class UPrimi
 			else
 			{
 				// Fill in EndOverlaps with overlaps valid at the end location.
-				const bool bMatchingScale = InitialTransform.GetScale3D().Equals(PrimComponent.GetComponentScale());
+				const bool bMatchingScale = FTransform::AreScale3DsEqual(InitialTransform, PrimComponent.GetComponentTransform());
 				if (bMatchingScale)
 				{
 					EndOverlapsPtr = PrimComponent.ConvertSweptOverlapsToCurrentOverlaps(
