Linux: PR #1207: Fix high CPU usage when Editor minimized/not in focus (UE-16465).

- Contributed by [REDACTED_PERSON], reworked to account for windowless applications.

diff --git a/Engine/Source/Runtime/Core/Private/Linux/LinuxApplication.cpp b/Engine/Source/Runtime/Core/Private/Linux/LinuxApplication.cpp
--- a/Engine/Source/Runtime/Core/Private/Linux/LinuxApplication.cpp
+++ b/Engine/Source/Runtime/Core/Private/Linux/LinuxApplication.cpp
@@ -79,7 +79,7 @@ FLinuxApplication::FLinuxApplication()
 	bAllowedToDeferMessageProcessing = true;
 	MouseCaptureWindow = NULL;
 	ControllerStates = new SDLControllerState[SDL_NumJoysticks()];
-	memset( ControllerStates, 0, sizeof(SDLControllerState) * SDL_NumJoysticks() );
+	FMemory::Memset(ControllerStates, 0, sizeof(SDLControllerState) * SDL_NumJoysticks());
 
 	fMouseWheelScrollAccel = 1.0f;
 	if (GConfig)
diff --git a/Engine/Source/Runtime/Core/Private/Linux/LinuxPlatformProcess.cpp b/Engine/Source/Runtime/Core/Private/Linux/LinuxPlatformProcess.cpp
--- a/Engine/Source/Runtime/Core/Private/Linux/LinuxPlatformProcess.cpp
+++ b/Engine/Source/Runtime/Core/Private/Linux/LinuxPlatformProcess.cpp
@@ -8,6 +8,7 @@
 #include <sys/resource.h>
 #include <sys/ioctl.h> // ioctl
 #include <asm/ioctls.h> // FIONREAD
+#include "LinuxApplication.h" // FLinuxApplication::IsForeground()
 
 void* FLinuxPlatformProcess::GetDllHandle( const TCHAR* Filename )
 {
@@ -953,8 +954,8 @@ bool FLinuxPlatformProcess::IsApplicationRunning( uint32 ProcessId )
 
 bool FLinuxPlatformProcess::IsThisApplicationForeground()
 {
-	STUBBED("FLinuxPlatformProcess::IsThisApplicationForeground");
-	return true;
+	extern FLinuxApplication* LinuxApplication;
+	return (LinuxApplication != nullptr) ? LinuxApplication->IsForeground() : true;
 }
 
 bool FLinuxPlatformProcess::IsApplicationRunning( const TCHAR* ProcName )
diff --git a/Engine/Source/Runtime/Core/Public/Linux/LinuxApplication.h b/Engine/Source/Runtime/Core/Public/Linux/LinuxApplication.h
--- a/Engine/Source/Runtime/Core/Public/Linux/LinuxApplication.h
+++ b/Engine/Source/Runtime/Core/Public/Linux/LinuxApplication.h
@@ -78,6 +78,13 @@ public:
 
 	virtual bool IsCursorDirectlyOverSlateWindow() const override;
 
+	/** Returns true if this application is foreground */
+	FORCEINLINE bool IsForeground()
+	{
+		// if there are no windows, consider ourselves foreground so servers and commandlets aren't impacted
+		return (Windows.Num() > 0) ? bActivateApp : true;
+	}
+
 private:
 
 	FLinuxApplication();
