Changing wording of localization text manager logging in an attempt to make it seem less like an error.

diff --git a/Engine/Source/Runtime/Core/Private/Internationalization/TextLocalizationManager.cpp b/Engine/Source/Runtime/Core/Private/Internationalization/TextLocalizationManager.cpp
--- a/Engine/Source/Runtime/Core/Private/Internationalization/TextLocalizationManager.cpp
+++ b/Engine/Source/Runtime/Core/Private/Internationalization/TextLocalizationManager.cpp
@@ -112,13 +112,13 @@ void EndInitTextLocalization()
 				if(RequestedCultureName != ValidCultureName)
 				{
 					// Make the user aware that the localization data belongs to a parent culture.
-					UE_LOG(LogTextLocalizationManager, Log, TEXT("The requested culture ('%s') has no localization data; parent culture's ('%s') localization data will be used."), *RequestedCultureName, *ValidCultureName);
+					UE_LOG(LogTextLocalizationManager, Log, TEXT("No specific translations for ('%s') exist, so ('%s') translations will be used."), *RequestedCultureName, *ValidCultureName);
 				}
 			}
 			else
 			{
 				// Fallback to English.
-				UE_LOG(LogTextLocalizationManager, Log, TEXT("The requested culture ('%s') has no localization data; falling back to 'en' for localization and internationalization data."), *RequestedCultureName);
+				UE_LOG(LogTextLocalizationManager, Log, TEXT("No translations for ('%s') exist, falling back to 'en' for localization and internationalization data."), *RequestedCultureName);
 				TargetCultureName = TEXT("en");
 			}
 		}
