Stats - Added command 'stat slow' (UECORE-45)
Stats - Don't cull children, if there is only one child

stat slow [-ms=1.0] [-depth=4] - display the game and render thread stats

diff --git a/Engine/Source/Runtime/Core/Private/Stats/StatsCommand.cpp b/Engine/Source/Runtime/Core/Private/Stats/StatsCommand.cpp
--- a/Engine/Source/Runtime/Core/Private/Stats/StatsCommand.cpp
+++ b/Engine/Source/Runtime/Core/Private/Stats/StatsCommand.cpp
@@ -43,15 +43,17 @@ struct FGroupFilter : public IItemFiler
 };
 
 /** Holds parameters used by the 'stat hier' or 'stat group ##' command. */
-struct FStatGroupParams
+struct FStatParams
 {
 	/** Default constructor. */
-	FStatGroupParams( const TCHAR* Cmd = nullptr )
+	FStatParams( const TCHAR* Cmd = nullptr )
 		: Group( Cmd, TEXT("group="), NAME_None )
 		, SortBy( Cmd, TEXT("sortby="), EStatCompareBy::Sum )
 		, MaxHistoryFrames( Cmd, TEXT("maxhistoryframes="), 60 )
-		, MaxHierarchyDepth( Cmd, TEXT("maxdepth="), 16 )
+		, MaxHierarchyDepth( Cmd, TEXT("maxdepth="), 4 )
+		, CullMs( Cmd, TEXT( "ms=" ), 0.2f )
 		, bReset( FCString::Stristr( Cmd, TEXT("-reset") ) != nullptr )
+		, bSlowMode( false )
 	{}
 
 	/**
@@ -72,7 +74,7 @@ struct FStatGroupParams
 	 *	Maximum number of frames to be included in the history. 
 	 *	-maxhistoryframes=[20:20-120]
 	 */
-	// @TODO yrx 2014-08-21 Replace with TParsedValueWithDefaultAndRange
+	// #YRX_STATS 2014-08-21 Replace with TParsedValueWithDefaultAndRange
 	TParsedValueWithDefault<int32> MaxHistoryFrames;
 
 	/**
@@ -81,8 +83,34 @@ struct FStatGroupParams
 	 */
 	TParsedValueWithDefault<int32> MaxHierarchyDepth;
 
+	/**
+	 *	Threshold when start culling stats, 
+	 *	if 0, disables culling
+	 * -ms=5.0f
+	 */
+	TParsedValueWithDefault<float> CullMs;
+
 	/** Whether to reset all collected data. */
 	bool bReset;
+
+	/** Whether to use the slow mode, which displays stats stack for the game and rendering thread. */
+	bool bSlowMode;
+};
+
+/** Holds parameters used by the 'stat slow' command. */
+struct FStatSlowParams : public FStatParams
+{
+	/** Default constructor. */
+	FStatSlowParams( const TCHAR* Cmd = nullptr )
+		: FStatParams( Cmd )
+	{
+		static FName NAME_Slow = TEXT( "Slow" );
+		Group = TParsedValueWithDefault<FName>( Cmd, TEXT( "group=" ), NAME_Slow );
+		CullMs = TParsedValueWithDefault<float>( Cmd, TEXT( "ms=" ), 1.0f );
+		MaxHierarchyDepth = TParsedValueWithDefault<int32>( Cmd, TEXT( "maxdepth=" ), 4 );
+		bSlowMode = true;
+		bReset = true;
+	}
 };
 
 void DumpHistoryFrame(FStatsThreadState const& StatsData, int64 TargetFrame, float DumpCull = 0.0f, int32 MaxDepth = MAX_int32, TCHAR const* Filter = NULL)
@@ -100,13 +128,13 @@ void DumpHistoryFrame(FStatsThreadState const& StatsData, int64 TargetFrame, flo
 		UE_LOG(LogStats, Log, TEXT("Stack ---------------"));
 		FRawStatStackNode Stack;
 		StatsData.UncondenseStackStats(TargetFrame, Stack);
+		Stack.AddSelf();
 		if (DumpCull != 0.0f)
 		{
-			Stack.Cull(int64(DumpCull / FPlatformTime::ToMilliseconds(1.0f)));
+			Stack.CullByCycles( int64( DumpCull / FPlatformTime::ToMilliseconds( 1 ) ) );		
 		}
-		Stack.AddNameHierarchy();
-		Stack.AddSelf();
-		Stack.DebugPrint(Filter, MaxDepth);
+		Stack.CullByDepth( MaxDepth );
+		Stack.DebugPrint(Filter);
 	}
 	if (DumpCull == 0.0f)
 	{
@@ -372,14 +400,14 @@ static void DumpHitch(int64 Frame)
 			{
 				GameThread = ChildIter.Value();
 				UE_LOG( LogStats, Log, TEXT( "------------------ Game Thread %.2fms" ), GameThreadTime * 1000.0f );
-				GameThread->Cull( MinCycles );
+				GameThread->CullByCycles( MinCycles );
 				GameThread->DebugPrint();
 			}
 			else if( ThreadName == FName( NAME_RenderThread ) )
 			{
 				RenderThread = ChildIter.Value();
 				UE_LOG( LogStats, Log, TEXT( "------------------ Render Thread (%s) %.2fms" ), *RenderThread->Meta.NameAndInfo.GetRawName().ToString(), RenderThreadTime * 1000.0f );
-				RenderThread->Cull( MinCycles );
+				RenderThread->CullByCycles( MinCycles );
 				RenderThread->DebugPrint();
 			}
 		}
@@ -400,8 +428,10 @@ static void DumpHitch(int64 Frame)
 
 #endif
 
-static bool HandleCommandBroadcast(const FName& InStatName, bool& bOutCurrentEnabled, bool& bOutOthersEnabled)
+static bool HandleToggleCommandBroadcast(const FName& InStatName, bool& bOutCurrentEnabled, bool& bOutOthersEnabled)
 {
+	// !!! Not thread-safe, calling game thread code from the stats thread. !!!
+
 	bOutCurrentEnabled = true;
 	bOutOthersEnabled = false;
 
@@ -457,14 +487,14 @@ FStatGroupGameThreadNotifier& FStatGroupGameThreadNotifier::Get()
 struct FInternalGroup
 {
 	/** Initialization constructor. */
-	FInternalGroup(const FName InGroupName, const FName InGroupCategory, const EStatDisplayMode::Type InDisplayMode, TSet<FName>& InEnabledItems, FString& InGroupDescription)
+	FInternalGroup(const FName InGroupName, const FName InGroupCategory, const EStatDisplayMode::Type InDisplayMode, TSet<FName>& InEnabledItems, const FString& InGroupDescription)
 		: GroupName( InGroupName )
 		, GroupCategory(InGroupCategory)
 		, DisplayMode( InDisplayMode )
+		, GroupDescription( InGroupDescription )
 	{
 		// To avoid copy.
 		Exchange( EnabledItems, InEnabledItems );
-		Exchange( GroupDescription, InGroupDescription );
 	}
 
 	/** Set of elements which should be included in this group stats. */
@@ -513,7 +543,7 @@ struct FHUDGroupManager
 	TArray<FComplexStatMessage> AggregatedNonStackStatsHistory;
 
 	/** Copy of the stat group command parameters. */
-	FStatGroupParams Params;
+	FStatParams Params;
 	
 	/** Number of frames for the root stat stack. */
 	int32 NumTotalStackFrames;
@@ -540,19 +570,37 @@ struct FHUDGroupManager
 	}
 
 	/** Handles hier or group command. */
-	void HandleCommand( const FStatGroupParams& InParams, const bool bHierarchy )
+	void HandleCommand( const FStatParams& InParams, const bool bHierarchy )
 	{
+		bool bCurrentEnabled, bOthersEnabled;
+
+		bool bResetData = false;
+		if (Params.bSlowMode != InParams.bSlowMode)
+		{
+			bResetData = true;
+		}
+
 		Params = InParams;
-		if( Params.ShouldReset() )
+		Params.bReset = bResetData;
+
+		if (Params.ShouldReset())
 		{
+			// Disable only stats groups, leave the fake FPS, Unit group untouched.
+			for (const auto& It : EnabledGroups)
+			{
+				HandleToggleCommandBroadcast( It.Key, bCurrentEnabled, bOthersEnabled );
+			}
+
+			EnabledGroups.Empty();
+			History.Empty();
 			NumTotalStackFrames = 0;
 		}
 
 		ResizeFramesHistory( Params.MaxHistoryFrames.Get() );
-
-		const FName MaybeGroupFName = FName(*(FString(TEXT("STATGROUP_")) + Params.Group.Get().GetPlainNameString()));
-		bool bCurrentEnabled, bOthersEnabled;
-		if (!HandleCommandBroadcast(MaybeGroupFName, bCurrentEnabled, bOthersEnabled))
+	
+		const FName MaybeGroupFName = FName(*(FString(TEXT("STATGROUP_")) + Params.Group.Get().GetPlainNameString()));	
+		const bool bResults = HandleToggleCommandBroadcast( MaybeGroupFName, bCurrentEnabled, bOthersEnabled );
+		if (!bResults)
 		{
 			// Remove all groups.
 			EnabledGroups.Empty();
@@ -591,14 +639,26 @@ struct FHUDGroupManager
 					GetStatsForGroup(EnabledItems, MaybeGroupFName);
 
 					const FStatMessage& Group = Stats.ShortNameToLongName.FindChecked(MaybeGroupFName);
-
 					const FName GroupCategory = Group.NameAndInfo.GetGroupCategory();
-
-					FString GroupDescription = Group.NameAndInfo.GetDescription();
+					const FString GroupDescription = Group.NameAndInfo.GetDescription();
 
 					EnabledGroups.Add(MaybeGroupFName, FInternalGroup(MaybeGroupFName, GroupCategory, bHierarchy ? EStatDisplayMode::Hierarchical : EStatDisplayMode::Flat, EnabledItems, GroupDescription));
 				}
 			}
+			else if (Params.bSlowMode)
+			{
+				const bool bEnabledSlowMode = EnabledGroups.Contains( MaybeGroupFName );
+				if (bEnabledSlowMode)
+				{
+					EnabledGroups.Remove( MaybeGroupFName );
+					NumTotalStackFrames = 0;
+				}
+				else
+				{
+					TSet<FName> EmptySet = TSet<FName>();
+					EnabledGroups.Add( MaybeGroupFName, FInternalGroup( MaybeGroupFName, NAME_None, EStatDisplayMode::Hierarchical, EmptySet, TEXT( "Slow hierarchy for game and render" ) ) );
+				}			
+			}
 		}
 
 		if( EnabledGroups.Num() && !bEnabled )
@@ -646,33 +706,78 @@ struct FHUDGroupManager
 		}
 	}
 
+	void LinearizeSlowStackForItems( const FComplexRawStatStackNode& StackNode, TArray<FComplexStatMessage>& out_HistoryStack, TArray<int32>& out_Indentation, int32 Depth )
+	{
+		// Ignore first call, this is the thread root.
+		const bool bToBeAdded = Depth > 0;// StackNode.ComplexStat.GetShortName() != FStatConstants::NAME_ThreadRoot;
+		if (bToBeAdded)
+		{
+			out_HistoryStack.Add( StackNode.ComplexStat );
+			out_Indentation.Add( Depth );
+		}
+
+		for (auto It = StackNode.Children.CreateConstIterator(); It; ++It)
+		{
+			const FComplexRawStatStackNode& Child = *It.Value();
+			LinearizeSlowStackForItems( Child, out_HistoryStack, out_Indentation, Depth + 1 );
+		}
+	}
+
 	void NewFrame(int64 TargetFrame)
 	{
 		SCOPE_CYCLE_COUNTER(STAT_HUDGroup);
 		check(bEnabled);
 
-		TSet<FName> HierEnabledItems;
-		for( auto It = EnabledGroups.CreateConstIterator(); It; ++It )
+		// Add a new frame to the history.
+		FHudFrame& NewFrame = History.FindOrAdd( TargetFrame );
+
+		const bool bUseSlowMode = Params.bSlowMode;
+		if (bUseSlowMode)
 		{
-			HierEnabledItems.Append( It.Value().EnabledItems );
-		}
-		FGroupFilter Filter(HierEnabledItems);
+			// Only for game thread and rendering thread.
+			Stats.UncondenseStackStats( TargetFrame, NewFrame.HierarchyInclusive, nullptr, nullptr );
 
-		// Add a new frame to the history.
-		FHudFrame& NewFrame = History.FindOrAdd(TargetFrame);
-		
-		// Generate root stats stack for current frame.
-		Stats.UncondenseStackStats( TargetFrame, NewFrame.HierarchyInclusive, &Filter, &NewFrame.NonStackStats );
-		NewFrame.HierarchyInclusive.Cull( TNumericLimits<int64>::Max(), Params.MaxHierarchyDepth.Get() );
-		//NewFrame.HierarchyInclusive.AddNameHierarchy();
-		NewFrame.HierarchyInclusive.AddSelf();
+			for (auto ChildIt = NewFrame.HierarchyInclusive.Children.CreateIterator(); ChildIt; ++ChildIt)
+			{
+				const FName ThreadName = ChildIt.Value()->Meta.NameAndInfo.GetShortName();
 
+				if (ThreadName == NAME_GameThread)
+				{
+					continue;
+				}
+				else if (ThreadName == NAME_RenderThread)
+				{
+					continue;
+				}
+
+				delete ChildIt.Value();
+				ChildIt.RemoveCurrent();
+			}
+		}
+		else
 		{
-			SCOPE_CYCLE_COUNTER(STAT_GetFlatAggregates);
-			Stats.GetInclusiveAggregateStackStats( TargetFrame, NewFrame.InclusiveAggregate, &Filter, false );
-			Stats.GetExclusiveAggregateStackStats( TargetFrame, NewFrame.ExclusiveAggregate, &Filter, false );
+			TSet<FName> HierEnabledItems;
+			for (auto It = EnabledGroups.CreateConstIterator(); It; ++It)
+			{
+				HierEnabledItems.Append( It.Value().EnabledItems );
+			}
+			FGroupFilter Filter( HierEnabledItems );		
+
+			// Generate root stats stack for current frame.
+			Stats.UncondenseStackStats( TargetFrame, NewFrame.HierarchyInclusive, &Filter, &NewFrame.NonStackStats );
+
+			{
+				SCOPE_CYCLE_COUNTER( STAT_GetFlatAggregates );
+				Stats.GetInclusiveAggregateStackStats( TargetFrame, NewFrame.InclusiveAggregate, &Filter, false );
+				Stats.GetExclusiveAggregateStackStats( TargetFrame, NewFrame.ExclusiveAggregate, &Filter, false );
+			}
 		}
 
+		NewFrame.HierarchyInclusive.AddSelf();
+		// To get the good performance we must pre-filter the results.
+		NewFrame.HierarchyInclusive.CullByCycles( int64( 0.001f / FPlatformTime::GetSecondsPerCycle() * 0.1f ) );
+		NewFrame.HierarchyInclusive.CullByDepth( Params.MaxHierarchyDepth.Get() );
+
 		// Aggregate hierarchical stats.
 		if( NumTotalStackFrames == 0 )
 		{
@@ -795,7 +900,15 @@ struct FHUDGroupManager
 
 			// Divide stats to get average values.
 			AggregatedHierarchyHistory.Divide( NumFrames );
-			AggregatedHierarchyHistory.CopyExclusivesFromSelf();
+			AggregatedHierarchyHistory.CopyExclusivesFromSelf();	
+			if (Params.CullMs.Get() != 0.0f)
+			{
+				AggregatedHierarchyHistory.CullByCycles( int64( Params.CullMs.Get() / FPlatformTime::ToMilliseconds( 1 ) ) );
+			}
+			AggregatedHierarchyHistory.CullByDepth( Params.MaxHierarchyDepth.Get() );
+
+			// Make sure the game thread is first.
+			AggregatedHierarchyHistory.Children.KeySort( TLess<FName>() );
 
 			FComplexStatUtils::DiviveStatArray( AggregatedFlatHistory, NumFrames, EComplexStatField::IncSum, EComplexStatField::IncAve );
 			FComplexStatUtils::DiviveStatArray( AggregatedFlatHistory, NumFrames, EComplexStatField::ExcSum, EComplexStatField::ExcAve );
@@ -815,50 +928,69 @@ struct FHUDGroupManager
 				ToGame->GroupNames.Add( GroupName );
 				ToGame->GroupDescriptions.Add( InternalGroup.GroupDescription );
 
-				if( InternalGroup.DisplayMode & EStatDisplayMode::Hierarchical )
+				if (Params.bSlowMode)
 				{
 					// Linearize stack stats for easier rendering.
-					LinearizeStackForItems( AggregatedHierarchyHistory, InternalGroup.EnabledItems, HudGroup.HierAggregate, HudGroup.Indentation, 0 );
+					LinearizeSlowStackForItems( AggregatedHierarchyHistory, HudGroup.HierAggregate, HudGroup.Indentation, 0 );
 				}
-				
-				if( InternalGroup.DisplayMode & EStatDisplayMode::Flat )
+				else
 				{
-					// Copy flat stats
-					for( int32 Index = 0; Index < AggregatedFlatHistory.Num(); ++Index )
+					if (InternalGroup.DisplayMode & EStatDisplayMode::Hierarchical)
+					{
+						// Linearize stack stats for easier rendering.
+						LinearizeStackForItems( AggregatedHierarchyHistory, InternalGroup.EnabledItems, HudGroup.HierAggregate, HudGroup.Indentation, 0 );
+					}
+
+					if (InternalGroup.DisplayMode & EStatDisplayMode::Flat)
+					{
+						// Copy flat stats
+						for (int32 Index = 0; Index < AggregatedFlatHistory.Num(); ++Index)
+						{
+							const FComplexStatMessage& AggregatedStatMessage = AggregatedFlatHistory[Index];
+							const bool bIsNonStackStat = !AggregatedStatMessage.NameAndInfo.GetFlag( EStatMetaFlags::IsPackedCCAndDuration );
+							const bool bToBeAdded = InternalGroup.EnabledItems.Contains( AggregatedStatMessage.NameAndInfo.GetRawName() );
+							if (bToBeAdded)
+							{
+								new(HudGroup.FlatAggregate) FComplexStatMessage( AggregatedStatMessage );
+							}
+						}
+					}
+
+					// Copy non-stack stats assigned to memory and counter groups.
+					for (int32 Index = 0; Index < AggregatedNonStackStatsHistory.Num(); ++Index)
 					{
-						const FComplexStatMessage& AggregatedStatMessage = AggregatedFlatHistory[Index];
-						const bool bIsNonStackStat = !AggregatedStatMessage.NameAndInfo.GetFlag(EStatMetaFlags::IsPackedCCAndDuration);
+						const FComplexStatMessage& AggregatedStatMessage = AggregatedNonStackStatsHistory[Index];
+						const bool bIsMemory = AggregatedStatMessage.NameAndInfo.GetFlag( EStatMetaFlags::IsMemory );
+						TArray<FComplexStatMessage>& Dest = bIsMemory ? HudGroup.MemoryAggregate : HudGroup.CountersAggregate;
+
 						const bool bToBeAdded = InternalGroup.EnabledItems.Contains( AggregatedStatMessage.NameAndInfo.GetRawName() );
-						if( bToBeAdded )
+						if (bToBeAdded)
 						{
-							new(HudGroup.FlatAggregate) FComplexStatMessage( AggregatedStatMessage );
+							new(Dest)FComplexStatMessage( AggregatedStatMessage );
 						}
 					}
 				}
 
-				// Copy non-stack stats assigned to memory and counter groups.
-				for( int32 Index = 0; Index < AggregatedNonStackStatsHistory.Num(); ++Index )
+				// Replace thread encoded id with the thread name.
+				for (auto& It : HudGroup.HierAggregate)
 				{
-					const FComplexStatMessage& AggregatedStatMessage = AggregatedNonStackStatsHistory[Index];
-					const bool bIsMemory = AggregatedStatMessage.NameAndInfo.GetFlag( EStatMetaFlags::IsMemory );
-					TArray<FComplexStatMessage>& Dest = bIsMemory ? HudGroup.MemoryAggregate : HudGroup.CountersAggregate; 
-
-					const bool bToBeAdded = InternalGroup.EnabledItems.Contains( AggregatedStatMessage.NameAndInfo.GetRawName() );
-					if( bToBeAdded )
+					FComplexStatMessage& StatMessage = It;
+					const FString StatDescription = StatMessage.NameAndInfo.GetDescription();
+					if (StatDescription.Contains( FStatConstants::ThreadNameMarker ))
 					{
-						new(Dest) FComplexStatMessage(AggregatedStatMessage);
-					}	
+						StatMessage.NameAndInfo.SetRawName( StatMessage.NameAndInfo.GetShortName() );
+					}
 				}
 			}
 
 			for (auto It = Stats.MemoryPoolToCapacityLongName.CreateConstIterator(); It; ++It)
 			{
-				FName LongName = It.Value();
+				const FName LongName = It.Value();
 				// dig out the abbreviation
 				{
-					FString LongNameStr = LongName.ToString();
-					int32 Open = LongNameStr.Find("[");
-					int32 Close = LongNameStr.Find("]");
+					const FString LongNameStr = LongName.ToString();
+					const int32 Open = LongNameStr.Find("[");
+					const int32 Close = LongNameStr.Find("]");
 					if (Open >= 0 && Close >= 0 && Open + 1 < Close)
 					{
 						FString Abbrev = LongNameStr.Mid(Open + 1, Close - Open - 1);
@@ -869,7 +1001,7 @@ struct FHUDGroupManager
 				FStatMessage const* Result = Stats.NotClearedEveryFrame.Find(LongName);
 				if (Result && Result->NameAndInfo.GetFlag(EStatMetaFlags::IsMemory))
 				{
-					int64 Capacity = Result->GetValue_int64();
+					const int64 Capacity = Result->GetValue_int64();
 					if (Capacity > 0)
 					{
 						ToGame->PoolCapacity.Add(It.Key(), Capacity);
@@ -902,28 +1034,6 @@ struct FHUDGroupManager
 		check(History.Num() <= Params.MaxHistoryFrames.Get());
 	}
 
-	void RebuildItems( TSet<FName>& out_EnabledItems, const TSet<FName>& InEnabledGroups )
-	{
-		out_EnabledItems.Empty();
-		for (auto It = InEnabledGroups.CreateConstIterator(); It; ++It)
-		{
-			TArray<FName> GroupItems;
-			Stats.Groups.MultiFind(*It, GroupItems);
-			for (int32 Index = 0; Index < GroupItems.Num(); Index++)
-			{
-				out_EnabledItems.Add(GroupItems[Index]); // short name
-				FStatMessage const* LongName = Stats.ShortNameToLongName.Find(GroupItems[Index]);
-				if (LongName)
-				{
-					out_EnabledItems.Add(LongName->NameAndInfo.GetRawName()); // long name
-				}
-			}
-		}
-
-		out_EnabledItems.Add(NAME_Self);
-		out_EnabledItems.Add(NAME_OtherChildren);
-	}
-
 	void GetStatsForGroup( TSet<FName>& out_EnabledItems, const FName GroupName )
 	{
 		out_EnabledItems.Empty();
@@ -1031,7 +1141,7 @@ struct FDumpMultiple
 			Stack->AddSelf();
 			if (DumpCull != 0.0f)
 			{
-				Stack->Cull(int64(DumpCull / FPlatformTime::ToMilliseconds(1.0f)));
+				Stack->CullByCycles( int64( DumpCull / FPlatformTime::ToMilliseconds( 1 ) ) );
 			}
 			Stack->DebugPrint(*NameFilter, MaxDepth);
 			delete Stack;
@@ -1085,11 +1195,12 @@ static void PrintStatsHelpToOutputDevice( FOutputDevice& Ar )
 	Ar.Log( TEXT("stat dumpcpu - dumps cpu stats"));
 
 	Ar.Log( TEXT("stat groupname[+] - toggles displaying stats group, + enables hierarchical display"));
-	Ar.Log( TEXT("stat hier -group=groupname [-sortby=name] [-maxhistoryframes=60] [-reset]"));
+	Ar.Log( TEXT("stat hier -group=groupname [-sortby=name] [-maxhistoryframes=60] [-reset] [-maxdepth=4]"));
 	Ar.Log( TEXT("    - groupname is a stat group like initviews or statsystem"));
 	Ar.Log( TEXT("    - sortby can be name (by stat FName), callcount (by number of calls, only for scoped cycle counters), num(by total inclusive time)"));
 	Ar.Log( TEXT("    - maxhistoryframes (default 60, number of frames used to generate the stats displayed on the hud)"));
 	Ar.Log( TEXT("    - reset (reset the accumulated history)"));
+	Ar.Log( TEXT("    - maxdepth (default 4, maximum depth for the hierarchy)"));
 	Ar.Log( TEXT("stat none - disables drawing all stats groups"));
 
 	Ar.Log( TEXT("stat group list|listall|enable name|disable name|none|all|default - manages stats groups"));
@@ -1107,6 +1218,8 @@ static void PrintStatsHelpToOutputDevice( FOutputDevice& Ar )
 
 	Ar.Log( TEXT("stat toggledebug - toggles tracking the most memory expensive stats"));
 
+	Ar.Log( TEXT( "stat slow [-ms=1.0] [-depth=4] - display the game and render thread stats" ) );
+
 	Ar.Log( TEXT("add -memoryprofiler in the command line to enable the memory profiling"));
 	Ar.Log( TEXT("stat stopfile - stops tracking all memory operations and writes the results to the file"));
 }
@@ -1285,7 +1398,7 @@ static void StatCmd(FString InCmd)
 	}
 	else if ( FParse::Command( &Cmd, TEXT( "none" ) ) )
 	{
-		FStatGroupParams Params;
+		FStatParams Params;
 		FHUDGroupManager::Get(Stats).HandleCommand(Params,false);
 	}
 	else if ( FParse::Command( &Cmd, TEXT( "group" ) ) )
@@ -1301,14 +1414,19 @@ static void StatCmd(FString InCmd)
 		if( FParse::Command( &Cmd, TEXT( "snapshot" ) ) )
 		{
 			// Put a snapshot marker in the raw stats.
-			// @TODO yrx 2014-12-03 
+			// #YRX_STATS 2014-12-03 
 		}
 	}
 	// @see FStatHierParams
 	else if ( FParse::Command( &Cmd, TEXT( "hier" ) ) )
 	{
-		FStatGroupParams Params( Cmd );
-		FHUDGroupManager::Get(Stats).HandleCommand(Params, true);
+		FStatParams Params( Cmd );
+		FHUDGroupManager::Get( Stats ).HandleCommand( Params, true );
+	}
+	else if (FParse::Command( &Cmd, TEXT( "slow" ) ))
+	{
+		FStatSlowParams Params( Cmd );
+		FHUDGroupManager::Get( Stats ).HandleCommand( Params, true );
 	}
 	else
 #endif
@@ -1329,13 +1447,13 @@ static void StatCmd(FString InCmd)
 			const FName MaybeGroupFName = FName(*MaybeGroup);
 #if STATS
 			// Try to parse.
-			FStatGroupParams Params( Cmd );
+			FStatParams Params( Cmd );
 			Params.Group.Set( MaybeGroupFName );
 			FHUDGroupManager::Get(Stats).HandleCommand(Params, bHierarchy);
 #else
 			// If stats aren't enabled, broadcast so engine stats can still be triggered
 			bool bCurrentEnabled, bOthersEnabled;
-			HandleCommandBroadcast(MaybeGroupFName, bCurrentEnabled, bOthersEnabled);
+			HandleToggleCommandBroadcast(MaybeGroupFName, bCurrentEnabled, bOthersEnabled);
 #endif
 		}
 		else
@@ -1435,6 +1553,9 @@ bool DirectStatsCommand(const TCHAR* Cmd, bool bBlockForCompletion /*= false*/,
 		else if ( FParse::Command( &TempCmd, TEXT( "memoryprofiler" ) ) )
 		{
 		}
+		else if( FParse::Command( &TempCmd, TEXT( "slow" ) ) )
+		{
+		}
 		else
 		{
 			bResult = false;
diff --git a/Engine/Source/Runtime/Core/Private/Stats/StatsData.cpp b/Engine/Source/Runtime/Core/Private/Stats/StatsData.cpp
--- a/Engine/Source/Runtime/Core/Private/Stats/StatsData.cpp
+++ b/Engine/Source/Runtime/Core/Private/Stats/StatsData.cpp
@@ -97,33 +97,66 @@ void FRawStatStackNode::Divide(uint32 Div)
 	}
 }
 
-
-void FRawStatStackNode::Cull(int64 MinCycles, int32 NoCullLevels)
+void FRawStatStackNode::CullByCycles( int64 MinCycles )
 {
-	FRawStatStackNode* Culled = NULL;
-	for (TMap<FName, FRawStatStackNode*>::TIterator It(Children); It; ++It)
+	FRawStatStackNode* Culled = nullptr;
+	const int32 NumChildren = Children.Num();
+	for (TMap<FName, FRawStatStackNode*>::TIterator It( Children ); It; ++It)
 	{
 		FRawStatStackNode* Child = It.Value();
-		if (NoCullLevels < 1 && FromPackedCallCountDuration_Duration(Child->Meta.GetValue_int64()) < MinCycles)
+		const int64 ChildCycles = Child->Meta.GetValue_Duration();
+		if (FromPackedCallCountDuration_Duration( Child->Meta.GetValue_int64() ) < MinCycles)
 		{
-			if (!Culled)
+			// Don't accumulate if we have just one child.
+			if (NumChildren > 1)
 			{
-				Culled = new FRawStatStackNode(FStatMessage(NAME_OtherChildren, EStatDataType::ST_int64, NULL, NULL, NULL, true, true));
-				Culled->Meta.NameAndInfo.SetFlag(EStatMetaFlags::IsPackedCCAndDuration, true);
-				Culled->Meta.Clear();
+				if (!Culled)
+				{
+					Culled = new FRawStatStackNode( FStatMessage( NAME_OtherChildren, EStatDataType::ST_int64, nullptr, nullptr, nullptr, true, true ) );
+					Culled->Meta.NameAndInfo.SetFlag( EStatMetaFlags::IsPackedCCAndDuration, true );
+					Culled->Meta.Clear();
+				}
+				FStatsUtils::AccumulateStat( Culled->Meta, Child->Meta, EStatOperation::Add, true );
+				delete Child;
+				It.RemoveCurrent();
+			}
+			else
+			{
+				// Remove children.
+				for (TMap<FName, FRawStatStackNode*>::TIterator It( Child->Children ); It; ++It)
+				{
+					delete It.Value();
+					It.RemoveCurrent();
+				}
 			}
-			FStatsUtils::AccumulateStat(Culled->Meta, Child->Meta, EStatOperation::Add, true);
-			delete Child;
-			It.RemoveCurrent();
 		}
-		else
+		else if (NumChildren > 0)
 		{
-			Child->Cull(MinCycles, NoCullLevels - 1);
+			Child->CullByCycles( MinCycles );
 		}
 	}
 	if (Culled)
 	{
-		Children.Add(NAME_OtherChildren, Culled);
+		Children.Add( NAME_OtherChildren, Culled );
+	}
+}
+
+void FRawStatStackNode::CullByDepth( int32 NoCullLevels )
+{
+	FRawStatStackNode* Culled = nullptr;
+	const int32 NumChildren = Children.Num();
+	for (TMap<FName, FRawStatStackNode*>::TIterator It( Children ); It; ++It)
+	{
+		FRawStatStackNode* Child = It.Value();
+		if (NoCullLevels < 1)
+		{
+			delete Child;
+			It.RemoveCurrent();
+		}
+		else
+		{
+			Child->CullByDepth( NoCullLevels - 1 );
+		}
 	}
 }
 
@@ -293,7 +326,7 @@ void FRawStatStackNode::DebugPrint(TCHAR const* Filter, int32 InMaxDepth, int32
 			UE_LOG(LogStats, Log, TEXT("%s%s"), FCString::Spc(Depth*2), *TmpDebugStr);
 		}
 
-		static int64 MinPrint = int64(.004f / FPlatformTime::ToMilliseconds(1.0f) + 0.5f);
+		static int64 MinPrint = int64( .004f / FPlatformTime::ToMilliseconds( 1 ) + 0.5f );
 		if (Children.Num())
 		{
 			TArray<FRawStatStackNode*> ChildArray;
@@ -309,7 +342,7 @@ void FRawStatStackNode::DebugPrint(TCHAR const* Filter, int32 InMaxDepth, int32
 				{
 					if (ChildArray[Index]->Meta.NameAndInfo.GetRawName().ToString().Contains(Filter))
 					{
-						ChildArray[Index]->DebugPrint(NULL, InMaxDepth, 0);
+						ChildArray[Index]->DebugPrint(nullptr, InMaxDepth, 0);
 					}
 					else
 					{
@@ -433,18 +466,75 @@ void FComplexRawStatStackNode::Divide(uint32 Div)
 	}
 }
 
+void FComplexRawStatStackNode::CullByCycles( int64 MinCycles )
+{
+	FComplexRawStatStackNode* Culled = nullptr;
+	const int32 NumChildren = Children.Num();
+	for (auto It = Children.CreateIterator(); It; ++It)
+	{
+		FComplexRawStatStackNode* Child = It.Value();
+		const int64 ChildCycles = Child->ComplexStat.GetValue_Duration( EComplexStatField::IncAve );
+		if (ChildCycles < MinCycles)
+		{
+			// Don't accumulate if we have just one child.
+			if (NumChildren > 1)
+			{	
+				// #YRX_STATS: 2015-06-09 Accumulate over complex stats
+				delete Child;
+				It.RemoveCurrent();
+			}
+			else
+			{
+				// Remove children.
+				for (auto ChildIt = Child->Children.CreateIterator(); ChildIt; ++ChildIt)
+				{
+					delete ChildIt.Value();
+					ChildIt.RemoveCurrent();
+				}
+			}
+		}
+		else if (NumChildren > 0)
+		{
+			Child->CullByCycles( MinCycles );
+		}
+	}
+	if (Culled)
+	{
+		Children.Add( NAME_OtherChildren, Culled );
+	}
+}
+
+void FComplexRawStatStackNode::CullByDepth( int32 NoCullLevels )
+{
+	FComplexRawStatStackNode* Culled = nullptr;
+	const int32 NumChildren = Children.Num();
+	for (auto It = Children.CreateIterator(); It; ++It)
+	{
+		FComplexRawStatStackNode* Child = It.Value();
+		if (NoCullLevels < 1)
+		{
+			delete Child;
+			It.RemoveCurrent();
+		}
+		else
+		{
+			Child->CullByDepth( NoCullLevels - 1 );
+		}
+	}
+}
+
 void FComplexRawStatStackNode::CopyExclusivesFromSelf()
 {
-	if( Children.Num() )
+	if (Children.Num())
 	{
 		const FComplexRawStatStackNode* SelfStat = Children.FindRef( NAME_Self );
-		if( SelfStat )
+		if (SelfStat)
 		{
-			ComplexStat.GetValue_int64(EComplexStatField::ExcAve) = SelfStat->ComplexStat.GetValue_int64(EComplexStatField::IncAve);
-			ComplexStat.GetValue_int64(EComplexStatField::ExcMax) = SelfStat->ComplexStat.GetValue_int64(EComplexStatField::IncMax);
+			ComplexStat.GetValue_int64( EComplexStatField::ExcAve ) = SelfStat->ComplexStat.GetValue_int64( EComplexStatField::IncAve );
+			ComplexStat.GetValue_int64( EComplexStatField::ExcMax ) = SelfStat->ComplexStat.GetValue_int64( EComplexStatField::IncMax );
 		}
 
-		for( auto It = Children.CreateIterator(); It; ++It )
+		for (auto It = Children.CreateIterator(); It; ++It)
 		{
 			FComplexRawStatStackNode* Child = It.Value();
 			Child->CopyExclusivesFromSelf();
@@ -483,10 +573,9 @@ FStatsThreadState::FStatsThreadState(int32 InHistoryFrames)
 {
 }
 
-// FStatsFileState:: ????
 //FStatsThreadState::FStatsThreadState(FString const& Filename)
-// void FStatsThreadState::AddMessages(TArray<FStatMessage>& InMessages)
-// @see moved to StatsFile.cpp
+//void FStatsThreadState::AddMessages(TArray<FStatMessage>& InMessages)
+//@see moved to StatsFile.cpp
 
 FStatsThreadState& FStatsThreadState::GetLocalState()
 {
@@ -732,7 +821,7 @@ void FStatsThreadState::AddToHistoryAndEmpty(FStatPacketArray& NewData)
 		{
 			FStatPacketArray& Frame = History.FindChecked(FrameNum);
 
-			FStatPacket const* PacketToCopyForNonFrame = NULL;
+			FStatPacket const* PacketToCopyForNonFrame = nullptr;
 
 			if( bFindMemoryExtensiveStats )
 			{
@@ -755,7 +844,7 @@ void FStatsThreadState::AddToHistoryAndEmpty(FStatPacketArray& NewData)
 
 				check(PacketToCopyForNonFrame);
 				FThreadStats* ThreadStats = FThreadStats::GetThreadStats();
-				FStatMessagesArray* NonFrameMessages = NULL;
+				FStatMessagesArray* NonFrameMessages = nullptr;
 
 				for (TMap<FName, FStatMessage>::TConstIterator It(NotClearedEveryFrame); It; ++It)
 				{
diff --git a/Engine/Source/Runtime/Core/Public/Stats/StatsData.h b/Engine/Source/Runtime/Core/Public/Stats/StatsData.h
--- a/Engine/Source/Runtime/Core/Public/Stats/StatsData.h
+++ b/Engine/Source/Runtime/Core/Public/Stats/StatsData.h
@@ -70,7 +70,10 @@ public:
 	TParsedValueWithDefault( const TCHAR* Stream, const TCHAR* Match, const T& Default )
 		: Value( Default )
 	{
-		ParseTypedValue<T>( Stream, Match, Value );
+		if (Stream && Match)
+		{
+			ParseTypedValue<T>( Stream, Match, Value );
+		}
 	}
 
 	const T& Get() const
@@ -254,8 +257,11 @@ struct CORE_API FRawStatStackNode
 	void MergeAdd(FRawStatStackNode const& Other);
 	void Divide(uint32 Div);
 
-	/** Cull this tree, merging children below MinCycles long **/
-	void Cull(int64 MinCycles, int32 NoCullLevels = 0);
+	/** Cull this tree, merging children below MinCycles long. */
+	void CullByCycles(int64 MinCycles);
+
+	/** Cull this tree, merging children below NoCullLevels. */
+	void CullByDepth(int32 NoCullLevels);
 
 	/** Adds name hiearchy. **/
 	void AddNameHierarchy(int32 CurrentPrefixDepth = 0);
@@ -377,6 +383,12 @@ public:
 	/** Divides this stack by the specified value. */
 	void Divide(uint32 Div);
 
+	/** Cull this tree, merging children below MinCycles long. */
+	void CullByCycles( int64 MinCycles );
+
+	/** Cull this tree, merging children below NoCullLevels. */
+	void CullByDepth( int32 NoCullLevels );
+
 	/** Copies exclusive times from the self node. **/
 	void CopyExclusivesFromSelf();
 
@@ -407,13 +419,14 @@ struct IItemFiler
 
 
 
-// @TODO yrx 2014-12-03 Separete stats thread state vs rawstats thread state?
-// @TODO yrx 2014-03-21 Move metadata functionality into a separate class
+// #YRX_STATS 2014-12-03 Separate stats thread state vs raw stats thread state?
+// #YRX_STATS 2014-03-21 Move metadata functionality into a separate class
+// 
 /**
-* Tracks stat state and history
-*  GetLocalState() is a singleton to the state for stats being collected in this executable.
-*  Other instances can be used to load stats for visualization.
-*/
+ * Tracks stat state and history
+ * GetLocalState() is a singleton to the state for stats being collected in this executable.
+ *  Other instances can be used to load stats for visualization.
+ */
 class CORE_API FStatsThreadState
 {
 	friend class FStatsThread;
@@ -449,6 +462,7 @@ public:
 	/** Toggles tracking the most memory expensive stats. */
 	void ToggleFindMemoryExtensiveStats();
 
+	/** Resets stats for raw stats. */
 	void ResetStatsForRawStats()
 	{
 		MaxNumStatMessages = 0;
@@ -457,10 +471,10 @@ public:
 
 private:
 	/** Internal method to accumulate any non-frame stats. **/
-	void ProcessNonFrameStats(FStatMessagesArray& Data, TSet<FName>* NonFrameStatsFound);
+	void ProcessNonFrameStats( FStatMessagesArray& Data, TSet<FName>* NonFrameStatsFound );
 
 	/** Internal method to place the data into the history, discard and broadcast any new frames to anyone who cares. **/
-	void AddToHistoryAndEmpty(FStatPacketArray& NewData);
+	void AddToHistoryAndEmpty( FStatPacketArray& NewData );
 
 	/** Does basic processing on the raw stats packets, discard and broadcast any new raw stats packets to anyone who cares. */
 	void ProcessRawStats(FStatPacketArray& NewData);
diff --git a/Engine/Source/Runtime/Engine/Private/StatsRender2.cpp b/Engine/Source/Runtime/Engine/Private/StatsRender2.cpp
--- a/Engine/Source/Runtime/Engine/Private/StatsRender2.cpp
+++ b/Engine/Source/Runtime/Engine/Private/StatsRender2.cpp
@@ -270,7 +270,7 @@ static int32 RenderCycle( const FComplexStatMessage& Item, class FCanvas* Canvas
 		const float InMs = FPlatformTime::ToMilliseconds(Item.GetValue_Duration(EComplexStatField::IncAve));
 		// Color will be determined by the average value of history
 		// If show inclusive and and show exclusive is on, then it will choose color based on inclusive average
-		// @TODO yrx 2014-08-21 This is slow, fix this.
+		// #YRX_STATS: 2015-06-09 This is slow, fix this
 		FString CounterName = Item.GetShortName().ToString();
 		CounterName.RemoveFromStart(TEXT("STAT_"));
 		GEngine->GetStatValueColoration(CounterName, InMs, Color);
