Modified the way the text filter handles "exact" string matching

UE-16745 - Cannot search for assets that have a dash (-) in the name

Previously, putting a string inside quotes would result in it performing an exact string match, however this left you with no way to perform a partial string match against a string containing special characters (such as -).

I've changed the behavior of the text filter syntax so that a + (pre-unary) denotes that a string should perform an exact match, and this can be applied to both quoted and unquoted strings (which now both perform partial matches by default). This matches the old behavior of the text filter.

diff --git a/Engine/Source/Runtime/Core/Private/Misc/TextFilterExpressionEvaluator.cpp b/Engine/Source/Runtime/Core/Private/Misc/TextFilterExpressionEvaluator.cpp
--- a/Engine/Source/Runtime/Core/Private/Misc/TextFilterExpressionEvaluator.cpp
+++ b/Engine/Source/Runtime/Core/Private/Misc/TextFilterExpressionEvaluator.cpp
@@ -87,6 +87,7 @@ DEFINE_TEXT_EXPRESSION_OPERATOR_NODE(FOr,						3,	0xF4778B51, 0xF535414D, 0x9C0E
 DEFINE_TEXT_EXPRESSION_OPERATOR_NODE(FAnd,						3,	0x7511397A, 0x02D24DC2, 0x86729800, 0xF454C320)
 DEFINE_TEXT_EXPRESSION_OPERATOR_NODE(FNot,						3,	0x03D78990, 0x41D04E26, 0x8E98AD2F, 0x74667868)
 
+DEFINE_TEXT_EXPRESSION_OPERATOR_NODE(FTextCmpExact,				1,	0x16374B83, 0x314F4624, 0x80BBA1C9, 0x1AB855CD)
 DEFINE_TEXT_EXPRESSION_OPERATOR_NODE(FTextCmpAnchor,			1,	0xC70749AD, 0x954D4A92, 0x95061222, 0xC8AC6D6D)
 
 DEFINE_EXPRESSION_NODE_TYPE(TextFilterExpressionParser::FTextToken,	0x09E49538, 0x633545E3, 0x84B5644F, 0x1F11628F);
@@ -95,8 +96,8 @@ DEFINE_EXPRESSION_NODE_TYPE(bool,									0xC1CD5DCF, 0x2AB44958, 0xB3FF4F8F, 0x
 namespace TextFilterExpressionParser
 {
 	/** This contains all the symbols that can define breaking points between text and an operator */
-	static const TCHAR BasicTextBreakingCharacters[]	= { '(', ')', '!', '-', '&', '|', '.', ' ' };						// ETextFilterExpressionEvaluatorMode::BasicString
-	static const TCHAR ComplexTextBreakingCharacters[]	= { '(', ')', '=', ':', '<', '>', '!', '-', '&', '|', '.', ' ' };	// ETextFilterExpressionEvaluatorMode::Complex
+	static const TCHAR BasicTextBreakingCharacters[]	= { '(', ')', '!', '-', '+', '&', '|', '.', ' ' };						// ETextFilterExpressionEvaluatorMode::BasicString
+	static const TCHAR ComplexTextBreakingCharacters[]	= { '(', ')', '=', ':', '<', '>', '!', '-', '+', '&', '|', '.', ' ' };	// ETextFilterExpressionEvaluatorMode::Complex
 
 	const TCHAR* FSubExpressionStart::Monikers[]		= { TEXT("(") };
 	const TCHAR* FSubExpressionEnd::Monikers[]			= { TEXT(")") };
@@ -112,6 +113,7 @@ namespace TextFilterExpressionParser
 	const TCHAR* FAnd::Monikers[]						= { TEXT("AND"), TEXT("&&"), TEXT("&") };
 	const TCHAR* FNot::Monikers[]						= { TEXT("NOT"), TEXT("!"), TEXT("-") };
 
+	const TCHAR* FTextCmpExact::Monikers[]				= { TEXT("+") };
 	const TCHAR* FTextCmpAnchor::Monikers[]				= { TEXT("...") };
 
 	enum class ETextNodeType : uint8
@@ -119,6 +121,7 @@ namespace TextFilterExpressionParser
 		Text,
 		SubExpressionStart,
 		SubExpressionEnd,
+		TextCmpExact,
 		TextCmpAnchor,
 		Other,
 	};
@@ -129,6 +132,7 @@ namespace TextFilterExpressionParser
 		if (InNode.Cast<FTextToken>())			return ETextNodeType::Text;
 		if (InNode.Cast<FSubExpressionStart>()) return ETextNodeType::SubExpressionStart;
 		if (InNode.Cast<FSubExpressionEnd>())	return ETextNodeType::SubExpressionEnd;
+		if (InNode.Cast<FTextCmpExact>())		return ETextNodeType::TextCmpExact;
 		if (InNode.Cast<FTextCmpAnchor>())		return ETextNodeType::TextCmpAnchor;
 		return ETextNodeType::Other;
 	}
@@ -236,11 +240,11 @@ namespace TextFilterExpressionParser
 		{
 			FString FinalString = TextToken.GetValue().GetString();
 			UnescapeQuotedString(FinalString, QuoteChar);
-			Consumer.Add(TextTokenWithQuotes, FExpressionNode(FTextToken(MoveTemp(FinalString), ETextFilterTextComparisonMode::Exact)));
+			Consumer.Add(TextTokenWithQuotes, FExpressionNode(FTextToken(MoveTemp(FinalString), ETextFilterTextComparisonMode::Partial)));
 		}
 		else
 		{
-			Consumer.Add(TextTokenWithQuotes, FExpressionNode(FTextToken(FString(), ETextFilterTextComparisonMode::Exact)));
+			Consumer.Add(TextTokenWithQuotes, FExpressionNode(FTextToken(FString(), ETextFilterTextComparisonMode::Partial)));
 		}
 
 		return TOptional<FExpressionError>();
@@ -430,55 +434,78 @@ bool FTextFilterExpressionEvaluator::SetFilterText(const FText& InFilterText)
 		{
 			auto& TmpLexTokens = LexResult.GetValue();
 
+			// We process the tokens in two passes (doing these separately minimizes the combinations required by step 2):
+			//	Pass 1) We process and remove any text comparison modification tokens (FTextCmpExact and FTextCmpAnchor)
+			//	Pass 2) We inject OR operators between any adjacent text and sub-expression tokens to mimic the old search filter behavior
 			TArray<FExpressionToken> FinalTokens;
 			FinalTokens.Reserve(TmpLexTokens.Num());
 
-			bool bIsComplexExpression = false;
-
-			ETextNodeType LastTokenType = ETextNodeType::Other;
-			for (FExpressionToken& TmpLexToken : TmpLexTokens)
+			// Pass 1
 			{
-				const ETextNodeType CurrentTokenType = GetTextNodeType(TmpLexToken.Node);
-
-				// We need to inject an OR operator between any consecutive text tokens
-				// This emulates the old behavior of the search filter, where a space between terms implied an OR operator
-				// We also do this between sub-expressions that aren't separated by an operator to ensure that we don't get multiple root nodes in our expression
-				if ((LastTokenType == ETextNodeType::Text || LastTokenType == ETextNodeType::SubExpressionEnd) &&
-					(CurrentTokenType == ETextNodeType::Text || CurrentTokenType == ETextNodeType::SubExpressionStart)
-					)
+				ETextNodeType LastTokenType = ETextNodeType::Other; // We have to cache this as we're moving tokens out of the TmpLexTokens array
+				for (FExpressionToken& CurrentToken : TmpLexTokens)
 				{
-					FinalTokens.Add(FExpressionToken(TmpLexToken.Context, FOr()));
-				}
+					const ETextNodeType CurrentTokenType = GetTextNodeType(CurrentToken.Node);
 
-				// We process TextCmpAnchor tokens during post-compilation, as they only ever operate on constants so don't need to be constantly re-evaluated for each item
-				if (LastTokenType == ETextNodeType::TextCmpAnchor && CurrentTokenType == ETextNodeType::Text)
-				{
-					// Pre-unary TextCmpAnchor - modify the current text token to performs an "ends with" string compare
-					const FTextToken* TextToken = TmpLexToken.Node.Cast<FTextToken>();
-					check(TextToken);
-					const_cast<FTextToken*>(TextToken)->SetTextComparisonMode(ETextFilterTextComparisonMode::EndsWith);
-				}
-				else if (LastTokenType == ETextNodeType::Text && CurrentTokenType == ETextNodeType::TextCmpAnchor)
-				{
-					// Post-unary TextCmpAnchor - modify the previous text token to perform a "starts with" string compare
-					const FTextToken* TextToken = FinalTokens.Last().Node.Cast<FTextToken>();
-					check(TextToken);
-					const_cast<FTextToken*>(TextToken)->SetTextComparisonMode(ETextFilterTextComparisonMode::StartsWith);
-				}
+					// We process TextCmpExact and TextCmpAnchor tokens during post-compilation, 
+					// as they only ever operate on constants so don't need to be constantly re-evaluated for each item
+					if (LastTokenType == ETextNodeType::TextCmpExact && CurrentTokenType == ETextNodeType::Text)
+					{
+						// Pre-unary TextCmpExact - modify the current text token to performs an "exact" string compare
+						const FTextToken* TextToken = CurrentToken.Node.Cast<FTextToken>();
+						check(TextToken);
+						const_cast<FTextToken*>(TextToken)->SetTextComparisonMode(ETextFilterTextComparisonMode::Exact);
+					}
+					else if (LastTokenType == ETextNodeType::TextCmpAnchor && CurrentTokenType == ETextNodeType::Text)
+					{
+						// Pre-unary TextCmpAnchor - modify the current text token to performs an "ends with" string compare
+						const FTextToken* TextToken = CurrentToken.Node.Cast<FTextToken>();
+						check(TextToken);
+						const_cast<FTextToken*>(TextToken)->SetTextComparisonMode(ETextFilterTextComparisonMode::EndsWith);
+					}
+					else if (LastTokenType == ETextNodeType::Text && CurrentTokenType == ETextNodeType::TextCmpAnchor)
+					{
+						// Post-unary TextCmpAnchor - modify the previous text token to perform a "starts with" string compare
+						const FTextToken* TextToken = FinalTokens.Last().Node.Cast<FTextToken>();
+						check(TextToken);
+						const_cast<FTextToken*>(TextToken)->SetTextComparisonMode(ETextFilterTextComparisonMode::StartsWith);
+					}
+					
+					// We strip out the TextCmpExact and TextCmpAnchor tokens, as we process them as part of this loop
+					if (CurrentTokenType != ETextNodeType::TextCmpExact && CurrentTokenType != ETextNodeType::TextCmpAnchor)
+					{
+						FinalTokens.Add(MoveTemp(CurrentToken));
+					}
 
-				// Key->value pairs require some form of comparison operator, so if we've found once of those, consider this a complex expression
-				if (!bIsComplexExpression && IsComplexNode(TmpLexToken.Node))
-				{
-					bIsComplexExpression = true;
+					LastTokenType = CurrentTokenType;
 				}
+			}
 
-				// We strip out the TextCmpAnchor tokens, as we process them as part of this loop
-				if (CurrentTokenType != ETextNodeType::TextCmpAnchor)
+			// Pass 2
+			bool bIsComplexExpression = false;
+			{
+				for (int32 FinalTokenIndex = 0; FinalTokenIndex < FinalTokens.Num(); ++FinalTokenIndex)
 				{
-					FinalTokens.Add(MoveTemp(TmpLexToken));
-				}
+					const FExpressionToken& CurrentToken = FinalTokens[FinalTokenIndex];
+
+					const ETextNodeType LastTokenType = (FinalTokenIndex > 0) ? GetTextNodeType(FinalTokens[FinalTokenIndex-1].Node) : ETextNodeType::Other;
+					const ETextNodeType CurrentTokenType = GetTextNodeType(CurrentToken.Node);
 
-				LastTokenType = CurrentTokenType;
+					// Key->value pairs require some form of comparison operator, so if we've found one of those, consider this to be a complex expression
+					if (!bIsComplexExpression && IsComplexNode(CurrentToken.Node))
+					{
+						bIsComplexExpression = true;
+					}
+
+					// We need to inject OR operators between any adjacent text and sub-expression tokens to mimic the old search filter behavior
+					if ((LastTokenType == ETextNodeType::Text || LastTokenType == ETextNodeType::SubExpressionEnd) &&
+						(CurrentTokenType == ETextNodeType::Text || CurrentTokenType == ETextNodeType::SubExpressionStart)
+						)
+					{
+						// Inject the new token between the current one and the previous one, and then skip over the newly added token
+						FinalTokens.Insert(FExpressionToken(CurrentToken.Context, FOr()), FinalTokenIndex++);
+					}
+				}
 			}
 
 			CompiledFilter = ExpressionParser::Compile(MoveTemp(FinalTokens), Grammar);
@@ -557,6 +584,7 @@ void FTextFilterExpressionEvaluator::ConstructExpressionParser()
 	TokenDefinitions.DefineToken(&ConsumeOperator<FOr>);
 	TokenDefinitions.DefineToken(&ConsumeOperator<FAnd>);
 	TokenDefinitions.DefineToken(&ConsumeOperator<FNot>);
+	TokenDefinitions.DefineToken(&ConsumeOperator<FTextCmpExact>);
 	TokenDefinitions.DefineToken(&ConsumeOperator<FTextCmpAnchor>);
 	TokenDefinitions.DefineToken(&ConsumeNumber);
 	TokenDefinitions.DefineToken(&ConsumeQuotedText);
@@ -572,8 +600,6 @@ void FTextFilterExpressionEvaluator::ConstructExpressionParser()
 	Grammar.DefineBinaryOperator<FOr>(2);
 	Grammar.DefineBinaryOperator<FAnd>(2);
 	Grammar.DefinePreUnaryOperator<FNot>();
-	Grammar.DefinePreUnaryOperator<FTextCmpAnchor>();
-	Grammar.DefinePostUnaryOperator<FTextCmpAnchor>();
 
 	JumpTable.MapBinary<FLessOrEqual>([](const FTextToken& A, const FTextToken& B, const ITextFilterExpressionContext* InContext)		{ return InContext->TestComplexExpression(FName(*A.GetString()), B.GetString(), ETextFilterComparisonOperation::LessOrEqual, B.GetTextComparisonMode()); });
 	JumpTable.MapBinary<FLess>([](const FTextToken& A, const FTextToken& B, const ITextFilterExpressionContext* InContext)				{ return InContext->TestComplexExpression(FName(*A.GetString()), B.GetString(), ETextFilterComparisonOperation::Less, B.GetTextComparisonMode()); });
diff --git a/Engine/Source/Runtime/Core/Private/Misc/TextFilterTests.cpp b/Engine/Source/Runtime/Core/Private/Misc/TextFilterTests.cpp
--- a/Engine/Source/Runtime/Core/Private/Misc/TextFilterTests.cpp
+++ b/Engine/Source/Runtime/Core/Private/Misc/TextFilterTests.cpp
@@ -144,7 +144,14 @@ namespace TextFilterTests
 		bResult &= InTestPayload.TestFilterExpression(TEXT("Woo..."), true);
 		bResult &= InTestPayload.TestFilterExpression(TEXT("...ble"), true);
 		bResult &= InTestPayload.TestFilterExpression(TEXT("Wo... AND ...le"), true);
+		bResult &= InTestPayload.TestFilterExpression(TEXT("'Wooble'"), true);
+		bResult &= InTestPayload.TestFilterExpression(TEXT("'Woo'"), true);
+		bResult &= InTestPayload.TestFilterExpression(TEXT("+'Wooble'"), true);
+		bResult &= InTestPayload.TestFilterExpression(TEXT("+'Woo'"), false);
+		bResult &= InTestPayload.TestFilterExpression(TEXT("+Wooble"), true);
+		bResult &= InTestPayload.TestFilterExpression(TEXT("+Woo"), false);
 		bResult &= InTestPayload.TestFilterExpression(TEXT("Wooble2"), false);
+		bResult &= InTestPayload.TestFilterExpression(TEXT("-Wooble2"), true);
 		bResult &= InTestPayload.TestFilterExpression(TEXT("Wooble2 Flibble"), true);
 		bResult &= InTestPayload.TestFilterExpression(TEXT("Wooble2 OR Flibble"), true);
 		bResult &= InTestPayload.TestFilterExpression(TEXT("Wooble2 AND Flibble"), false);
