Stats - Small tweak to 'stat slow'

diff --git a/Engine/Source/Runtime/Core/Private/Stats/StatsCommand.cpp b/Engine/Source/Runtime/Core/Private/Stats/StatsCommand.cpp
--- a/Engine/Source/Runtime/Core/Private/Stats/StatsCommand.cpp
+++ b/Engine/Source/Runtime/Core/Private/Stats/StatsCommand.cpp
@@ -105,7 +105,7 @@ struct FStatSlowParams : public FStatParams
 		: FStatParams( Cmd )
 	{
 		static FName NAME_Slow = TEXT( "Slow" );
-		Group = TParsedValueWithDefault<FName>( Cmd, TEXT( "group=" ), NAME_Slow );
+		Group = TParsedValueWithDefault<FName>( nullptr, nullptr, NAME_Slow );
 		CullMs = TParsedValueWithDefault<float>( Cmd, TEXT( "ms=" ), 1.0f );
 		MaxHierarchyDepth = TParsedValueWithDefault<int32>( Cmd, TEXT( "maxdepth=" ), 4 );
 		bSlowMode = true;
@@ -656,7 +656,7 @@ struct FHUDGroupManager
 				else
 				{
 					TSet<FName> EmptySet = TSet<FName>();
-					EnabledGroups.Add( MaybeGroupFName, FInternalGroup( MaybeGroupFName, NAME_None, EStatDisplayMode::Hierarchical, EmptySet, TEXT( "Slow hierarchy for game and render" ) ) );
+					EnabledGroups.Add( MaybeGroupFName, FInternalGroup( MaybeGroupFName, NAME_None, EStatDisplayMode::Hierarchical, EmptySet, TEXT( "Hierarchy for game and render" ) ) );
 				}			
 			}
 		}
