Fixed shadowed variables

diff --git a/Engine/Source/Runtime/Core/Private/Stats/StatsCommand.cpp b/Engine/Source/Runtime/Core/Private/Stats/StatsCommand.cpp
--- a/Engine/Source/Runtime/Core/Private/Stats/StatsCommand.cpp
+++ b/Engine/Source/Runtime/Core/Private/Stats/StatsCommand.cpp
@@ -972,9 +972,9 @@ struct FHUDGroupManager
 				}
 
 				// Replace thread encoded id with the thread name.
-				for (auto& It : HudGroup.HierAggregate)
+				for (auto& HierIt : HudGroup.HierAggregate)
 				{
-					FComplexStatMessage& StatMessage = It;
+					FComplexStatMessage& StatMessage = HierIt;
 					const FString StatDescription = StatMessage.NameAndInfo.GetDescription();
 					if (StatDescription.Contains( FStatConstants::ThreadNameMarker ))
 					{
diff --git a/Engine/Source/Runtime/Core/Private/Stats/StatsData.cpp b/Engine/Source/Runtime/Core/Private/Stats/StatsData.cpp
--- a/Engine/Source/Runtime/Core/Private/Stats/StatsData.cpp
+++ b/Engine/Source/Runtime/Core/Private/Stats/StatsData.cpp
@@ -123,10 +123,10 @@ void FRawStatStackNode::CullByCycles( int64 MinCycles )
 			else
 			{
 				// Remove children.
-				for (TMap<FName, FRawStatStackNode*>::TIterator It( Child->Children ); It; ++It)
+				for (TMap<FName, FRawStatStackNode*>::TIterator InnerIt( Child->Children ); InnerIt; ++InnerIt)
 				{
-					delete It.Value();
-					It.RemoveCurrent();
+					delete InnerIt.Value();
+					InnerIt.RemoveCurrent();
 				}
 			}
 		}
