Undeprecated some module manager functions until we design a proper alternative
(See CL 2562810 for more context)

diff --git a/Engine/Source/Runtime/Core/Private/Modules/ModuleManager.cpp b/Engine/Source/Runtime/Core/Private/Modules/ModuleManager.cpp
--- a/Engine/Source/Runtime/Core/Private/Modules/ModuleManager.cpp
+++ b/Engine/Source/Runtime/Core/Private/Modules/ModuleManager.cpp
@@ -940,9 +940,7 @@ void FModuleManager::UnloadOrAbandonModuleWithCallback(const FName InModuleName,
 {
 	auto Module = FindModuleChecked(InModuleName);
 	
-PRAGMA_DISABLE_DEPRECATION_WARNINGS
 	Module->Module->PreUnloadCallback();
-PRAGMA_ENABLE_DEPRECATION_WARNINGS
 
 	const bool bIsHotReloadable = DoesLoadedModuleHaveUObjects( InModuleName );
 	if (!bAbandonOnly && bIsHotReloadable && Module->Module->SupportsDynamicReloading())
@@ -971,9 +969,7 @@ bool FModuleManager::LoadModuleWithCallback( const FName InModuleName, FOutputDe
 
 	if (bWasSuccessful && LoadedModule.IsValid())
 	{
-PRAGMA_DISABLE_DEPRECATION_WARNINGS
 		LoadedModule->PostLoadCallback();
-PRAGMA_ENABLE_DEPRECATION_WARNINGS
 	}
 	else
 	{
diff --git a/Engine/Source/Runtime/Core/Public/Modules/ModuleInterface.h b/Engine/Source/Runtime/Core/Public/Modules/ModuleInterface.h
--- a/Engine/Source/Runtime/Core/Public/Modules/ModuleInterface.h
+++ b/Engine/Source/Runtime/Core/Public/Modules/ModuleInterface.h
@@ -32,7 +32,6 @@ public:
 	/**
 	 * Called before the module has been unloaded
 	 */
-	DEPRECATED(4.9, "PreUnloadCallback() was deprecated because it was not called in all expected cases and generally was not useful.  Override the ShutdownModule() function instead.")
 	virtual void PreUnloadCallback()
 	{
 	}
@@ -40,7 +39,6 @@ public:
 	/**
 	 * Called after the module has been reloaded
 	 */
-	DEPRECATED(4.9, "PostLoadCallback() was deprecated because it was not called in all expected cases and generally was not useful.  Override the StartupModule() function instead.")
 	virtual void PostLoadCallback()
 	{
 	}
