Add -fullcrashdump commandline parameter for forcing any crash minidumps to contain full memory info. Useful for QA to use as launch param when tracking down rare/hard-to-repro crashes.

diff --git a/Engine/Source/Runtime/Core/Private/GenericPlatform/GenericPlatformContext.cpp b/Engine/Source/Runtime/Core/Private/GenericPlatform/GenericPlatformContext.cpp
--- a/Engine/Source/Runtime/Core/Private/GenericPlatform/GenericPlatformContext.cpp
+++ b/Engine/Source/Runtime/Core/Private/GenericPlatform/GenericPlatformContext.cpp
@@ -28,6 +28,7 @@ namespace NCachedCrashContextProperties
 	static FString PrimaryGPUBrand;
 	static FString UserName;
 	static FString DefaultLocale;
+	static int32 CrashDumpMode;
 
 	static FString CrashGUID;
 }
@@ -51,6 +52,17 @@ void FGenericCrashContext::Initialize()
 	NCachedCrashContextProperties::UserName = FPlatformProcess::UserName();
 	NCachedCrashContextProperties::DefaultLocale = FPlatformMisc::GetDefaultLocale();
 
+	// Using the -fullcrashdump parameter will cause full memory minidumps to be created for crashes
+	CrashDumpMode = (int32)ECrashDumpMode::Default;
+	if (FCommandLine::IsInitialized())
+	{
+		const TCHAR* CmdLine = FCommandLine::Get();
+		if (FParse::Param( CmdLine, TEXT("fullcrashdump") ))
+		{
+			CrashDumpMode = (int32)ECrashDumpMode::FullDump;
+		}
+	}
+
 	const FGuid Guid = FGuid::NewGuid();
 	NCachedCrashContextProperties::CrashGUID = FString::Printf(TEXT("UE4CC-%s-%s"), *NCachedCrashContextProperties::PlatformNameIni, *Guid.ToString(EGuidFormats::Digits));
 
@@ -127,6 +139,8 @@ void FGenericCrashContext::SerializeContentToBuffer()
 	AddCrashProperty(TEXT("MiscOSVersionMajor"), *NCachedCrashContextProperties::OsVersion);
 	AddCrashProperty(TEXT("MiscOSVersionMinor"), *NCachedCrashContextProperties::OsSubVersion);
 
+	AddCrashProperty( TEXT( "CrashDumpMode" ), CrashDumpMode );
+
 
 	// @TODO yrx 2014-10-08 Move to the crash report client.
 	/*if( CanUseUnsafeAPI() )
@@ -180,6 +194,11 @@ const FString& FGenericCrashContext::GetUniqueCrashName()
 	return NCachedCrashContextProperties::CrashGUID;
 }
 
+const bool FGenericCrashContext::IsFullCrashDump()
+{
+	return (CrashDumpMode == (int32)ECrashDumpMode::FullDump);
+}
+
 void FGenericCrashContext::SerializeAsXML( const TCHAR* Filename )
 {
 	SerializeContentToBuffer();
diff --git a/Engine/Source/Runtime/Core/Private/Windows/WindowsPlatformCrashContext.cpp b/Engine/Source/Runtime/Core/Private/Windows/WindowsPlatformCrashContext.cpp
--- a/Engine/Source/Runtime/Core/Private/Windows/WindowsPlatformCrashContext.cpp
+++ b/Engine/Source/Runtime/Core/Private/Windows/WindowsPlatformCrashContext.cpp
@@ -62,6 +62,10 @@ bool WriteMinidump(const TCHAR* Path, LPEXCEPTION_POINTERS ExceptionInfo)
 	CrashContextStreamInformation.UserStreamArray = &CrashContextStream;
 
 	MINIDUMP_TYPE MinidumpType = MiniDumpNormal;//(MINIDUMP_TYPE)(MiniDumpWithPrivateReadWriteMemory|MiniDumpWithDataSegs|MiniDumpWithHandleData|MiniDumpWithFullMemoryInfo|MiniDumpWithThreadInfo|MiniDumpWithUnloadedModules);
+	if (CrashContext.IsFullCrashDump())
+	{
+		MinidumpType = (MINIDUMP_TYPE)(MiniDumpWithFullMemory|MiniDumpWithFullMemoryInfo|MiniDumpWithHandleData|MiniDumpWithThreadInfo|MiniDumpWithUnloadedModules);
+	}
 
 	const BOOL Result = MiniDumpWriteDump(GetCurrentProcess(), GetCurrentProcessId(), FileHandle, MinidumpType, &DumpExceptionInfo, &CrashContextStreamInformation, NULL);
 	CloseHandle(FileHandle);
diff --git a/Engine/Source/Runtime/Core/Public/GenericPlatform/GenericPlatformContext.h b/Engine/Source/Runtime/Core/Public/GenericPlatform/GenericPlatformContext.h
--- a/Engine/Source/Runtime/Core/Public/GenericPlatform/GenericPlatformContext.h
+++ b/Engine/Source/Runtime/Core/Public/GenericPlatform/GenericPlatformContext.h
@@ -47,6 +47,16 @@ enum class ECrashDescVersions
 	VER_2_AddedNewProperties,
 };
 
+/** Enumerates crash dump modes. */
+enum class ECrashDumpMode : int32
+{
+	/** Default minidump settings. */
+	Default = 0,
+
+	/** Full memory crash minidump */
+	FullDump = 1,
+};
+
 /**
  *	Contains a runtime crash's properties that are common for all platforms.
  *	This may change in the future.
@@ -87,6 +97,11 @@ public:
 	 */
 	const FString& GetUniqueCrashName();
 
+	/**
+	 * @return whether this crash is a full memory minidump
+	 */
+	const bool IsFullCrashDump();
+
 	/** Serializes crash's informations to the specified filename. */
 	void SerializeAsXML( const TCHAR* Filename );
 
