Fix CIS (compile error in FGenericCrashContext).

diff --git a/Engine/Source/Runtime/Core/Private/GenericPlatform/GenericPlatformContext.cpp b/Engine/Source/Runtime/Core/Private/GenericPlatform/GenericPlatformContext.cpp
--- a/Engine/Source/Runtime/Core/Private/GenericPlatform/GenericPlatformContext.cpp
+++ b/Engine/Source/Runtime/Core/Private/GenericPlatform/GenericPlatformContext.cpp
@@ -53,13 +53,13 @@ void FGenericCrashContext::Initialize()
 	NCachedCrashContextProperties::DefaultLocale = FPlatformMisc::GetDefaultLocale();
 
 	// Using the -fullcrashdump parameter will cause full memory minidumps to be created for crashes
-	CrashDumpMode = (int32)ECrashDumpMode::Default;
+	NCachedCrashContextProperties::CrashDumpMode = (int32)ECrashDumpMode::Default;
 	if (FCommandLine::IsInitialized())
 	{
 		const TCHAR* CmdLine = FCommandLine::Get();
 		if (FParse::Param( CmdLine, TEXT("fullcrashdump") ))
 		{
-			CrashDumpMode = (int32)ECrashDumpMode::FullDump;
+			NCachedCrashContextProperties::CrashDumpMode = (int32)ECrashDumpMode::FullDump;
 		}
 	}
 
@@ -139,7 +139,7 @@ void FGenericCrashContext::SerializeContentToBuffer()
 	AddCrashProperty(TEXT("MiscOSVersionMajor"), *NCachedCrashContextProperties::OsVersion);
 	AddCrashProperty(TEXT("MiscOSVersionMinor"), *NCachedCrashContextProperties::OsSubVersion);
 
-	AddCrashProperty( TEXT( "CrashDumpMode" ), CrashDumpMode );
+	AddCrashProperty( TEXT( "CrashDumpMode" ), NCachedCrashContextProperties::CrashDumpMode );
 
 
 	// @TODO yrx 2014-10-08 Move to the crash report client.
@@ -196,7 +196,7 @@ const FString& FGenericCrashContext::GetUniqueCrashName()
 
 const bool FGenericCrashContext::IsFullCrashDump()
 {
-	return (CrashDumpMode == (int32)ECrashDumpMode::FullDump);
+	return (NCachedCrashContextProperties::CrashDumpMode == (int32)ECrashDumpMode::FullDump);
 }
 
 void FGenericCrashContext::SerializeAsXML( const TCHAR* Filename )
