Added TSet::Reset

This will empty the set without performing a reallocation. TMap::Reset now calls this function on its internal set (rather than calling Empty(Num()) which could potentially reallocate if Num() != Max()).

diff --git a/Engine/Source/Runtime/Core/Public/Containers/Map.h b/Engine/Source/Runtime/Core/Public/Containers/Map.h
--- a/Engine/Source/Runtime/Core/Public/Containers/Map.h
+++ b/Engine/Source/Runtime/Core/Public/Containers/Map.h
@@ -252,7 +252,7 @@ public:
     /** Efficiently empties out the map but preserves all allocations and capacities */
     FORCEINLINE void Reset()
     {
-        Empty(Num());
+        Pairs.Reset();
     }
 
 	/** Shrinks the pair set to avoid slack. */
diff --git a/Engine/Source/Runtime/Core/Public/Containers/Set.h b/Engine/Source/Runtime/Core/Public/Containers/Set.h
--- a/Engine/Source/Runtime/Core/Public/Containers/Set.h
+++ b/Engine/Source/Runtime/Core/Public/Containers/Set.h
@@ -305,6 +305,19 @@ public:
 		}
 	}
 
+	/** Efficiently empties out the set but preserves all allocations and capacities */
+    void Reset()
+    {
+		// Reset the elements array.
+		Elements.Reset();
+
+		// Clear the references to the elements that have now been removed.
+		for(int32 HashIndex = 0;HashIndex < HashSize;HashIndex++)
+		{
+			GetTypedHash(HashIndex) = FSetElementId();
+		}
+    }
+
 	/** Shrinks the set's element storage to avoid slack. */
 	FORCEINLINE void Shrink()
 	{
