Linux: gracefully handle inability to get display metrics instead of dying.

- PR #1232 contributed by [REDACTED_PERSON].
- Fixes cooking on Linux headless machines where getting display metrics is impossible.

diff --git a/Engine/Source/Runtime/Core/Private/Linux/LinuxApplication.cpp b/Engine/Source/Runtime/Core/Private/Linux/LinuxApplication.cpp
--- a/Engine/Source/Runtime/Core/Private/Linux/LinuxApplication.cpp
+++ b/Engine/Source/Runtime/Core/Private/Linux/LinuxApplication.cpp
@@ -1277,16 +1277,18 @@ bool FLinuxApplication::TryCalculatePopupWindowPosition( const FPlatformRect& In
 
 void FDisplayMetrics::GetDisplayMetrics(FDisplayMetrics& OutDisplayMetrics)
 {
-	if (!FPlatformMisc::PlatformInitMultimedia()) //	will not initialize more than once
+	int NumDisplays = 0;
+
+	if (FPlatformMisc::PlatformInitMultimedia()) //	will not initialize more than once
 	{
-		// consider making non-fatal and just returning bullshit? (which can be checked for and handled more gracefully)
-		UE_LOG(LogInit, Fatal, TEXT("FDisplayMetrics::GetDisplayMetrics: PlatformInitMultimedia() failed, cannot get display metrics"));
-		// unreachable
-		return;
+		NumDisplays = SDL_GetNumVideoDisplays();
+	}
+	else
+	{
+		UE_LOG(LogInit, Warning, TEXT("FDisplayMetrics::GetDisplayMetrics: PlatformInitMultimedia() failed, cannot get display metrics"));
 	}
 
 	// loop over all monitors to determine which one is the best
-	int NumDisplays = SDL_GetNumVideoDisplays();
 	if (NumDisplays <= 0)
 	{
 		OutDisplayMetrics.PrimaryDisplayWorkAreaRect.Left = 0;
