#ue4 - Remove static constant initialization from VectorSin() SSE implementation.

diff --git a/Engine/Source/Runtime/Core/Public/Math/UnrealMathSSE.h b/Engine/Source/Runtime/Core/Public/Math/UnrealMathSSE.h
--- a/Engine/Source/Runtime/Core/Public/Math/UnrealMathSSE.h
+++ b/Engine/Source/Runtime/Core/Public/Math/UnrealMathSSE.h
@@ -1043,6 +1043,22 @@ FORCEINLINE VectorRegister VectorLog2(const VectorRegister& X)
 	return MakeVectorRegister((float)log2(VectorGetComponent(X, 0)), (float)log2(VectorGetComponent(X, 1)), (float)log2(VectorGetComponent(X, 2)), (float)log2(VectorGetComponent(X, 3)));
 }
 
+
+/**
+ * Using "static const float ..." or "static const VectorRegister ..." in functions creates the branch and code to construct those constants.
+ * Doing this in FORCEINLINE not only means you introduce a branch per static, but you bloat the inlined code immensely.
+ * Defining these constants at the global scope causes them to be created at startup, and avoids the cost at the function level.
+ * Doing it at the function level is okay for anything that is a simple "const float", but usage of "sqrt()" here forces actual function calls.
+ */
+namespace VectorSinConstantsSSE
+{
+	static const float p = 0.225f;
+	static const float a = (16 * sqrt(p));
+	static const float b = ((1 - p) / sqrt(p));
+	static const VectorRegister A = MakeVectorRegister(a, a, a, a);
+	static const VectorRegister B = MakeVectorRegister(b, b, b, b);
+}
+
 FORCEINLINE VectorRegister VectorSin(const VectorRegister& X)
 {
 	//Sine approximation using a squared parabola restrained to f(0) = 0, f(PI) = 0, f(PI/2) = 1.
@@ -1051,17 +1067,10 @@ FORCEINLINE VectorRegister VectorSin(const VectorRegister& X)
 	//Average error of 0.000128
 	//Max error of 0.001091
 
-	static const float p = 0.225f;
-	static const VectorRegister P = MakeVectorRegister(p, p, p, p);
-	static const float a = 16 * sqrt(p);
-	static const VectorRegister A = MakeVectorRegister(a, a, a, a);
-	static const float b = (1 - p) / sqrt(p);
-	static const VectorRegister B = MakeVectorRegister(b, b, b, b);
-
 	VectorRegister y = VectorMultiply(X, GlobalVectorConstants::OneOverTwoPi);
 	y = VectorSubtract(y, VectorFloor(VectorAdd(y, GlobalVectorConstants::FloatOneHalf)));
-	y = VectorMultiply(A, VectorMultiply(y, VectorSubtract(GlobalVectorConstants::FloatOneHalf, VectorAbs(y))));
-	return VectorMultiply(y, VectorAdd(B, VectorAbs(y)));
+	y = VectorMultiply(VectorSinConstantsSSE::A, VectorMultiply(y, VectorSubtract(GlobalVectorConstants::FloatOneHalf, VectorAbs(y))));
+	return VectorMultiply(y, VectorAdd(VectorSinConstantsSSE::B, VectorAbs(y)));
 }
 
 FORCEINLINE VectorRegister VectorCos(const VectorRegister& X)
