TIsDerivedFrom<> now supports const!

--------
Integrated using branch UE4-To-UE4-Orion (reversed) of change#2559198 by [REDACTED_PERSON] on 2015/05/20 14:06:54.

diff --git a/Engine/Source/Runtime/Core/Public/Templates/UnrealTypeTraits.h b/Engine/Source/Runtime/Core/Public/Templates/UnrealTypeTraits.h
--- a/Engine/Source/Runtime/Core/Public/Templates/UnrealTypeTraits.h
+++ b/Engine/Source/Runtime/Core/Public/Templates/UnrealTypeTraits.h
@@ -41,6 +41,7 @@ struct TIsDerivedFrom
 	// Overloading Test() s.t. only calling it with something that is
 	// a BaseType (or inherited from the BaseType) will return a Yes.
 	static Yes& Test( BaseType* );
+	static Yes& Test( const BaseType* );
 	static No& Test( ... );
 
 	// Makes a DerivedType ptr.
