Stats - Stats from other threads like task graph or pool thread should be synced with the game thread
Stats - More stats for async tasks

diff --git a/Engine/Source/Runtime/Core/Public/Async/AsyncWork.h b/Engine/Source/Runtime/Core/Public/Async/AsyncWork.h
--- a/Engine/Source/Runtime/Core/Public/Async/AsyncWork.h
+++ b/Engine/Source/Runtime/Core/Public/Async/AsyncWork.h
@@ -6,7 +6,7 @@
 
 #pragma once
 
-DECLARE_STATS_GROUP(TEXT("ThreadPool Async Tasks"), STATGROUP_ThreadPoolAsyncTasks, STATCAT_Advanced);
+
 
 /**
 	FAutoDeleteAsyncTask - template task for jobs that delete themselves when complete
@@ -29,10 +29,10 @@ DECLARE_STATS_GROUP(TEXT("ThreadPool Async Tasks"), STATGROUP_ThreadPoolAsyncTas
 			... do the work here
 		}
 
- 		FORCEINLINE TStatId GetStatId() const
- 		{
+		FORCEINLINE TStatId GetStatId() const
+		{
 			RETURN_QUICK_DECLARE_CYCLE_STAT(ExampleAutoDeleteAsyncTask, STATGROUP_ThreadPoolAsyncTasks);
- 		}
+		}
 	};
 
 	start an example job
@@ -268,6 +268,9 @@ class FAsyncTask
 	**/
 	void Start(bool bForceSynchronous, FQueuedThreadPool* InQueuedPool)
 	{
+		FScopeCycleCounter Scope( Task.GetStatId(), true );
+		DECLARE_SCOPE_CYCLE_COUNTER( TEXT( "FAsyncTask::Start" ), STAT_FAsyncTask_Start, STATGROUP_ThreadPoolAsyncTasks );
+
 		FPlatformMisc::MemoryBarrier();
 		CheckIdle();  // can't start a job twice without it being completed first
 		WorkNotFinishedCounter.Increment();
@@ -364,6 +367,9 @@ class FAsyncTask
 		FPlatformMisc::MemoryBarrier();
 		if (QueuedPool)
 		{
+			FScopeCycleCounter Scope( Task.GetStatId() );
+			DECLARE_SCOPE_CYCLE_COUNTER( TEXT( "FAsyncTask::SyncCompletion" ), STAT_FAsyncTask_SyncCompletion, STATGROUP_ThreadPoolAsyncTasks );
+
 			check(DoneEvent); // if it is not done yet, we must have an event
 			DoneEvent->Wait();
 			QueuedPool = 0;
diff --git a/Engine/Source/Runtime/Core/Public/Stats/Stats2.h b/Engine/Source/Runtime/Core/Public/Stats/Stats2.h
--- a/Engine/Source/Runtime/Core/Public/Stats/Stats2.h
+++ b/Engine/Source/Runtime/Core/Public/Stats/Stats2.h
@@ -1038,7 +1038,7 @@ typedef TChunkedArray<FStatMessage,(uint32)EStatMessagesArrayConstants::MESSAGES
 */
 struct FStatPacket
 {
-	/** Assigned later, this is the frame number this packet is for. @see FStatsThreadState::ScanForAdvance */
+	/** Assigned later, this is the frame number this packet is for. @see FStatsThreadState::ScanForAdvance or FThreadStats::DetectAndUpdateCurrentGameFrame */
 	int64 Frame;
 	/** ThreadId this packet came from **/
 	uint32 ThreadId;
@@ -1238,6 +1238,8 @@ public:
 			const bool bFrameHasChanged = FStats::GameThreadStatsFrame > CurrentGameFrame;
 			if( bFrameHasChanged )
 			{
+				// Other threads are one frame behind so set the frame to the previous one.
+				Packet.AssignFrame( CurrentGameFrame );
 				CurrentGameFrame = FStats::GameThreadStatsFrame;
 				return true;
 			}
@@ -1727,6 +1729,11 @@ struct FStat_##StatName\
 	DECLARE_STAT(CounterName,StatId,GroupId,EStatDataType::ST_int64, false, false, FPlatformMemory::MCR_Invalid); \
 	static DEFINE_STAT(StatId)
 
+/** FName stat that allows sending a string based data. */
+#define DECLARE_FNAME_STAT(CounterName,StatId,GroupId) \
+	DECLARE_STAT(CounterName,StatId,GroupId,EStatDataType::ST_FName, false, false, FPlatformMemory::MCR_Invalid); \
+	static DEFINE_STAT(StatId)
+
 /** This is a fake stat, mostly used to implement memory message or other custom stats that don't easily fit into the system. */
 #define DECLARE_PTR_STAT(CounterName,StatId,GroupId)\
 	DECLARE_STAT(CounterName,StatId,GroupId,EStatDataType::ST_Ptr, false, false, FPlatformMemory::MCR_Invalid); \
@@ -1982,6 +1989,7 @@ DECLARE_STATS_GROUP(TEXT("Streaming Details"),STATGROUP_StreamingDetails, STATCA
 DECLARE_STATS_GROUP(TEXT("Streaming"),STATGROUP_Streaming, STATCAT_Advanced);
 DECLARE_STATS_GROUP(TEXT("Target Platform"),STATGROUP_TargetPlatform, STATCAT_Advanced);
 DECLARE_STATS_GROUP(TEXT("Text"),STATGROUP_Text, STATCAT_Advanced);
+DECLARE_STATS_GROUP(TEXT("ThreadPool Async Tasks"),STATGROUP_ThreadPoolAsyncTasks, STATCAT_Advanced);
 DECLARE_STATS_GROUP(TEXT("Threading"),STATGROUP_Threading, STATCAT_Advanced);
 DECLARE_STATS_GROUP(TEXT("Threads"),STATGROUP_Threads, STATCAT_Advanced);
 DECLARE_STATS_GROUP(TEXT("Tickables"),STATGROUP_Tickables, STATCAT_Advanced);
