Stats - Added stats for wait/trigger in events (core)

diff --git a/Engine/Source/Runtime/Core/Private/HAL/ThreadingBase.cpp b/Engine/Source/Runtime/Core/Private/HAL/ThreadingBase.cpp
--- a/Engine/Source/Runtime/Core/Private/HAL/ThreadingBase.cpp
+++ b/Engine/Source/Runtime/Core/Private/HAL/ThreadingBase.cpp
@@ -6,6 +6,9 @@
 #include "LockFreeList.h"
 #include "StatsData.h"
 
+DEFINE_STAT( STAT_EventWaitWithId );
+DEFINE_STAT( STAT_EventTriggerWithId );
+
 
 /** The global thread pool */
 FQueuedThreadPool* GThreadPool = nullptr;
@@ -168,6 +171,59 @@ FSingleThreadManager& FSingleThreadManager::Get()
 	return Singleton;
 }
 
+
+/*-----------------------------------------------------------------------------
+	FEvent, FScopedEvent
+-----------------------------------------------------------------------------*/
+
+uint32 FEvent::EventUniqueId = 0;
+
+void FEvent::AdvanceStats()
+{
+#if	STATS
+	EventId = FPlatformAtomics::InterlockedAdd( (int32*)&EventUniqueId, 1 );
+	EventStartCycles = 0;
+#endif // STATS
+}
+
+void FEvent::WaitForStats()
+{
+#if	STATS
+	// Only start counting on the first wait, trigger will "close" the history.
+	if( FThreadStats::IsCollectingData() && EventStartCycles == 0 )
+	{
+		FThreadStats* ThreadStats = FThreadStats::GetThreadStats();
+		const uint64 PacketEventIdAndCycles = ((uint64)EventId << 32) | 0;
+		ThreadStats->AddMessage( GET_STATFNAME( STAT_EventWaitWithId ), EStatOperation::SpecialMessageMarker, PacketEventIdAndCycles );
+		EventStartCycles = FPlatformTime::Cycles();
+	}
+#endif // STATS
+}
+
+void FEvent::TriggerForStats()
+{
+#if	STATS
+	// Only add wait-trigger pairs.
+	if( EventStartCycles > 0 && FThreadStats::IsCollectingData() )
+	{
+		FThreadStats* ThreadStats = FThreadStats::GetThreadStats();
+		const uint32 EndCycles = FPlatformTime::Cycles();
+		const int32 DeltaCycles = int32( EndCycles - EventStartCycles );
+		const uint64 PacketEventIdAndCycles = ((uint64)EventId << 32) | DeltaCycles;
+		ThreadStats->AddMessage( GET_STATFNAME( STAT_EventTriggerWithId ), EStatOperation::SpecialMessageMarker, PacketEventIdAndCycles );
+
+		AdvanceStats();
+	}
+#endif // STATS
+}
+
+void FEvent::ResetForStats()
+{
+#if	STATS
+	AdvanceStats();
+#endif // STATS
+}
+
 FScopedEvent::FScopedEvent()
 	: Event(FEventPool<EEventPoolTypes::AutoReset>::Get().GetEventFromPool())
 { }
@@ -322,8 +378,11 @@ protected:
 	{
 		while (!TimeToDie)
 		{
-			// Wait for some work to do
-			DoWorkEvent->Wait();
+			{				
+				DECLARE_SCOPE_CYCLE_COUNTER( TEXT( "FQueuedThread::Run.WaitForWork" ), STAT_FQueuedThread_Run_WaitForWork, STATGROUP_ThreadPoolAsyncTasks );
+				// Wait for some work to do
+				DoWorkEvent->Wait();
+			}
 			IQueuedWork* LocalQueuedWork = QueuedWork;
 			QueuedWork = nullptr;
 			FPlatformMisc::MemoryBarrier();
@@ -406,6 +465,8 @@ public:
 	 */
 	void DoWork(IQueuedWork* InQueuedWork)
 	{
+		DECLARE_SCOPE_CYCLE_COUNTER( TEXT( "FQueuedThread::DoWork" ), STAT_FQueuedThread_DoWork, STATGROUP_ThreadPoolAsyncTasks );
+
 		check(QueuedWork == nullptr && "Can't do more than one task at a time");
 		// Tell the thread the work to be done
 		QueuedWork = InQueuedWork;
diff --git a/Engine/Source/Runtime/Core/Private/Stats/StatsCommand.cpp b/Engine/Source/Runtime/Core/Private/Stats/StatsCommand.cpp
--- a/Engine/Source/Runtime/Core/Private/Stats/StatsCommand.cpp
+++ b/Engine/Source/Runtime/Core/Private/Stats/StatsCommand.cpp
@@ -195,6 +195,139 @@ void DumpNonFrame(FStatsThreadState const& StatsData)
 	}
 }
 
+/** Returns stats based stack as human readable string. */
+static FString GetHumanReadableCallstack( const TArray<FStatNameAndInfo>& StatsStack )
+{
+	FString Result;
+
+	for (int32 Index = StatsStack.Num() - 1; Index >= 0; --Index)
+	{
+		const FStatNameAndInfo& NameAndInfo = StatsStack[Index];
+
+		const FString ShortName = NameAndInfo.GetShortName().GetPlainNameString();
+		FString Desc = NameAndInfo.GetDescription();
+		Desc.Trim();
+
+		// For threads use the thread name, as the description contains encoded thread id.
+		const FName GroupName = NameAndInfo.GetGroupName();
+		if (GroupName == TEXT( "STATGROUP_Threads" ))
+		{
+			Desc.Empty();
+		}
+
+		if (Desc.Len() == 0)
+		{
+			Result += ShortName;
+		}
+		else
+		{
+			Result += Desc;
+		}
+
+		if (Index > 0)
+		{
+			Result += TEXT( " <- " );
+		}
+	}
+
+	Result.ReplaceInline( TEXT( "STAT_" ), TEXT( "" ), ESearchCase::CaseSensitive );
+	return Result;
+}
+
+/** Dumps event history if specified thread name is the as for the printing event. Removes already listed events from the history. */
+void DumpEventsHistoryIfThreadValid( TArray<FEventData>& EventsHistoryForFrame, const FName ThreadName, float MinDurationToDisplay )
+{
+	bool bIgnoreGameAndRender = false;
+	if (ThreadName == NAME_None)
+	{
+		bIgnoreGameAndRender = true;
+	}
+
+	UE_LOG( LogStats, Log, TEXT( "Displaying events history for %s" ), *ThreadName.GetPlainNameString() );
+	for( int32 Index = 0; Index < EventsHistoryForFrame.Num(); ++Index )
+	{
+		const FEventData& EventStats = EventsHistoryForFrame[Index];
+		if (EventStats.DurationMS < MinDurationToDisplay)
+		{
+			break;
+		}
+		
+		const FName EventThreadName = EventStats.WaitStackStats[0].GetShortName();
+		if (EventThreadName == ThreadName || bIgnoreGameAndRender)
+		{
+			UE_LOG( LogStats, Log, TEXT( "Duration: %.2f MS" ), EventStats.DurationMS );
+			UE_LOG( LogStats, Log, TEXT( " Wait   : %s" ), *GetHumanReadableCallstack( EventStats.WaitStackStats ) );
+			UE_LOG( LogStats, Log, TEXT( " Trigger: %s" ), *GetHumanReadableCallstack( EventStats.TriggerStackStats ) );
+
+			EventsHistoryForFrame.RemoveAt( Index--, 1, false );
+		}	
+	}
+}
+
+static FDelegateHandle DumpEventsDelegateHandle;
+
+/** For the specified frame dumps events history to the log. */
+void DumpEvents( int64 TargetFrame, float DumpEventsCullMS, bool bDisplayAllThreads )
+{
+	FStatsThreadState& Stats = FStatsThreadState::GetLocalState();
+
+	// Prepare data.
+	const TArray<FStatMessage> Data = Stats.GetCondensedHistory( TargetFrame );
+
+	// Events
+	struct FSortByDurationMS
+	{
+		FORCEINLINE bool operator()( const FEventData& A, const FEventData& B ) const
+		{
+			// Sort descending
+			return B.DurationMS < A.DurationMS;
+		}
+	};
+	
+	TArray<FEventData> EventsHistoryForFrame;
+	for( const auto& It : Stats.EventsHistory )
+	{
+		if (It.Value.Frame >= TargetFrame /*- STAT_FRAME_SLOP*/ && It.Value.HasValidStacks() && It.Value.DurationMS > DumpEventsCullMS)
+		{
+			EventsHistoryForFrame.Add( It.Value );
+		}
+	}
+
+	// Don't print the header if we don't have data.
+	if (EventsHistoryForFrame.Num() == 0)
+	{
+		return;
+	}
+
+	UE_LOG( LogStats, Log, TEXT( "----------------------------------------" ) );
+	UE_LOG( LogStats, Log, TEXT( "Events history: Single frame %lld, greater than %2.1f ms" ), TargetFrame, DumpEventsCullMS );
+	
+
+	EventsHistoryForFrame.Sort( FSortByDurationMS() );
+
+	// First print all events that wait on the game thread.
+	DumpEventsHistoryIfThreadValid( EventsHistoryForFrame, NAME_GameThread, DumpEventsCullMS );
+
+	// Second print all events that wait on the rendering thread.
+	DumpEventsHistoryIfThreadValid( EventsHistoryForFrame, NAME_RenderThread, DumpEventsCullMS );
+
+	if (bDisplayAllThreads)
+	{
+		// Print all the remaining events.
+		DumpEventsHistoryIfThreadValid( EventsHistoryForFrame, NAME_None, DumpEventsCullMS );
+	}
+
+	UE_LOG( LogStats, Log, TEXT( "----------------------------------------" ) );
+}
+
+void DumpEventsOnce( int64 TargetFrame, float DumpEventsCullMS, bool bDisplayAllThreads )
+{
+	FStatsThreadState& Stats = FStatsThreadState::GetLocalState();
+	DumpEvents( TargetFrame, DumpEventsCullMS, bDisplayAllThreads );
+	StatsMasterEnableSubtract();
+	Stats.NewFrameDelegate.Remove( DumpEventsDelegateHandle );
+}
+
 void DumpCPUSummary(FStatsThreadState const& StatsData, int64 TargetFrame)
 {
 	UE_LOG(LogStats, Log, TEXT("CPU Summary: Single Frame %lld ---------------------------------"), TargetFrame);
@@ -230,7 +363,7 @@ void DumpCPUSummary(FStatsThreadState const& StatsData, int64 TargetFrame)
 		// The description of a thread group contains the thread name marker
 		const FString Desc = Item.NameAndInfo.GetDescription();
 		bool bIsThread = Desc.StartsWith( FStatConstants::ThreadNameMarker );
-		bool bIsStall = !bIsThread && Desc.StartsWith("CPU Stall");
+		bool bIsStall = !bIsThread && Desc.StartsWith("CPU Stall"); // TArray<FName> StallStats/StatMessages
 		
 		EStatOperation::Type Op = Item.NameAndInfo.GetField<EStatOperation>();
 		if ((Op == EStatOperation::ChildrenStart || Op == EStatOperation::ChildrenEnd ||  Op == EStatOperation::Leaf) && Item.NameAndInfo.GetFlag(EStatMetaFlags::IsCycle))
@@ -423,6 +556,9 @@ static void DumpHitch(int64 Frame)
 		}
 
 		LastHitchFrame = Frame;
+
+		// Display events, but only the large ones.
+		DumpEvents( Frame, 1.0f, false );
 	}
 }
 
@@ -584,7 +720,7 @@ struct FHUDGroupManager
 		Params = InParams;
 		Params.bReset = bResetData;
 
-		if (Params.ShouldReset())
+		if( Params.ShouldReset() )
 		{
 			// Disable only stats groups, leave the fake FPS, Unit group untouched.
 			for (const auto& It : EnabledGroups)
@@ -598,8 +734,8 @@ struct FHUDGroupManager
 		}
 
 		ResizeFramesHistory( Params.MaxHistoryFrames.Get() );
-	
-		const FName MaybeGroupFName = FName(*(FString(TEXT("STATGROUP_")) + Params.Group.Get().GetPlainNameString()));	
+
+		const FName MaybeGroupFName = FName(*(FString(TEXT("STATGROUP_")) + Params.Group.Get().GetPlainNameString()));
 		const bool bResults = HandleToggleCommandBroadcast( MaybeGroupFName, bCurrentEnabled, bOthersEnabled );
 		if (!bResults)
 		{
@@ -757,21 +893,21 @@ struct FHUDGroupManager
 		}
 		else
 		{
-			TSet<FName> HierEnabledItems;
-			for (auto It = EnabledGroups.CreateConstIterator(); It; ++It)
-			{
-				HierEnabledItems.Append( It.Value().EnabledItems );
-			}
-			FGroupFilter Filter( HierEnabledItems );		
+		TSet<FName> HierEnabledItems;
+		for( auto It = EnabledGroups.CreateConstIterator(); It; ++It )
+		{
+			HierEnabledItems.Append( It.Value().EnabledItems );
+		}
+		FGroupFilter Filter(HierEnabledItems);
 
-			// Generate root stats stack for current frame.
-			Stats.UncondenseStackStats( TargetFrame, NewFrame.HierarchyInclusive, &Filter, &NewFrame.NonStackStats );
+		// Generate root stats stack for current frame.
+		Stats.UncondenseStackStats( TargetFrame, NewFrame.HierarchyInclusive, &Filter, &NewFrame.NonStackStats );
 
-			{
-				SCOPE_CYCLE_COUNTER( STAT_GetFlatAggregates );
-				Stats.GetInclusiveAggregateStackStats( TargetFrame, NewFrame.InclusiveAggregate, &Filter, false );
-				Stats.GetExclusiveAggregateStackStats( TargetFrame, NewFrame.ExclusiveAggregate, &Filter, false );
-			}
+		{
+			SCOPE_CYCLE_COUNTER(STAT_GetFlatAggregates);
+			Stats.GetInclusiveAggregateStackStats( TargetFrame, NewFrame.InclusiveAggregate, &Filter, false );
+			Stats.GetExclusiveAggregateStackStats( TargetFrame, NewFrame.ExclusiveAggregate, &Filter, false );
+		}
 		}
 
 		NewFrame.HierarchyInclusive.AddSelf();
@@ -901,7 +1037,7 @@ struct FHUDGroupManager
 
 			// Divide stats to get average values.
 			AggregatedHierarchyHistory.Divide( NumFrames );
-			AggregatedHierarchyHistory.CopyExclusivesFromSelf();	
+			AggregatedHierarchyHistory.CopyExclusivesFromSelf();
 			if (Params.CullMs.Get() != 0.0f)
 			{
 				AggregatedHierarchyHistory.CullByCycles( int64( Params.CullMs.Get() / FPlatformTime::ToMilliseconds( 1 ) ) );
@@ -936,41 +1072,41 @@ struct FHUDGroupManager
 				}
 				else
 				{
-					if (InternalGroup.DisplayMode & EStatDisplayMode::Hierarchical)
-					{
-						// Linearize stack stats for easier rendering.
-						LinearizeStackForItems( AggregatedHierarchyHistory, InternalGroup.EnabledItems, HudGroup.HierAggregate, HudGroup.Indentation, 0 );
-					}
-
-					if (InternalGroup.DisplayMode & EStatDisplayMode::Flat)
+				if( InternalGroup.DisplayMode & EStatDisplayMode::Hierarchical )
+				{
+					// Linearize stack stats for easier rendering.
+					LinearizeStackForItems( AggregatedHierarchyHistory, InternalGroup.EnabledItems, HudGroup.HierAggregate, HudGroup.Indentation, 0 );
+				}
+				
+				if( InternalGroup.DisplayMode & EStatDisplayMode::Flat )
+				{
+					// Copy flat stats
+					for( int32 Index = 0; Index < AggregatedFlatHistory.Num(); ++Index )
 					{
-						// Copy flat stats
-						for (int32 Index = 0; Index < AggregatedFlatHistory.Num(); ++Index)
+						const FComplexStatMessage& AggregatedStatMessage = AggregatedFlatHistory[Index];
+						const bool bIsNonStackStat = !AggregatedStatMessage.NameAndInfo.GetFlag(EStatMetaFlags::IsPackedCCAndDuration);
+						const bool bToBeAdded = InternalGroup.EnabledItems.Contains( AggregatedStatMessage.NameAndInfo.GetRawName() );
+						if( bToBeAdded )
 						{
-							const FComplexStatMessage& AggregatedStatMessage = AggregatedFlatHistory[Index];
-							const bool bIsNonStackStat = !AggregatedStatMessage.NameAndInfo.GetFlag( EStatMetaFlags::IsPackedCCAndDuration );
-							const bool bToBeAdded = InternalGroup.EnabledItems.Contains( AggregatedStatMessage.NameAndInfo.GetRawName() );
-							if (bToBeAdded)
-							{
-								new(HudGroup.FlatAggregate) FComplexStatMessage( AggregatedStatMessage );
-							}
+							new(HudGroup.FlatAggregate) FComplexStatMessage( AggregatedStatMessage );
 						}
 					}
+				}
 
-					// Copy non-stack stats assigned to memory and counter groups.
-					for (int32 Index = 0; Index < AggregatedNonStackStatsHistory.Num(); ++Index)
-					{
-						const FComplexStatMessage& AggregatedStatMessage = AggregatedNonStackStatsHistory[Index];
-						const bool bIsMemory = AggregatedStatMessage.NameAndInfo.GetFlag( EStatMetaFlags::IsMemory );
-						TArray<FComplexStatMessage>& Dest = bIsMemory ? HudGroup.MemoryAggregate : HudGroup.CountersAggregate;
+				// Copy non-stack stats assigned to memory and counter groups.
+				for( int32 Index = 0; Index < AggregatedNonStackStatsHistory.Num(); ++Index )
+				{
+					const FComplexStatMessage& AggregatedStatMessage = AggregatedNonStackStatsHistory[Index];
+					const bool bIsMemory = AggregatedStatMessage.NameAndInfo.GetFlag( EStatMetaFlags::IsMemory );
+					TArray<FComplexStatMessage>& Dest = bIsMemory ? HudGroup.MemoryAggregate : HudGroup.CountersAggregate; 
 
-						const bool bToBeAdded = InternalGroup.EnabledItems.Contains( AggregatedStatMessage.NameAndInfo.GetRawName() );
-						if (bToBeAdded)
-						{
-							new(Dest)FComplexStatMessage( AggregatedStatMessage );
-						}
-					}
+					const bool bToBeAdded = InternalGroup.EnabledItems.Contains( AggregatedStatMessage.NameAndInfo.GetRawName() );
+					if( bToBeAdded )
+					{
+						new(Dest) FComplexStatMessage(AggregatedStatMessage);
+					}	
 				}
+			}
 
 				// Replace thread encoded id with the thread name.
 				for (auto& HierIt : HudGroup.HierAggregate)
@@ -1072,7 +1208,6 @@ static int32 MaxDepth = MAX_int32;
 static FString NameFilter;
 static FDelegateHandle DumpFrameDelegateHandle;
 static FDelegateHandle DumpCPUDelegateHandle;
-static FDelegateHandle DumpMemoryDelegateHandle;
 
 static void DumpFrame(int64 Frame)
 {
@@ -1191,7 +1326,8 @@ static void PrintStatsHelpToOutputDevice( FOutputDevice& Ar )
 	Ar.Log( TEXT("    stat dumpframe -ms=.001 -root=shadow"));
 
 	Ar.Log( TEXT("stat dumpave|dumpmax|dumpsum  [-start | -stop | -num=30] [-ms=5.0] [-depth=maxint] - aggregate stats over multiple frames"));
-	Ar.Log( TEXT("stat dumphitches  - dumps hitches"));
+	Ar.Log( TEXT("stat dumphitches - toggles dumping hitches"));
+	Ar.Log( TEXT("stat dumpevents [-ms=0.2] [-all] - toggles dumping events history for slow events, -all adds other threads besides game and render"));
 	Ar.Log( TEXT("stat dumpnonframe - dumps non-frame stats, usually memory stats"));
 	Ar.Log( TEXT("stat dumpcpu - dumps cpu stats"));
 
@@ -1219,7 +1355,7 @@ static void PrintStatsHelpToOutputDevice( FOutputDevice& Ar )
 
 	Ar.Log( TEXT("stat toggledebug - toggles tracking the most memory expensive stats"));
 
-	Ar.Log( TEXT( "stat slow [-ms=1.0] [-depth=4] - display the game and render thread stats" ) );
+	Ar.Log( TEXT( "stat slow [-ms=1.0] [-depth=4] - toggles displaying the game and render thread stats" ) );
 
 	Ar.Log( TEXT("add -memoryprofiler in the command line to enable the memory profiling"));
 	Ar.Log( TEXT("stat stopfile - stops tracking all memory operations and writes the results to the file"));
@@ -1350,6 +1486,15 @@ static void StatCmd(FString InCmd)
 			UE_LOG(LogStats, Log, TEXT( "**************************** %d hitches    %8.0fms total hitch time" ), HitchIndex, TotalHitchTime);
 		}
 	}
+	else if (FParse::Command( &Cmd, TEXT( "DumpEvents" ) ))
+	{
+		float DumpEventsCullMS = 0.1f;
+		FParse::Value( Cmd, TEXT( "MS=" ), DumpEventsCullMS );
+		const bool bDisplayAllThreads = FParse::Param( Cmd, TEXT( "all" ) );
+
+		StatsMasterEnableAdd();
+		DumpEventsDelegateHandle = Stats.NewFrameDelegate.AddStatic( &DumpEventsOnce, DumpEventsCullMS, bDisplayAllThreads );
+	}
 	else if( FParse::Command( &Cmd, TEXT( "STARTFILE" ) ) )
 	{
 		FString Filename;
@@ -1427,7 +1572,7 @@ static void StatCmd(FString InCmd)
 	else if (FParse::Command( &Cmd, TEXT( "slow" ) ))
 	{
 		FStatSlowParams Params( Cmd );
-		FHUDGroupManager::Get( Stats ).HandleCommand( Params, true );
+		FHUDGroupManager::Get(Stats).HandleCommand(Params, true);
 	}
 	else
 #endif
@@ -1527,6 +1672,9 @@ bool DirectStatsCommand(const TCHAR* Cmd, bool bBlockForCompletion /*= false*/,
 		else if( FParse::Command(&TempCmd,TEXT("DUMPHITCHES")) )
 		{
 		}
+		else if (FParse::Command( &TempCmd, TEXT( "DumpEvents" ) ))
+		{
+		}
 		else if( FParse::Command( &TempCmd, TEXT( "STOPFILE" ) ) )
 		{
 		}
diff --git a/Engine/Source/Runtime/Core/Private/Stats/StatsData.cpp b/Engine/Source/Runtime/Core/Private/Stats/StatsData.cpp
--- a/Engine/Source/Runtime/Core/Private/Stats/StatsData.cpp
+++ b/Engine/Source/Runtime/Core/Private/Stats/StatsData.cpp
@@ -29,6 +29,9 @@ const FString FStatConstants::StatsFileRawExtension = TEXT( ".ue4statsraw" );
 
 const FString FStatConstants::ThreadNameMarker = TEXT( "Thread_" );
 
+const FName FStatConstants::NAME_EventWaitWithId = FStatNameAndInfo( GET_STATFNAME( STAT_EventWaitWithId ), true ).GetRawName();
+const FName FStatConstants::NAME_EventTriggerWithId = FStatNameAndInfo( GET_STATFNAME( STAT_EventTriggerWithId ), true ).GetRawName();
+
 /*-----------------------------------------------------------------------------
 	FRawStatStackNode
 -----------------------------------------------------------------------------*/
@@ -697,10 +700,24 @@ void FStatsThreadState::ScanForAdvance(FStatPacketArray& NewData)
 	uint32 Count = 0;
 	for (int32 Index = 0; Index < NewData.Packets.Num(); Index++)
 	{
-
-		const int64 FrameNum = NewData.Packets[Index]->ThreadType == EThreadType::Renderer ? CurrentRenderFrame : CurrentGameFrame;
-		NewData.Packets[Index]->AssignFrame( FrameNum );
-		
+		const EThreadType::Type ThreadType = NewData.Packets[Index]->ThreadType;
+		if ( ThreadType == EThreadType::Renderer)
+		{
+			NewData.Packets[Index]->AssignFrame( CurrentRenderFrame );
+		}
+		else if (ThreadType == EThreadType::Game)
+		{
+			NewData.Packets[Index]->AssignFrame( CurrentGameFrame );
+		}
+		else if (ThreadType == EThreadType::Other)
+		{
+			// @see FThreadStats::DetectAndUpdateCurrentGameFrame 
+		}
+		else
+		{
+			checkf( 0, TEXT( "Unknown thread type" ) );
+		}
+				
 		const FStatMessagesArray& Data = NewData.Packets[Index]->StatMessages;
 		ScanForAdvance(Data);
 		Count += Data.Num();
@@ -745,15 +762,15 @@ void FStatsThreadState::ProcessNonFrameStats(FStatMessagesArray& Data, TSet<FNam
 				Op != EStatOperation::ChildrenEnd &&
 				Op != EStatOperation::Leaf &&
 				Op != EStatOperation::AdvanceFrameEventGameThread &&
-				Op != EStatOperation::AdvanceFrameEventRenderThread /*&&
-				Op != EStatOperation::Memory*/
+				Op != EStatOperation::AdvanceFrameEventRenderThread
 				))
 			{
 				UE_LOG(LogStats, Fatal, TEXT( "Stat %s was not cleared every frame, but was used with a scope cycle counter." ), *Item.NameAndInfo.GetRawName().ToString() );
 			}
 			else
 			{
-				if( Op != EStatOperation::Memory )
+				// Ignore any memory or special messages, they shouldn't be treated as regular stats messages.
+				if( Op != EStatOperation::Memory && Op != EStatOperation::SpecialMessageMarker )
 				{
 					FStatMessage* Result = NotClearedEveryFrame.Find(Item.NameAndInfo.GetRawName());
 					if (!Result)
@@ -788,6 +805,7 @@ void FStatsThreadState::AddToHistoryAndEmpty(FStatPacketArray& NewData)
 		ShortNameToLongName.Empty();
 		Groups.Empty();
 		History.Empty();
+		EventsHistory.Empty();
 		return;
 	}
 
@@ -909,13 +927,20 @@ void FStatsThreadState::AddToHistoryAndEmpty(FStatPacketArray& NewData)
 	for (auto It = History.CreateIterator(); It; ++It)
 	{
 		int64 ThisFrame = It.Key();
-		FStatPacketArray& StatPacketArray = It.Value();
 		if (ThisFrame <= LastFullFrameMetaAndNonFrame && ThisFrame < MinFrameToKeep)
 		{
 			check(ThisFrame <= LastFullFrameMetaAndNonFrame);
 			It.RemoveCurrent();
 		}
 	}
+	for (auto It = EventsHistory.CreateIterator(); It; ++It)
+	{
+		int64 ThisFrame = It.Value().Frame;
+		if (ThisFrame <= LastFullFrameProcessed && ThisFrame < MinFrameToKeep)
+		{
+			It.RemoveCurrent();
+		}
+	}
 	for (auto It = CondensedStackHistory.CreateIterator(); It; ++It)
 	{
 		int64 ThisFrame = It.Key();
@@ -1212,10 +1237,9 @@ void FStatsThreadState::GetRawStackStats(int64 TargetFrame, FRawStatStackNode& R
 				check(Item.NameAndInfo.GetFlag(EStatMetaFlags::DummyAlwaysOne));  // we should never be sending short names to the stats anymore
 
 				EStatOperation::Type Op = Item.NameAndInfo.GetField<EStatOperation>();
-
-				FName LongName = Item.NameAndInfo.GetRawName();
+				const FName LongName = Item.NameAndInfo.GetRawName();
 				if (Op == EStatOperation::CycleScopeStart || Op == EStatOperation::CycleScopeEnd)
-				{
+				{				
 					check(Item.NameAndInfo.GetFlag(EStatMetaFlags::IsCycle));
 					if (Op == EStatOperation::CycleScopeStart)
 					{
@@ -1230,7 +1254,6 @@ void FStatsThreadState::GetRawStackStats(int64 TargetFrame, FRawStatStackNode& R
 						Stack.Add(Result);
 						StartStack.Add(&Item);
 						Current = Result;
-
 					}
 					if (Op == EStatOperation::CycleScopeEnd)
 					{
@@ -1241,6 +1264,65 @@ void FStatsThreadState::GetRawStackStats(int64 TargetFrame, FRawStatStackNode& R
 						Current = Stack.Last();
 					}
 				}
+				// We are using here EStatOperation::SpecialMessageMarker to indicate custom stat messages
+				// At this moment only these messages are supported:
+				//	EventWaitWithId
+				//	EventTriggerWithId
+				else if( Op == EStatOperation::SpecialMessageMarker )
+				{
+					const FName RawName = Item.NameAndInfo.GetRawName();
+					const uint64 PacketEventIdAndCycles = Item.GetValue_Ptr();
+					const uint32 EventId = uint32(PacketEventIdAndCycles >> 32);
+					const uint32 EventCycles = uint32(PacketEventIdAndCycles & MAX_uint32);
+
+					if (FStatConstants::NAME_EventWaitWithId == RawName)
+					{
+						TArray<FStatNameAndInfo> EventWaitStack;
+						for (const auto& It : Stack)
+						{
+							EventWaitStack.Add( It->Meta.NameAndInfo );
+						}
+
+#if	UE_BUILD_DEBUG
+						// Debug check, detect duplicates.
+						FEventData* EventPtr = EventsHistory.Find( EventId );
+						if (EventPtr && EventPtr->WaitStackStats.Num() > 0)
+						{
+							int32 k = 0; k++;
+						}
+#endif // UE_BUILD_DEBUG
+
+						FEventData& EventStats = EventsHistory.FindOrAdd( EventId );
+						EventStats.WaitStackStats = EventWaitStack;
+						EventStats.Frame = EventStats.HasValidStacks() ? TargetFrame : 0; // Only to maintain history.
+					}
+
+					if (FStatConstants::NAME_EventTriggerWithId == RawName)
+					{
+						TArray<FStatNameAndInfo> EventTriggerStack;
+						for (const auto& It : Stack)
+						{
+							EventTriggerStack.Add( It->Meta.NameAndInfo );
+						}
+
+#if	UE_BUILD_DEBUG
+						// Debug check, detect duplicates.
+						FEventData* EventPtr = EventsHistory.Find( EventId );
+						if (EventPtr && EventPtr->TriggerStackStats.Num() > 0)
+						{
+							int32 k = 0; k++;
+						}
+#endif // UE_BUILD_DEBUG
+
+						FEventData& EventStats = EventsHistory.FindOrAdd( EventId );
+
+						EventStats.TriggerStackStats = EventTriggerStack;
+						EventStats.Duration = EventCycles;
+						EventStats.DurationMS = FPlatformTime::ToMilliseconds( EventCycles );
+						EventStats.Frame = EventStats.HasValidStacks() ? TargetFrame : 0; // Only to maintain history.
+					}
+
+				}
 				else if( Op == EStatOperation::Memory )
 				{
 					// Should never happen.
diff --git a/Engine/Source/Runtime/Core/Public/HAL/ThreadingBase.h b/Engine/Source/Runtime/Core/Public/HAL/ThreadingBase.h
--- a/Engine/Source/Runtime/Core/Public/HAL/ThreadingBase.h
+++ b/Engine/Source/Runtime/Core/Public/HAL/ThreadingBase.h
@@ -20,9 +20,9 @@ enum EThreadPriority
 
 	/** Below normal priority. */
 	TPri_BelowNormal,
-    
-    /** Highest priority. */
-    TPri_Highest
+	
+	/** Highest priority. */
+	TPri_Highest
 };
 
 
@@ -103,8 +103,39 @@ public:
 		return Wait(WaitTime.GetTotalMilliseconds(), bIgnoreThreadIdleStats);
 	}
 
+	/** Default constructor. */
+	FEvent()
+		: EventId( 0 )
+		, EventStartCycles( 0 )
+	{}
+
 	/** Virtual destructor. */
-	virtual ~FEvent() { }
+	virtual ~FEvent() 
+	{}
+
+	// DO NOT MODIFY THESE
+
+	/** Advances stats associated with this event. Used to monitor wait->trigger history. */
+	void AdvanceStats();
+
+protected:
+	/** Sends to the stats a special messages which encodes a wait for the event. */
+	void WaitForStats();
+
+	/** Send to the stats a special message which encodes a trigger for the event. */
+	void TriggerForStats();
+
+	/** Resets start cycles to 0. */
+	void ResetForStats();
+
+	/** Counter used to generate an unique id for the events. */
+	static uint32 EventUniqueId;
+
+	/** An unique id of this event. */
+	uint32 EventId;
+
+	/** Greater than 0, if the event called wait. */
+	uint32 EventStartCycles;
 };
 
 
diff --git a/Engine/Source/Runtime/Core/Public/Misc/CoreStats.h b/Engine/Source/Runtime/Core/Public/Misc/CoreStats.h
--- a/Engine/Source/Runtime/Core/Public/Misc/CoreStats.h
+++ b/Engine/Source/Runtime/Core/Public/Misc/CoreStats.h
@@ -63,4 +63,14 @@ DECLARE_CYCLE_STAT_EXTERN(TEXT("Pump Messages"),STAT_PumpMessages,STATGROUP_Thre
 DECLARE_FLOAT_COUNTER_STAT_EXTERN(TEXT("Percentage CPU utilization"),STAT_CPUTimePct,STATGROUP_Threading, CORE_API);
 DECLARE_FLOAT_COUNTER_STAT_EXTERN(TEXT("Percentage CPU utilization (relative to one core)"),STAT_CPUTimePctRelative,STATGROUP_Threading, CORE_API);
 
+
+/*-----------------------------------------------------------------------------
+	CPU Stalls
+-----------------------------------------------------------------------------*/
+
 DECLARE_CYCLE_STAT_EXTERN(TEXT("CPU Stall - Sleep"), STAT_Sleep, STATGROUP_CPUStalls, CORE_API);
+DECLARE_CYCLE_STAT_EXTERN(TEXT("CPU Stall - Wait For Event" ), STAT_EventWait, STATGROUP_CPUStalls,);
+
+/** The id will be stored as uint64 in the stat message. */
+DECLARE_PTR_STAT_EXTERN( TEXT( "CPU Stall - Wait For Event with ID" ), STAT_EventWaitWithId, STATGROUP_CPUStalls, );
+DECLARE_PTR_STAT_EXTERN( TEXT( "CPU Stall - Trigger For Event with ID" ), STAT_EventTriggerWithId, STATGROUP_CPUStalls, );
diff --git a/Engine/Source/Runtime/Core/Public/Stats/StatsData.h b/Engine/Source/Runtime/Core/Public/Stats/StatsData.h
--- a/Engine/Source/Runtime/Core/Public/Stats/StatsData.h
+++ b/Engine/Source/Runtime/Core/Public/Stats/StatsData.h
@@ -35,6 +35,12 @@ struct CORE_API FStatConstants
 
 	/** Indicates that the item is a thread. */
 	static const FString ThreadNameMarker;
+
+	/** A raw name for the event wait with id. */
+	static const FName NAME_EventWaitWithId;
+
+	/** A raw name for the event trigger with id. */
+	static const FName NAME_EventTriggerWithId;
 };
 
 namespace LexicalConversion
@@ -406,6 +412,37 @@ protected:
 	}
 };
 
+
+/** Holds information about event history, callstacks for the wait and the corresponding trigger. */
+struct FEventData
+{
+	/** Default constructor. */
+	FEventData()
+		: Frame( 0 )
+		, Duration( 0 )
+		, DurationMS( 0.0f )
+	{}
+
+	/**
+	 * @return true, if the wait and trigger stacks are not empty.
+	 */
+	bool HasValidStacks() const
+	{
+		return WaitStackStats.Num() > 0 && TriggerStackStats.Num() > 0;
+	}
+
+	/** Stack for the event's wait. */
+	TArray<FStatNameAndInfo> WaitStackStats;
+	/** Stat for the event's trigger. */
+	TArray<FStatNameAndInfo> TriggerStackStats;
+	/** Frame, used to remove older records from the history. */
+	int64 Frame;
+	/** Duration of the event's wait, in cycles. */
+	uint32 Duration;
+	/** Duration of the event's wait, in milliseconds. */
+	float DurationMS;
+};
+
 //@todo split header
 
 /**
@@ -575,6 +612,9 @@ public:
 	/** Map from a short name to the SetLongName message that defines the metadata and long name for this stat **/
 	TMap<FName, FStatMessage> ShortNameToLongName;
 
+	/** Map from the unique event id to the event stats. */
+	mutable TMap<uint32, FEventData> EventsHistory;
+
 	/** Map from memory pool to long name**/
 	TMap<FPlatformMemory::EMemoryCounterRegion, FName> MemoryPoolToCapacityLongName;
 
