Stats - Minor optimization (Added stats for wait/trigger in events)

diff --git a/Engine/Source/Runtime/Core/Private/HAL/ThreadingBase.cpp b/Engine/Source/Runtime/Core/Private/HAL/ThreadingBase.cpp
--- a/Engine/Source/Runtime/Core/Private/HAL/ThreadingBase.cpp
+++ b/Engine/Source/Runtime/Core/Private/HAL/ThreadingBase.cpp
@@ -192,9 +192,8 @@ void FEvent::WaitForStats()
 	// Only start counting on the first wait, trigger will "close" the history.
 	if( FThreadStats::IsCollectingData() && EventStartCycles == 0 )
 	{
-		FThreadStats* ThreadStats = FThreadStats::GetThreadStats();
 		const uint64 PacketEventIdAndCycles = ((uint64)EventId << 32) | 0;
-		ThreadStats->AddMessage( GET_STATFNAME( STAT_EventWaitWithId ), EStatOperation::SpecialMessageMarker, PacketEventIdAndCycles );
+		FThreadStats::AddMessage( GET_STATFNAME( STAT_EventWaitWithId ), EStatOperation::SpecialMessageMarker, PacketEventIdAndCycles );
 		EventStartCycles = FPlatformTime::Cycles();
 	}
 #endif // STATS
@@ -206,11 +205,10 @@ void FEvent::TriggerForStats()
 	// Only add wait-trigger pairs.
 	if( EventStartCycles > 0 && FThreadStats::IsCollectingData() )
 	{
-		FThreadStats* ThreadStats = FThreadStats::GetThreadStats();
 		const uint32 EndCycles = FPlatformTime::Cycles();
 		const int32 DeltaCycles = int32( EndCycles - EventStartCycles );
 		const uint64 PacketEventIdAndCycles = ((uint64)EventId << 32) | DeltaCycles;
-		ThreadStats->AddMessage( GET_STATFNAME( STAT_EventTriggerWithId ), EStatOperation::SpecialMessageMarker, PacketEventIdAndCycles );
+		FThreadStats::AddMessage( GET_STATFNAME( STAT_EventTriggerWithId ), EStatOperation::SpecialMessageMarker, PacketEventIdAndCycles );
 
 		AdvanceStats();
 	}
