#ue4 - 30% Faster FQuat::Normalize() and FQuat::GetNormalized() when vector intrinsics are enabled, and surrounding code has the quat in a VectorRegister (the common case).

diff --git a/Engine/Source/Runtime/Core/Public/Math/Quat.h b/Engine/Source/Runtime/Core/Public/Math/Quat.h
--- a/Engine/Source/Runtime/Core/Public/Math/Quat.h
+++ b/Engine/Source/Runtime/Core/Public/Math/Quat.h
@@ -778,7 +778,7 @@ FORCEINLINE bool FQuat::operator!=( const FQuat& Q ) const
 	return VectorMaskBits(VectorCompareNE(A, B)) != 0x00;
 #else
 	return X != Q.X || Y != Q.Y || Z != Q.Z || W != Q.W;
-#endif
+#endif // PLATFORM_ENABLE_VECTORINTRINSICS
 }
 
 
@@ -790,9 +790,20 @@ FORCEINLINE float FQuat::operator|( const FQuat& Q ) const
 
 FORCEINLINE void FQuat::Normalize(float Tolerance)
 {
+#if PLATFORM_ENABLE_VECTORINTRINSICS
+	const VectorRegister Vector = VectorLoadAligned(this);
+
+	const VectorRegister SquareSum = VectorDot4(Vector, Vector);
+	const VectorRegister NonZeroMask = VectorCompareGE(SquareSum, VectorLoadFloat1(&Tolerance));
+	const VectorRegister InvLength = VectorReciprocalSqrtAccurate(SquareSum);
+	const VectorRegister NormalizedVector = VectorMultiply(InvLength, Vector);
+	VectorRegister Result = VectorSelect(NonZeroMask, NormalizedVector, GlobalVectorConstants::Float0001);
+
+	VectorStoreAligned(Result, this);
+#else
 	const float SquareSum = X * X + Y * Y + Z * Z + W * W;
 
-	if (SquareSum > Tolerance)
+	if (SquareSum >= Tolerance)
 	{
 		const float Scale = FMath::InvSqrt(SquareSum);
 
@@ -805,6 +816,7 @@ FORCEINLINE void FQuat::Normalize(float Tolerance)
 	{
 		*this = FQuat::Identity;
 	}
+#endif // PLATFORM_ENABLE_VECTORINTRINSICS
 }
 
 
