Changed the text filter to inject AND tokens at spaces (rather than OR tokens) as this matches the old behavior of the filter.

diff --git a/Engine/Source/Runtime/Core/Private/Misc/TextFilterExpressionEvaluator.cpp b/Engine/Source/Runtime/Core/Private/Misc/TextFilterExpressionEvaluator.cpp
--- a/Engine/Source/Runtime/Core/Private/Misc/TextFilterExpressionEvaluator.cpp
+++ b/Engine/Source/Runtime/Core/Private/Misc/TextFilterExpressionEvaluator.cpp
@@ -436,7 +436,7 @@ bool FTextFilterExpressionEvaluator::SetFilterText(const FText& InFilterText)
 
 			// We process the tokens in two passes (doing these separately minimizes the combinations required by step 2):
 			//	Pass 1) We process and remove any text comparison modification tokens (FTextCmpExact and FTextCmpAnchor)
-			//	Pass 2) We inject OR operators between any adjacent text and sub-expression tokens to mimic the old search filter behavior
+			//	Pass 2) We inject AND operators between any adjacent text and sub-expression tokens to mimic the old search filter behavior
 			TArray<FExpressionToken> FinalTokens;
 			FinalTokens.Reserve(TmpLexTokens.Num());
 
@@ -497,13 +497,13 @@ bool FTextFilterExpressionEvaluator::SetFilterText(const FText& InFilterText)
 						bIsComplexExpression = true;
 					}
 
-					// We need to inject OR operators between any adjacent text and sub-expression tokens to mimic the old search filter behavior
+					// We need to inject AND operators between any adjacent text and sub-expression tokens to mimic the old search filter behavior
 					if ((LastTokenType == ETextNodeType::Text || LastTokenType == ETextNodeType::SubExpressionEnd) &&
 						(CurrentTokenType == ETextNodeType::Text || CurrentTokenType == ETextNodeType::SubExpressionStart)
 						)
 					{
 						// Inject the new token between the current one and the previous one, and then skip over the newly added token
-						FinalTokens.Insert(FExpressionToken(CurrentToken.Context, FOr()), FinalTokenIndex++);
+						FinalTokens.Insert(FExpressionToken(CurrentToken.Context, FAnd()), FinalTokenIndex++);
 					}
 				}
 			}
