Stats - Temp compile fix

diff --git a/Engine/Source/Runtime/Core/Public/Misc/CoreStats.h b/Engine/Source/Runtime/Core/Public/Misc/CoreStats.h
--- a/Engine/Source/Runtime/Core/Public/Misc/CoreStats.h
+++ b/Engine/Source/Runtime/Core/Public/Misc/CoreStats.h
@@ -69,7 +69,7 @@ DECLARE_FLOAT_COUNTER_STAT_EXTERN(TEXT("Percentage CPU utilization (relative to
 -----------------------------------------------------------------------------*/
 
 DECLARE_CYCLE_STAT_EXTERN(TEXT("CPU Stall - Sleep"), STAT_Sleep, STATGROUP_CPUStalls, CORE_API);
-DECLARE_CYCLE_STAT_EXTERN(TEXT("CPU Stall - Wait For Event" ), STAT_EventWait, STATGROUP_CPUStalls,);
+//DECLARE_CYCLE_STAT_EXTERN(TEXT("CPU Stall - Wait For Event" ), STAT_EventWait, STATGROUP_CPUStalls,);
 
 /** The id will be stored as uint64 in the stat message. */
 DECLARE_PTR_STAT_EXTERN( TEXT( "CPU Stall - Wait For Event with ID" ), STAT_EventWaitWithId, STATGROUP_CPUStalls, );
