Stats - Added stats for wait/trigger in events (platform)

diff --git a/Engine/Source/Runtime/Core/Private/GenericPlatform/GenericPlatformProcess.cpp b/Engine/Source/Runtime/Core/Private/GenericPlatform/GenericPlatformProcess.cpp
--- a/Engine/Source/Runtime/Core/Private/GenericPlatform/GenericPlatformProcess.cpp
+++ b/Engine/Source/Runtime/Core/Private/GenericPlatform/GenericPlatformProcess.cpp
@@ -12,6 +12,8 @@
 #endif
 
 DEFINE_STAT(STAT_Sleep);
+DEFINE_STAT(STAT_EventWait);
+
 
 void* FGenericPlatformProcess::GetDllHandle( const TCHAR* Filename )
 {
@@ -297,11 +299,11 @@ void FGenericPlatformProcess::ConditionalSleep(const TFunctionRef<bool()>& Condi
 
 #include "PThreadEvent.h"
 
-DECLARE_CYCLE_STAT(TEXT("CPU Stall - Wait For Event"),STAT_EventWait,STATGROUP_CPUStalls);
-
 bool FPThreadEvent::Wait(uint32 WaitTime, const bool bIgnoreThreadIdleStats /*= false*/)
 {
-	FScopeCycleCounter Counter(StatID);
+	//WaitForStats();
+
+	SCOPE_CYCLE_COUNTER(STAT_EventWait);
 	FThreadIdleStats::FScopeIdle Scope(bIgnoreThreadIdleStats);
 
 	check(bInitialized);
diff --git a/Engine/Source/Runtime/Core/Private/Misc/EventPool.cpp b/Engine/Source/Runtime/Core/Private/Misc/EventPool.cpp
--- a/Engine/Source/Runtime/Core/Private/Misc/EventPool.cpp
+++ b/Engine/Source/Runtime/Core/Private/Misc/EventPool.cpp
@@ -3,27 +3,3 @@
 
 #include "CorePrivatePCH.h"
 #include "EventPool.h"
-
-TStatId FEventStats::CreateStatID()
-{
-#if	STATS
-	static FThreadSafeCounter UniqueId;
-	const int32 PreviousId = UniqueId.Add(1);
-	const FString UniqueEventName = FString::Printf( TEXT("CPU Stall - Wait for event %3i"), PreviousId );
-
-	const FName StatName = FName( *UniqueEventName );
-	FStartupMessages::Get().AddMetadata( StatName, *UniqueEventName, 
-		STAT_GROUP_TO_FStatGroup( STATGROUP_CPUStalls )::GetGroupName(), 
-		STAT_GROUP_TO_FStatGroup( STATGROUP_CPUStalls )::GetGroupCategory(), 
-		STAT_GROUP_TO_FStatGroup( STATGROUP_CPUStalls )::GetDescription(),
-		true, EStatDataType::ST_int64, true );
-
-	return IStatGroupEnableManager::Get().GetHighPerformanceEnableForStat(StatName, 
-		STAT_GROUP_TO_FStatGroup( STATGROUP_CPUStalls )::GetGroupName(), 
-		STAT_GROUP_TO_FStatGroup( STATGROUP_CPUStalls )::GetGroupCategory(), 
-		STAT_GROUP_TO_FStatGroup( STATGROUP_CPUStalls )::DefaultEnable,
-		true, EStatDataType::ST_int64, *UniqueEventName, true);
-#else
-	return TStatId();
-#endif // STATS
-}
diff --git a/Engine/Source/Runtime/Core/Private/Misc/EventPool.h b/Engine/Source/Runtime/Core/Private/Misc/EventPool.h
--- a/Engine/Source/Runtime/Core/Private/Misc/EventPool.h
+++ b/Engine/Source/Runtime/Core/Private/Misc/EventPool.h
@@ -15,13 +15,6 @@ enum class EEventPoolTypes
 	ManualReset
 };
 
-/** Helper class used to create a new stat id for events. */
-struct FEventStats
-{
-	/** Creates a new stats id that can be used for events. */
-	CORE_API static TStatId CreateStatID();
-};
-
 /**
  * Template class for event pools.
  *
@@ -67,6 +60,8 @@ public:
 			PRAGMA_ENABLE_DEPRECATION_WARNINGS
 		}
 
+		Result->AdvanceStats();
+
 		check(Result);
 		return Result;
 	}
diff --git a/Engine/Source/Runtime/Core/Private/Windows/WindowsEvent.h b/Engine/Source/Runtime/Core/Private/Windows/WindowsEvent.h
--- a/Engine/Source/Runtime/Core/Private/Windows/WindowsEvent.h
+++ b/Engine/Source/Runtime/Core/Private/Windows/WindowsEvent.h
@@ -34,7 +34,7 @@ public:
 		// Create the event and default it to non-signaled
 		Event = CreateEvent(nullptr, bIsManualReset, 0, nullptr);
 		ManualReset = bIsManualReset;
-		StatID = FEventStats::CreateStatID();
+
 		return Event != nullptr;
 	}
 
@@ -43,32 +43,16 @@ public:
 		return ManualReset;
 	}
 
-	virtual void Trigger() override
-	{
-		check(Event);
-
-		SetEvent(Event);
-	}
-
-	virtual void Reset() override
-	{
-		check(Event);
-
-		ResetEvent(Event);
-	}
-
+	virtual void Trigger() override;
+	virtual void Reset() override;
 	virtual bool Wait( uint32 WaitTime, const bool bIgnoreThreadIdleStats = false ) override;
 
 private:
-
 	/** Holds the handle to the event. */
 	HANDLE Event;
 
 	/** Whether the signaled state of the event needs to be reset manually. */
 	bool ManualReset;
-
-	/** Stat ID of this event. */
-	TStatId StatID;
 };
 
 
diff --git a/Engine/Source/Runtime/Core/Private/Windows/WindowsPlatformProcess.cpp b/Engine/Source/Runtime/Core/Private/Windows/WindowsPlatformProcess.cpp
--- a/Engine/Source/Runtime/Core/Private/Windows/WindowsPlatformProcess.cpp
+++ b/Engine/Source/Runtime/Core/Private/Windows/WindowsPlatformProcess.cpp
@@ -977,9 +977,9 @@ bool FWindowsPlatformProcess::ResolveNetworkPath( FString InUNCPath, FString& Ou
 
 void FWindowsPlatformProcess::Sleep( float Seconds )
 {
-	SCOPE_CYCLE_COUNTER(STAT_Sleep);
+	SCOPE_CYCLE_COUNTER( STAT_Sleep );
 	FThreadIdleStats::FScopeIdle Scope;
-	SleepNoStats(Seconds);
+	SleepNoStats( Seconds );
 }
 
 void FWindowsPlatformProcess::SleepNoStats(float Seconds)
@@ -1021,10 +1021,27 @@ FEvent* FWindowsPlatformProcess::CreateSynchEvent(bool bIsManualReset)
 
 bool FEventWin::Wait(uint32 WaitTime, const bool bIgnoreThreadIdleStats /*= false*/)
 {
-	FScopeCycleCounter Counter(StatID);
-	FThreadIdleStats::FScopeIdle Scope(bIgnoreThreadIdleStats);
-	check(Event);
-	return (WaitForSingleObject(Event, WaitTime) == WAIT_OBJECT_0);
+	WaitForStats();
+
+	SCOPE_CYCLE_COUNTER( STAT_EventWait );
+	check( Event );
+
+	FThreadIdleStats::FScopeIdle Scope( bIgnoreThreadIdleStats );
+	return (WaitForSingleObject( Event, WaitTime ) == WAIT_OBJECT_0);
+}
+
+void FEventWin::Trigger()
+{
+	TriggerForStats();
+	check( Event );
+	SetEvent( Event );
+}
+
+void FEventWin::Reset()
+{
+	ResetForStats();
+	check( Event );
+	ResetEvent( Event );
 }
 
 #include "HideWindowsPlatformTypes.h"
diff --git a/Engine/Source/Runtime/Core/Public/HAL/PThreadEvent.h b/Engine/Source/Runtime/Core/Public/HAL/PThreadEvent.h
--- a/Engine/Source/Runtime/Core/Public/HAL/PThreadEvent.h
+++ b/Engine/Source/Runtime/Core/Public/HAL/PThreadEvent.h
@@ -24,9 +24,6 @@ class FPThreadEvent
     pthread_mutex_t Mutex;
     pthread_cond_t Condition;
 
-	/** Stat ID of this event. */
-	TStatId StatID;
-
 	inline void LockEventMutex()
 	{
 		const int rc = pthread_mutex_lock(&Mutex);
@@ -117,7 +114,6 @@ public:
 				pthread_mutex_destroy(&Mutex);
 			}
 		}
-		StatID = FEventStats::CreateStatID();
 		return RetVal;
 	}
 
@@ -128,6 +124,8 @@ public:
 
 	virtual void Trigger() override
 	{
+		//TriggerForStats();
+
 		check(bInitialized);
 
 		LockEventMutex();
@@ -148,15 +146,17 @@ public:
 			check(rc == 0);
 		}
 
-		UnlockEventMutex();
+		UnlockEventMutex();	
 	}
 
 	virtual void Reset() override
 	{
+		//ResetForStats();
+
 		check(bInitialized);
 		LockEventMutex();
 		Triggered = TRIGGERED_NONE;
-		UnlockEventMutex();
+		UnlockEventMutex();	
 	}
 
 	virtual bool Wait( uint32 WaitTime = (uint32)-1, const bool bIgnoreThreadIdleStats = false ) override;
diff --git a/Engine/Source/Runtime/Core/Public/Misc/CoreStats.h b/Engine/Source/Runtime/Core/Public/Misc/CoreStats.h
--- a/Engine/Source/Runtime/Core/Public/Misc/CoreStats.h
+++ b/Engine/Source/Runtime/Core/Public/Misc/CoreStats.h
@@ -69,7 +69,7 @@ DECLARE_FLOAT_COUNTER_STAT_EXTERN(TEXT("Percentage CPU utilization (relative to
 -----------------------------------------------------------------------------*/
 
 DECLARE_CYCLE_STAT_EXTERN(TEXT("CPU Stall - Sleep"), STAT_Sleep, STATGROUP_CPUStalls, CORE_API);
-//DECLARE_CYCLE_STAT_EXTERN(TEXT("CPU Stall - Wait For Event" ), STAT_EventWait, STATGROUP_CPUStalls,);
+DECLARE_CYCLE_STAT_EXTERN(TEXT("CPU Stall - Wait For Event" ), STAT_EventWait, STATGROUP_CPUStalls,);
 
 /** The id will be stored as uint64 in the stat message. */
 DECLARE_PTR_STAT_EXTERN( TEXT( "CPU Stall - Wait For Event with ID" ), STAT_EventWaitWithId, STATGROUP_CPUStalls, );
