Fixed some cases where the text filter would fail to update

The text filter was stripping whitespace and storing the stripped version. This meant that a search term containing only whitespace could never be cleared again, as the filter never notified of a change in text.

The other issue was that some code was performing a pre-test to see if the text had changed, however this code was doing case-insensitive comparisons, which no longer suffice as the operators are case-sensitive.

diff --git a/Engine/Source/Editor/ContentBrowser/Private/SAssetPicker.cpp b/Engine/Source/Editor/ContentBrowser/Private/SAssetPicker.cpp
--- a/Engine/Source/Editor/ContentBrowser/Private/SAssetPicker.cpp
+++ b/Engine/Source/Editor/ContentBrowser/Private/SAssetPicker.cpp
@@ -303,7 +303,8 @@ FText SAssetPicker::GetHighlightedText() const
 
 void SAssetPicker::SetSearchBoxText(const FText& InSearchText)
 {
-	if ( !InSearchText.EqualToCaseIgnored(TextFilter->GetRawFilterText()) )
+	// Has anything changed? (need to test case as the operators are case-sensitive)
+	if (!InSearchText.ToString().Equals(TextFilter->GetRawFilterText().ToString(), ESearchCase::CaseSensitive))
 	{
 		TextFilter->SetRawFilterText(InSearchText);
 		if (InSearchText.IsEmpty())
diff --git a/Engine/Source/Editor/ContentBrowser/Private/SContentBrowser.cpp b/Engine/Source/Editor/ContentBrowser/Private/SContentBrowser.cpp
--- a/Engine/Source/Editor/ContentBrowser/Private/SContentBrowser.cpp
+++ b/Engine/Source/Editor/ContentBrowser/Private/SContentBrowser.cpp
@@ -1240,7 +1240,8 @@ void SContentBrowser::NewFolderRequested(const FString& SelectedPath)
 
 void SContentBrowser::SetSearchBoxText(const FText& InSearchText)
 {
-	if (!InSearchText.EqualToCaseIgnored(TextFilter->GetRawFilterText()))
+	// Has anything changed? (need to test case as the operators are case-sensitive)
+	if (!InSearchText.ToString().Equals(TextFilter->GetRawFilterText().ToString(), ESearchCase::CaseSensitive))
 	{
 		TextFilter->SetRawFilterText( InSearchText );
 		SearchBoxPtr->SetError( TextFilter->GetFilterErrorText() );
diff --git a/Engine/Source/Runtime/Core/Private/Misc/TextFilterExpressionEvaluator.cpp b/Engine/Source/Runtime/Core/Private/Misc/TextFilterExpressionEvaluator.cpp
--- a/Engine/Source/Runtime/Core/Private/Misc/TextFilterExpressionEvaluator.cpp
+++ b/Engine/Source/Runtime/Core/Private/Misc/TextFilterExpressionEvaluator.cpp
@@ -407,19 +407,16 @@ bool FTextFilterExpressionEvaluator::SetFilterText(const FText& InFilterText)
 	using namespace TextFilterExpressionParser;
 	using TextFilterExpressionParser::FTextToken;
 
-	// Sanitize the new string
-	const FText NewText = FText::TrimPrecedingAndTrailing(InFilterText);
-
-	// Has anything changed?
-	if (NewText.EqualTo(FilterText))
+	// Has anything changed? (need to test case as the operators are case-sensitive)
+	if (FilterText.ToString().Equals(InFilterText.ToString(), ESearchCase::CaseSensitive))
 	{
 		return false;
 	}
 
 	FilterType = ETextFilterExpressionType::Invalid;
-	FilterText = NewText;
+	FilterText = InFilterText;
 
-	if (FilterText.IsEmpty())
+	if (FilterText.IsEmptyOrWhitespace())
 	{
 		// If the filter is empty, just clear out the compiled filter so we match everything
 		FilterType = ETextFilterExpressionType::Empty;
@@ -429,7 +426,7 @@ bool FTextFilterExpressionEvaluator::SetFilterText(const FText& InFilterText)
 	else
 	{
 		// Otherwise, re-parse the filter text into a compiled expression
-		auto LexResult = ExpressionParser::Lex(*NewText.ToString(), TokenDefinitions);
+		auto LexResult = ExpressionParser::Lex(*FilterText.ToString(), TokenDefinitions);
 		if (LexResult.IsValid())
 		{
 			auto& TmpLexTokens = LexResult.GetValue();
