Fixed automation test

The automatically injected operator was changed from OR to AND, and this test needed updating.

diff --git a/Engine/Source/Runtime/Core/Private/Misc/TextFilterTests.cpp b/Engine/Source/Runtime/Core/Private/Misc/TextFilterTests.cpp
--- a/Engine/Source/Runtime/Core/Private/Misc/TextFilterTests.cpp
+++ b/Engine/Source/Runtime/Core/Private/Misc/TextFilterTests.cpp
@@ -152,7 +152,7 @@ namespace TextFilterTests
 		bResult &= InTestPayload.TestFilterExpression(TEXT("+Woo"), false);
 		bResult &= InTestPayload.TestFilterExpression(TEXT("Wooble2"), false);
 		bResult &= InTestPayload.TestFilterExpression(TEXT("-Wooble2"), true);
-		bResult &= InTestPayload.TestFilterExpression(TEXT("Wooble2 Flibble"), true);
+		bResult &= InTestPayload.TestFilterExpression(TEXT("Wooble2 Flibble"), false);
 		bResult &= InTestPayload.TestFilterExpression(TEXT("Wooble2 OR Flibble"), true);
 		bResult &= InTestPayload.TestFilterExpression(TEXT("Wooble2 AND Flibble"), false);
 		bResult &= InTestPayload.TestFilterExpression(TEXT("Wooble && !Flibble"), false);
