Stats - Fixed an issue with serialization of FNames

diff --git a/Engine/Source/Runtime/Core/Public/Stats/Stats2.h b/Engine/Source/Runtime/Core/Public/Stats/Stats2.h
--- a/Engine/Source/Runtime/Core/Public/Stats/Stats2.h
+++ b/Engine/Source/Runtime/Core/Public/Stats/Stats2.h
@@ -255,14 +255,21 @@ struct EStatMetaFlags
 {
 	enum Type
 	{
-		Invalid							= 0x0,
-		DummyAlwaysOne					= 0x1,  // this bit is always one and is used for error checking
-		HasLongNameAndMetaInfo			= 0x2,  // if true, then this message contains all meta data as well and the name is long
-		IsCycle							= 0x4,	// if true, then this message contains and int64 cycle or IsPackedCCAndDuration
-		IsMemory						= 0x8,  // if true, then this message contains a memory stat
-		IsPackedCCAndDuration			= 0x10, // if true, then this is actually two uint32s, the cycle count and the call count, see FromPackedCallCountDuration_Duration
-		ShouldClearEveryFrame			= 0x20, // if true, then this stat is cleared every frame
-		SendingFName					= 0x40, // used only on disk on on the wire, indicates that we serialized the fname string.
+		Invalid							= 0x00,
+		/** this bit is always one and is used for error checking. */
+		DummyAlwaysOne					= 0x01,  
+		/** if true, then this message contains all meta data as well and the name is long. NOT USED. */
+		HasLongNameAndMetaInfo			= 0x02,
+		/** if true, then this message contains and int64 cycle or IsPackedCCAndDuration. */
+		IsCycle							= 0x04,	
+		/** if true, then this message contains a memory stat. */
+		IsMemory						= 0x08, 
+		/** if true, then this is actually two uint32s, the cycle count and the call count, see FromPackedCallCountDuration_Duration. */
+		IsPackedCCAndDuration			= 0x10,
+		/** if true, then this stat is cleared every frame. */
+		ShouldClearEveryFrame			= 0x20, 
+		/** used only on disk on on the wire, indicates that we serialized the FName string. */
+		SendingFName					= 0x40,
 
 		Num								= 0x80,
 		Mask							= 0xff,
@@ -432,6 +439,7 @@ public:
 
 	/**
 	 * The encoded FName with the correct, original Number
+	 * The original number usually is 0
 	 */
 	FORCEINLINE_STATS FName GetRawName() const
 	{
@@ -444,7 +452,8 @@ public:
 	}
 
 	/**
-	 * The encoded FName with the correct, original Number
+	 * The encoded FName with the encoded, new Number
+	 * The number contains all encoded metadata
 	 */
 	FORCEINLINE_STATS FName GetEncodedName() const
 	{
@@ -589,7 +598,15 @@ private:
 	/** For ST_Ptr. */
 	uint64	Ptr;
 	/** ST_int64 and IsPackedCCAndDuration. */
-	int32	CCAndDuration[2];
+	uint32	CCAndDuration[2];
+	/** For FName. */
+	const TCHAR* GetName() const
+	{
+		// Based on DebugFName
+		static TCHAR TempName[256];
+		FCString::Strcpy( TempName, *FName::SafeString( (int32)Cycles ) );
+		return TempName;
+	}
 };
 
 /**
@@ -699,7 +716,7 @@ struct FStatMessage
 		NameAndInfo.SetField<EStatOperation>(InStatOperation);
 		checkStats(NameAndInfo.GetField<EStatDataType>() == EStatDataType::ST_FName);
 		checkStats(NameAndInfo.GetFlag(EStatMetaFlags::IsCycle) == false);
-		GetValue_FName() = Value;
+		GetValue_FMinimalName() = NameToMinimalName(Value);
 	}
 
 	/**
@@ -1770,7 +1787,7 @@ struct FStat_##StatName\
 /** This is a fake stat, mostly used to implement memory message or other custom stats that don't easily fit into the system. */
 #define DECLARE_PTR_STAT_EXTERN(CounterName,StatId,GroupId, API) \
 	DECLARE_STAT(CounterName,StatId,GroupId,EStatDataType::ST_Ptr, false, false, FPlatformMemory::MCR_Invalid); \
-	extern API DEFINE_STAT(StatId)
+	extern API DEFINE_STAT(StatId);
 
 #define DECLARE_MEMORY_STAT_EXTERN(CounterName,StatId,GroupId, API) \
 	DECLARE_STAT(CounterName,StatId,GroupId,EStatDataType::ST_int64, false, false, FPlatformMemory::MCR_Physical); \
diff --git a/Engine/Source/Runtime/Core/Public/Stats/StatsFile.h b/Engine/Source/Runtime/Core/Public/Stats/StatsFile.h
--- a/Engine/Source/Runtime/Core/Public/Stats/StatsFile.h
+++ b/Engine/Source/Runtime/Core/Public/Stats/StatsFile.h
@@ -692,7 +692,7 @@ public:
 			case EStatDataType::ST_FName:
 			{
 				FStatNameAndInfo Payload( ReadFName( Ar, bHasFNameMap ) );
-				Result.GetValue_FName() = Payload.GetRawName();
+				Result.GetValue_FMinimalName() = NameToMinimalName( Payload.GetRawName() );
 				break;
 			}
 
