Stats - Stats packet should start with Frame set to 1

diff --git a/Engine/Source/Runtime/Core/Public/Stats/Stats2.h b/Engine/Source/Runtime/Core/Public/Stats/Stats2.h
--- a/Engine/Source/Runtime/Core/Public/Stats/Stats2.h
+++ b/Engine/Source/Runtime/Core/Public/Stats/Stats2.h
@@ -1070,7 +1070,7 @@ struct FStatPacket
 
 	/** constructor **/
 	FStatPacket()
-		: Frame(0)
+		: Frame(1)
 		, ThreadId(0)
 		, ThreadType(EThreadType::Invalid)
 		, bBrokenCallstacks(false)
