Stats - Added stats command
stat namedmarker #markername# - adds a custom marker to the stats stream (from the console)
STAT_ADD_CUSTOMMESSAGE_PTR/NAME( Stat, Value ) (from the code)

diff --git a/Engine/Source/Runtime/Core/Private/HAL/ThreadingBase.cpp b/Engine/Source/Runtime/Core/Private/HAL/ThreadingBase.cpp
--- a/Engine/Source/Runtime/Core/Private/HAL/ThreadingBase.cpp
+++ b/Engine/Source/Runtime/Core/Private/HAL/ThreadingBase.cpp
@@ -193,7 +193,7 @@ void FEvent::WaitForStats()
 	if( FThreadStats::IsCollectingData() && EventStartCycles == 0 )
 	{
 		const uint64 PacketEventIdAndCycles = ((uint64)EventId << 32) | 0;
-		FThreadStats::AddMessage( GET_STATFNAME( STAT_EventWaitWithId ), EStatOperation::SpecialMessageMarker, PacketEventIdAndCycles );
+		STAT_ADD_CUSTOMMESSAGE_PTR( STAT_EventWaitWithId, PacketEventIdAndCycles );
 		EventStartCycles = FPlatformTime::Cycles();
 	}
 #endif // STATS
@@ -208,7 +208,7 @@ void FEvent::TriggerForStats()
 		const uint32 EndCycles = FPlatformTime::Cycles();
 		const int32 DeltaCycles = int32( EndCycles - EventStartCycles );
 		const uint64 PacketEventIdAndCycles = ((uint64)EventId << 32) | DeltaCycles;
-		FThreadStats::AddMessage( GET_STATFNAME( STAT_EventTriggerWithId ), EStatOperation::SpecialMessageMarker, PacketEventIdAndCycles );
+		STAT_ADD_CUSTOMMESSAGE_PTR( STAT_EventTriggerWithId, PacketEventIdAndCycles );
 
 		AdvanceStats();
 	}
diff --git a/Engine/Source/Runtime/Core/Private/Stats/Stats2.cpp b/Engine/Source/Runtime/Core/Private/Stats/Stats2.cpp
--- a/Engine/Source/Runtime/Core/Private/Stats/Stats2.cpp
+++ b/Engine/Source/Runtime/Core/Private/Stats/Stats2.cpp
@@ -33,6 +33,7 @@ DECLARE_CYCLE_STAT(TEXT("Flush Raw Stats"),STAT_FlushRawStats,STATGROUP_StatSyst
 DECLARE_MEMORY_STAT( TEXT("Stats Descriptions"), STAT_StatDescMemory, STATGROUP_StatSystem );
 
 DEFINE_STAT(STAT_FrameTime);
+DEFINE_STAT(STAT_NamedMarker);
 
 /*-----------------------------------------------------------------------------
 	FStats2
@@ -1073,6 +1074,8 @@ void FThreadStats::CheckForCollectingStartupStats()
 		FStatsMallocProfilerProxy::Get()->SetState( true );
 		DirectStatsCommand( TEXT( "stat startfileraw" ), true );
 	}
+
+	STAT_ADD_CUSTOMMESSAGE_NAME( STAT_NamedMarker, TEXT("CheckForCollectingStartupStats") );
 }
 
 void FThreadStats::ExplicitFlush(bool DiscardCallstack)
diff --git a/Engine/Source/Runtime/Core/Private/Stats/StatsCommand.cpp b/Engine/Source/Runtime/Core/Private/Stats/StatsCommand.cpp
--- a/Engine/Source/Runtime/Core/Private/Stats/StatsCommand.cpp
+++ b/Engine/Source/Runtime/Core/Private/Stats/StatsCommand.cpp
@@ -1359,6 +1359,8 @@ static void PrintStatsHelpToOutputDevice( FOutputDevice& Ar )
 
 	Ar.Log( TEXT("add -memoryprofiler in the command line to enable the memory profiling"));
 	Ar.Log( TEXT("stat stopfile - stops tracking all memory operations and writes the results to the file"));
+
+	Ar.Log( TEXT("stat namedmarker #markername# - adds a custom marker to the stats stream"));
 }
 
 // @TODO yrx 2014-12-01 Move to StatsFile.cpp/.h
@@ -1555,12 +1557,29 @@ static void StatCmd(FString InCmd)
 	{
 		FStatsThreadState::GetLocalState().ToggleFindMemoryExtensiveStats();
 	}
-	else if ( FParse::Command( &Cmd, TEXT( "memoryprofiler" ) ) )
+	else if ( FParse::Command( &Cmd, TEXT( "namedmarker" ) ) )
 	{
-		if( FParse::Command( &Cmd, TEXT( "snapshot" ) ) )
+		FString MarkerName;
+		FParse::Token( Cmd, MarkerName, false );
+		
+		struct FLocal
+		{
+			static void OnGameThread( FString InMarkerName )
+			{
+				const FName NAME_Marker = FName( *InMarkerName );
+				STAT_ADD_CUSTOMMESSAGE_NAME( STAT_NamedMarker, NAME_Marker );
+				UE_LOG( LogStats, Log, TEXT( "Added from console STAT_NamedMarker: %s" ), *InMarkerName );
+			}
+		};
+
+		if (!MarkerName.IsEmpty())
 		{
-			// Put a snapshot marker in the raw stats.
-			// #YRX_STATS 2014-12-03 
+			// This will be executed on the game thread.
+			FSimpleDelegateGraphTask::CreateAndDispatchWhenReady
+			(
+				FSimpleDelegateGraphTask::FDelegate::CreateStatic( &FLocal::OnGameThread, MarkerName ),
+				TStatId(), nullptr, ENamedThreads::GameThread
+			);
 		}
 	}
 	// @see FStatHierParams
@@ -1705,6 +1724,9 @@ bool DirectStatsCommand(const TCHAR* Cmd, bool bBlockForCompletion /*= false*/,
 		else if( FParse::Command( &TempCmd, TEXT( "slow" ) ) )
 		{
 		}
+		else if (FParse::Command( &TempCmd, TEXT( "namedmarker" ) ))
+		{
+		}
 		else
 		{
 			bResult = false;
diff --git a/Engine/Source/Runtime/Core/Private/Stats/StatsData.cpp b/Engine/Source/Runtime/Core/Private/Stats/StatsData.cpp
--- a/Engine/Source/Runtime/Core/Private/Stats/StatsData.cpp
+++ b/Engine/Source/Runtime/Core/Private/Stats/StatsData.cpp
@@ -31,6 +31,7 @@ const FString FStatConstants::ThreadNameMarker = TEXT( "Thread_" );
 
 const FName FStatConstants::NAME_EventWaitWithId = FStatNameAndInfo( GET_STATFNAME( STAT_EventWaitWithId ), true ).GetRawName();
 const FName FStatConstants::NAME_EventTriggerWithId = FStatNameAndInfo( GET_STATFNAME( STAT_EventTriggerWithId ), true ).GetRawName();
+const FName FStatConstants::NAME_NamedMarker = FStatNameAndInfo( GET_STATFNAME( STAT_NamedMarker ), true ).GetRawName(); 
 
 /*-----------------------------------------------------------------------------
 	FRawStatStackNode
@@ -1268,60 +1269,69 @@ void FStatsThreadState::GetRawStackStats(int64 TargetFrame, FRawStatStackNode& R
 				// At this moment only these messages are supported:
 				//	EventWaitWithId
 				//	EventTriggerWithId
+				//	StatMarker
 				else if( Op == EStatOperation::SpecialMessageMarker )
 				{
+					//const FName EncName = Item.NameAndInfo.GetEncodedName();
 					const FName RawName = Item.NameAndInfo.GetRawName();
+
 					const uint64 PacketEventIdAndCycles = Item.GetValue_Ptr();
 					const uint32 EventId = uint32(PacketEventIdAndCycles >> 32);
 					const uint32 EventCycles = uint32(PacketEventIdAndCycles & MAX_uint32);
 
-					if (FStatConstants::NAME_EventWaitWithId == RawName)
+					if (RawName == FStatConstants::NAME_EventWaitWithId || RawName==FStatConstants::NAME_EventTriggerWithId)
 					{
-						TArray<FStatNameAndInfo> EventWaitStack;
-						for (const auto& It : Stack)
+						if (FStatConstants::NAME_EventWaitWithId == RawName)
 						{
-							EventWaitStack.Add( It->Meta.NameAndInfo );
-						}
+							TArray<FStatNameAndInfo> EventWaitStack;
+							for (const auto& It : Stack)
+							{
+								EventWaitStack.Add( It->Meta.NameAndInfo );
+							}
 
 #if	UE_BUILD_DEBUG
-						// Debug check, detect duplicates.
-						FEventData* EventPtr = EventsHistory.Find( EventId );
-						if (EventPtr && EventPtr->WaitStackStats.Num() > 0)
-						{
-							int32 k = 0; k++;
-						}
+							// Debug check, detect duplicates.
+							FEventData* EventPtr = EventsHistory.Find( EventId );
+							if (EventPtr && EventPtr->WaitStackStats.Num() > 0)
+							{
+								int32 k = 0; k++;
+							}
 #endif // UE_BUILD_DEBUG
 
-						FEventData& EventStats = EventsHistory.FindOrAdd( EventId );
-						EventStats.WaitStackStats = EventWaitStack;
-						EventStats.Frame = EventStats.HasValidStacks() ? TargetFrame : 0; // Only to maintain history.
-					}
+							FEventData& EventStats = EventsHistory.FindOrAdd( EventId );
+							EventStats.WaitStackStats = EventWaitStack;
+							EventStats.Frame = EventStats.HasValidStacks() ? TargetFrame : 0; // Only to maintain history.
+						}
 
-					if (FStatConstants::NAME_EventTriggerWithId == RawName)
-					{
-						TArray<FStatNameAndInfo> EventTriggerStack;
-						for (const auto& It : Stack)
+						if (FStatConstants::NAME_EventTriggerWithId == RawName)
 						{
-							EventTriggerStack.Add( It->Meta.NameAndInfo );
-						}
+							TArray<FStatNameAndInfo> EventTriggerStack;
+							for (const auto& It : Stack)
+							{
+								EventTriggerStack.Add( It->Meta.NameAndInfo );
+							}
 
 #if	UE_BUILD_DEBUG
-						// Debug check, detect duplicates.
-						FEventData* EventPtr = EventsHistory.Find( EventId );
-						if (EventPtr && EventPtr->TriggerStackStats.Num() > 0)
-						{
-							int32 k = 0; k++;
-						}
+							// Debug check, detect duplicates.
+							FEventData* EventPtr = EventsHistory.Find( EventId );
+							if (EventPtr && EventPtr->TriggerStackStats.Num() > 0)
+							{
+								int32 k = 0; k++;
+							}
 #endif // UE_BUILD_DEBUG
 
-						FEventData& EventStats = EventsHistory.FindOrAdd( EventId );
+							FEventData& EventStats = EventsHistory.FindOrAdd( EventId );
 
-						EventStats.TriggerStackStats = EventTriggerStack;
-						EventStats.Duration = EventCycles;
-						EventStats.DurationMS = FPlatformTime::ToMilliseconds( EventCycles );
-						EventStats.Frame = EventStats.HasValidStacks() ? TargetFrame : 0; // Only to maintain history.
+							EventStats.TriggerStackStats = EventTriggerStack;
+							EventStats.Duration = EventCycles;
+							EventStats.DurationMS = FPlatformTime::ToMilliseconds( EventCycles );
+							EventStats.Frame = EventStats.HasValidStacks() ? TargetFrame : 0; // Only to maintain history.
+						}
 					}
+					else if (RawName == FStatConstants::NAME_NamedMarker)
+					{
 
+					}
 				}
 				else if( Op == EStatOperation::Memory )
 				{
diff --git a/Engine/Source/Runtime/Core/Public/Stats/Stats.h b/Engine/Source/Runtime/Core/Public/Stats/Stats.h
--- a/Engine/Source/Runtime/Core/Public/Stats/Stats.h
+++ b/Engine/Source/Runtime/Core/Public/Stats/Stats.h
@@ -386,6 +386,7 @@ FORCEINLINE void StatsMasterEnableSubtract(int32 Value = 1)
 #define DECLARE_DWORD_COUNTER_STAT(CounterName,StatId,GroupId)
 #define DECLARE_FLOAT_ACCUMULATOR_STAT(CounterName,StatId,GroupId)
 #define DECLARE_DWORD_ACCUMULATOR_STAT(CounterName,StatId,GroupId)
+#define DECLARE_FNAME_STAT(CounterName,StatId,GroupId, API)
 #define DECLARE_PTR_STAT(CounterName,StatId,GroupId)
 #define DECLARE_MEMORY_STAT(CounterName,StatId,GroupId)
 #define DECLARE_MEMORY_STAT_POOL(CounterName,StatId,GroupId,Pool)
@@ -394,6 +395,7 @@ FORCEINLINE void StatsMasterEnableSubtract(int32 Value = 1)
 #define DECLARE_DWORD_COUNTER_STAT_EXTERN(CounterName,StatId,GroupId, API)
 #define DECLARE_FLOAT_ACCUMULATOR_STAT_EXTERN(CounterName,StatId,GroupId, API)
 #define DECLARE_DWORD_ACCUMULATOR_STAT_EXTERN(CounterName,StatId,GroupId, API)
+#define DECLARE_FNAME_STAT_EXTERN(CounterName,StatId,GroupId, API)
 #define DECLARE_PTR_STAT_EXTERN(CounterName,StatId,GroupId, API)
 #define DECLARE_MEMORY_STAT_EXTERN(CounterName,StatId,GroupId, API)
 #define DECLARE_MEMORY_STAT_POOL_EXTERN(CounterName,StatId,GroupId,Pool, API)
@@ -416,6 +418,8 @@ FORCEINLINE void StatsMasterEnableSubtract(int32 Value = 1)
 #define SET_MEMORY_STAT(StatId,Value)
 #define SET_DWORD_STAT(StatId,Value)
 #define SET_FLOAT_STAT(StatId,Value)
+#define STAT_ADD_CUSTOMMESSAGE_NAME(StatId,Value)
+#define STAT_ADD_CUSTOMMESSAGE_PTR(StatId,Value)
 
 #define SET_CYCLE_COUNTER_FName(Stat,Cycles)
 #define INC_DWORD_STAT_FName(Stat)
diff --git a/Engine/Source/Runtime/Core/Public/Stats/Stats2.h b/Engine/Source/Runtime/Core/Public/Stats/Stats2.h
--- a/Engine/Source/Runtime/Core/Public/Stats/Stats2.h
+++ b/Engine/Source/Runtime/Core/Public/Stats/Stats2.h
@@ -600,12 +600,9 @@ private:
 	/** ST_int64 and IsPackedCCAndDuration. */
 	uint32	CCAndDuration[2];
 	/** For FName. */
-	const TCHAR* GetName() const
+	CORE_API const FString GetName() const
 	{
-		// Based on DebugFName
-		static TCHAR TempName[256];
-		FCString::Strcpy( TempName, *FName::SafeString( (int32)Cycles ) );
-		return TempName;
+		return FName::SafeString( (int32)Cycles );
 	}
 };
 
@@ -1784,6 +1781,11 @@ struct FStat_##StatName\
 	DECLARE_STAT(CounterName,StatId,GroupId,EStatDataType::ST_int64, false, false, FPlatformMemory::MCR_Invalid); \
 	extern API DEFINE_STAT(StatId);
 
+/** FName stat that allows sending a string based data. */
+#define DECLARE_FNAME_STAT_EXTERN(CounterName,StatId,GroupId, API) \
+	DECLARE_STAT(CounterName,StatId,GroupId,EStatDataType::ST_FName, false, false, FPlatformMemory::MCR_Invalid); \
+	extern API DEFINE_STAT(StatId);
+
 /** This is a fake stat, mostly used to implement memory message or other custom stats that don't easily fit into the system. */
 #define DECLARE_PTR_STAT_EXTERN(CounterName,StatId,GroupId, API) \
 	DECLARE_STAT(CounterName,StatId,GroupId,EStatDataType::ST_Ptr, false, false, FPlatformMemory::MCR_Invalid); \
@@ -1890,7 +1892,14 @@ struct FStat_##StatName\
 {\
 	FThreadStats::AddMessage(GET_STATFNAME(Stat), EStatOperation::Set, double(Value));\
 }
-
+#define STAT_ADD_CUSTOMMESSAGE_NAME(Stat,Value) \
+{\
+	FThreadStats::AddMessage(GET_STATFNAME(Stat), EStatOperation::SpecialMessageMarker, FName(Value));\
+}
+#define STAT_ADD_CUSTOMMESSAGE_PTR(Stat,Value) \
+{\
+	FThreadStats::AddMessage(GET_STATFNAME(Stat), EStatOperation::SpecialMessageMarker, uint64(Value));\
+}
 
 #define SET_CYCLE_COUNTER_FName(Stat,Cycles) \
 {\
@@ -2016,6 +2025,6 @@ DECLARE_STATS_GROUP(TEXT("UObjects"),STATGROUP_UObjects, STATCAT_Advanced);
 DECLARE_STATS_GROUP(TEXT("User"),STATGROUP_User, STATCAT_Advanced);
 
 DECLARE_CYCLE_STAT_EXTERN(TEXT("FrameTime"),STAT_FrameTime,STATGROUP_Engine, CORE_API);
-
+DECLARE_FNAME_STAT_EXTERN(TEXT("NamedMarker"),STAT_NamedMarker,STATGROUP_StatSystem, CORE_API);
 
 #endif
diff --git a/Engine/Source/Runtime/Core/Public/Stats/StatsData.h b/Engine/Source/Runtime/Core/Public/Stats/StatsData.h
--- a/Engine/Source/Runtime/Core/Public/Stats/StatsData.h
+++ b/Engine/Source/Runtime/Core/Public/Stats/StatsData.h
@@ -41,6 +41,9 @@ struct CORE_API FStatConstants
 
 	/** A raw name for the event trigger with id. */
 	static const FName NAME_EventTriggerWithId;
+
+	/** A raw name for the stat marker. */
+	static const FName NAME_NamedMarker;
 };
 
 namespace LexicalConversion
