Merging CL 2573182 - Removed some redundant calls from FPaths::MakeStandardFilename

diff --git a/Engine/Source/Runtime/Core/Private/Misc/Paths.cpp b/Engine/Source/Runtime/Core/Private/Misc/Paths.cpp
--- a/Engine/Source/Runtime/Core/Private/Misc/Paths.cpp
+++ b/Engine/Source/Runtime/Core/Private/Misc/Paths.cpp
@@ -694,12 +694,6 @@ void FPaths::MakeStandardFilename(FString& InPath)
 	
 	FString Standardized = FPaths::ConvertRelativePathToFull(InPath);
 
-	// fix separators and remove any trailing slashes
-	FPaths::NormalizeDirectoryName(Standardized);
-	
-	// remove inner ..'s
-	FPaths::CollapseRelativeDirectories(Standardized);
-
 	// remove duplicate slashes
 	FPaths::RemoveDuplicateSlashes(Standardized);
 
