Merging CL 2572016 - fix the splash screen on OS X 10.11

diff --git a/Engine/Source/Runtime/Core/Private/Mac/MacPlatformSplash.cpp b/Engine/Source/Runtime/Core/Private/Mac/MacPlatformSplash.cpp
--- a/Engine/Source/Runtime/Core/Private/Mac/MacPlatformSplash.cpp
+++ b/Engine/Source/Runtime/Core/Private/Mac/MacPlatformSplash.cpp
@@ -4,6 +4,7 @@
 #include "Misc/App.h"
 #include "EngineVersion.h"
 #include "EngineBuildSettings.h"
+#include "CocoaThread.h"
 
 /**
  * Simple window class that overrides a couple of functions, so the splash window can be moved to front even if it's borderless
@@ -316,23 +317,25 @@ void FMacPlatformSplash::Show()
 			GSplashScreenTextRects[ SplashTextType::CopyrightInfo ].origin.y = OriginY;
 		}
 
-		// Create bordeless window with size from NSImage
-		GSplashWindow = [[FSplashWindow alloc] initWithContentRect: ContentRect styleMask: 0 backing: NSBackingStoreBuffered defer: NO];
-		[GSplashWindow setContentView: [[UE4SplashView alloc] initWithFrame: ContentRect]];
+		MainThreadCall(^{
+			// Create bordeless window with size from NSImage
+			GSplashWindow = [[FSplashWindow alloc] initWithContentRect: ContentRect styleMask: 0 backing: NSBackingStoreBuffered defer: NO];
+			[GSplashWindow setContentView: [[UE4SplashView alloc] initWithFrame: ContentRect]];
 
-		if( GSplashWindow )
-		{
-			// Show window
-			[GSplashWindow setHasShadow:YES];
-			[GSplashWindow center];
-			[GSplashWindow orderFront: nil];
-			[NSApp activateIgnoringOtherApps:YES];
-		}
-		else
-		{
-			[GSplashScreenImage release];
-			GSplashScreenImage = NULL;
-		}
+			if( GSplashWindow )
+			{
+				// Show window
+				[GSplashWindow setHasShadow:YES];
+				[GSplashWindow center];
+				[GSplashWindow orderFront: nil];
+				[NSApp activateIgnoringOtherApps:YES];
+			}
+			else
+			{
+				[GSplashScreenImage release];
+				GSplashScreenImage = NULL;
+			}
+		}, NSDefaultRunLoopMode, true);
 
 		FPlatformMisc::PumpMessages(true);
 	}
