Merging Mac high precision mouse mode fixes

diff --git a/Engine/Source/Runtime/Core/Private/Mac/MacApplication.cpp b/Engine/Source/Runtime/Core/Private/Mac/MacApplication.cpp
--- a/Engine/Source/Runtime/Core/Private/Mac/MacApplication.cpp
+++ b/Engine/Source/Runtime/Core/Private/Mac/MacApplication.cpp
@@ -1071,6 +1071,8 @@ void FMacApplication::OnApplicationWillResignActive()
 		}
 	}
 
+	SetHighPrecisionMouseMode(false, nullptr);
+
 	((FMacCursor*)Cursor.Get())->UpdateVisibility();
 
 	// If editor thread doesn't have the focus, don't suck up too much CPU time.
@@ -1416,4 +1418,4 @@ void FDisplayMetrics::GetDisplayMetrics(FDisplayMetrics& OutDisplayMetrics)
 	OutDisplayMetrics.PrimaryDisplayWorkAreaRect.Top = ScreenFrame.size.height - (VisibleFrame.origin.y + VisibleFrame.size.height);
 	OutDisplayMetrics.PrimaryDisplayWorkAreaRect.Right = VisibleFrame.origin.x + VisibleFrame.size.width;
 	OutDisplayMetrics.PrimaryDisplayWorkAreaRect.Bottom = OutDisplayMetrics.PrimaryDisplayWorkAreaRect.Top + VisibleFrame.size.height;
-}
\ No newline at end of file
+}
diff --git a/Engine/Source/Runtime/Core/Private/Mac/MacCursor.cpp b/Engine/Source/Runtime/Core/Private/Mac/MacCursor.cpp
--- a/Engine/Source/Runtime/Core/Private/Mac/MacCursor.cpp
+++ b/Engine/Source/Runtime/Core/Private/Mac/MacCursor.cpp
@@ -241,6 +241,10 @@ void FMacCursor::SetPosition( const int32 X, const int32 Y )
 void FMacCursor::SetType( const EMouseCursor::Type InNewCursor )
 {
 	check( InNewCursor < EMouseCursor::TotalCursorCount );
+	if (MacApplication && CurrentType == EMouseCursor::None && InNewCursor != EMouseCursor::None)
+	{
+		MacApplication->SetHighPrecisionMouseMode(false, nullptr);
+	}
 	CurrentType = InNewCursor;
 	CurrentCursor = CursorHandles[InNewCursor];
 	if( CurrentCursor )
