Platform: Added SSIZE_T type definition (equivalent for ssize_t)

diff --git a/Engine/Source/Runtime/Core/Public/GenericPlatform/GenericPlatform.h b/Engine/Source/Runtime/Core/Public/GenericPlatform/GenericPlatform.h
--- a/Engine/Source/Runtime/Core/Public/GenericPlatform/GenericPlatform.h
+++ b/Engine/Source/Runtime/Core/Public/GenericPlatform/GenericPlatform.h
@@ -53,6 +53,7 @@ struct FGenericPlatformTypes
 	typedef SelectIntPointerType<uint32, uint64, sizeof(void*)>::TIntPointer UPTRINT;	// unsigned int the same size as a pointer
 	typedef SelectIntPointerType<int32, int64, sizeof(void*)>::TIntPointer PTRINT;		// signed int the same size as a pointer
 	typedef UPTRINT SIZE_T;																// unsigned int the same size as a pointer
+	typedef PTRINT SSIZE_T;																// unsigned int the same size as a pointer
 
 	typedef int32					TYPE_OF_NULL;
 	typedef decltype(nullptr)		TYPE_OF_NULLPTR;
diff --git a/Engine/Source/Runtime/Core/Public/HAL/Platform.h b/Engine/Source/Runtime/Core/Public/HAL/Platform.h
--- a/Engine/Source/Runtime/Core/Public/HAL/Platform.h
+++ b/Engine/Source/Runtime/Core/Public/HAL/Platform.h
@@ -621,6 +621,7 @@ typedef FPlatformTypes::CHAR32		UTF32CHAR;	///< A 32-bit character containing a
 typedef FPlatformTypes::UPTRINT UPTRINT;		///< An unsigned integer the same size as a pointer
 typedef FPlatformTypes::PTRINT PTRINT;			///< A signed integer the same size as a pointer
 typedef FPlatformTypes::SIZE_T SIZE_T;			///< An unsigned integer the same size as a pointer, the same as UPTRINT
+typedef FPlatformTypes::SSIZE_T SSIZE_T;		///< An integer the same size as a pointer, the same as PTRINT
 
 typedef FPlatformTypes::TYPE_OF_NULL	TYPE_OF_NULL;		///< The type of the NULL constant.
 typedef FPlatformTypes::TYPE_OF_NULLPTR	TYPE_OF_NULLPTR;	///< The type of the C++ nullptr keyword.
diff --git a/Engine/Source/Runtime/Core/Public/HTML5/HTML5Platform.h b/Engine/Source/Runtime/Core/Public/HTML5/HTML5Platform.h
--- a/Engine/Source/Runtime/Core/Public/HTML5/HTML5Platform.h
+++ b/Engine/Source/Runtime/Core/Public/HTML5/HTML5Platform.h
@@ -15,6 +15,7 @@
 struct FHTML5Types : public FGenericPlatformTypes
 {
 	typedef unsigned long SIZE_T;
+	typedef long SSIZE_T;
 #if !PLATFORM_HTML5_WIN32
 	// Emscripten uses the musl libc implementation, where NULL is defined as '0L', which is of type long.
 	typedef long					TYPE_OF_NULL;
diff --git a/Engine/Source/Runtime/Core/Public/IOS/IOSPlatform.h b/Engine/Source/Runtime/Core/Public/IOS/IOSPlatform.h
--- a/Engine/Source/Runtime/Core/Public/IOS/IOSPlatform.h
+++ b/Engine/Source/Runtime/Core/Public/IOS/IOSPlatform.h
@@ -12,6 +12,7 @@
 struct FIOSPlatformTypes : public FGenericPlatformTypes
 {
 	typedef size_t				SIZE_T;
+	typedef ssize_t				SSIZE_T;
 	typedef decltype(NULL)		TYPE_OF_NULL;
 	typedef char16_t			CHAR16;
 };
diff --git a/Engine/Source/Runtime/Core/Public/Mac/MacPlatform.h b/Engine/Source/Runtime/Core/Public/Mac/MacPlatform.h
--- a/Engine/Source/Runtime/Core/Public/Mac/MacPlatform.h
+++ b/Engine/Source/Runtime/Core/Public/Mac/MacPlatform.h
@@ -13,6 +13,7 @@ struct FMacPlatformTypes : public FGenericPlatformTypes
 {
 	typedef unsigned int		DWORD;
 	typedef size_t				SIZE_T;
+	typedef ssize_t				SSIZE_T;
 	typedef decltype(NULL)		TYPE_OF_NULL;
 	typedef char16_t			CHAR16;
 };
diff --git a/Engine/Source/Runtime/Core/Public/WinRT/WinRTARMPlatform.h b/Engine/Source/Runtime/Core/Public/WinRT/WinRTARMPlatform.h
--- a/Engine/Source/Runtime/Core/Public/WinRT/WinRTARMPlatform.h
+++ b/Engine/Source/Runtime/Core/Public/WinRT/WinRTARMPlatform.h
@@ -18,6 +18,7 @@ struct FWinRTTypes : public FGenericPlatformTypes
 	// defined in windefs.h, even though this is equivalent, the compiler doesn't think so
 	typedef unsigned long	DWORD;
 	typedef unsigned long	SIZE_T;
+	typedef long			SSIZE_T;
 };
 
 typedef FWinRTTypes FPlatformTypes;
diff --git a/Engine/Source/Runtime/Core/Public/WinRT/WinRTPlatform.h b/Engine/Source/Runtime/Core/Public/WinRT/WinRTPlatform.h
--- a/Engine/Source/Runtime/Core/Public/WinRT/WinRTPlatform.h
+++ b/Engine/Source/Runtime/Core/Public/WinRT/WinRTPlatform.h
@@ -17,6 +17,7 @@ struct FWinRTTypes : public FGenericPlatformTypes
 	// defined in windefs.h, even though this is equivalent, the compiler doesn't think so
 	typedef unsigned long       DWORD;
 	typedef unsigned __int64	SIZE_T;
+	typedef __int64				SSIZE_T;
 
 	typedef decltype(__nullptr) TYPE_OF_NULLPTR;
 };
diff --git a/Engine/Source/Runtime/Core/Public/Windows/WIndowsPlatform.h b/Engine/Source/Runtime/Core/Public/Windows/WIndowsPlatform.h
--- a/Engine/Source/Runtime/Core/Public/Windows/WIndowsPlatform.h
+++ b/Engine/Source/Runtime/Core/Public/Windows/WIndowsPlatform.h
@@ -10,8 +10,10 @@ struct FWindowsPlatformTypes : public FGenericPlatformTypes
 {
 #ifdef _WIN64
 	typedef unsigned __int64	SIZE_T;
+	typedef __int64				SSIZE_T;
 #else
 	typedef unsigned long		SIZE_T;
+	typedef long				SSIZE_T;
 #endif
 };
 
