Platform: Fixed MacOS ssize_t

diff --git a/Engine/Source/Runtime/Core/Public/Mac/MacPlatform.h b/Engine/Source/Runtime/Core/Public/Mac/MacPlatform.h
--- a/Engine/Source/Runtime/Core/Public/Mac/MacPlatform.h
+++ b/Engine/Source/Runtime/Core/Public/Mac/MacPlatform.h
@@ -13,7 +13,6 @@ struct FMacPlatformTypes : public FGenericPlatformTypes
 {
 	typedef unsigned int		DWORD;
 	typedef size_t				SIZE_T;
-	typedef ssize_t				SSIZE_T;
 	typedef decltype(NULL)		TYPE_OF_NULL;
 	typedef char16_t			CHAR16;
 };
