Platform: Fixed iOS ssize_t

diff --git a/Engine/Source/Runtime/Core/Public/IOS/IOSPlatform.h b/Engine/Source/Runtime/Core/Public/IOS/IOSPlatform.h
--- a/Engine/Source/Runtime/Core/Public/IOS/IOSPlatform.h
+++ b/Engine/Source/Runtime/Core/Public/IOS/IOSPlatform.h
@@ -12,7 +12,6 @@
 struct FIOSPlatformTypes : public FGenericPlatformTypes
 {
 	typedef size_t				SIZE_T;
-	typedef ssize_t				SSIZE_T;
 	typedef decltype(NULL)		TYPE_OF_NULL;
 	typedef char16_t			CHAR16;
 };
