New stats for events enabled for pthread platforms

diff --git a/Engine/Source/Runtime/Core/Private/GenericPlatform/GenericPlatformProcess.cpp b/Engine/Source/Runtime/Core/Private/GenericPlatform/GenericPlatformProcess.cpp
--- a/Engine/Source/Runtime/Core/Private/GenericPlatform/GenericPlatformProcess.cpp
+++ b/Engine/Source/Runtime/Core/Private/GenericPlatform/GenericPlatformProcess.cpp
@@ -301,7 +301,7 @@ void FGenericPlatformProcess::ConditionalSleep(const TFunctionRef<bool()>& Condi
 
 bool FPThreadEvent::Wait(uint32 WaitTime, const bool bIgnoreThreadIdleStats /*= false*/)
 {
-	//WaitForStats();
+	WaitForStats();
 
 	SCOPE_CYCLE_COUNTER(STAT_EventWait);
 	FThreadIdleStats::FScopeIdle Scope(bIgnoreThreadIdleStats);
diff --git a/Engine/Source/Runtime/Core/Public/HAL/PThreadEvent.h b/Engine/Source/Runtime/Core/Public/HAL/PThreadEvent.h
--- a/Engine/Source/Runtime/Core/Public/HAL/PThreadEvent.h
+++ b/Engine/Source/Runtime/Core/Public/HAL/PThreadEvent.h
@@ -124,7 +124,7 @@ public:
 
 	virtual void Trigger() override
 	{
-		//TriggerForStats();
+		TriggerForStats();
 
 		check(bInitialized);
 
@@ -151,7 +151,7 @@ public:
 
 	virtual void Reset() override
 	{
-		//ResetForStats();
+		ResetForStats();
 
 		check(bInitialized);
 		LockEventMutex();
