Added better messages for ensuring retrieving vtable pointer.

diff --git a/Engine/Source/Editor/UnrealEd/Private/HierarchicalLOD.cpp b/Engine/Source/Editor/UnrealEd/Private/HierarchicalLOD.cpp
--- a/Engine/Source/Editor/UnrealEd/Private/HierarchicalLOD.cpp
+++ b/Engine/Source/Editor/UnrealEd/Private/HierarchicalLOD.cpp
@@ -115,7 +115,7 @@ FHierarchicalLODBuilder::FHierarchicalLODBuilder(class UWorld* InWorld)
 FHierarchicalLODBuilder::FHierarchicalLODBuilder()
 	: World(nullptr)
 {
-	EnsureRetrievingVTablePtr();
+	EnsureRetrievingVTablePtrDuringCtor(TEXT("FHierarchicalLODBuilder()"));
 }
 #endif // WITH_HOT_RELOAD_CTORS
 
diff --git a/Engine/Source/Runtime/Core/Private/Misc/ConfigCacheIni.cpp b/Engine/Source/Runtime/Core/Private/Misc/ConfigCacheIni.cpp
--- a/Engine/Source/Runtime/Core/Private/Misc/ConfigCacheIni.cpp
+++ b/Engine/Source/Runtime/Core/Private/Misc/ConfigCacheIni.cpp
@@ -1350,7 +1350,7 @@ FConfigCacheIni::FConfigCacheIni(EConfigCacheType InType)
 #if WITH_HOT_RELOAD_CTORS
 FConfigCacheIni::FConfigCacheIni()
 {
-	EnsureRetrievingVTablePtr();
+	EnsureRetrievingVTablePtrDuringCtor(TEXT("FConfigCacheIni()"));
 }
 #endif // WITH_HOT_RELOAD_CTORS
 
diff --git a/Engine/Source/Runtime/Core/Private/Misc/CoreMisc.cpp b/Engine/Source/Runtime/Core/Private/Misc/CoreMisc.cpp
--- a/Engine/Source/Runtime/Core/Private/Misc/CoreMisc.cpp
+++ b/Engine/Source/Runtime/Core/Private/Misc/CoreMisc.cpp
@@ -970,8 +970,8 @@ FBoolConfigValueHelper::FBoolConfigValueHelper(const TCHAR* Section, const TCHAR
 #if WITH_HOT_RELOAD_CTORS
 bool GIsRetrievingVTablePtr = false;
 
-void EnsureRetrievingVTablePtr()
+void EnsureRetrievingVTablePtrDuringCtor(const TCHAR* CtorSignature)
 {
-	UE_CLOG(!GIsRetrievingVTablePtr, LogCore, Fatal, TEXT("This should be used only during vtable ptr retrieval process."));
+	UE_CLOG(!GIsRetrievingVTablePtr, LogCore, Fatal, TEXT("The %s constructor is for internal usage only for hot-reload purposes. Please do NOT use it."), CtorSignature);
 }
 #endif // WITH_HOT_RELOAD_CTORS
diff --git a/Engine/Source/Runtime/Core/Public/CoreGlobals.h b/Engine/Source/Runtime/Core/Public/CoreGlobals.h
--- a/Engine/Source/Runtime/Core/Public/CoreGlobals.h
+++ b/Engine/Source/Runtime/Core/Public/CoreGlobals.h
@@ -309,8 +309,11 @@ extern CORE_API TArray<FScriptTraceStackNode> GScriptStack;
 
 #if WITH_HOT_RELOAD_CTORS
 /**
- * Ensures that current thread is during retrieval of vtable ptr
- * of some UClass.
+ * Ensures that current thread is during retrieval of vtable ptr of some
+ * UClass.
+ *
+ * @param CtorSignature The signature of the ctor currently running to
+ *		construct proper error message.
  */
-CORE_API void EnsureRetrievingVTablePtr();
+CORE_API void EnsureRetrievingVTablePtrDuringCtor(const TCHAR* CtorSignature);
 #endif // WITH_HOT_RELOAD_CTORS
\ No newline at end of file
diff --git a/Engine/Source/Runtime/Core/Public/Templates/SharedPointer.h b/Engine/Source/Runtime/Core/Public/Templates/SharedPointer.h
--- a/Engine/Source/Runtime/Core/Public/Templates/SharedPointer.h
+++ b/Engine/Source/Runtime/Core/Public/Templates/SharedPointer.h
@@ -163,7 +163,7 @@ public:
 		: Object(new ObjectType())
 		, SharedReferenceCount(Object)
 	{
-		EnsureRetrievingVTablePtr();
+		EnsureRetrievingVTablePtrDuringCtor(TEXT("TSharedRef()"));
 		Init(Object);
 	}
 #endif // WITH_HOT_RELOAD_CTORS
diff --git a/Engine/Source/Runtime/CoreUObject/Private/UObject/Obj.cpp b/Engine/Source/Runtime/CoreUObject/Private/UObject/Obj.cpp
--- a/Engine/Source/Runtime/CoreUObject/Private/UObject/Obj.cpp
+++ b/Engine/Source/Runtime/CoreUObject/Private/UObject/Obj.cpp
@@ -55,7 +55,7 @@ UObject::UObject( EStaticConstructor, EObjectFlags InFlags )
 #if WITH_HOT_RELOAD_CTORS
 UObject::UObject(FVTableHelper& Helper)
 {
-	EnsureRetrievingVTablePtr();
+	EnsureRetrievingVTablePtrDuringCtor(TEXT("UObject(FVTableHelper& Helper)"));
 
 	static struct FUseVTableConstructorsCache
 	{
diff --git a/Engine/Source/Runtime/CoreUObject/Public/UObject/ObjectBase.h b/Engine/Source/Runtime/CoreUObject/Public/UObject/ObjectBase.h
--- a/Engine/Source/Runtime/CoreUObject/Public/UObject/ObjectBase.h
+++ b/Engine/Source/Runtime/CoreUObject/Public/UObject/ObjectBase.h
@@ -138,7 +138,7 @@ public:
 	/** DO NOT USE. This constructor is for internal usage only for hot-reload purposes. */
 	COREUOBJECT_API FVTableHelper()
 	{
-		EnsureRetrievingVTablePtr();
+		EnsureRetrievingVTablePtrDuringCtor(TEXT("FVTableHelper()"));
 	}
 };
 #endif // WITH_HOT_RELOAD_CTORS
diff --git a/Engine/Source/Runtime/Engine/Public/GenericOctree.inl b/Engine/Source/Runtime/Engine/Public/GenericOctree.inl
--- a/Engine/Source/Runtime/Engine/Public/GenericOctree.inl
+++ b/Engine/Source/Runtime/Engine/Public/GenericOctree.inl
@@ -358,7 +358,7 @@ template<typename ElementType, typename OctreeSemantics>
 TOctree<ElementType, OctreeSemantics>::TOctree()
 	: RootNode(nullptr)
 {
-	EnsureRetrievingVTablePtr();
+	EnsureRetrievingVTablePtrDuringCtor(TEXT("TOctree()"));
 }
 #endif // WITH_HOT_RELOAD_CTORS
 
diff --git a/Engine/Source/Runtime/Engine/Public/GenericQuadTree.h b/Engine/Source/Runtime/Engine/Public/GenericQuadTree.h
--- a/Engine/Source/Runtime/Engine/Public/GenericQuadTree.h
+++ b/Engine/Source/Runtime/Engine/Public/GenericQuadTree.h
@@ -151,7 +151,7 @@ TQuadTree<ElementType, NodeCapacity>::TQuadTree(const FBox2D& Box)
 template <typename ElementType, int32 NodeCapacity>
 TQuadTree<ElementType, NodeCapacity>::TQuadTree()
 {
-	EnsureRetrievingVTablePtr();
+	EnsureRetrievingVTablePtrDuringCtor(TEXT("TQuadTree()"));
 }
 #endif // WITH_HOT_RELOAD_CTORS
 
diff --git a/Engine/Source/Runtime/Engine/Public/ShowFlags.h b/Engine/Source/Runtime/Engine/Public/ShowFlags.h
--- a/Engine/Source/Runtime/Engine/Public/ShowFlags.h
+++ b/Engine/Source/Runtime/Engine/Public/ShowFlags.h
@@ -143,7 +143,7 @@ struct FEngineShowFlags
 	/** DO NOT USE. This constructor is for internal usage only for hot-reload purposes. */
 	FEngineShowFlags()
 	{
-		EnsureRetrievingVTablePtr();
+		EnsureRetrievingVTablePtrDuringCtor(TEXT("FEngineShowFlags()"));
 		Init(ESFIM_Game);
 	}
 #endif // WITH_HOT_RELOAD_CTORS
