Merging increased stack sizes on Mac

diff --git a/Engine/Source/Runtime/Core/Private/Mac/CocoaThread.cpp b/Engine/Source/Runtime/Core/Private/Mac/CocoaThread.cpp
--- a/Engine/Source/Runtime/Core/Private/Mac/CocoaThread.cpp
+++ b/Engine/Source/Runtime/Core/Private/Mac/CocoaThread.cpp
@@ -454,7 +454,7 @@ void RunGameThread(id Target, SEL Selector)
 	
 	// Create a separate game thread and set it to the stack size to be the same as the main thread default of 8MB ( http://developer.apple.com/library/mac/#qa/qa1419/_index.html )
 	FCocoaGameThread* GameThread = [[FCocoaGameThread alloc] initWithTarget:Target selector:Selector object:nil];
-	[GameThread setStackSize:8*1024*1024];
+	[GameThread setStackSize:16*1024*1024];
 	[GameThread start];
 #else
 	[Target performSelector:Selector withObject:nil];
diff --git a/Engine/Source/Runtime/Core/Public/Mac/MacPlatformRunnableThread.h b/Engine/Source/Runtime/Core/Public/Mac/MacPlatformRunnableThread.h
--- a/Engine/Source/Runtime/Core/Public/Mac/MacPlatformRunnableThread.h
+++ b/Engine/Source/Runtime/Core/Public/Mac/MacPlatformRunnableThread.h
@@ -30,7 +30,7 @@ private:
 	virtual int GetDefaultStackSize() override
 	{
 		// default is 512 KB, we need more
-		return 2 * 1024 * 1024;
+		return 4 * 1024 * 1024;
 	}
 
 	/**
