Fixed %*s formatting on platforms that don't use system vswprintf

diff --git a/Engine/Source/Runtime/Core/Private/GenericPlatform/StandardPlatformString.cpp b/Engine/Source/Runtime/Core/Private/GenericPlatform/StandardPlatformString.cpp
--- a/Engine/Source/Runtime/Core/Private/GenericPlatform/StandardPlatformString.cpp
+++ b/Engine/Source/Runtime/Core/Private/GenericPlatform/StandardPlatformString.cpp
@@ -78,9 +78,14 @@ static int32 GetFormattingInfo(const WIDECHAR* Format, FFormatInfo& OutInfo)
 
 	OutInfo.Type = *Format++;
 
-	const int32 FormatLength = Format - FormatStart;
+	int32 FormatLength = Format - FormatStart;
 
 	FMemory::Memcpy(OutInfo.Format, FormatStart, FormatLength * sizeof(WIDECHAR));
+	if (OutInfo.HasDynamicWidth && FChar::ToLower(OutInfo.Type) == LITERAL(WIDECHAR, 's'))
+	{
+		OutInfo.Format[FormatLength - 1] = 'l';
+		OutInfo.Format[FormatLength++] = 's';
+	}
 	OutInfo.Format[FormatLength] = 0;
 
 	return FormatLength;
@@ -129,7 +134,7 @@ static const WIDECHAR* GetFormattedArgument(const FFormatInfo& Info, VA_LIST_REF
 			if (String)
 			{
 				InOutLength = swprintf(Formatted, InOutLength, Info.Format, Width, String);
-				return String;
+				return Formatted;
 			}
 			else
 			{
