Fixed setting of the current working dir on Mac for paths that contain non-ansi characters

diff --git a/Engine/Source/Runtime/Core/Private/Mac/MacPlatformProcess.cpp b/Engine/Source/Runtime/Core/Private/Mac/MacPlatformProcess.cpp
--- a/Engine/Source/Runtime/Core/Private/Mac/MacPlatformProcess.cpp
+++ b/Engine/Source/Runtime/Core/Private/Mac/MacPlatformProcess.cpp
@@ -681,8 +681,7 @@ const TCHAR* FMacPlatformProcess::BaseDir()
 				BasePath = [BasePath stringByDeletingLastPathComponent];
 			}
 		}
-		const char *BaseDir = [FileManager fileSystemRepresentationWithPath:BasePath];
-		FCString::Strcpy(Result, MAX_PATH, ANSI_TO_TCHAR(BaseDir));
+		FCString::Strcpy(Result, MAX_PATH, *FString(BasePath));
 		FCString::Strcat(Result, TEXT("/"));
 	}
 	return Result;
@@ -814,7 +813,7 @@ const TCHAR* FMacPlatformProcess::UserName(bool bOnlyAlphaNumeric)
 
 void FMacPlatformProcess::SetCurrentWorkingDirectoryToBaseDir()
 {
-	chdir(TCHAR_TO_ANSI(BaseDir()));
+	chdir([FString(BaseDir()).GetNSString() fileSystemRepresentation]);
 }
 
 const TCHAR* FMacPlatformProcess::ExecutableName(bool bRemoveExtension)
