Merging GetTypeHash for __uint128_t for Mac

diff --git a/Engine/Source/Runtime/Core/Public/Templates/TypeHash.h b/Engine/Source/Runtime/Core/Public/Templates/TypeHash.h
--- a/Engine/Source/Runtime/Core/Public/Templates/TypeHash.h
+++ b/Engine/Source/Runtime/Core/Public/Templates/TypeHash.h
@@ -85,6 +85,15 @@ inline uint32 GetTypeHash( const int64 A )
 	return (uint32)A+((uint32)(A>>32) * 23);
 }
 
+#if PLATFORM_MAC
+inline uint32 GetTypeHash( const __uint128_t A )
+{
+	uint64 Low = (uint64)A;
+	uint64 High = (uint64)(A >> 64);
+	return GetTypeHash(Low) ^ GetTypeHash(High);
+}
+#endif
+
 inline uint32 GetTypeHash( float Value )
 {
 	return *(uint32*)&Value;
