Ignore some more warnings on Mac to compile with new Clang in Xcode 7

diff --git a/Engine/Source/Runtime/Core/Public/Mac/MacPlatformCompilerPreSetup.h b/Engine/Source/Runtime/Core/Public/Mac/MacPlatformCompilerPreSetup.h
--- a/Engine/Source/Runtime/Core/Public/Mac/MacPlatformCompilerPreSetup.h
+++ b/Engine/Source/Runtime/Core/Public/Mac/MacPlatformCompilerPreSetup.h
@@ -102,3 +102,6 @@
 // Rationale: we don't want to suppress those as there are plans to address them (e.g. UE-12341), but breaking builds due to these warnings is very expensive
 // since they cannot be caught by all compilers that we support. They are deemed to be relatively safe to be ignored, at least until all SDKs/toolchains start supporting them.
 #pragma clang diagnostic warning "-Wreorder"
+#pragma clang diagnostic ignored "-Winconsistent-missing-override"
+#pragma clang diagnostic ignored "-Wundefined-bool-conversion"
+#pragma clang diagnostic ignored "-Wunused-local-typedef"
