Corrected behavior of FText::FindText which would not properly ensure that the text found actually has a display string that is associated with the source string provided (if provided).

diff --git a/Engine/Source/Runtime/Core/Private/Internationalization/Text.cpp b/Engine/Source/Runtime/Core/Private/Internationalization/Text.cpp
--- a/Engine/Source/Runtime/Core/Private/Internationalization/Text.cpp
+++ b/Engine/Source/Runtime/Core/Private/Internationalization/Text.cpp
@@ -992,7 +992,7 @@ FString FText::GetInvariantTimeZone()
 
 bool FText::FindText( const FString& Namespace, const FString& Key, FText& OutText, const FString* const SourceString )
 {
-	TSharedPtr< FString, ESPMode::ThreadSafe > FoundString = FTextLocalizationManager::Get().FindDisplayString( Namespace, Key );
+	TSharedPtr< FString, ESPMode::ThreadSafe > FoundString = FTextLocalizationManager::Get().FindDisplayString( Namespace, Key, SourceString );
 
 	if ( FoundString.IsValid() )
 	{
@@ -1266,8 +1266,9 @@ TSharedPtr< FString, ESPMode::ThreadSafe > FText::GetSourceString() const
 
 bool FText::IdenticalTo( const FText& Other ) const
 {
-	// If both instances point to the same string, then both instances are considered identical
-	// This is fast as it skips a lexical compare, however it can also return true for two instances that have identical strings, but in different pointers
+	// If both instances point to the same string, then both instances are considered identical.
+	// This is fast as it skips a lexical compare, however it can also return false for two instances that have identical strings, but in different pointers.
+	// For instance, this method will return false for two FText objects created from FText::FromString("Wooble") as they each have unique, non-shared instances.
 	return DisplayString == Other.DisplayString;
 }
 
