Stat - Stats memory profiler improvements, added realloc message (UECORE-167)

diff --git a/Engine/Source/Runtime/Core/Private/Stats/Stats2.cpp b/Engine/Source/Runtime/Core/Private/Stats/Stats2.cpp
--- a/Engine/Source/Runtime/Core/Private/Stats/Stats2.cpp
+++ b/Engine/Source/Runtime/Core/Private/Stats/Stats2.cpp
@@ -91,8 +91,6 @@ void FStats::AdvanceFrame( bool bDiscardCallstack, const FOnAdvanceRenderingThre
 
 /* Todo
 
-FStatsThreadState::FStatsThreadState(FString const& Filename) - needs to be more careful about files that are "cut off". Should look for < 16 bytes left OR next 8 bytes are zero. Bad files will end with zeros and it can't be a valid record.
-
 ST_int64 combined with IsPackedCCAndDuration is fairly bogus, ST_uint32_pair should probably be a first class type. IsPackedCCAndDuration can stay, it says what kind of data is in the pair. remove FromPackedCallCountDuration_Duration et al
 
 averages stay as int64, it might be better if they were floats, but then we would probably want a ST_float_pair too.
@@ -904,7 +902,7 @@ void FThreadStats::CheckEnable()
 	}
 }
 
-void FThreadStats::Flush(bool bHasBrokenCallstacks /*= false*/, bool bForceFlush /*= false*/)
+void FThreadStats::Flush( bool bHasBrokenCallstacks /*= false*/, bool bForceFlush /*= false*/ )
 {
 	if (bMasterDisableForever)
 	{
@@ -956,6 +954,18 @@ void FThreadStats::FlushRegularStats( bool bHasBrokenCallstacks, bool bForceFlus
 		{
 			UE_LOG( LogStats, Verbose, TEXT( "StatMessage Packet has more than %i messages.  Ignoring for the presize history." ), (int32)PRESIZE_MAX_SIZE );
 		}
+
+		// Other threads are one frame behind so set the frame to the previous one.
+		FRunnableThread* RunnableThread = FRunnableThread::GetRunnableThread();
+		if (RunnableThread)
+		{
+			const FString& ThreadName = RunnableThread->GetThreadName();
+			if (ThreadName.Contains( TEXT( "PoolThread" ) ) && Packet.StatMessages.Num() > 128)
+			{
+				UE_LOG( LogStats, Log, TEXT( "%12s Packet T: %i, Frame: %i//%i/%i, Num: %i" ), *ThreadName, Packet.ThreadId, Packet.Frame, FStats::GameThreadStatsFrame, CurrentGameFrame, Packet.StatMessages.Num() );
+			}
+		}
+
 		FStatPacket* ToSend = new FStatPacket(Packet);
 		Exchange(ToSend->StatMessages, Packet.StatMessages);
 		ToSend->bBrokenCallstacks = bHasBrokenCallstacks;
diff --git a/Engine/Source/Runtime/Core/Private/Stats/StatsData.cpp b/Engine/Source/Runtime/Core/Private/Stats/StatsData.cpp
--- a/Engine/Source/Runtime/Core/Private/Stats/StatsData.cpp
+++ b/Engine/Source/Runtime/Core/Private/Stats/StatsData.cpp
@@ -702,13 +702,14 @@ void FStatsThreadState::ScanForAdvance(FStatPacketArray& NewData)
 	for (int32 Index = 0; Index < NewData.Packets.Num(); Index++)
 	{
 		const EThreadType::Type ThreadType = NewData.Packets[Index]->ThreadType;
+		FStatPacket* Packet = NewData.Packets[Index];
 		if ( ThreadType == EThreadType::Renderer)
 		{
-			NewData.Packets[Index]->AssignFrame( CurrentRenderFrame );
+			Packet->AssignFrame( CurrentRenderFrame );
 		}
 		else if (ThreadType == EThreadType::Game)
 		{
-			NewData.Packets[Index]->AssignFrame( CurrentGameFrame );
+			Packet->AssignFrame( CurrentGameFrame );
 		}
 		else if (ThreadType == EThreadType::Other)
 		{
@@ -1024,14 +1025,8 @@ void FStatsThreadState::ResetRegularStats()
 
 void FStatsThreadState::UpdateStatMessagesMemoryUsage()
 {
-	// .UpdateStatMessagesMemoryUsage
 	const int32 CurrentNumStatMessages = NumStatMessages.GetValue();
-	//MaxNumStatMessages = FMath::Max(MaxNumStatMessages,CurrentNumStatMessages);
-
-	if( CurrentNumStatMessages > MaxNumStatMessages )
-	{
-		MaxNumStatMessages = CurrentNumStatMessages;
-	}
+	MaxNumStatMessages = FMath::Max(MaxNumStatMessages,CurrentNumStatMessages);
 
 	TotalNumStatMessages += CurrentNumStatMessages;
 	SET_MEMORY_STAT( STAT_StatMessagesMemory, CurrentNumStatMessages*sizeof(FStatMessage) );
@@ -1055,10 +1050,10 @@ void FStatsThreadState::UpdateStatMessagesMemoryUsage()
 		ToGame->GroupDescriptions.Add( Total );
 
 		FSimpleDelegateGraphTask::CreateAndDispatchWhenReady
-			(
+		(
 			FSimpleDelegateGraphTask::FDelegate::CreateRaw(&FHUDGroupGameThreadRenderer::Get(), &FHUDGroupGameThreadRenderer::NewData, ToGame),
 			TStatId(), nullptr, ENamedThreads::GameThread
-			);
+		);
 	}
 }
 
diff --git a/Engine/Source/Runtime/Core/Private/Stats/StatsFile.cpp b/Engine/Source/Runtime/Core/Private/Stats/StatsFile.cpp
--- a/Engine/Source/Runtime/Core/Private/Stats/StatsFile.cpp
+++ b/Engine/Source/Runtime/Core/Private/Stats/StatsFile.cpp
@@ -320,9 +320,8 @@ void IStatsWriteFile::WriteHeader()
 	Ar << Magic;
 
 	// Serialize dummy header, overwritten in Finalize.
-	Header.Version = EStatMagicWithHeader::VERSION_5;
+	Header.Version = EStatMagicWithHeader::VERSION_LATEST;
 	Header.PlatformName = FPlatformProperties::PlatformName();
-	//Header.bRawStatsFile = bIsRawStatsFile;
 	Ar << Header;
 
 	// Serialize metadata.
@@ -440,7 +439,7 @@ void FStatsWriteFile::SetDataDelegate( bool bSet )
 
 void FStatsWriteFile::WriteFrame( int64 TargetFrame, bool bNeedFullMetadata )
 {
-	// @TODO yrx 2014-11-25 Add stat startfile -num=number of frames to capture
+	// #YRX_STATS: 2015-06-17 Add stat startfile -num=number of frames to capture
 
 	SCOPE_CYCLE_COUNTER( STAT_StreamFile );
 
diff --git a/Engine/Source/Runtime/Core/Private/Stats/StatsMallocProfilerProxy.cpp b/Engine/Source/Runtime/Core/Private/Stats/StatsMallocProfilerProxy.cpp
--- a/Engine/Source/Runtime/Core/Private/Stats/StatsMallocProfilerProxy.cpp
+++ b/Engine/Source/Runtime/Core/Private/Stats/StatsMallocProfilerProxy.cpp
@@ -8,18 +8,21 @@
 #include "StatsMallocProfilerProxy.h"
 
 /** Fake stat group and memory stats. */
-DECLARE_STATS_GROUP(TEXT("Memory Profiler"), STATGROUP_MemoryProfiler, STATCAT_Advanced);
+DECLARE_STATS_GROUP( TEXT( "Memory Profiler" ), STATGROUP_MemoryProfiler, STATCAT_Advanced );
 
 DECLARE_PTR_STAT( TEXT( "Memory Free Ptr" ), STAT_Memory_FreePtr, STATGROUP_MemoryProfiler );
 DECLARE_PTR_STAT( TEXT( "Memory Alloc Ptr" ), STAT_Memory_AllocPtr, STATGROUP_MemoryProfiler );
+DECLARE_PTR_STAT( TEXT( "Memory Realloc Ptr" ), STAT_Memory_ReallocPtr, STATGROUP_MemoryProfiler );
 DECLARE_MEMORY_STAT( TEXT( "Memory Alloc Size" ), STAT_Memory_AllocSize, STATGROUP_MemoryProfiler );
 DECLARE_MEMORY_STAT( TEXT( "Memory Operation Sequence Tag" ), STAT_Memory_OperationSequenceTag, STATGROUP_MemoryProfiler );
+DECLARE_FNAME_STAT( TEXT( "Memory Snapshot" ), STAT_Memory_Snapshot, STATGROUP_MemoryProfiler );
 
 /** Stats for memory usage by the profiler. */
-DECLARE_DWORD_COUNTER_STAT(TEXT("Profiler AllocPtr Calls"),	STAT_Memory_AllocPtr_Calls,	STATGROUP_StatSystem);
-DECLARE_DWORD_COUNTER_STAT(TEXT("Profiler FreePtr Calls"),	STAT_Memory_FreePtr_Calls,	STATGROUP_StatSystem);
-DECLARE_MEMORY_STAT( TEXT( "Profiler AllocPtr" ),			STAT_Memory_AllocPtr_Mem,	STATGROUP_StatSystem );
-DECLARE_MEMORY_STAT( TEXT( "Profiler FreePtr" ),			STAT_Memory_FreePtr_Mem,	STATGROUP_StatSystem );
+DECLARE_DWORD_COUNTER_STAT( TEXT( "Profiler AllocPtr Calls" ), STAT_Memory_AllocPtr_Calls, STATGROUP_StatSystem );
+DECLARE_DWORD_COUNTER_STAT( TEXT( "Profiler ReallocPtr Calls" ), STAT_Memory_ReallocPtr_Calls, STATGROUP_StatSystem );
+DECLARE_DWORD_COUNTER_STAT( TEXT( "Profiler FreePtr Calls" ), STAT_Memory_FreePtr_Calls, STATGROUP_StatSystem );
+DECLARE_MEMORY_STAT( TEXT( "Profiler AllocPtr" ), STAT_Memory_AllocPtr_Mem, STATGROUP_StatSystem );
+DECLARE_MEMORY_STAT( TEXT( "Profiler FreePtr" ), STAT_Memory_FreePtr_Mem, STATGROUP_StatSystem );
 
 //#define DEBUG_MALLOC_PROXY
 
@@ -80,19 +83,17 @@ void FStatsMallocProfilerProxy::InitializeStatsMetadata()
 	// Malloc profiler proxy needs to be disabled otherwise it will hit infinite recursion in DoSetup.
 	// Needs to be changed if we want to support boot time memory profiling.
 	const FName NameAllocPtr = GET_STATFNAME(STAT_Memory_AllocPtr);
+	const FName NameReallocPtr = GET_STATFNAME(STAT_Memory_ReallocPtr);
 	const FName NameFreePtr = GET_STATFNAME(STAT_Memory_FreePtr);
 	const FName NameAllocSize = GET_STATFNAME(STAT_Memory_AllocSize);
 	const FName NameOperationSequenceTag = GET_STATFNAME(STAT_Memory_OperationSequenceTag);
 
 	GET_STATFNAME(STAT_Memory_AllocPtr_Calls);
+	GET_STATFNAME(STAT_Memory_ReallocPtr_Calls);
 	GET_STATFNAME(STAT_Memory_FreePtr_Calls);
 
 	GET_STATFNAME(STAT_Memory_AllocPtr_Mem);
 	GET_STATFNAME(STAT_Memory_FreePtr_Mem);
-
-	//StatPtr_STAT_Memory_AllocPtr;
-	//StatPtr_STAT_Memory_FreePtr;
-	//StatPtr_STAT_Memory_AllocSize;
 }
 
 void FStatsMallocProfilerProxy::TrackAlloc( void* Ptr, int64 Size, int32 SequenceTag )
@@ -114,7 +115,7 @@ void FStatsMallocProfilerProxy::TrackAlloc( void* Ptr, int64 Size, int32 Sequenc
 		
 			if( ThreadStats->MemoryMessageScope == 0 )
 			{
-#if	_DEBUG
+#if	UE_BUILD_DEBUG
 				if( ThreadStats->Packet.StatMessages.Num() > 0 && ThreadStats->Packet.StatMessages.Num() % 32767 == 0 )
 				{
 					ThreadStats->MemoryMessageScope++;
@@ -122,7 +123,7 @@ void FStatsMallocProfilerProxy::TrackAlloc( void* Ptr, int64 Size, int32 Sequenc
 					UE_LOG( LogStats, Verbose, TEXT( "ThreadID: %i, Current: %.1f" ), FPlatformTLS::GetCurrentThreadId(), InvMB*(int64)ThreadStats->Packet.StatMessages.Num()*sizeof( FStatMessage ) );
 					ThreadStats->MemoryMessageScope--;
 				}
-#endif // _DEBUG
+#endif // UE_BUILD_DEBUG
 
 				// 48 bytes per allocation.
 				ThreadStats->AddMemoryMessage( GET_STATFNAME(STAT_Memory_AllocPtr), (uint64)(UPTRINT)Ptr | (uint64)EMemoryOperation::Alloc );
@@ -162,6 +163,34 @@ void FStatsMallocProfilerProxy::TrackFree( void* Ptr, int32 SequenceTag )
 	}
 }
 
+void FStatsMallocProfilerProxy::TrackRealloc( void* OldPtr, void* NewPtr, int64 NewSize, int32 SequenceTag )
+{
+	if (bEnabled)
+	{
+		if (OldPtr != nullptr && NewSize > 0)
+		{
+			FThreadStats* ThreadStats = FThreadStats::GetThreadStats();
+			if (ThreadStats->MemoryMessageScope == 0)
+			{
+				// 64 bytes per reallocation. 80 for Free/Alloc
+				ThreadStats->AddMemoryMessage( GET_STATFNAME( STAT_Memory_FreePtr ), (uint64)(UPTRINT)OldPtr | (uint64)EMemoryOperation::Realloc );
+				ThreadStats->AddMemoryMessage( GET_STATFNAME( STAT_Memory_AllocPtr ), (uint64)(UPTRINT)NewPtr | (uint64)EMemoryOperation::Realloc );
+				ThreadStats->AddMemoryMessage( GET_STATFNAME( STAT_Memory_AllocSize ), NewSize );
+				ThreadStats->AddMemoryMessage( GET_STATFNAME( STAT_Memory_OperationSequenceTag ), (int64)SequenceTag );
+				ReallocPtrCalls.Increment();
+			}
+		}
+		else if (OldPtr == nullptr)
+		{
+			TrackAlloc( NewPtr, NewSize, SequenceTag );
+		}
+		else
+		{
+			TrackFree( OldPtr, SequenceTag );
+		}
+	}
+}
+
 void* FStatsMallocProfilerProxy::Malloc( SIZE_T Size, uint32 Alignment )
 {
 	const int32 SequenceTag = MemorySequenceTag.Increment();
@@ -173,13 +202,9 @@ void* FStatsMallocProfilerProxy::Malloc( SIZE_T Size, uint32 Alignment )
 
 void* FStatsMallocProfilerProxy::Realloc( void* OldPtr, SIZE_T NewSize, uint32 Alignment )
 {
-	// @TODO yrx 2015-02-18 Add TrackRealloc to optimize
-	const int32 FreeSequenceTag = MemorySequenceTag.Increment();
-	TrackFree( OldPtr, FreeSequenceTag );
-
-	const int32 AllocSequenceTag = MemorySequenceTag.Increment();
+	const int32 SequenceTag = MemorySequenceTag.Increment();
 	void* NewPtr = UsedMalloc->Realloc( OldPtr, NewSize, Alignment );
-	TrackAlloc( NewPtr, (int64)NewSize, AllocSequenceTag );
+	TrackRealloc( OldPtr, NewPtr, (int64)NewSize, SequenceTag );
 	return NewPtr;
 }
 
@@ -197,9 +222,11 @@ void FStatsMallocProfilerProxy::UpdateStats()
 	if( bEnabled )
 	{
 		const int32 NumAllocPtrCalls = AllocPtrCalls.GetValue();
+		const int32 NumReallocPtrCalls = ReallocPtrCalls.GetValue();
 		const int32 NumFreePtrCalls = FreePtrCalls.GetValue();
 
 		SET_DWORD_STAT( STAT_Memory_AllocPtr_Calls, NumAllocPtrCalls );
+		SET_DWORD_STAT( STAT_Memory_ReallocPtr_Calls, NumReallocPtrCalls );
 		SET_DWORD_STAT( STAT_Memory_FreePtr_Calls, NumFreePtrCalls );
 
 		// AllocPtr, AllocSize
@@ -208,6 +235,7 @@ void FStatsMallocProfilerProxy::UpdateStats()
 		SET_MEMORY_STAT( STAT_Memory_FreePtr_Mem, NumFreePtrCalls*(sizeof(FStatMessage)*1) );
 
 		AllocPtrCalls.Reset();
+		ReallocPtrCalls.Reset();
 		FreePtrCalls.Reset();
 	}
 }
diff --git a/Engine/Source/Runtime/Core/Public/Stats/Stats2.h b/Engine/Source/Runtime/Core/Public/Stats/Stats2.h
--- a/Engine/Source/Runtime/Core/Public/Stats/Stats2.h
+++ b/Engine/Source/Runtime/Core/Public/Stats/Stats2.h
@@ -297,7 +297,7 @@ struct EMemoryRegion
 };
 
 /** Memory operation for STAT_Memory_AllocPtr. */
-enum class EMemoryOperation
+enum class EMemoryOperation : uint8
 {
 	/** Invalid. */
 	Invalid,	
@@ -305,6 +305,8 @@ enum class EMemoryOperation
 	Alloc,
 	/** Free. */
 	Free,
+	/** Realloc. */
+	Realloc,
 
 	Num,
 	Mask = 0x7,
diff --git a/Engine/Source/Runtime/Core/Public/Stats/StatsFile.h b/Engine/Source/Runtime/Core/Public/Stats/StatsFile.h
--- a/Engine/Source/Runtime/Core/Public/Stats/StatsFile.h
+++ b/Engine/Source/Runtime/Core/Public/Stats/StatsFile.h
@@ -44,7 +44,17 @@ namespace EStatMagicWithHeader
 		 *	New low-level raw stats with memory profiler functionality.
 		 *  !!CAUTION!! Not backward compatible with version 4.
 		 */
-		VERSION_5 = 5
+		VERSION_5 = 5,
+
+		/**
+		 *	Added realloc message to avoid sending alloc and free
+		 *	Should also reduce the amount of all messages
+		 *  !!CAUTION!! Not backward compatible with version 5.
+		 */
+		VERSION_6 = 6,
+
+		/** Latest version. */
+		VERSION_LATEST = VERSION_6,
 	};
 }
 
diff --git a/Engine/Source/Runtime/Core/Public/Stats/StatsMallocProfilerProxy.h b/Engine/Source/Runtime/Core/Public/Stats/StatsMallocProfilerProxy.h
--- a/Engine/Source/Runtime/Core/Public/Stats/StatsMallocProfilerProxy.h
+++ b/Engine/Source/Runtime/Core/Public/Stats/StatsMallocProfilerProxy.h
@@ -25,6 +25,9 @@ private:
 	/** Malloc calls made this frame. */
 	FThreadSafeCounter AllocPtrCalls;
 
+	/** Realloc calls made this frame. */
+	FThreadSafeCounter ReallocPtrCalls;
+
 	/** Free calls made this frame. */
 	FThreadSafeCounter FreePtrCalls;
 
@@ -69,6 +72,17 @@ public:
 	 */
 	void TrackFree( void* Ptr, int32 SequenceTag );
 
+
+	/** 
+	 *	Tracks realloc operation
+	 *	
+	 * @param	Old - Previous pointer
+	 * @param	New - New pointer
+	 * @param	NewSize - Size of allocated pointer
+	 * @param	SequenceTag - Previous value for the memory sequence tag
+	 */
+	void TrackRealloc( void* OldPtr, void* NewPtr, int64 NewSize, int32 SequenceTag );
+
 	virtual void* Malloc( SIZE_T Size, uint32 Alignment ) override;
 	virtual void* Realloc( void* OldPtr, SIZE_T NewSize, uint32 Alignment ) override;
 	virtual void Free( void* Ptr ) override;
