CIS Fix

diff --git a/Engine/Source/Runtime/Core/Private/Stats/Stats2.cpp b/Engine/Source/Runtime/Core/Private/Stats/Stats2.cpp
--- a/Engine/Source/Runtime/Core/Private/Stats/Stats2.cpp
+++ b/Engine/Source/Runtime/Core/Private/Stats/Stats2.cpp
@@ -954,18 +954,6 @@ void FThreadStats::FlushRegularStats( bool bHasBrokenCallstacks, bool bForceFlus
 		{
 			UE_LOG( LogStats, Verbose, TEXT( "StatMessage Packet has more than %i messages.  Ignoring for the presize history." ), (int32)PRESIZE_MAX_SIZE );
 		}
-
-		// Other threads are one frame behind so set the frame to the previous one.
-		FRunnableThread* RunnableThread = FRunnableThread::GetRunnableThread();
-		if (RunnableThread)
-		{
-			const FString& ThreadName = RunnableThread->GetThreadName();
-			if (ThreadName.Contains( TEXT( "PoolThread" ) ) && Packet.StatMessages.Num() > 128)
-			{
-				UE_LOG( LogStats, Log, TEXT( "%12s Packet T: %i, Frame: %i//%i/%i, Num: %i" ), *ThreadName, Packet.ThreadId, Packet.Frame, FStats::GameThreadStatsFrame, CurrentGameFrame, Packet.StatMessages.Num() );
-			}
-		}
-
 		FStatPacket* ToSend = new FStatPacket(Packet);
 		Exchange(ToSend->StatMessages, Packet.StatMessages);
 		ToSend->bBrokenCallstacks = bHasBrokenCallstacks;
