Fix a pthread mutexattr leak.

diff --git a/Engine/Source/Runtime/Core/Public/HAL/PThreadCriticalSection.h b/Engine/Source/Runtime/Core/Public/HAL/PThreadCriticalSection.h
--- a/Engine/Source/Runtime/Core/Public/HAL/PThreadCriticalSection.h
+++ b/Engine/Source/Runtime/Core/Public/HAL/PThreadCriticalSection.h
@@ -29,6 +29,7 @@ public:
 		pthread_mutexattr_init(&MutexAttributes);
 		pthread_mutexattr_settype(&MutexAttributes, PTHREAD_MUTEX_RECURSIVE);
 		pthread_mutex_init(&Mutex, &MutexAttributes);
+		pthread_mutexattr_destroy(&MutexAttributes);
 	}
 
 	/**
