Make sure the UE4Game directory created
#ue-17385
#ue4
#android

diff --git a/Engine/Source/Runtime/Core/Private/Android/AndroidFile.cpp b/Engine/Source/Runtime/Core/Private/Android/AndroidFile.cpp
--- a/Engine/Source/Runtime/Core/Private/Android/AndroidFile.cpp
+++ b/Engine/Source/Runtime/Core/Private/Android/AndroidFile.cpp
@@ -883,6 +883,10 @@ public:
 				MountOBB(*(OBBDir2 / PatchOBBName));
 			}
 		}
+
+		// make sure the base path directory exists (UE4Game)
+		mkdir(TCHAR_TO_UTF8(*GetFileBasePath()), 0766);
+
 		return true;
 	}
 
