merge automation fix @ CL 2570986

diff --git a/Engine/Source/Runtime/Core/Private/Misc/AutomationTest.cpp b/Engine/Source/Runtime/Core/Private/Misc/AutomationTest.cpp
--- a/Engine/Source/Runtime/Core/Private/Misc/AutomationTest.cpp
+++ b/Engine/Source/Runtime/Core/Private/Misc/AutomationTest.cpp
@@ -437,14 +437,19 @@ bool FAutomationTestFramework::ShouldTestContent(const FString& Path) const
 		GConfig->GetArray( TEXT("/Script/Engine.AutomationTestSettings"), TEXT("TestLevelFolders"), TestLevelFolders, GEngineIni);
 	}
 
+	bool bMatchingDirectory = false;
 	for ( const FString& Folder : TestLevelFolders )
 	{
 		const FString PatternToCheck = FString::Printf(TEXT("/%s/"), *Folder);
 		if ( Path.Contains(*PatternToCheck) )
 		{
-			return false;
+			bMatchingDirectory = true;
 		}
 	}
+	if (bMatchingDirectory)
+	{
+		return true;
+	}
 
 	FString DevelopersPath = FPaths::GameDevelopersDir().LeftChop(1);
 	return bDeveloperDirectoryIncluded || bVisualCommandletFilterOn || !Path.StartsWith(DevelopersPath);
