Fix for FAndroidPlatformMemory::BinnedAllocFromOS for SDK 21
#ue4
#android

diff --git a/Engine/Source/Runtime/Core/Private/Android/AndroidMemory.cpp b/Engine/Source/Runtime/Core/Private/Android/AndroidMemory.cpp
--- a/Engine/Source/Runtime/Core/Private/Android/AndroidMemory.cpp
+++ b/Engine/Source/Runtime/Core/Private/Android/AndroidMemory.cpp
@@ -96,7 +96,8 @@ FMalloc* FAndroidPlatformMemory::BaseAllocator()
 void* FAndroidPlatformMemory::BinnedAllocFromOS( SIZE_T Size )
 {
 	// valloc was deprecated, this is a functional equivalent, for SDK 21
-	return memalign(sysconf(_SC_PAGESIZE), Size);
+	const FPlatformMemoryConstants& MemoryConstants = FPlatformMemory::GetConstants();
+	return memalign(MemoryConstants.PageSize, Size);
 }
 
 void FAndroidPlatformMemory::BinnedFreeToOS( void* Ptr )
