Linux: refuse to run as root.

- Prevents surprisingly common user errors of attempting to re-run tools with 'sudo' in case of failure and then messing things up even more.
- There might be legitimate use cases for running as root, and if those are identified the check will need to be amended.

diff --git a/Engine/Source/Runtime/Core/Private/Linux/LinuxPlatformMemory.cpp b/Engine/Source/Runtime/Core/Private/Linux/LinuxPlatformMemory.cpp
--- a/Engine/Source/Runtime/Core/Private/Linux/LinuxPlatformMemory.cpp
+++ b/Engine/Source/Runtime/Core/Private/Linux/LinuxPlatformMemory.cpp
@@ -25,6 +25,16 @@ void FLinuxPlatformMemory::Init()
 
 class FMalloc* FLinuxPlatformMemory::BaseAllocator()
 {
+	// This function gets executed very early, way before main() (because global constructors will allocate memory).
+	// This makes it ideal, if unobvious, place for a root privilege check.
+	if (geteuid() == 0)
+	{
+		fprintf(stderr, "Refusing to run with the root privileges.\n");
+		FPlatformMisc::RequestExit(true);
+		// unreachable
+		return nullptr;
+	}
+
 	enum EAllocatorToUse
 	{
 		Ansi,
