#ue4 - Fix VectorMod() working incorrectly for negative values (now matches std fmodf). Fix VectorSinCos() doing incorrect range reduction. Improve math tests at startup.

- Math tests: added VectorMod, VectorSinCos.
- Math tests: added more FQuat <-> FRotator conversions and cleaned it up (make more important tests earlier, because those could cause later ones to fail).

diff --git a/Engine/Source/Runtime/Core/Private/Tests/Math/UnrealMathTest.cpp b/Engine/Source/Runtime/Core/Private/Tests/Math/UnrealMathTest.cpp
--- a/Engine/Source/Runtime/Core/Private/Tests/Math/UnrealMathTest.cpp
+++ b/Engine/Source/Runtime/Core/Private/Tests/Math/UnrealMathTest.cpp
@@ -10,7 +10,7 @@ MS_ALIGN( 16 ) static float GScratch[16] GCC_ALIGN( 16 );
 static float GSum;
 static bool GPassing;
 
-#define MATHTEST_INLINE FORCEINLINE_DEBUGGABLE
+#define MATHTEST_INLINE FORCENOINLINE // if you want to do performance testing change to FORCEINLINE or FORCEINLINE_DEBUGGABLE
 
 /**
  * Tests if two vectors (xyzw) are bitwise equal
@@ -56,6 +56,24 @@ bool TestVectorsEqual( VectorRegister Vec0, VectorRegister Vec1, float Tolerance
 }
 
 
+/**
+ * Enforce tolerance per-component, not summed error
+ */
+bool TestVectorsEqual_ComponentWiseError(VectorRegister Vec0, VectorRegister Vec1, float Tolerance = 0.0f)
+{
+	VectorStoreAligned(Vec0, GScratch + 0);
+	VectorStoreAligned(Vec1, GScratch + 4);
+	bool bPassing = true;
+	for (int32 Component = 0; Component < 4; Component++)
+	{
+		float Diff = GScratch[Component + 0] - GScratch[Component + 4];
+		bPassing &= FMath::IsNearlyZero(Diff, Tolerance);
+	}
+	GPassing &= bPassing;
+	return bPassing;
+}
+
+
 /**
  * Tests if two vectors (xyz) are equal within an optional tolerance
  *
@@ -307,7 +325,7 @@ VectorRegister TestVectorTransformVector(const VectorRegister&  VecP,  const voi
 * @param Rotator FRotator 
 * @return Rotation as a quaternion.
 */
-FORCENOINLINE FQuat TestRotatorToQuaternion( const FRotator& Rotator)
+MATHTEST_INLINE FQuat TestRotatorToQuaternion( const FRotator& Rotator)
 {
 	const float CR = FMath::Cos(FMath::DegreesToRadians(Rotator.Roll  * 0.5f));
 	const float CP = FMath::Cos(FMath::DegreesToRadians(Rotator.Pitch * 0.5f));
@@ -476,6 +494,16 @@ void LogRotatorTest(bool bExpected, const TCHAR* TestName, const FRotator& A, co
 	}
 }
 
+void LogQuaternionTest(const TCHAR* TestName, const FQuat& A, const FQuat& B, bool bComparison)
+{
+	if (bComparison == false)
+	{
+		UE_LOG(LogUnrealMathTest, Log, TEXT("%s: %s"), bComparison ? TEXT("PASSED") : TEXT("FAILED"), TestName);
+		UE_LOG(LogUnrealMathTest, Log, TEXT("%s.Equals(%s) = %d"), *A.ToString(), *B.ToString(), bComparison);
+		GPassing = false;
+	}
+}
+
 
 // Normalize tests
 
@@ -501,6 +529,51 @@ MATHTEST_INLINE VectorRegister TestVectorNormalize_InvSqrtEst(const VectorRegist
 }
 
 
+// A Mod M
+MATHTEST_INLINE VectorRegister TestReferenceMod(const VectorRegister& A, const VectorRegister& M)
+{
+	return MakeVectorRegister(
+		(float)fmod(VectorGetComponent(A, 0), VectorGetComponent(M, 0)),
+		(float)fmod(VectorGetComponent(A, 1), VectorGetComponent(M, 1)),
+		(float)fmod(VectorGetComponent(A, 2), VectorGetComponent(M, 2)),
+		(float)fmod(VectorGetComponent(A, 3), VectorGetComponent(M, 3)));
+}
+
+// SinCos
+MATHTEST_INLINE void TestReferenceSinCos(VectorRegister& S, VectorRegister& C, const VectorRegister& VAngles)
+{
+	S = MakeVectorRegister(
+			FMath::Sin(VectorGetComponent(VAngles, 0)),
+			FMath::Sin(VectorGetComponent(VAngles, 1)),
+			FMath::Sin(VectorGetComponent(VAngles, 2)),
+			FMath::Sin(VectorGetComponent(VAngles, 3))
+			);
+
+	C = MakeVectorRegister(
+			FMath::Cos(VectorGetComponent(VAngles, 0)),
+			FMath::Cos(VectorGetComponent(VAngles, 1)),
+			FMath::Cos(VectorGetComponent(VAngles, 2)),
+			FMath::Cos(VectorGetComponent(VAngles, 3))
+			);
+}
+
+MATHTEST_INLINE void TestFastSinCos(VectorRegister& S, VectorRegister& C, const VectorRegister& VAngles)
+{
+	float SFloat[4], CFloat[4];
+	FMath::SinCos(&SFloat[0], &CFloat[0], VectorGetComponent(VAngles, 0));
+	FMath::SinCos(&SFloat[1], &CFloat[1], VectorGetComponent(VAngles, 1));
+	FMath::SinCos(&SFloat[2], &CFloat[2], VectorGetComponent(VAngles, 2));
+	FMath::SinCos(&SFloat[3], &CFloat[3], VectorGetComponent(VAngles, 3));
+
+	S = VectorLoad(SFloat);
+	C = VectorLoad(CFloat);
+}
+
+MATHTEST_INLINE void TestVectorSinCos(VectorRegister& S, VectorRegister& C, const VectorRegister& VAngles)
+{
+	VectorSinCos(&S, &C, &VAngles);
+}
+
 /**
  * Helper debugf function to print out success or failure information for a test
  *
@@ -903,6 +976,19 @@ bool FVectorRegisterAbstractionTest::RunTest(const FString& Parameters)
 	V3 = VectorMultiply(VectorMultiply(V1, V1), V0);
 	LogTest( TEXT("VectorReciprocalSqrtAccurate"), TestVectorsEqual( VectorOne(), V3, 1e-6f ) );
 
+	// VectorMod
+	V0 = MakeVectorRegister(0.0f, 3.2f, 2.8f, 10.0f);
+	V1 = MakeVectorRegister(2.0f, 1.2f, 2.0f,  3.0f);
+	V2 = TestReferenceMod(V0, V1);
+	V3 = VectorMod(V0, V1);
+	LogTest( TEXT("VectorMod positive"), TestVectorsEqual(V2, V3));
+
+	V0 = MakeVectorRegister(-2.0f,  3.2f, -2.8f, -10.0f);
+	V1 = MakeVectorRegister(-1.5f, -1.2f,  2.0f,   3.0f);
+	V2 = TestReferenceMod(V0, V1);
+	V3 = VectorMod(V0, V1);
+	LogTest( TEXT("VectorMod negative"), TestVectorsEqual(V2, V3));
+
 	FMatrix	M0, M1, M2, M3;
 	FVector Eye, LookAt, Up;	
 	// Create Look at Matrix
@@ -952,6 +1038,95 @@ bool FVectorRegisterAbstractionTest::RunTest(const FString& Parameters)
 
 	FQuat Q0, Q1, Q2, Q3;
 
+	// SinCos tests
+	{
+		const VectorRegister QuadrantDegreesArray[] = {
+			MakeVectorRegister( 0.0f, 10.0f, 20.0f, 30.0f),
+			MakeVectorRegister(45.0f, 60.0f, 70.0f, 80.0f),
+		};
+
+		const float SinCosTolerance = 1e-6f;
+		const int32 Cycles = 3; // Go through a full circle this many times (negative and positive)
+		for (int32 OffsetQuadrant = -4*Cycles; OffsetQuadrant <= 4*Cycles; ++OffsetQuadrant)
+		{
+			const float OffsetFloat = (float)OffsetQuadrant * 90.0f; // Add 90 degrees repeatedly to cover all quadrants and wrap a few times
+			const VectorRegister VOffset = VectorLoadFloat1(&OffsetFloat);
+			for (VectorRegister const& VDegrees : QuadrantDegreesArray)
+			{
+				const VectorRegister VAnglesDegrees = VectorAdd(VOffset, VDegrees);
+				const VectorRegister VAngles = VectorMultiply(VAnglesDegrees, GlobalVectorConstants::DEG_TO_RAD);
+				VectorRegister S[3], C[3];
+				TestReferenceSinCos(S[0], C[0], VAngles);
+				TestFastSinCos(S[1], C[1], VAngles);
+				TestVectorSinCos(S[2], C[2], VAngles);
+				LogTest(TEXT("SinCos (Sin): Ref vs Fast"), TestVectorsEqual_ComponentWiseError(S[0], S[1], SinCosTolerance));
+				LogTest(TEXT("SinCos (Cos): Ref vs Fast"), TestVectorsEqual_ComponentWiseError(C[0], C[1], SinCosTolerance));
+				LogTest(TEXT("SinCos (Sin): Ref vs Vec"), TestVectorsEqual_ComponentWiseError(S[0], S[2], SinCosTolerance));
+				LogTest(TEXT("SinCos (Cos): Ref vs Vec"), TestVectorsEqual_ComponentWiseError(C[0], C[2], SinCosTolerance));
+			}
+		}
+	}
+
+	// Quat<->Rotator conversions and equality
+	{
+		const float Nudge = KINDA_SMALL_NUMBER * 0.25f;
+		const FRotator RotArray[] = {
+			FRotator(0.f, 0.f, 0.f),
+			FRotator(Nudge, -Nudge, Nudge),
+			FRotator(+180.f, -180.f, +180.f),
+			FRotator(-180.f, +180.f, -180.f),
+			FRotator(+45.0f - Nudge, -120.f + Nudge, +270.f - Nudge),
+			FRotator(-45.0f + Nudge, +120.f - Nudge, -270.f + Nudge),
+			FRotator(+315.f - 360.f, -240.f - 360.f, -90.0f - 360.f),
+			FRotator(-315.f + 360.f, +240.f + 360.f, +90.0f + 360.f),
+		};
+
+		// FRotator tests
+		{
+			// Equality test
+			const float RotTolerance = KINDA_SMALL_NUMBER;
+			for (auto const& A : RotArray)
+			{
+				for (auto const& B : RotArray)
+				{
+					const bool bExpected = TestRotatorEqual0(A, B, RotTolerance);
+					LogRotatorTest(bExpected, TEXT("TestRotatorEqual1"), A, B, TestRotatorEqual1(A, B, RotTolerance));
+					LogRotatorTest(bExpected, TEXT("TestRotatorEqual2"), A, B, TestRotatorEqual2(A, B, RotTolerance));
+					LogRotatorTest(bExpected, TEXT("TestRotatorEqual3"), A, B, TestRotatorEqual3(A, B, RotTolerance));
+				}
+			}
+		}
+
+		// Quaternion conversion test
+		const float QuatTolerance = 1e-6f;
+		for (auto const& A : RotArray)
+		{
+			const FQuat QA = TestRotatorToQuaternion(A);
+			const FQuat QB = A.Quaternion();
+			LogQuaternionTest(TEXT("TestRotatorToQuaternion"), QA, QB, TestQuatsEqual(QA, QB, QuatTolerance));
+		}
+	}
+
+	// Rotator->Quat->Rotator
+	{
+		const FRotator RotArray[] ={
+			FRotator(30.0f, -45.0f, 90.0f),
+			FRotator(45.0f, 60.0f, 120.0f),
+			FRotator(90.f, 0.f, 0.f),
+			FRotator(-90.f, 0.f, 0.f),
+			FRotator(150.f, 0.f, 0.f)
+		};
+
+		for (FRotator const& Rotator0 : RotArray)
+		{
+			Q0 = TestRotatorToQuaternion(Rotator0);
+			FRotator Rotator1 = Q0.Rotator();
+			FRotator Rotator2 = TestQuaternionToRotator(Q0);
+			LogTest(TEXT("Rotator->Quat->Rotator"), Rotator1.Equals(Rotator2, 1e-5f));
+		}
+	}
+
+	// Quat / Rotator conversion to vectors, matrices
 	{
 		FRotator Rotator0;
 		Rotator0 = FRotator(30.0f, -45.0f, 90.0f);
@@ -986,6 +1161,7 @@ bool FVectorRegisterAbstractionTest::RunTest(const FString& Parameters)
 		LogTest(TEXT("Test2 FQuatRotationMatrix"), TestVectorsEqual3_NoVec(FV0, FV1, 1e-6f));
 	}
 
+	// Quat Rotation tests
 	{
 		// Use these Quats...
 		const FQuat TestQuats[] = {
@@ -1032,71 +1208,31 @@ bool FVectorRegisterAbstractionTest::RunTest(const FString& Parameters)
 		}
 	}
 
+	// Quat multiplication
 	{
-		FRotator Rotator0, Rotator1, Rotator2;
-		Rotator0 = FRotator(30.0f, -45.0f, 90.0f);
-		Q0 = TestRotatorToQuaternion(Rotator0);
-		Rotator1 = Q0.Rotator();
-		Rotator2 = TestQuaternionToRotator(Q0);
-		LogTest(TEXT("TestQuaternionToRotator"), Rotator1.Equals(Rotator2, 1e-5f));
-
-		Rotator0 = FRotator(45.0f, 60.0f, 120.0f);
-		Q0 = TestRotatorToQuaternion(Rotator0);
-		Rotator1 = Q0.Rotator();
-		Rotator2 = TestQuaternionToRotator(Q0);
-		LogTest(TEXT("TestQuaternionToRotator"), Rotator1.Equals(Rotator2, 1e-5f));
+		Q0 = FQuat(FRotator(30.0f, -45.0f, 90.0f));
+		Q1 = FQuat(FRotator(45.0f, 60.0f, 120.0f));
+		VectorQuaternionMultiply(&Q2, &Q0, &Q1);
+		TestVectorQuaternionMultiply(&Q3, &Q0, &Q1);
+		LogTest(TEXT("VectorQuaternionMultiply"), TestQuatsEqual(Q2, Q3, 1e-6f));
+		V0 = VectorLoadAligned(&Q0);
+		V1 = VectorLoadAligned(&Q1);
+		V2 = VectorQuaternionMultiply2(V0, V1);
+		V3 = VectorLoadAligned(&Q3);
+		LogTest(TEXT("VectorQuaternionMultiply2"), TestVectorsEqual(V2, V3, 1e-6f));
+
+		Q0 = FQuat(FRotator(0.0f, 180.0f, 45.0f));
+		Q1 = FQuat(FRotator(-120.0f, -90.0f, 0.0f));
+		VectorQuaternionMultiply(&Q2, &Q0, &Q1);
+		TestVectorQuaternionMultiply(&Q3, &Q0, &Q1);
+		LogTest(TEXT("VectorMatrixInverse"), TestQuatsEqual(Q2, Q3, 1e-6f));
+		V0 = VectorLoadAligned(&Q0);
+		V1 = VectorLoadAligned(&Q1);
+		V2 = VectorQuaternionMultiply2(V0, V1);
+		V3 = VectorLoadAligned(&Q3);
+		LogTest(TEXT("VectorQuaternionMultiply2"), TestVectorsEqual(V2, V3, 1e-6f));
 	}
 
-	Q0 = FQuat(FRotator(30.0f, -45.0f, 90.0f));
-	Q1 = FQuat(FRotator(45.0f,  60.0f, 120.0f));
-	VectorQuaternionMultiply(&Q2, &Q0, &Q1);
-	TestVectorQuaternionMultiply(&Q3, &Q0, &Q1);
-	LogTest( TEXT("VectorQuaternionMultiply"), TestQuatsEqual(Q2, Q3, 1e-6f));
-	V0 = VectorLoadAligned(&Q0);
-	V1 = VectorLoadAligned(&Q1);
-	V2 = VectorQuaternionMultiply2(V0, V1);
-	V3 = VectorLoadAligned(&Q3);
-	LogTest( TEXT("VectorQuaternionMultiply2"), TestVectorsEqual( V2, V3, 1e-6f ) );
-
-	Q0 = FQuat(FRotator(0.0f, 180.0f, 45.0f));
-	Q1 = FQuat(FRotator(-120.0f, -90.0f, 0.0f));
-	VectorQuaternionMultiply(&Q2, &Q0, &Q1);
-	TestVectorQuaternionMultiply(&Q3, &Q0, &Q1);
-	LogTest( TEXT("VectorMatrixInverse"), TestQuatsEqual(Q2, Q3, 1e-6f) );
-	V0 = VectorLoadAligned(&Q0);
-	V1 = VectorLoadAligned(&Q1);
-	V2 = VectorQuaternionMultiply2(V0, V1);
-	V3 = VectorLoadAligned(&Q3);
-	LogTest( TEXT("VectorQuaternionMultiply2"), TestVectorsEqual(V2, V3, 1e-6f) );
-
-
-	// FRotator tests
-	{
-		const float Nudge = KINDA_SMALL_NUMBER * 0.25f;
-		FRotator RotArray[] ={
-			FRotator(0.f, 0.f, 0.f),
-			FRotator(Nudge, -Nudge, Nudge),
-			FRotator(+180.f, -180.f, +180.f),
-			FRotator(-180.f, +180.f, -180.f),
-			FRotator(+45.0f - Nudge, -120.f + Nudge, +270.f - Nudge),
-			FRotator(-45.0f + Nudge, +120.f - Nudge, -270.f + Nudge),
-			FRotator(+315.f - 360.f, -240.f - 360.f, -90.0f - 360.f),
-			FRotator(-315.f + 360.f, +240.f + 360.f, +90.0f + 360.f),
-		};
-
-		// Equality test
-		const float RotTolerance = KINDA_SMALL_NUMBER;
-		for (auto const& A : RotArray)
-		{
-			for (auto const& B : RotArray)
-			{
-				const bool bExpected = TestRotatorEqual0(A, B, RotTolerance);
-				LogRotatorTest(bExpected, TEXT("TestRotatorEqual1"), A, B, TestRotatorEqual1(A, B, RotTolerance));
-				LogRotatorTest(bExpected, TEXT("TestRotatorEqual2"), A, B, TestRotatorEqual2(A, B, RotTolerance));
-				LogRotatorTest(bExpected, TEXT("TestRotatorEqual3"), A, B, TestRotatorEqual3(A, B, RotTolerance));
-			}
-		}
-	}
 
 	if (!GPassing)
 	{
diff --git a/Engine/Source/Runtime/Core/Public/Math/UnrealMathSSE.h b/Engine/Source/Runtime/Core/Public/Math/UnrealMathSSE.h
--- a/Engine/Source/Runtime/Core/Public/Math/UnrealMathSSE.h
+++ b/Engine/Source/Runtime/Core/Public/Math/UnrealMathSSE.h
@@ -1003,7 +1003,7 @@ FORCEINLINE VectorRegister VectorFloor(const VectorRegister& X)
 
 FORCEINLINE VectorRegister VectorMod(const VectorRegister& X, const VectorRegister& Y)
 {
-	VectorRegister Temp = VectorFloor(VectorDivide(X, Y));
+	VectorRegister Temp = VectorTruncate(VectorDivide(X, Y));
 	return VectorSubtract(X, VectorMultiply(Y, Temp));
 }
 
@@ -1089,8 +1089,14 @@ FORCEINLINE VectorRegister VectorCos(const VectorRegister& X)
 FORCEINLINE void VectorSinCos(VectorRegister* RESTRICT VSinAngles, VectorRegister* RESTRICT VCosAngles, const VectorRegister* RESTRICT VAngles)
 {
 	// Map to [-pi, pi]
-	VectorRegister X = VectorFloor(VectorMultiplyAdd(*VAngles, GlobalVectorConstants::OneOverTwoPi, GlobalVectorConstants::FloatOneHalf));
-	X = VectorSubtract(*VAngles, VectorMultiply(GlobalVectorConstants::TwoPi, X));
+	// X = A - 2pi * round(A/2pi)
+	// Note the round(), not truncate(). In this case round() can round halfway cases using round-to-nearest-even OR round-to-nearest.
+
+	// Quotient = round(A/2pi)
+	VectorRegister Quotient = VectorMultiply(*VAngles, GlobalVectorConstants::OneOverTwoPi);
+	Quotient = _mm_cvtepi32_ps(_mm_cvtps_epi32(Quotient)); // round to nearest even is the default rounding mode but that's fine here.
+	// X = A - 2pi * Quotient
+	VectorRegister X = VectorSubtract(*VAngles, VectorMultiply(GlobalVectorConstants::TwoPi, Quotient));
 
 	// Map in [-pi/2,pi/2]
 	VectorRegister sign = VectorBitwiseAnd(X, GlobalVectorConstants::SignBit);
diff --git a/Engine/Source/Runtime/Core/Public/Math/UnrealMathVectorConstants.h b/Engine/Source/Runtime/Core/Public/Math/UnrealMathVectorConstants.h
--- a/Engine/Source/Runtime/Core/Public/Math/UnrealMathVectorConstants.h
+++ b/Engine/Source/Runtime/Core/Public/Math/UnrealMathVectorConstants.h
@@ -14,6 +14,7 @@ namespace GlobalVectorConstants
 	static const VectorRegister Float111_Minus1 = MakeVectorRegister( 1.f, 1.f, 1.f, -1.f );
 	static const VectorRegister FloatMinus1_111= MakeVectorRegister( -1.f, 1.f, 1.f, 1.f );
 	static const VectorRegister FloatOneHalf = MakeVectorRegister( 0.5f, 0.5f, 0.5f, 0.5f );
+	static const VectorRegister FloatMinusOneHalf = MakeVectorRegister( -0.5f, -0.5f, -0.5f, -0.5f );
 	static const VectorRegister KindaSmallNumber = MakeVectorRegister( KINDA_SMALL_NUMBER, KINDA_SMALL_NUMBER, KINDA_SMALL_NUMBER, KINDA_SMALL_NUMBER );
 	static const VectorRegister SmallNumber = MakeVectorRegister( SMALL_NUMBER, SMALL_NUMBER, SMALL_NUMBER, SMALL_NUMBER );
 	static const VectorRegister ThreshQuatNormalized = MakeVectorRegister( THRESH_QUAT_NORMALIZED, THRESH_QUAT_NORMALIZED, THRESH_QUAT_NORMALIZED, THRESH_QUAT_NORMALIZED );
