Check for XMPP existing before trying to load it. Fixes warnings if you don't have XMPP (because it's in a /NotForLicensees/ folder), which also happens for everything we build in Rocket.

diff --git a/Engine/Source/Runtime/Core/Private/Modules/ModuleManager.cpp b/Engine/Source/Runtime/Core/Private/Modules/ModuleManager.cpp
--- a/Engine/Source/Runtime/Core/Private/Modules/ModuleManager.cpp
+++ b/Engine/Source/Runtime/Core/Private/Modules/ModuleManager.cpp
@@ -100,6 +100,12 @@ void FModuleManager::FindModules(const TCHAR* WildcardWithoutExtension, TArray<F
 #endif
 }
 
+bool FModuleManager::ModuleExists(const TCHAR* ModuleName)
+{
+	TArray<FName> Names;
+	FindModules(ModuleName, Names);
+	return Names.Num() > 0;
+}
 
 bool FModuleManager::IsModuleLoaded( const FName InModuleName )
 {
@@ -329,93 +335,93 @@ TSharedPtr<IModuleInterface> FModuleManager::LoadModuleWithFailureReason(const F
 	}
 	else
 	{
-	// Make sure this isn't a module that we had previously loaded, and then unloaded at shutdown time.
-	//
-	// If this assert goes off, your trying to load a module during the shutdown phase that was already
-	// cleaned up.  The easiest way to fix this is to change your code to query for an already-loaded
-	// module instead of trying to load it directly.
-	checkf((!ModuleInfo->bWasUnloadedAtShutdown), TEXT("Attempted to load module '%s' that was already unloaded at shutdown.  FModuleManager::LoadModule() was called to load a module that was previously loaded, and was unloaded at shutdown time.  If this assert goes off, your trying to load a module during the shutdown phase that was already cleaned up.  The easiest way to fix this is to change your code to query for an already-loaded module instead of trying to load it directly."), *InModuleName.ToString());
-
-	// Check if we're statically linked with the module.  Those modules register with the module manager using a static variable,
-	// so hopefully we already know about the name of the module and how to initialize it.
-	const FInitializeStaticallyLinkedModule* ModuleInitializerPtr = StaticallyLinkedModuleInitializers.Find(InModuleName);
+		// Make sure this isn't a module that we had previously loaded, and then unloaded at shutdown time.
+		//
+		// If this assert goes off, your trying to load a module during the shutdown phase that was already
+		// cleaned up.  The easiest way to fix this is to change your code to query for an already-loaded
+		// module instead of trying to load it directly.
+		checkf((!ModuleInfo->bWasUnloadedAtShutdown), TEXT("Attempted to load module '%s' that was already unloaded at shutdown.  FModuleManager::LoadModule() was called to load a module that was previously loaded, and was unloaded at shutdown time.  If this assert goes off, your trying to load a module during the shutdown phase that was already cleaned up.  The easiest way to fix this is to change your code to query for an already-loaded module instead of trying to load it directly."), *InModuleName.ToString());
+
+		// Check if we're statically linked with the module.  Those modules register with the module manager using a static variable,
+		// so hopefully we already know about the name of the module and how to initialize it.
+		const FInitializeStaticallyLinkedModule* ModuleInitializerPtr = StaticallyLinkedModuleInitializers.Find(InModuleName);
 		if (ModuleInitializerPtr != NULL)
-	{
-		const FInitializeStaticallyLinkedModule& ModuleInitializer(*ModuleInitializerPtr);
+		{
+			const FInitializeStaticallyLinkedModule& ModuleInitializer(*ModuleInitializerPtr);
 
-		// Initialize the module!
-		ModuleInfo->Module = MakeShareable(ModuleInitializer.Execute());
+			// Initialize the module!
+			ModuleInfo->Module = MakeShareable(ModuleInitializer.Execute());
 
-		if (ModuleInfo->Module.IsValid())
-		{
-			// Startup the module
-			ModuleInfo->Module->StartupModule();
+			if (ModuleInfo->Module.IsValid())
+			{
+				// Startup the module
+				ModuleInfo->Module->StartupModule();
 
-			// Module was started successfully!  Fire callbacks.
-			ModulesChangedEvent.Broadcast(InModuleName, EModuleChangeReason::ModuleLoaded);
+				// Module was started successfully!  Fire callbacks.
+				ModulesChangedEvent.Broadcast(InModuleName, EModuleChangeReason::ModuleLoaded);
 
-			// Set the return parameter
-			LoadedModule = ModuleInfo->Module;
-		}
-		else
-		{
-			UE_LOG(LogModuleManager, Warning, TEXT("ModuleManager: Unable to load module '%s' because InitializeModule function failed (returned NULL pointer.)"), *InModuleName.ToString());
-			OutFailureReason = EModuleLoadResult::FailedToInitialize;
+				// Set the return parameter
+				LoadedModule = ModuleInfo->Module;
+			}
+			else
+			{
+				UE_LOG(LogModuleManager, Warning, TEXT("ModuleManager: Unable to load module '%s' because InitializeModule function failed (returned NULL pointer.)"), *InModuleName.ToString());
+				OutFailureReason = EModuleLoadResult::FailedToInitialize;
+			}
 		}
-	}
 #if IS_MONOLITHIC
 		else
 		{
-	// Monolithic builds that do not have the initializer were *not found* during the build step, so return FileNotFound
-	// (FileNotFound is an acceptable error in some case - ie loading a content only project)
-	UE_LOG(LogModuleManager, Warning, TEXT("ModuleManager: Module '%s' not found - its StaticallyLinkedModuleInitializers function is null."), *InModuleName.ToString());
-	OutFailureReason = EModuleLoadResult::FileNotFound;
+			// Monolithic builds that do not have the initializer were *not found* during the build step, so return FileNotFound
+			// (FileNotFound is an acceptable error in some case - ie loading a content only project)
+			UE_LOG(LogModuleManager, Warning, TEXT("ModuleManager: Module '%s' not found - its StaticallyLinkedModuleInitializers function is null."), *InModuleName.ToString());
+			OutFailureReason = EModuleLoadResult::FileNotFound;
 		}
 #endif
 #if !IS_MONOLITHIC
 		else
 		{
-	// Make sure that any UObjects that need to be registered were already processed before we go and
-	// load another module.  We just do this so that we can easily tell whether UObjects are present
-	// in the module being loaded.
-	if (bCanProcessNewlyLoadedObjects)
-	{
-		ProcessLoadedObjectsCallback.Broadcast();
-	}
+			// Make sure that any UObjects that need to be registered were already processed before we go and
+			// load another module.  We just do this so that we can easily tell whether UObjects are present
+			// in the module being loaded.
+			if (bCanProcessNewlyLoadedObjects)
+			{
+				ProcessLoadedObjectsCallback.Broadcast();
+			}
 
-	// Try to dynamically load the DLL
+			// Try to dynamically load the DLL
 
-	UE_LOG(LogModuleManager, Verbose, TEXT("ModuleManager: Load Module '%s' DLL '%s'"), *InModuleName.ToString(), *ModuleInfo->Filename);
+			UE_LOG(LogModuleManager, Verbose, TEXT("ModuleManager: Load Module '%s' DLL '%s'"), *InModuleName.ToString(), *ModuleInfo->Filename);
 
-	// Determine which file to load for this module.
-	const FString ModuleFileToLoad = FPaths::ConvertRelativePathToFull(ModuleInfo->Filename);
+			// Determine which file to load for this module.
+			const FString ModuleFileToLoad = FPaths::ConvertRelativePathToFull(ModuleInfo->Filename);
 
-	// Clear the handle and set it again below if the module is successfully loaded
-	ModuleInfo->Handle = NULL;
+			// Clear the handle and set it again below if the module is successfully loaded
+			ModuleInfo->Handle = NULL;
 
-	// Skip this check if file manager has not yet been initialized
+			// Skip this check if file manager has not yet been initialized
 			if (FPaths::FileExists(ModuleFileToLoad))
-	{
+			{
 				if (CheckModuleCompatibility(*ModuleFileToLoad))
-	{
-	ModuleInfo->Handle = FPlatformProcess::GetDllHandle(*ModuleFileToLoad);
+				{
+					ModuleInfo->Handle = FPlatformProcess::GetDllHandle(*ModuleFileToLoad);
 					if (ModuleInfo->Handle != NULL)
-	{
-	// First things first.  If the loaded DLL has UObjects in it, then their generated code's
-	// static initialization will have run during the DLL loading phase, and we'll need to
-	// go in and make sure those new UObject classes are properly registered.
-	{
-		// Sometimes modules are loaded before even the UObject systems are ready.  We need to assume
-		// these modules aren't using UObjects.
-		if (bCanProcessNewlyLoadedObjects)
-		{
-			// OK, we've verified that loading the module caused new UObject classes to be
-			// registered, so we'll treat this module as a module with UObjects in it.
-			ProcessLoadedObjectsCallback.Broadcast();
-		}
-	}
+					{
+						// First things first.  If the loaded DLL has UObjects in it, then their generated code's
+						// static initialization will have run during the DLL loading phase, and we'll need to
+						// go in and make sure those new UObject classes are properly registered.
+						{
+							// Sometimes modules are loaded before even the UObject systems are ready.  We need to assume
+							// these modules aren't using UObjects.
+							if (bCanProcessNewlyLoadedObjects)
+							{
+								// OK, we've verified that loading the module caused new UObject classes to be
+								// registered, so we'll treat this module as a module with UObjects in it.
+								ProcessLoadedObjectsCallback.Broadcast();
+							}
+						}
 
-	// Find our "InitializeModule" global function, which must exist for all module DLLs
+						// Find our "InitializeModule" global function, which must exist for all module DLLs
 						FInitializeModuleFunctionPtr InitializeModuleFunctionPtr =
 							(FInitializeModuleFunctionPtr)FPlatformProcess::GetDllExport(ModuleInfo->Handle, TEXT("InitializeModule"));
 						if (InitializeModuleFunctionPtr != NULL)
@@ -435,21 +441,21 @@ TSharedPtr<IModuleInterface> FModuleManager::LoadModuleWithFailureReason(const F
 								LoadedModule = ModuleInfo->Module;
 							}
 							else
-	{
+							{
 								UE_LOG(LogModuleManager, Warning, TEXT("ModuleManager: Unable to load module '%s' because InitializeModule function failed (returned NULL pointer.)"), *ModuleFileToLoad);
 
-		FPlatformProcess::FreeDllHandle(ModuleInfo->Handle);
-		ModuleInfo->Handle = NULL;
-		OutFailureReason = EModuleLoadResult::FailedToInitialize;
+								FPlatformProcess::FreeDllHandle(ModuleInfo->Handle);
+								ModuleInfo->Handle = NULL;
+								OutFailureReason = EModuleLoadResult::FailedToInitialize;
 							}
-	}
+						}
 						else
-	{
+						{
 							UE_LOG(LogModuleManager, Warning, TEXT("ModuleManager: Unable to load module '%s' because InitializeModule function was not found."), *ModuleFileToLoad);
 
-		FPlatformProcess::FreeDllHandle(ModuleInfo->Handle);
-		ModuleInfo->Handle = NULL;
-		OutFailureReason = EModuleLoadResult::FailedToInitialize;
+							FPlatformProcess::FreeDllHandle(ModuleInfo->Handle);
+							ModuleInfo->Handle = NULL;
+							OutFailureReason = EModuleLoadResult::FailedToInitialize;
 						}
 					}
 					else
@@ -469,7 +475,7 @@ TSharedPtr<IModuleInterface> FModuleManager::LoadModuleWithFailureReason(const F
 				UE_LOG(LogModuleManager, Warning, TEXT("ModuleManager: Unable to load module '%s' because the file '%s' was not found."), *InModuleName.ToString(), *ModuleFileToLoad);
 				OutFailureReason = EModuleLoadResult::FileNotFound;
 			}
-	}
+		}
 #endif
 	}
 
diff --git a/Engine/Source/Runtime/Core/Public/Modules/ModuleManager.h b/Engine/Source/Runtime/Core/Public/Modules/ModuleManager.h
--- a/Engine/Source/Runtime/Core/Public/Modules/ModuleManager.h
+++ b/Engine/Source/Runtime/Core/Public/Modules/ModuleManager.h
@@ -309,6 +309,14 @@ public:
 	 */
 	void FindModules( const TCHAR* WildcardWithoutExtension, TArray<FName>& OutModules );
 
+	/**
+	 * Determines if a module with the given name exists, regardless of whether it is currently loaded.
+	 *
+	 * @param ModuleName Name of the module to look for.
+	 * @return Whether the module exists.
+	 */
+	bool ModuleExists(const TCHAR* ModuleName);
+
 	/**
 	 * Gets the number of loaded modules.
 	 *
diff --git a/Engine/Source/Runtime/Launch/Private/LaunchEngineLoop.cpp b/Engine/Source/Runtime/Launch/Private/LaunchEngineLoop.cpp
--- a/Engine/Source/Runtime/Launch/Private/LaunchEngineLoop.cpp
+++ b/Engine/Source/Runtime/Launch/Private/LaunchEngineLoop.cpp
@@ -2831,7 +2831,10 @@ bool FEngineLoop::AppInit( )
 	// Because most, if not all online subsystems will load HTTP themselves. This can cause problems though, as HTTP will be loaded *after* OSS, 
 	// and if OSS holds on to resources allocated by it, this will cause crash (modules are being unloaded in LIFO order with no dependency tracking).
 	// Loading HTTP before OSS works around this problem by making ModuleManager unload HTTP after OSS, at the cost of extra module for the few OSS (like Null) that don't use it.
-	FModuleManager::Get().LoadModule(TEXT("XMPP"));
+	if(FModuleManager::Get().ModuleExists(TEXT("XMPP")))
+	{
+		FModuleManager::Get().LoadModule(TEXT("XMPP"));
+	}
 	FModuleManager::Get().LoadModule(TEXT("HTTP"));
 	FModuleManager::Get().LoadModule(TEXT("OnlineSubsystem"));
 	FModuleManager::Get().LoadModule(TEXT("OnlineSubsystemUtils"));
