Stats - Track allocations of size 0

diff --git a/Engine/Source/Runtime/Core/Private/Stats/StatsMallocProfilerProxy.cpp b/Engine/Source/Runtime/Core/Private/Stats/StatsMallocProfilerProxy.cpp
--- a/Engine/Source/Runtime/Core/Private/Stats/StatsMallocProfilerProxy.cpp
+++ b/Engine/Source/Runtime/Core/Private/Stats/StatsMallocProfilerProxy.cpp
@@ -100,38 +100,35 @@ void FStatsMallocProfilerProxy::TrackAlloc( void* Ptr, int64 Size, int32 Sequenc
 {
 	if( bEnabled )
 	{
-		if( Size > 0 )
-		{
-			FThreadStats* ThreadStats = FThreadStats::GetThreadStats();
+		FThreadStats* ThreadStats = FThreadStats::GetThreadStats();
 
 #ifdef DEBUG_MALLOC_PROXY
-			if( GThreadStatsToDumpMemory == ThreadStats && ThreadStats->MemoryMessageScope == 0 )
+		if (GThreadStatsToDumpMemory == ThreadStats && ThreadStats->MemoryMessageScope == 0)
+		{
+			ThreadStats->MemoryMessageScope++;
+			UE_LOG( LogStats, Warning, TEXT( "TrackAlloc, %llu, %lli, %i, %i" ), (uint64)(UPTRINT)Ptr, Size, SequenceTag, ThreadStats->MemoryMessageScope );
+			ThreadStats->MemoryMessageScope--;
+		}
+#endif // DEBUG_MALLOC_PROXY
+
+		if (ThreadStats->MemoryMessageScope == 0)
+		{
+#if	UE_BUILD_DEBUG
+			if (ThreadStats->Packet.StatMessages.Num() > 0 && ThreadStats->Packet.StatMessages.Num() % 32767 == 0)
 			{
 				ThreadStats->MemoryMessageScope++;
-				UE_LOG( LogStats, Warning, TEXT( "TrackAlloc, %llu, %lli, %i, %i" ), (uint64)(UPTRINT)Ptr, Size, SequenceTag, ThreadStats->MemoryMessageScope );
+				const double InvMB = 1.0f / 1024.0f / 1024.0f;
+				UE_LOG( LogStats, Verbose, TEXT( "ThreadID: %i, Current: %.1f" ), FPlatformTLS::GetCurrentThreadId(), InvMB*(int64)ThreadStats->Packet.StatMessages.Num()*sizeof( FStatMessage ) );
 				ThreadStats->MemoryMessageScope--;
 			}
-#endif // DEBUG_MALLOC_PROXY
-		
-			if( ThreadStats->MemoryMessageScope == 0 )
-			{
-#if	UE_BUILD_DEBUG
-				if( ThreadStats->Packet.StatMessages.Num() > 0 && ThreadStats->Packet.StatMessages.Num() % 32767 == 0 )
-				{
-					ThreadStats->MemoryMessageScope++;
-					const double InvMB = 1.0f / 1024.0f / 1024.0f;
-					UE_LOG( LogStats, Verbose, TEXT( "ThreadID: %i, Current: %.1f" ), FPlatformTLS::GetCurrentThreadId(), InvMB*(int64)ThreadStats->Packet.StatMessages.Num()*sizeof( FStatMessage ) );
-					ThreadStats->MemoryMessageScope--;
-				}
 #endif // UE_BUILD_DEBUG
 
-				// 48 bytes per allocation.
-				ThreadStats->AddMemoryMessage( GET_STATFNAME(STAT_Memory_AllocPtr), (uint64)(UPTRINT)Ptr | (uint64)EMemoryOperation::Alloc );
-				ThreadStats->AddMemoryMessage( GET_STATFNAME(STAT_Memory_AllocSize), Size );			
-				ThreadStats->AddMemoryMessage( GET_STATFNAME(STAT_Memory_OperationSequenceTag), (int64)SequenceTag );			
-				AllocPtrCalls.Increment();
-			}
-		}	
+			// 48 bytes per allocation.
+			ThreadStats->AddMemoryMessage( GET_STATFNAME( STAT_Memory_AllocPtr ), (uint64)(UPTRINT)Ptr | (uint64)EMemoryOperation::Alloc );
+			ThreadStats->AddMemoryMessage( GET_STATFNAME( STAT_Memory_AllocSize ), Size );
+			ThreadStats->AddMemoryMessage( GET_STATFNAME( STAT_Memory_OperationSequenceTag ), (int64)SequenceTag );
+			AllocPtrCalls.Increment();
+		}
 	}
 }
 
