An empty text filter should match anything

diff --git a/Engine/Source/Runtime/Core/Private/Misc/TextFilterExpressionEvaluator.cpp b/Engine/Source/Runtime/Core/Private/Misc/TextFilterExpressionEvaluator.cpp
--- a/Engine/Source/Runtime/Core/Private/Misc/TextFilterExpressionEvaluator.cpp
+++ b/Engine/Source/Runtime/Core/Private/Misc/TextFilterExpressionEvaluator.cpp
@@ -537,6 +537,11 @@ bool FTextFilterExpressionEvaluator::TestTextFilter(const ITextFilterExpressionC
 	using namespace TextFilterExpressionParser;
 	using TextFilterExpressionParser::FTextToken;
 
+	if (FilterType == ETextFilterExpressionType::Empty)
+	{
+		return true;
+	}
+
 	if (CompiledFilter.IsSet())
 	{
 		auto& CompiledResult = CompiledFilter.GetValue();
