Allow TDoubleLinkedList to have nodes reordered (or moved to another list) without deleting and reallocating the node memory
Add a Contains helper method to TDoubleLinkedList

diff --git a/Engine/Source/Runtime/Core/Public/Containers/List.h b/Engine/Source/Runtime/Core/Public/Containers/List.h
--- a/Engine/Source/Runtime/Core/Public/Containers/List.h
+++ b/Engine/Source/Runtime/Core/Public/Containers/List.h
@@ -325,8 +325,12 @@ public:
 	 */
 	bool AddHead( const ElementType& InElement )
 	{
-		TDoubleLinkedListNode* NewNode = new TDoubleLinkedListNode(InElement);
-		if ( NewNode == nullptr )
+		return AddHead(new TDoubleLinkedListNode(InElement));
+	}
+
+	bool AddHead( TDoubleLinkedListNode* NewNode )
+	{
+		if (NewNode == nullptr)
 		{
 			return false;
 		}
@@ -356,7 +360,11 @@ public:
 	 */
 	bool AddTail( const ElementType& InElement )
 	{
-		TDoubleLinkedListNode* NewNode = new TDoubleLinkedListNode(InElement);
+		return AddTail(new TDoubleLinkedListNode(InElement));
+	}
+
+	bool AddTail( TDoubleLinkedListNode* NewNode )
+	{
 		if ( NewNode == nullptr )
 		{
 			return false;
@@ -388,17 +396,21 @@ public:
 	 */
 	bool InsertNode( const ElementType& InElement, TDoubleLinkedListNode* NodeToInsertBefore=nullptr )
 	{
-		if ( NodeToInsertBefore == nullptr || NodeToInsertBefore == HeadNode )
-		{
-			return AddHead(InElement);
-		}
+		return InsertNode(new TDoubleLinkedListNode(InElement), NodeToInsertBefore);
+	}
 
-		TDoubleLinkedListNode* NewNode = new TDoubleLinkedListNode(InElement);
+	bool InsertNode( TDoubleLinkedListNode* NewNode, TDoubleLinkedListNode* NodeToInsertBefore=nullptr )
+	{
 		if ( NewNode == nullptr )
 		{
 			return false;
 		}
 
+		if ( NodeToInsertBefore == nullptr || NodeToInsertBefore == HeadNode )
+		{
+			return AddHead(NewNode);
+		}
+
 		NewNode->PrevNode = NodeToInsertBefore->PrevNode;
 		NewNode->NextNode = NodeToInsertBefore;
 
@@ -427,7 +439,7 @@ public:
 	 * @param NodeToRemove The node to remove.
 	 * @see Empty, InsertNode
 	 */
-	void RemoveNode( TDoubleLinkedListNode* NodeToRemove )
+	void RemoveNode( TDoubleLinkedListNode* NodeToRemove, bool bDeleteNode = true )
 	{
 		if ( NodeToRemove != nullptr )
 		{
@@ -435,7 +447,16 @@ public:
 			if ( Num() == 1 )
 			{
 				checkSlow(NodeToRemove == HeadNode);
-				Empty();
+				if (bDeleteNode)
+				{
+					Empty();
+				}
+				else
+				{
+					NodeToRemove->NextNode = NodeToRemove->PrevNode = nullptr;
+					HeadNode = TailNode = nullptr;
+					SetListSize(0);
+				}
 				return;
 			}
 
@@ -456,8 +477,14 @@ public:
 				NodeToRemove->PrevNode->NextNode = NodeToRemove->NextNode;
 			}
 
-			delete NodeToRemove;
-			NodeToRemove = nullptr;
+			if (bDeleteNode)
+			{
+				delete NodeToRemove;
+			}
+			else
+			{
+				NodeToRemove->NextNode = NodeToRemove->PrevNode = nullptr;
+			}
 			SetListSize(ListSize - 1);
 		}
 	}
@@ -465,12 +492,12 @@ public:
 	/** Removes all nodes from the list. */
 	void Empty()
 	{
-		TDoubleLinkedListNode* pNode;
+		TDoubleLinkedListNode* Node;
 		while ( HeadNode != nullptr )
 		{
-			pNode = HeadNode->NextNode;
+			Node = HeadNode->NextNode;
 			delete HeadNode;
-			HeadNode = pNode;
+			HeadNode = Node;
 		}
 
 		HeadNode = TailNode = nullptr;
@@ -509,18 +536,23 @@ public:
 	 */
 	TDoubleLinkedListNode* FindNode( const ElementType& InElement )
 	{
-		TDoubleLinkedListNode* pNode = HeadNode;
-		while ( pNode != nullptr )
+		TDoubleLinkedListNode* Node = HeadNode;
+		while ( Node != nullptr )
 		{
-			if ( pNode->GetValue() == InElement )
+			if ( Node->GetValue() == InElement )
 			{
 				break;
 			}
 
-			pNode = pNode->NextNode;
+			Node = Node->NextNode;
 		}
 
-		return pNode;
+		return Node;
+	}
+
+	bool Contains( const ElementType& InElement )
+	{
+		return (FindNode(InElement) != nullptr);
 	}
 
 	/**
