Fix RHI thread support in Shipping builds.

diff --git a/Engine/Source/Runtime/Core/Public/HAL/Platform.h b/Engine/Source/Runtime/Core/Public/HAL/Platform.h
--- a/Engine/Source/Runtime/Core/Public/HAL/Platform.h
+++ b/Engine/Source/Runtime/Core/Public/HAL/Platform.h
@@ -279,6 +279,14 @@
 	#define PLATFORM_USES_ANSI_STRING_FOR_EXTERNAL_PROFILING 1
 #endif
 
+#ifndef PLATFORM_RHITHREAD_DEFAULT_BYPASS
+	#define PLATFORM_RHITHREAD_DEFAULT_BYPASS					1
+#endif
+
+#ifndef PLATFORM_CAN_TOGGLE_RHITHREAD_IN_SHIPPING
+	#define PLATFORM_CAN_TOGGLE_RHITHREAD_IN_SHIPPING			0
+#endif
+
 // deprecated, do not use
 #define PLATFORM_HAS_THREADSAFE_RHIGetRenderQueryResult	#
 #define PLATFORM_SUPPORTS_RHI_THREAD #
diff --git a/Engine/Source/Runtime/Core/Public/Windows/WIndowsPlatform.h b/Engine/Source/Runtime/Core/Public/Windows/WIndowsPlatform.h
--- a/Engine/Source/Runtime/Core/Public/Windows/WIndowsPlatform.h
+++ b/Engine/Source/Runtime/Core/Public/Windows/WIndowsPlatform.h
@@ -70,6 +70,10 @@ typedef FWindowsPlatformTypes FPlatformTypes;
 #endif
 #define PLATFORM_COMPILER_HAS_TCHAR_WMAIN					1
 
+#define PLATFORM_RHITHREAD_DEFAULT_BYPASS					1
+#define PLATFORM_CAN_TOGGLE_RHITHREAD_IN_SHIPPING			1
+
+
 // Intrinsics for 128-bit atomics on Windows platform requires Windows 8 or higher (WINVER>0x0602)
 // http://msdn.microsoft.com/en-us/library/windows/desktop/hh972640.aspx
 #define PLATFORM_HAS_128BIT_ATOMICS							(!HACK_HEADER_GENERATOR && PLATFORM_64BITS && (WINVER >= 0x602))
diff --git a/Engine/Source/Runtime/RHI/Public/RHICommandList.h b/Engine/Source/Runtime/RHI/Public/RHICommandList.h
--- a/Engine/Source/Runtime/RHI/Public/RHICommandList.h
+++ b/Engine/Source/Runtime/RHI/Public/RHICommandList.h
@@ -2601,7 +2601,7 @@ class RHI_API FRHICommandListExecutor
 public:
 	enum
 	{
-		DefaultBypass = 1
+		DefaultBypass = PLATFORM_RHITHREAD_DEFAULT_BYPASS
 	};
 	FRHICommandListExecutor()
 		: bLatchedBypass(!!DefaultBypass)
@@ -2618,7 +2618,7 @@ public:
 
 	FORCEINLINE_DEBUGGABLE bool Bypass()
 	{
-#if !UE_BUILD_SHIPPING
+#if !UE_BUILD_SHIPPING || PLATFORM_CAN_TOGGLE_RHITHREAD_IN_SHIPPING
 		return bLatchedBypass;
 #else
 		return !!DefaultBypass;
