Cleared bits after a Reset.

diff --git a/Engine/Source/Runtime/Core/Public/Containers/BitArray.h b/Engine/Source/Runtime/Core/Public/Containers/BitArray.h
--- a/Engine/Source/Runtime/Core/Public/Containers/BitArray.h
+++ b/Engine/Source/Runtime/Core/Public/Containers/BitArray.h
@@ -328,6 +328,9 @@ public:
 	 */
 	void Reset()
 	{
+		// We need this because iterators often use whole DWORDs when masking, which includes off-the-end elements
+		FMemory::Memset(GetData(), 0, FMath::DivideAndRoundUp(NumBits, NumBitsPerDWORD) * sizeof(uint32));
+
 		NumBits = 0;
 	}
 
