Improving editor startup times by a few (3-4) seconds (that's 20%-30% in tested cases).

diff --git a/Engine/Source/Developer/HotReload/Private/HotReload.cpp b/Engine/Source/Developer/HotReload/Private/HotReload.cpp
--- a/Engine/Source/Developer/HotReload/Private/HotReload.cpp
+++ b/Engine/Source/Developer/HotReload/Private/HotReload.cpp
@@ -651,6 +651,8 @@ ECompilationResult::Type FHotReloadModule::DoHotReloadInternal(bool bRecompileFi
 
 	if (CompilationResult == ECompilationResult::Succeeded)
 	{
+		FModuleManager::Get().ResetModulePathsCache();
+
 		FFeedbackContext& ErrorsFC = UClass::GetDefaultPropertiesFeedbackContext();
 		ErrorsFC.Errors.Empty();
 		ErrorsFC.Warnings.Empty();
@@ -820,6 +822,8 @@ ECompilationResult::Type FHotReloadModule::RebindPackagesInternal(TArray<UPackag
 
 	if (bCanRebind)
 	{
+		FModuleManager::Get().ResetModulePathsCache();
+
 		bIsHotReloadingFromEditor = true;
 
 		const double StartTime = FPlatformTime::Seconds();
diff --git a/Engine/Source/Runtime/Core/Private/Modules/ModuleManager.cpp b/Engine/Source/Runtime/Core/Private/Modules/ModuleManager.cpp
--- a/Engine/Source/Runtime/Core/Private/Modules/ModuleManager.cpp
+++ b/Engine/Source/Runtime/Core/Private/Modules/ModuleManager.cpp
@@ -128,12 +128,12 @@ bool FModuleManager::IsModuleLoaded( const FName InModuleName )
 }
 
 
-bool FModuleManager::IsModuleUpToDate( const FName InModuleName ) const
+bool FModuleManager::IsModuleUpToDate(const FName InModuleName)
 {
 	TMap<FName, FString> ModulePathMap;
 	FindModulePaths(*InModuleName.ToString(), ModulePathMap);
 
-	if(ModulePathMap.Num() != 1)
+	if (ModulePathMap.Num() != 1)
 	{
 		return false;
 	}
@@ -886,9 +886,31 @@ void FModuleManager::GetModuleFilenameFormat(bool bGameModule, FString& OutPrefi
 	OutSuffix += FPlatformProcess::GetModuleExtension();
 }
 
+void FModuleManager::ResetModulePathsCache()
+{
+	ModulePathsCache.Empty();
+	bModulePathsCacheInitialized = false;
+}
 
-void FModuleManager::FindModulePaths(const TCHAR* NamePattern, TMap<FName, FString> &OutModulePaths) const
+void FModuleManager::FindModulePaths(const TCHAR* NamePattern, TMap<FName, FString> &OutModulePaths, bool bCanUseCache /*= true*/)
 {
+	if (!bModulePathsCacheInitialized)
+	{
+		bModulePathsCacheInitialized = true;
+		const bool bCanUseCacheWhileGeneratingIt = false;
+		FindModulePaths(TEXT("*"), ModulePathsCache, bCanUseCacheWhileGeneratingIt);
+	}
+
+	if (bCanUseCache)
+	{
+		// Try to use cache first
+		if (const FString* ModulePathPtr = ModulePathsCache.Find(NamePattern))
+		{
+			OutModulePaths.Add(FName(NamePattern), *ModulePathPtr);
+			return;
+		}
+	}
+
 	// Search through the engine directory
 	FindModulePathsInDirectory(FPlatformProcess::GetModulesDirectory(), false, NamePattern, OutModulePaths);
 
diff --git a/Engine/Source/Runtime/Core/Public/Modules/ModuleManager.h b/Engine/Source/Runtime/Core/Public/Modules/ModuleManager.h
--- a/Engine/Source/Runtime/Core/Public/Modules/ModuleManager.h
+++ b/Engine/Source/Runtime/Core/Public/Modules/ModuleManager.h
@@ -373,7 +373,7 @@ public:
 	 * @param InModuleName The base name of the module file.
 	 * @return true if module exists and is up to date, false otherwise.
 	 */
-	bool IsModuleUpToDate( const FName InModuleName ) const;
+	bool IsModuleUpToDate( const FName InModuleName );
 
 	/**
 	 * Determines whether the specified module contains UObjects.  The module must already be loaded into
@@ -454,6 +454,7 @@ protected:
 	 */
 	FModuleManager( )
 		: bCanProcessNewlyLoadedObjects(false)
+		, bModulePathsCacheInitialized(false)
 	{ }
 
 protected:
@@ -507,6 +508,9 @@ public:
 
 	void AddModuleToModulesList(const FName InModuleName, TSharedRef<FModuleInfo>& ModuleInfo);
 
+	/** Clears module path cache */
+	void ResetModulePathsCache();
+
 private:
 	/** Thread safe module finding routine. */
 	TSharedRef<FModuleInfo>* FindModule(FName InModuleName);
@@ -519,7 +523,7 @@ private:
 	static void GetModuleFilenameFormat(bool bGameModule, FString& OutPrefix, FString& OutSuffix);
 
 	/** Finds modules matching a given name wildcard. */
-	void FindModulePaths(const TCHAR *NamePattern, TMap<FName, FString> &OutModulePaths) const;
+	void FindModulePaths(const TCHAR *NamePattern, TMap<FName, FString> &OutModulePaths, bool bCanUseCache = true);
 
 	/** Finds modules matching a given name wildcard within a given directory. */
 	void FindModulePathsInDirectory(const FString &DirectoryName, bool bIsGameDirectory, const TCHAR *NamePattern, TMap<FName, FString> &OutModulePaths) const;
@@ -538,6 +542,12 @@ private:
 	/** True if module manager should automatically register new UObjects discovered while loading C++ modules */
 	bool bCanProcessNewlyLoadedObjects;
 
+	/** True if module paths cache has already been initialized */
+	bool bModulePathsCacheInitialized;
+
+	/** Cache of known module paths. Used for performance. Can increase editor startup times by up to 30% */
+	TMap<FName, FString> ModulePathsCache;
+
 	/** Multicast delegate that will broadcast a notification when modules are loaded, unloaded, or
 	    our set of known modules changes */
 	FModulesChangedEvent ModulesChangedEvent;
