PR #1309: [Core][Windows] Don't crash when GetRawInputData() fails. (Contributed by [REDACTED_PERSON])

diff --git a/Engine/Source/Runtime/Core/Private/Windows/WindowsApplication.cpp b/Engine/Source/Runtime/Core/Private/Windows/WindowsApplication.cpp
--- a/Engine/Source/Runtime/Core/Private/Windows/WindowsApplication.cpp
+++ b/Engine/Source/Runtime/Core/Private/Windows/WindowsApplication.cpp
@@ -769,36 +769,33 @@ int32 FWindowsApplication::ProcessMessage( HWND hwnd, uint32 msg, WPARAM wParam,
 		case WM_INPUT:
 			{
 				uint32 Size = 0;
-				TArray<uint8> RawData;
-
 				::GetRawInputData((HRAWINPUT)lParam, RID_INPUT, NULL, &Size, sizeof(RAWINPUTHEADER));
-				RawData.AddZeroed( Size );
 
-				RAWINPUT* Raw = NULL;
-				if (::GetRawInputData((HRAWINPUT)lParam, RID_INPUT, RawData.GetData(), &Size, sizeof(RAWINPUTHEADER)) == Size )
-				{
-					Raw = (RAWINPUT*)RawData.GetData();
-				}
+				TScopedPointer<uint8> RawData(new uint8[Size]);
 
-				if (Raw->header.dwType == RIM_TYPEMOUSE) 
+				if (::GetRawInputData((HRAWINPUT)lParam, RID_INPUT, RawData.GetOwnedPointer(), &Size, sizeof(RAWINPUTHEADER)) == Size )
 				{
-					const bool IsAbsoluteInput = (Raw->data.mouse.usFlags & MOUSE_MOVE_ABSOLUTE) == MOUSE_MOVE_ABSOLUTE;
-					if( IsAbsoluteInput )
+					const RAWINPUT* const Raw = (const RAWINPUT* const)RawData.GetOwnedPointer();
+
+					if (Raw->header.dwType == RIM_TYPEMOUSE) 
 					{
-						// Since the raw input is coming in as absolute it is likely the user is using a tablet
-						// or perhaps is interacting through a virtual desktop
-						DeferMessage( CurrentNativeEventWindowPtr, hwnd, msg, wParam, lParam, 0, 0, MOUSE_MOVE_ABSOLUTE );
-						return 1;
-					}
+						const bool IsAbsoluteInput = (Raw->data.mouse.usFlags & MOUSE_MOVE_ABSOLUTE) == MOUSE_MOVE_ABSOLUTE;
+						if( IsAbsoluteInput )
+						{
+							// Since the raw input is coming in as absolute it is likely the user is using a tablet
+							// or perhaps is interacting through a virtual desktop
+							DeferMessage( CurrentNativeEventWindowPtr, hwnd, msg, wParam, lParam, 0, 0, MOUSE_MOVE_ABSOLUTE );
+							return 1;
+						}
 
-					// Since raw input is coming in as relative it is likely a traditional mouse device
-					const int xPosRelative = Raw->data.mouse.lLastX;
-					const int yPosRelative = Raw->data.mouse.lLastY;
+						// Since raw input is coming in as relative it is likely a traditional mouse device
+						const int xPosRelative = Raw->data.mouse.lLastX;
+						const int yPosRelative = Raw->data.mouse.lLastY;
 
-					DeferMessage( CurrentNativeEventWindowPtr, hwnd, msg, wParam, lParam, xPosRelative, yPosRelative, MOUSE_MOVE_RELATIVE );
-					return 1;
-				} 
-				
+						DeferMessage( CurrentNativeEventWindowPtr, hwnd, msg, wParam, lParam, xPosRelative, yPosRelative, MOUSE_MOVE_RELATIVE );
+						return 1;
+					} 
+				}
 			}
 			break;
 
