Stats - Added DebugLeakTest method for leaking memory of purpose, disabled by default

diff --git a/Engine/Source/Runtime/Core/Private/Stats/Stats2.cpp b/Engine/Source/Runtime/Core/Private/Stats/Stats2.cpp
--- a/Engine/Source/Runtime/Core/Private/Stats/Stats2.cpp
+++ b/Engine/Source/Runtime/Core/Private/Stats/Stats2.cpp
@@ -1,6 +1,7 @@
 // Copyright 1998-2015 Epic Games, Inc. All Rights Reserved.
 
 #include "CorePrivatePCH.h"
+#include "TaskGraphInterfaces.h"
 
 /*-----------------------------------------------------------------------------
 	Global
@@ -35,6 +36,112 @@ DECLARE_MEMORY_STAT( TEXT("Stats Descriptions"), STAT_StatDescMemory, STATGROUP_
 DEFINE_STAT(STAT_FrameTime);
 DEFINE_STAT(STAT_NamedMarker);
 
+/*-----------------------------------------------------------------------------
+	DebugLeakTest, for the stats based memory profiler
+-----------------------------------------------------------------------------*/
+
+#if	!UE_BUILD_SHIPPING
+
+static TAutoConsoleVariable<int32> CVarEnableLeakTest(
+	TEXT( "debug.EnableLeakTest" ),
+	0,
+	TEXT( "If set to 1, enables leak test, for testing stats based memory profiler" )
+	);
+
+void DebugLeakTest()
+{
+	if (CVarEnableLeakTest.GetValueOnGameThread() == 1)
+	{
+		if (GFrameCounter == 60)
+		{
+			DirectStatsCommand( TEXT( "stat namedmarker Frame-060" ), true );
+		}
+
+		if (GFrameCounter == 120)
+		{
+			DirectStatsCommand( TEXT( "stat namedmarker Frame-120" ), true );
+		}
+
+
+		if (GFrameCounter == 240)
+		{
+			DirectStatsCommand( TEXT( "stat namedmarker Frame-240" ), true );
+		}
+
+		if (GFrameCounter == 300)
+		{
+			GIsRequestingExit = true;
+		}
+
+		// Realloc.
+		static TArray<uint8> Array;
+		static int32 Initial = 1;
+		{
+			DECLARE_SCOPE_CYCLE_COUNTER( TEXT( "LeakTest::Realloc" ), Stat_LeakTest_Realloc, STATGROUP_Quick );
+			Array.AddZeroed( Initial );
+			Initial += 100;
+		}
+
+		if (GFrameCounter == 300)
+		{
+			UE_LOG( LogTemp, Warning, TEXT( "Stat_ReallocTest: %i / %i" ), Array.GetAllocatedSize(), Initial );
+		}
+
+		// General memory leak.
+		{
+			DECLARE_SCOPE_CYCLE_COUNTER( TEXT( "LeakTest::NewInt8" ), Stat_LeakTest_NewInt8, STATGROUP_Quick );
+			int8* Leak = new int8[1000 * 1000];
+		}
+
+
+		if (GFrameCounter <= 250)
+		{
+			// Background threads memory test.
+			struct FAllocTask
+			{
+				static void Alloc()
+				{
+					DECLARE_SCOPE_CYCLE_COUNTER( TEXT( "FAllocTask::Alloc" ), Stat_FAllocTask_Alloc, STATGROUP_Quick );
+
+					int* IntAlloc = new int[1000];
+					int8* LeakTask = new int8[100000];
+					delete[] IntAlloc;
+				}
+			};
+
+			for (int32 Index = 0; Index < 40; ++Index)
+			{
+				FSimpleDelegateGraphTask::CreateAndDispatchWhenReady( FSimpleDelegateGraphTask::FDelegate::CreateStatic( FAllocTask::Alloc ), TStatId() );
+			}
+
+			class FAllocPool : public FNonAbandonableTask
+			{
+			public:
+				void DoWork()
+				{
+					DECLARE_SCOPE_CYCLE_COUNTER( TEXT( "FAllocPool::DoWork" ), Stat_FAllocPool_DoWork, STATGROUP_Quick );
+
+					int* IntAlloc = new int[1000];
+					int8* LeakTask = new int8[100000];
+					delete[] IntAlloc;
+				}
+
+				TStatId GetStatId() const
+				{
+					return TStatId();
+				}
+			};
+
+			for (int32 Index = 0; Index < 40; ++Index)
+			{
+				(new FAutoDeleteAsyncTask<FAllocPool>())->StartBackgroundTask();
+			}
+		}
+	}
+}
+
+#endif // !UE_BUILD_SHIPPING
+
 /*-----------------------------------------------------------------------------
 	FStats2
 -----------------------------------------------------------------------------*/
@@ -1166,6 +1273,10 @@ void FThreadStats::WaitForStats()
 
 		LastFramesEvents[(CurrentEventIndex + MAX_STAT_LAG - 1) % MAX_STAT_LAG] = TGraphTask<FNullGraphTask>::CreateTask(NULL, ENamedThreads::GameThread).ConstructAndDispatchWhenReady(GET_STATID(STAT_FNullGraphTask_StatWaitFence), FPlatformProcess::SupportsMultithreading() ? ENamedThreads::StatsThread : ENamedThreads::GameThread);
 		CurrentEventIndex++;
+
+#if	!UE_BUILD_SHIPPING
+		DebugLeakTest();
+#endif // !UE_BUILD_SHIPPING
 	}
 }
 
