CIS fix

diff --git a/Engine/Source/Runtime/Core/Private/Misc/TextFilterUtils.cpp b/Engine/Source/Runtime/Core/Private/Misc/TextFilterUtils.cpp
--- a/Engine/Source/Runtime/Core/Private/Misc/TextFilterUtils.cpp
+++ b/Engine/Source/Runtime/Core/Private/Misc/TextFilterUtils.cpp
@@ -104,7 +104,7 @@ namespace FastToUpper
 		0,		// Y
 		0,		// Z
 		0,		// [
-		0,		// \ 
+		0,		// Backslash
 		0,		// ]
 		0,		// ^
 		0,		// _
