Call PreUnloadCallback before Shutdown for Modules

(modified PR #1243 Contributed by [REDACTED_PERSON])

diff --git a/Engine/Source/Runtime/Core/Private/Modules/ModuleManager.cpp b/Engine/Source/Runtime/Core/Private/Modules/ModuleManager.cpp
--- a/Engine/Source/Runtime/Core/Private/Modules/ModuleManager.cpp
+++ b/Engine/Source/Runtime/Core/Private/Modules/ModuleManager.cpp
@@ -580,9 +580,11 @@ void FModuleManager::UnloadModulesAtShutdown()
 	{
 		FName ModuleName;
 		int32 LoadOrder;
-		FModulePair(FName InModuleName, int32 InLoadOrder)
+		TSharedPtr<IModuleInterface> Module;
+		FModulePair(FName InModuleName, int32 InLoadOrder, TSharedPtr<IModuleInterface> InModule)
 			: ModuleName(InModuleName)
 			, LoadOrder(InLoadOrder)
+			, Module(InModule)
 		{
 			check(LoadOrder > 0); // else was never initialized
 		}
@@ -602,18 +604,25 @@ void FModuleManager::UnloadModulesAtShutdown()
 			// Only if the module supports shutting down in this phase
 			if( ModuleInfo->Module->SupportsAutomaticShutdown() )
 			{
-				new (ModulesToUnload) FModulePair(ModuleIt.Key, ModuleIt.Value->LoadOrder);
+				new (ModulesToUnload)FModulePair(ModuleIt.Key, ModuleIt.Value->LoadOrder, ModuleInfo->Module);
 			}
 		}
 	}
 
 	ModulesToUnload.Sort();
-
-	for (int32 Index = 0; Index < ModulesToUnload.Num(); Index++)
+	
+	// Call PreUnloadCallback on all modules first
+	for (FModulePair& ModuleToUnload : ModulesToUnload)
+	{
+		ModuleToUnload.Module->PreUnloadCallback();
+		ModuleToUnload.Module = nullptr;
+	}
+	// Now actually unload all modules
+	for (FModulePair& ModuleToUnload : ModulesToUnload)
 	{
-		UE_LOG(LogModuleManager, Log, TEXT( "Shutting down and abandoning module %s (%d)" ), *ModulesToUnload[Index].ModuleName.ToString(), ModulesToUnload[Index].LoadOrder );
+		UE_LOG(LogModuleManager, Log, TEXT("Shutting down and abandoning module %s (%d)"), *ModuleToUnload.ModuleName.ToString(), ModuleToUnload.LoadOrder);
 		const bool bIsShutdown = true;
-		UnloadModule( ModulesToUnload[Index].ModuleName, bIsShutdown );
+		UnloadModule(ModuleToUnload.ModuleName, bIsShutdown);
 		UE_LOG(LogModuleManager, Verbose, TEXT( "Returned from UnloadModule." ));
 	}
 }
