Stats - Fixed an issue with calculating non-frame stats on pooled threads

diff --git a/Engine/Source/Runtime/Core/Private/HAL/ThreadingBase.cpp b/Engine/Source/Runtime/Core/Private/HAL/ThreadingBase.cpp
--- a/Engine/Source/Runtime/Core/Private/HAL/ThreadingBase.cpp
+++ b/Engine/Source/Runtime/Core/Private/HAL/ThreadingBase.cpp
@@ -9,6 +9,7 @@
 DEFINE_STAT( STAT_EventWaitWithId );
 DEFINE_STAT( STAT_EventTriggerWithId );
 
+DECLARE_DWORD_COUNTER_STAT( TEXT( "ThreadPoolDummyCounter" ), STAT_ThreadPoolDummyCounter, STATGROUP_ThreadPoolAsyncTasks );
 
 /** The global thread pool */
 FQueuedThreadPool* GThreadPool = nullptr;
@@ -376,11 +377,17 @@ protected:
 	{
 		while (!TimeToDie)
 		{
+			// This will force sending the stats packet from the previous frame.
+			SET_DWORD_STAT( STAT_ThreadPoolDummyCounter, 0 );
+			// We need to wait for shorter amount of time
+			bool bContinueWaiting = true;
+			while( bContinueWaiting )
 			{				
 				DECLARE_SCOPE_CYCLE_COUNTER( TEXT( "FQueuedThread::Run.WaitForWork" ), STAT_FQueuedThread_Run_WaitForWork, STATGROUP_ThreadPoolAsyncTasks );
 				// Wait for some work to do
-				DoWorkEvent->Wait();
+				bContinueWaiting = !DoWorkEvent->Wait( 10 );
 			}
+
 			IQueuedWork* LocalQueuedWork = QueuedWork;
 			QueuedWork = nullptr;
 			FPlatformMisc::MemoryBarrier();
@@ -614,8 +621,8 @@ public:
 		FScopeLock sl(SynchQueue);
 		if (QueuedThreads.Num() > 0)
 		{
-			// Figure out which thread is available
-			int32 Index = QueuedThreads.Num() - 1;
+			// Cycle through all available threads to make sure that stats are up to date.
+			int32 Index = 0;
 			// Grab that thread to use
 			Thread = QueuedThreads[Index];
 			// Remove it from the list so no one else grabs it
