Device Clipboard Support for Crashes on IOS

diff --git a/Engine/Source/Runtime/Core/Private/IOS/IOSPlatformOutputDevices.cpp b/Engine/Source/Runtime/Core/Private/IOS/IOSPlatformOutputDevices.cpp
--- a/Engine/Source/Runtime/Core/Private/IOS/IOSPlatformOutputDevices.cpp
+++ b/Engine/Source/Runtime/Core/Private/IOS/IOSPlatformOutputDevices.cpp
@@ -81,6 +81,9 @@ void FOutputDeviceIOSError::HandleError()
 	GIsRunning = 0;
 	GIsCriticalError = 1;
 	GLogConsole = NULL;
+    
+    GErrorHist[ARRAY_COUNT(GErrorHist)-1] = 0;
+    FPlatformMisc::ClipboardCopy(GErrorHist);
 
 	GLog->Flush();
 }
