PS4 - Crash Context fixes for PS4
* PS4 specific memory allocator
* Cache some properties at startup as we can't do it inside the crash handler

diff --git a/Engine/Source/Runtime/Core/Private/Android/AndroidMisc.cpp b/Engine/Source/Runtime/Core/Private/Android/AndroidMisc.cpp
--- a/Engine/Source/Runtime/Core/Private/Android/AndroidMisc.cpp
+++ b/Engine/Source/Runtime/Core/Private/Android/AndroidMisc.cpp
@@ -9,7 +9,7 @@
 #include <string.h>
 
 #include "AndroidPlatformCrashContext.h"
-#include "MallocCrash.h"
+#include "PlatformMallocCrash.h"
 #include "AndroidJavaMessageBox.h"
 
 DECLARE_LOG_CATEGORY_EXTERN(LogEngine, Log, All);
@@ -380,7 +380,7 @@ void (* GCrashHandlerPointer)(const FGenericCrashContext& Context) = NULL;
 void PlatformCrashHandler(int32 Signal, siginfo* Info, void* Context)
 {
 	// Switch to malloc crash.
-	//FMallocCrash::Get().SetAsGMalloc(); @todo uncomment after verification
+	//FGenericPlatformMallocCrash::Get().SetAsGMalloc(); @todo uncomment after verification
 
 	fprintf(stderr, "Signal %d caught.\n", Signal);
 
diff --git a/Engine/Source/Runtime/Core/Private/GenericPlatform/GenericPlatformContext.cpp b/Engine/Source/Runtime/Core/Private/GenericPlatform/GenericPlatformContext.cpp
--- a/Engine/Source/Runtime/Core/Private/GenericPlatform/GenericPlatformContext.cpp
+++ b/Engine/Source/Runtime/Core/Private/GenericPlatform/GenericPlatformContext.cpp
@@ -12,6 +12,10 @@ bool FGenericCrashContext::bIsInitialized = false;
 
 namespace NCachedCrashContextProperties
 {
+	static bool bIsInternalBuild;
+	static bool bIsPerforceBuild;
+	static bool bIsSourceDistribution;
+	static bool bIsUE4Release;
 	static FString ExecutableName;
 	static FString PlatformName;
 	static FString PlatformNameIni;
@@ -35,6 +39,11 @@ namespace NCachedCrashContextProperties
 
 void FGenericCrashContext::Initialize()
 {
+	NCachedCrashContextProperties::bIsInternalBuild = FEngineBuildSettings::IsInternalBuild();
+	NCachedCrashContextProperties::bIsPerforceBuild = FEngineBuildSettings::IsPerforceBuild();
+	NCachedCrashContextProperties::bIsSourceDistribution = FEngineBuildSettings::IsSourceDistribution();
+	NCachedCrashContextProperties::bIsUE4Release = FApp::IsEngineInstalled();
+
 	NCachedCrashContextProperties::ExecutableName = FPlatformProcess::ExecutableName();
 	NCachedCrashContextProperties::PlatformName = FPlatformProperties::PlatformName();
 	NCachedCrashContextProperties::PlatformNameIni = FPlatformProperties::IniPlatformName();
@@ -87,10 +96,9 @@ void FGenericCrashContext::SerializeContentToBuffer()
 	// It means that the crashed process needs to be alive as long as the crash report client is working on the process.
 	// Proposal, not implemented yet.
 	AddCrashProperty( TEXT( "ProcessId" ), FPlatformProcess::GetCurrentProcessId() );
-
-	AddCrashProperty( TEXT( "IsInternalBuild" ), (int32)FEngineBuildSettings::IsInternalBuild() );
-	AddCrashProperty( TEXT( "IsPerforceBuild" ), (int32)FEngineBuildSettings::IsPerforceBuild() );
-	AddCrashProperty( TEXT( "IsSourceDistribution" ), (int32)FEngineBuildSettings::IsSourceDistribution() );
+	AddCrashProperty( TEXT( "IsInternalBuild" ), (int32)NCachedCrashContextProperties::bIsInternalBuild );
+	AddCrashProperty( TEXT( "IsPerforceBuild" ), (int32)NCachedCrashContextProperties::bIsPerforceBuild );
+	AddCrashProperty( TEXT( "IsSourceDistribution" ), (int32)NCachedCrashContextProperties::bIsSourceDistribution );
 
 	// Add common crash properties.
 	AddCrashProperty( TEXT( "GameName" ), FApp::GetGameName() );
@@ -105,7 +113,7 @@ void FGenericCrashContext::SerializeContentToBuffer()
 	AddCrashProperty( TEXT( "LanguageLCID" ), FInternationalization::Get().GetCurrentCulture()->GetLCID() );
 	AddCrashProperty(TEXT("DefaultLocale"), *NCachedCrashContextProperties::DefaultLocale);
 
-	AddCrashProperty( TEXT( "IsUE4Release" ), (int32)FApp::IsEngineInstalled() );
+	AddCrashProperty( TEXT( "IsUE4Release" ), (int32)NCachedCrashContextProperties::bIsUE4Release);
 
 	// Remove periods from user names to match AutoReporter user names
 	// The name prefix is read by CrashRepository.AddNewCrash in the website code
diff --git a/Engine/Source/Runtime/Core/Private/GenericPlatform/GenericPlatformMallocCrash.cpp b/Engine/Source/Runtime/Core/Private/GenericPlatform/GenericPlatformMallocCrash.cpp
new file mode 100644
--- /dev/null
+++ b/Engine/Source/Runtime/Core/Private/GenericPlatform/GenericPlatformMallocCrash.cpp
@@ -0,0 +1,480 @@
+// Copyright 1998-2015 Epic Games, Inc. All Rights Reserved.
+
+#include "CorePrivatePCH.h"
+#include "GenericPlatformMallocCrash.h"
+
+/** Describes a pool. */
+struct FPoolDesc
+{
+	FPoolDesc( const uint32 InSize, const uint32 InNumAllocs ):
+		Size(InSize),
+		NumAllocs(InNumAllocs)
+	{}
+
+	/** Size of the pool. */
+	const uint32 Size;
+
+	/** Number of allocations in the pool. */
+	const uint32 NumAllocs;
+};
+
+/** Generated by the FGenericPlatformMallocCrash::PrintPoolsUsage. */
+const FPoolDesc& FGenericPlatformMallocCrash::GetPoolDesc( uint32 Index ) const
+{
+	static const FPoolDesc AllPoolDesc[NUM_POOLS]=
+	{
+		FPoolDesc(   64, 224),
+		FPoolDesc(   96, 144),
+		FPoolDesc(  128,  80),
+		FPoolDesc(  192, 560),
+		FPoolDesc(  256, 384),
+		FPoolDesc(  384, 208),
+		FPoolDesc(  512,  48),
+		FPoolDesc(  768,  32),
+		FPoolDesc( 1024,  32),
+		FPoolDesc( 2048,  32),
+		FPoolDesc( 4096,  32),
+		FPoolDesc( 8192,  32),
+		FPoolDesc(16384,  16),
+		FPoolDesc(32768,  16),
+	};
+	return AllPoolDesc[Index];
+}
+
+struct FMallocCrashPool
+{
+	uint32 NumUsed;
+	uint32 MaxUsedIndex;
+	uint32 MaxNumUsed;
+	uint32 TotalNumUsed;
+
+	/** Allocation size for this pool. */
+	const uint32 AllocationSize;
+
+	/** Fixed list of allocations for the specified size for this pool. */
+	FPtrInfo* Allocations[FGenericPlatformMallocCrash::MAX_NUM_ALLOCS_IN_POOL];
+
+	/** Maximum number of allocations that can be made for this pool. */
+	const uint32 MaxNumAllocations;
+
+	/** Memory allocated in the pool and memory used by a fixed array. */
+	uint32 AllocatedMemory;
+
+	FMallocCrashPool( const FPoolDesc& PoolDesc, FGenericPlatformMallocCrash& Outer ): 
+		NumUsed(0),
+		MaxUsedIndex(0),
+		MaxNumUsed(0),
+		TotalNumUsed(0),
+		AllocationSize( PoolDesc.Size+FGenericPlatformMallocCrash::PER_ALLOC_OVERHEAD ),
+		MaxNumAllocations( PoolDesc.NumAllocs )
+	{
+		for( uint32 Index = 0; Index < MaxNumAllocations; ++Index )
+		{
+			uint8* NewPtr = (uint8*)Outer.AllocateFromSmallPool( AllocationSize );
+			Allocations[Index] = new(NewPtr) FPtrInfo( NewPtr+FGenericPlatformMallocCrash::PER_ALLOC_OVERHEAD );
+		}
+
+		// Zero the rest.
+		for( uint32 Index = MaxNumAllocations; Index < FGenericPlatformMallocCrash::MAX_NUM_ALLOCS_IN_POOL; ++Index )
+		{
+			Allocations[Index] = nullptr;
+		}
+
+		AllocatedMemory = MaxNumAllocations*AllocationSize + sizeof(Allocations);
+
+#if	_DEBUG
+		FPlatformMisc::LowLevelOutputDebugStringf( TEXT( "FMallocCrashPool Size=%5u Num=%4i Mem=%8u\n" ), PoolDesc.Size, MaxNumAllocations, AllocatedMemory );
+#endif // _DEBUG
+	}
+
+	uint8* AllocateFromPool( uint32 InAllocationSize )
+	{
+		DebugVerify();
+
+		// Find an unused allocation.
+		FPtrInfo* PtrInfo = nullptr;
+		for( uint32 Index = 0; Index < MaxNumAllocations; ++Index )
+		{
+			FPtrInfo* PtrIt = Allocations[Index];
+			if( PtrIt->Size == 0 )
+			{
+				PtrInfo = PtrIt;
+				MaxUsedIndex = FMath::Max(MaxUsedIndex,Index);
+				break;
+			}
+		}
+
+		if( PtrInfo )
+		{
+			NumUsed++;
+			TotalNumUsed++;
+			MaxNumUsed = FMath::Max(MaxNumUsed,NumUsed);
+			PtrInfo->Size = InAllocationSize;
+
+			FMemory::Memset( (void*)PtrInfo->Ptr, FGenericPlatformMallocCrash::MEM_TAG, PtrInfo->Size );
+			DebugVerify();
+
+			//FPlatformMisc::LowLevelOutputDebugStringf( TEXT( "Malloc Size=%u, PooledPtr=0x%016llx \n" ), AllocationSize, (uint64)PtrInfo->Ptr );
+
+			return PtrInfo->Ptr;
+		}
+		else
+		{
+			FPlatformMisc::DebugBreak();
+			FPlatformMisc::LowLevelOutputDebugStringf( TEXT( "AllocateFromPool run out of memory allocating %u bytes for %u allocations\n" ), InAllocationSize, MaxNumAllocations );
+			FPlatformMisc::LowLevelOutputDebugString( TEXT( "Please increase MaxNumAllocations for that pool, exiting...\n" ) );
+			FPlatformMisc::RequestExit( true );
+		}
+		return nullptr;
+	}
+
+	/** Tries to free a pointer. */
+	void TryFreeFromPool( uint8* Ptr )
+	{
+		//const uint32 PtrSize = FGenericPlatformMallocCrash::GetAllocationSize(Ptr);
+		//FPlatformMisc::LowLevelOutputDebugStringf( TEXT( "Free SizeWithOverhead=%u, PooledPtr=0x%016llx \n" ), PtrSize, (uint64)Ptr );
+
+		bool bRemoved = false;
+		for( uint32 Index = 0; Index < MaxNumAllocations; ++Index )
+		{
+			FPtrInfo* PtrIt = Allocations[Index];
+			if( PtrIt->Ptr == Ptr )
+			{
+				FMemory::Memset( (void*)PtrIt->Ptr, FGenericPlatformMallocCrash::MEM_WIPETAG, PtrIt->Size );
+
+				PtrIt->Size = 0;
+				bRemoved = true;
+				NumUsed--;
+
+				break;
+			}
+		}
+
+		if( !bRemoved )
+		{
+			FPlatformMisc::DebugBreak();
+		}
+
+		DebugVerify();
+	}
+
+private:
+	void DebugVerify()
+	{
+#if	_DEBUG
+		for( uint32 Index = 0; Index < MaxNumAllocations; ++Index )
+		{
+			FPtrInfo* PtrIt = Allocations[Index];
+			if( PtrIt->Size > 32768 )
+			{
+				FPlatformMisc::DebugBreak();
+			}
+		}
+#endif // _DEBUG
+	}
+};
+
+FGenericPlatformMallocCrash::FGenericPlatformMallocCrash( FMalloc* MainMalloc ) :
+	CrashedThreadId( 0 ),
+	LargeMemoryPoolOffset( 0 ),
+	SmallMemoryPoolOffset( 0 ),
+	PreviousMalloc( MainMalloc )
+{
+	const uint32 LargeMemoryPoolSize = Align(LARGE_MEMORYPOOL_SIZE,SafePageSize());
+	LargeMemoryPool = (uint8*)FPlatformMemory::BinnedAllocFromOS( LargeMemoryPoolSize );
+	SmallMemoryPool = (uint8*)FPlatformMemory::BinnedAllocFromOS( (SIZE_T)GetSmallPoolTotalSize() );
+	if( !SmallMemoryPool || !LargeMemoryPool )
+	{
+		FPlatformMisc::LowLevelOutputDebugString( TEXT( "Memory pools allocations failed, exiting...\n" ) );
+		FPlatformMisc::RequestExit(true);
+	}
+
+	if((UPTRINT(LargeMemoryPool) & (REQUIRED_ALIGNMENT - 1)) != 0 || (UPTRINT(SmallMemoryPool) & (REQUIRED_ALIGNMENT - 1)) != 0)
+	{
+		FPlatformMisc::LowLevelOutputDebugString( TEXT( "OS allocations must be aligned to a value multiple of 16, exiting...\n" ) );
+		FPlatformMisc::RequestExit(true);
+	}
+
+	InitializeSmallPools();
+#if	_DEBUG
+	FPlatformMisc::LowLevelOutputDebugStringf( TEXT( "FGenericPlatformMallocCrash overhead is %u bytes\n" ), LargeMemoryPoolSize+GetSmallPoolTotalSize() );
+#endif // _DEBUG
+}
+
+FGenericPlatformMallocCrash::~FGenericPlatformMallocCrash()
+{
+}
+
+FGenericPlatformMallocCrash& FGenericPlatformMallocCrash::Get( FMalloc* MainMalloc /*= nullptr*/ )
+{
+	static FGenericPlatformMallocCrash CrashMalloc( MainMalloc );
+	return CrashMalloc;
+}
+
+void FGenericPlatformMallocCrash::SetAsGMalloc()
+{
+	InternalLock.Lock();
+	GMalloc = this;
+	CrashedThreadId = FPlatformTLS::GetCurrentThreadId();
+}
+
+void* FGenericPlatformMallocCrash::Malloc( SIZE_T Size, uint32 Alignment )
+{
+	const uint32 Size32 = (uint32)Size;
+	if( Alignment > 16 )
+	{
+		FPlatformMisc::DebugBreak();
+		FPlatformMisc::LowLevelOutputDebugString( TEXT( "Alignment > 16 is not supported\n" ) );
+	}
+
+	if( IsOnCrashedThread() )
+	{
+		FMallocCrashPool* Pool = FindPoolFromSize( Size32 );
+		if( Pool )
+		{
+			const uint8* PooledPtr = Pool->AllocateFromPool( Size32 );	
+			return (void*)PooledPtr;
+		}
+		else
+		{
+			const uint32 SizeWithOverhead = Size32 + PER_ALLOC_OVERHEAD;
+			LargeMemoryPoolOffset = Align( LargeMemoryPoolOffset, REQUIRED_ALIGNMENT );
+
+			if( LargeMemoryPoolOffset + SizeWithOverhead <= LARGE_MEMORYPOOL_SIZE )
+			{
+				const uint32 ReturnMemoryPoolOffset = LargeMemoryPoolOffset;
+				LargeMemoryPoolOffset += SizeWithOverhead;
+
+				FPtrInfo* PtrInfo = (FPtrInfo*)(LargeMemoryPool+ReturnMemoryPoolOffset);
+				PtrInfo->Size = Size32;
+				PtrInfo->Ptr = LargeMemoryPool+ReturnMemoryPoolOffset+PER_ALLOC_OVERHEAD;
+
+				FPlatformMisc::LowLevelOutputDebugStringf( TEXT( "Malloc Size=%u LargeMemoryPoolOffset=%u \n" ), Size32, LargeMemoryPoolOffset );
+				return (void*)PtrInfo->Ptr;
+			}
+			else
+			{
+				FPlatformMisc::DebugBreak();
+				FPlatformMisc::LowLevelOutputDebugStringf( TEXT( "MallocCrash run out of memory allocating %u bytes, free %u bytes\n" ), Size32, LARGE_MEMORYPOOL_SIZE-LargeMemoryPoolOffset );
+				FPlatformMisc::LowLevelOutputDebugString( TEXT( "Please increase LARGE_MEMORYPOOL_SIZE, exiting...\n" ) );
+				FPlatformMisc::RequestExit( true );			
+			}
+		}
+	}
+	return nullptr;
+}
+
+void* FGenericPlatformMallocCrash::Realloc( void* Ptr, SIZE_T NewSize, uint32 Alignment )
+{
+	if( IsOnCrashedThread() )
+	{
+		void* Result = nullptr;
+		if( Ptr && NewSize )
+		{
+			SIZE_T PtrSize = 0;
+			const bool bPreviousMalloc = NewSize > 0 && Ptr && !(IsPtrInLargePool(Ptr)||IsPtrInSmallPool(Ptr));
+
+			if( bPreviousMalloc )
+			{
+				// At this moment, we can safely get allocation size only from the binned malloc, this may change in future.
+				if( FCStringWide::Strcmp( PreviousMalloc->GetDescriptiveName(), TEXT("binned") ) == 0 )
+				{
+					// Realloc from the previous allocator.
+					PreviousMalloc->GetAllocationSize(Ptr,PtrSize);
+					if( PtrSize == 0 )
+					{
+						PtrSize = 0;
+					}
+				}
+				// There is nothing we can do about it.
+				else
+				{
+					FPlatformMisc::LowLevelOutputDebugString( TEXT( "Realloc from previous malloc, exiting...\n" ) );
+					FPlatformMisc::RequestExit( true );	
+				}
+			}
+			else
+			{
+				PtrSize = GetAllocationSize(Ptr);
+			}
+			
+			Result = Malloc( NewSize, REQUIRED_ALIGNMENT );
+			FMemory::Memcpy( Result, Ptr, FMath::Min(NewSize,PtrSize) );
+			
+			if( PtrSize > 32768 )
+			{
+				FPlatformMisc::LowLevelOutputDebugStringf( TEXT( "Realloc PtrSize=%u NewSize=%u PooledPtr=0x%016llx\n" ), (uint32)PtrSize, (uint32)NewSize, (uint64)Ptr );
+			}
+
+			Free( Ptr );
+		}
+		else if( Ptr == nullptr )
+		{
+			Result = Malloc( NewSize, REQUIRED_ALIGNMENT );
+		}
+		else
+		{
+			Free( Ptr );
+			Result = nullptr;
+		}
+		return Result;
+	}
+	return nullptr;
+}
+
+void FGenericPlatformMallocCrash::Free( void* Ptr )
+{
+	if( IsOnCrashedThread() )
+	{
+		if( IsPtrInSmallPool(Ptr) )
+		{
+			FMallocCrashPool* Pool = FindPoolFromSize( GetAllocationSize(Ptr) );
+			if( Pool )
+			{
+				Pool->TryFreeFromPool( (uint8*)Ptr );
+			}
+			else
+			{
+				FPlatformMisc::DebugBreak();
+			}
+		}
+		else if( IsPtrInLargePool(Ptr) )
+		{
+			// Not implemented yet.
+		}
+		else
+		{
+			// From the previous allocator.
+		}
+	}
+}
+
+void FGenericPlatformMallocCrash::PrintPoolsUsage()
+{
+#if	_DEBUG
+	FPlatformMisc::LowLevelOutputDebugString( TEXT( "FPoolDesc used\n" ) );
+	for( uint32 Index = 0; Index < FGenericPlatformMallocCrash::NUM_POOLS; ++Index )
+	{
+		const FMallocCrashPool& CrashPool = *Pools[Index];
+		FPlatformMisc::LowLevelOutputDebugStringf( TEXT( "FPoolDesc(%5u,%4u),\n" ), CrashPool.AllocationSize-FGenericPlatformMallocCrash::PER_ALLOC_OVERHEAD, CrashPool.MaxUsedIndex );
+	}
+
+	FPlatformMisc::LowLevelOutputDebugString( TEXT( "FPoolDesc tweaked\n" ) );
+	for( uint32 Index = 0; Index < FGenericPlatformMallocCrash::NUM_POOLS; ++Index )
+	{
+		const FMallocCrashPool& CrashPool = *Pools[Index];
+		FPlatformMisc::LowLevelOutputDebugStringf( TEXT( "FPoolDesc(%5u,%4u),\n" ), CrashPool.AllocationSize-FGenericPlatformMallocCrash::PER_ALLOC_OVERHEAD, Align(CrashPool.MaxUsedIndex*2+16,16) );
+	}
+	FPlatformMisc::LowLevelOutputDebugStringf( TEXT( "LargeMemoryPoolOffset=%u\n" ), LargeMemoryPoolOffset );
+#endif // _DEBUG
+}
+
+bool FGenericPlatformMallocCrash::IsOnCrashedThread() const
+{
+	// Suspend threads other than the crashed one to prevent serious memory errors.
+	// Only the crashed thread can do anything meaningful from here anyway.
+	if( CrashedThreadId == FPlatformTLS::GetCurrentThreadId() )
+	{
+		return true;
+	}
+	else
+	{
+		FPlatformProcess::SleepInfinite();
+		return false;
+	}
+}
+
+bool FGenericPlatformMallocCrash::IsPtrInLargePool( void* Ptr ) const
+{
+	const bool bResult = (Ptr >= &LargeMemoryPool[0] && Ptr < &LargeMemoryPool[LARGE_MEMORYPOOL_SIZE]);
+	return bResult;
+}
+
+bool FGenericPlatformMallocCrash::IsPtrInSmallPool( void* Ptr ) const
+{
+	const bool bResult = (Ptr >= &SmallMemoryPool[0] && Ptr < &SmallMemoryPool[GetSmallPoolTotalSize()]);
+	return bResult;
+}
+
+bool FGenericPlatformMallocCrash::GetAllocationSize( void *Original, SIZE_T &SizeOut )
+{
+	SizeOut = FGenericPlatformMallocCrash::GetAllocationSize(Original);
+	return true;
+}
+
+uint32 FGenericPlatformMallocCrash::GetSmallPoolTotalSize() const
+{
+	static uint32 TotalSize = 0;
+	if( TotalSize == 0 )
+	{
+		for( uint32 Index = 0; Index < NUM_POOLS; ++Index )
+		{
+			const FPoolDesc& PoolDesc = GetPoolDesc(Index);
+			check(PoolDesc.NumAllocs%16==0);
+			check(PoolDesc.Size%16==0);
+			TotalSize += PoolDesc.NumAllocs*(PoolDesc.Size+PER_ALLOC_OVERHEAD);
+		}
+
+		TotalSize = Align(TotalSize,SafePageSize());
+	}
+	return TotalSize;
+}
+
+void FGenericPlatformMallocCrash::InitializeSmallPools()
+{
+	uint32 MallocCrashOverhead = 0;
+
+	for( uint32 Index = 0; Index < NUM_POOLS; ++Index )
+	{
+		Pools[Index] = new FMallocCrashPool( GetPoolDesc(Index), *this );
+		MallocCrashOverhead += Pools[Index]->AllocatedMemory;
+	}
+
+	check(SmallMemoryPoolOffset<=GetSmallPoolTotalSize());
+}
+
+FMallocCrashPool* FGenericPlatformMallocCrash::FindPoolFromSize( uint32 AllocationSize ) const
+{
+	for( uint32 Index = 0; Index < NUM_POOLS; ++Index )
+	{
+		FMallocCrashPool* Pool = Pools[Index];
+		if( AllocationSize <= Pool->AllocationSize-PER_ALLOC_OVERHEAD )
+		{
+			return Pool;
+		}
+	}
+
+	// Use large allocation pool.
+	return nullptr;
+}
+
+uint8* FGenericPlatformMallocCrash::AllocateFromSmallPool( uint32 AllocationSize )
+{
+	if( SmallMemoryPoolOffset + AllocationSize <= GetSmallPoolTotalSize() )
+	{
+		const uint32 ReturnMemoryPoolOffset = SmallMemoryPoolOffset;
+		SmallMemoryPoolOffset += AllocationSize;
+		return (uint8*)SmallMemoryPool+ReturnMemoryPoolOffset;
+	}
+	else
+	{
+		check(0);
+	}
+	return nullptr;
+}
+
+uint32 FGenericPlatformMallocCrash::GetAllocationSize( void *Original )
+{
+	FPtrInfo* PtrInfo = (FPtrInfo*)((uint8*)Original-PER_ALLOC_OVERHEAD);
+	return PtrInfo->Size;
+}
+
+uint32 FGenericPlatformMallocCrash::SafePageSize()
+{
+	const uint32 PageSize = (uint32)FPlatformMemory::GetStats().PageSize;
+	if( PageSize == 0 )
+	{
+		return 65536u;
+	}
+	return PageSize;
+}
diff --git a/Engine/Source/Runtime/Core/Private/HAL/MallocCrash.cpp b/Engine/Source/Runtime/Core/Private/HAL/MallocCrash.cpp
deleted file mode 100644
diff --git a/Engine/Source/Runtime/Core/Private/HAL/UnrealMemory.cpp b/Engine/Source/Runtime/Core/Private/HAL/UnrealMemory.cpp
--- a/Engine/Source/Runtime/Core/Private/HAL/UnrealMemory.cpp
+++ b/Engine/Source/Runtime/Core/Private/HAL/UnrealMemory.cpp
@@ -13,7 +13,7 @@
 #include "MallocDebug.h"
 #include "MallocProfiler.h"
 #include "MallocThreadSafeProxy.h"
-#include "MallocCrash.h"
+#include "PlatformMallocCrash.h"
 
 
 /** Helper function called on first allocation to create and initialize GMalloc */
@@ -21,7 +21,7 @@ void GCreateMalloc()
 {
 	GMalloc = FPlatformMemory::BaseAllocator();
 	// Setup malloc crash as soon as possible.
-	FMallocCrash::Get( GMalloc );
+	FPlatformMallocCrash::Get( GMalloc );
 
 // so now check to see if we are using a Mem Profiler which wraps the GMalloc
 #if USE_MALLOC_PROFILER
diff --git a/Engine/Source/Runtime/Core/Private/Misc/OutputDevice.cpp b/Engine/Source/Runtime/Core/Private/Misc/OutputDevice.cpp
--- a/Engine/Source/Runtime/Core/Private/Misc/OutputDevice.cpp
+++ b/Engine/Source/Runtime/Core/Private/Misc/OutputDevice.cpp
@@ -3,7 +3,6 @@
 #include "CorePrivatePCH.h"
 #include "ExceptionHandling.h"
 #include "VarargsHelper.h"
-#include "MallocCrash.h"
 
 
 DEFINE_LOG_CATEGORY_STATIC(LogOutputDevice, Log, All);
diff --git a/Engine/Source/Runtime/Core/Private/Windows/WindowsPlatformCrashContext.cpp b/Engine/Source/Runtime/Core/Private/Windows/WindowsPlatformCrashContext.cpp
--- a/Engine/Source/Runtime/Core/Private/Windows/WindowsPlatformCrashContext.cpp
+++ b/Engine/Source/Runtime/Core/Private/Windows/WindowsPlatformCrashContext.cpp
@@ -1,7 +1,7 @@
 // Copyright 1998-2015 Epic Games, Inc. All Rights Reserved.
 
 #include "CorePrivatePCH.h"
-#include "MallocCrash.h"
+#include "PlatformMallocCrash.h"
 #include "ExceptionHandling.h"
 #include "WindowsPlatformCrashContext.h"
 #include "EngineVersion.h"
diff --git a/Engine/Source/Runtime/Core/Public/GenericPlatform/GenericPlatformMallocCrash.h b/Engine/Source/Runtime/Core/Public/GenericPlatform/GenericPlatformMallocCrash.h
new file mode 100644
--- /dev/null
+++ b/Engine/Source/Runtime/Core/Public/GenericPlatform/GenericPlatformMallocCrash.h
@@ -0,0 +1,146 @@
+// Copyright 1998-2015 Epic Games, Inc. All Rights Reserved.
+
+#pragma once
+
+/** Describes a pointer. */
+struct FPtrInfo
+{
+	/** Size of the allocation. */
+	uint64 Size;
+	/** Address of the allocation. */
+	uint8* Ptr;
+
+#if PLATFORM_32BITS
+	/** Explicit padding for 32 bit builds */
+	uint8 Padding[4];
+#endif
+
+	FPtrInfo():
+		Size(0),
+		Ptr(0)
+	{}
+
+	FPtrInfo( void* NewPtr ):
+		Size(0),
+		Ptr((uint8*)NewPtr)
+	{}
+};
+
+struct FPoolDesc;
+struct FMallocCrashPool;
+
+/**
+ * Simple pooled memory allocator that uses preallocated memory.
+ * Instance of this class replaces GMalloc after a crash, so we can use dynamic memory allocation even if the app crashes due to OOM.
+ */
+struct FGenericPlatformMallocCrash final : public FMalloc
+{
+	friend struct FPoolDesc;
+	friend struct FMallocCrashPool;
+
+private:
+	enum
+	{
+		LARGE_MEMORYPOOL_SIZE = 2 * 1024 * 1024,
+		REQUIRED_ALIGNMENT = 16,
+		PER_ALLOC_OVERHEAD = sizeof(FPtrInfo),
+
+		NUM_POOLS = 14,
+		MAX_NUM_ALLOCS_IN_POOL = 2048,
+
+		MEM_TAG     = 0xfe,
+		MEM_WIPETAG = 0xcd,
+	};
+
+public:
+	FGenericPlatformMallocCrash( FMalloc* MainMalloc );
+	virtual ~FGenericPlatformMallocCrash();
+
+	/** Creates a new instance. */
+	static CORE_API FGenericPlatformMallocCrash& Get( FMalloc* MainMalloc = nullptr );
+
+	/**
+	 * Sets as GMalloc.
+	 * This method locks to the thread that crashed.
+	 * Any next calls to this method or GMalloc from other threads will dead-lock those threads,
+	 * but this is ok, because we are shutting down.
+	 * This also fixes a lot of potential issues of using dynamic allocation during crash dumping.
+	 * 
+	 * @warning
+	 * This is not super safe, may interfere with other allocations, ie. replacing vtable during running 
+	 * the code from the previous malloc will probably crash other threads.
+	 */
+	void SetAsGMalloc();
+
+	// FMalloc interface.
+	virtual void* Malloc( SIZE_T Size, uint32 Alignment ) override;
+
+	virtual void* Realloc( void* Ptr, SIZE_T NewSize, uint32 Alignment ) override;
+
+	virtual void Free( void* /*Ptr*/ ) override;
+
+	virtual bool GetAllocationSize( void *Original, SIZE_T &SizeOut ) override;
+
+	virtual bool IsInternallyThreadSafe() const override
+	{
+		return true;
+	}
+
+	virtual bool ValidateHeap() override
+	{
+		// Nothing to do here.
+		return true;
+	}
+
+	virtual const TCHAR * GetDescriptiveName() override
+	{
+		return TEXT( "MallocCrash" );
+	}
+
+	CORE_API void PrintPoolsUsage();
+
+protected:
+	/** Whether it is safe to call crash malloc's methods. */
+	bool IsOnCrashedThread() const;
+
+	bool IsPtrInLargePool(void* Ptr) const;
+	bool IsPtrInSmallPool(void* Ptr) const;
+
+	const FPoolDesc& GetPoolDesc( uint32 Index ) const;
+	uint32 GetSmallPoolTotalSize() const;
+
+	void InitializeSmallPools();
+	FMallocCrashPool* FindPoolFromSize( uint32 AllocationSize ) const;
+	uint8* AllocateFromSmallPool( uint32 AllocationSize );
+
+	static uint32 GetAllocationSize(void *Original);
+
+	/**
+	 * @return page size, if page size is not initialized, returns 64k.
+	 */
+	static uint32 SafePageSize();
+
+protected:
+	/** Used to lock crash malloc to one thread. */
+	FCriticalSection InternalLock;
+
+	/** The id of the thread that crashed. */
+	int32 CrashedThreadId;
+
+	/** Preallocated memory pool for large allocations. */
+	uint8* LargeMemoryPool;
+
+	/** Current position in the large memory pool. */
+	uint32 LargeMemoryPoolOffset;
+
+	/** Preallocated memory pool for small allocations. */
+	uint8* SmallMemoryPool;
+
+	/** Current position in the small memory pool. */
+	uint32 SmallMemoryPoolOffset;
+
+	/** Previously used malloc. */
+	FMalloc* PreviousMalloc;
+
+	FMallocCrashPool* Pools[NUM_POOLS];
+};
diff --git a/Engine/Source/Runtime/Core/Public/HAL/MallocCrash.h b/Engine/Source/Runtime/Core/Public/HAL/MallocCrash.h
deleted file mode 100644
diff --git a/Engine/Source/Runtime/Core/Public/HAL/PlatformMallocCrash.h b/Engine/Source/Runtime/Core/Public/HAL/PlatformMallocCrash.h
new file mode 100644
--- /dev/null
+++ b/Engine/Source/Runtime/Core/Public/HAL/PlatformMallocCrash.h
@@ -0,0 +1,11 @@
+// Copyright 1998-2015 Epic Games, Inc. All Rights Reserved.
+#pragma once
+
+#include "HAL/Platform.h"
+
+#if PLATFORM_PS4
+#include "PS4/PS4MallocCrash.h"
+#else
+#include "GenericPlatform/GenericPlatformMallocCrash.h"
+typedef FGenericPlatformMallocCrash FPlatformMallocCrash;
+#endif
\ No newline at end of file
diff --git a/Engine/Source/Runtime/Launch/Private/Windows/LaunchWindows.cpp b/Engine/Source/Runtime/Launch/Private/Windows/LaunchWindows.cpp
--- a/Engine/Source/Runtime/Launch/Private/Windows/LaunchWindows.cpp
+++ b/Engine/Source/Runtime/Launch/Private/Windows/LaunchWindows.cpp
@@ -2,7 +2,7 @@
 	
 #include "LaunchPrivatePCH.h"
 #include "ExceptionHandling.h"
-#include "MallocCrash.h"
+#include "PlatformMallocCrash.h"
 
 #if UE_BUILD_DEBUG
 #include <crtdbg.h>
@@ -211,7 +211,7 @@ int32 WINAPI WinMain( HINSTANCE hInInstance, HINSTANCE hPrevInstance, char*, int
 			ErrorLevel = 1;
 			GError->HandleError();
 			LaunchStaticShutdownAfterError();
-			FMallocCrash::Get().PrintPoolsUsage();
+			FPlatformMallocCrash::Get().PrintPoolsUsage();
 			FPlatformMisc::RequestExit( true );
 		}
 #endif
