Made simple expression searches match more naturally

You can now search for assets containing - or + without having to quote your search terms.

The TextCmpExact and TextCmpAnchor operators have been removed, and are now processed as part of the string lexing - this allows searches for assets with odd names (eg, +-Diffuse) to work like they used to.

Expressions that cannot possibly match anything (ie, an expression that contains no text tokens) will now fallback to testing their entire input as a single text token.

The automatic AND operator insertion now works correctly when inserting an AND operator before a pre-unary NOT operator (eg, "thing -other" will now be correctly treated as "thing AND -other").

diff --git a/Engine/Source/Runtime/Core/Private/Misc/TextFilterExpressionEvaluator.cpp b/Engine/Source/Runtime/Core/Private/Misc/TextFilterExpressionEvaluator.cpp
--- a/Engine/Source/Runtime/Core/Private/Misc/TextFilterExpressionEvaluator.cpp
+++ b/Engine/Source/Runtime/Core/Private/Misc/TextFilterExpressionEvaluator.cpp
@@ -9,21 +9,30 @@ namespace TextFilterExpressionParser
 	class FTextToken
 	{
 	public:
-		FTextToken(FTextFilterString InString, ETextFilterTextComparisonMode InTextComparisonMode)
+		enum class EInvertBasicStringResult : uint8
+		{
+			No,
+			Yes,
+		};
+
+		FTextToken(FTextFilterString InString, ETextFilterTextComparisonMode InTextComparisonMode, const EInvertBasicStringResult InInvertBasicStringResult)
 			: String(MoveTemp(InString))
 			, TextComparisonMode(MoveTemp(InTextComparisonMode))
+			, InvertBasicStringResult(InInvertBasicStringResult)
 		{
 		}
 
 		FTextToken(const FTextToken& Other)
 			: String(Other.String)
 			, TextComparisonMode(Other.TextComparisonMode)
+			, InvertBasicStringResult(Other.InvertBasicStringResult)
 		{
 		}
 
 		FTextToken(FTextToken&& Other)
 			: String(MoveTemp(Other.String))
-			, TextComparisonMode(MoveTemp(Other.TextComparisonMode))
+			, TextComparisonMode(Other.TextComparisonMode)
+			, InvertBasicStringResult(Other.InvertBasicStringResult)
 		{
 		}
 
@@ -31,13 +40,15 @@ namespace TextFilterExpressionParser
 		{
 			String = Other.String;
 			TextComparisonMode = Other.TextComparisonMode;
+			InvertBasicStringResult = Other.InvertBasicStringResult;
 			return *this;
 		}
 
 		FTextToken& operator=(FTextToken&& Other)
 		{
 			String = MoveTemp(Other.String);
-			TextComparisonMode = MoveTemp(Other.TextComparisonMode);
+			TextComparisonMode = Other.TextComparisonMode;
+			InvertBasicStringResult = Other.InvertBasicStringResult;
 			return *this;
 		}
 
@@ -51,19 +62,27 @@ namespace TextFilterExpressionParser
 			return String.AsName();
 		}
 
-		ETextFilterTextComparisonMode GetTextComparisonMode() const
+		bool EvaluateAsBasicStringExpression(const ITextFilterExpressionContext* InContext) const
 		{
-			return TextComparisonMode;
+			if (String.IsEmpty())
+			{
+				// An empty string is always considered to be true
+				return true;
+			}
+
+			const bool bMatched = InContext->TestBasicStringExpression(String, TextComparisonMode);
+			return (InvertBasicStringResult == EInvertBasicStringResult::No) ? bMatched : !bMatched;
 		}
 
-		void SetTextComparisonMode(ETextFilterTextComparisonMode InTextComparisonMode)
+		ETextFilterTextComparisonMode GetTextComparisonMode() const
 		{
-			TextComparisonMode = InTextComparisonMode;
+			return TextComparisonMode;
 		}
 
 	private:
 		FTextFilterString String;
 		ETextFilterTextComparisonMode TextComparisonMode;
+		EInvertBasicStringResult InvertBasicStringResult;
 	};
 }
 
@@ -92,17 +111,18 @@ DEFINE_TEXT_EXPRESSION_OPERATOR_NODE(FOr,						3,	0xF4778B51, 0xF535414D, 0x9C0E
 DEFINE_TEXT_EXPRESSION_OPERATOR_NODE(FAnd,						3,	0x7511397A, 0x02D24DC2, 0x86729800, 0xF454C320)
 DEFINE_TEXT_EXPRESSION_OPERATOR_NODE(FNot,						3,	0x03D78990, 0x41D04E26, 0x8E98AD2F, 0x74667868)
 
-DEFINE_TEXT_EXPRESSION_OPERATOR_NODE(FTextCmpExact,				1,	0x16374B83, 0x314F4624, 0x80BBA1C9, 0x1AB855CD)
-DEFINE_TEXT_EXPRESSION_OPERATOR_NODE(FTextCmpAnchor,			1,	0xC70749AD, 0x954D4A92, 0x95061222, 0xC8AC6D6D)
-
 DEFINE_EXPRESSION_NODE_TYPE(TextFilterExpressionParser::FTextToken,	0x09E49538, 0x633545E3, 0x84B5644F, 0x1F11628F);
 DEFINE_EXPRESSION_NODE_TYPE(bool,									0xC1CD5DCF, 0x2AB44958, 0xB3FF4F8F, 0xE665D121);
 
 namespace TextFilterExpressionParser
 {
-	/** This contains all the symbols that can define breaking points between text and an operator */
-	static const TCHAR BasicTextBreakingCharacters[]	= { '(', ')', '!', '-', '+', '&', '|', '.', ' ' };						// ETextFilterExpressionEvaluatorMode::BasicString
-	static const TCHAR ComplexTextBreakingCharacters[]	= { '(', ')', '=', ':', '<', '>', '!', '-', '+', '&', '|', '.', ' ' };	// ETextFilterExpressionEvaluatorMode::Complex
+	/**
+	 * This contains all the symbols that can define breaking points between text and an operator
+	 * Note: We don't include + and - in this list as these are valid to use inside names, and we should consume 
+	 *		 them as part of the text token as people usually put a space between their search terms
+	 */
+	static const TCHAR BasicTextBreakingCharacters[]	= { '(', ')', '!', '&', '|', ' ' };						// ETextFilterExpressionEvaluatorMode::BasicString
+	static const TCHAR ComplexTextBreakingCharacters[]	= { '(', ')', '=', ':', '<', '>', '!', '&', '|', ' ' };	// ETextFilterExpressionEvaluatorMode::Complex
 
 	const TCHAR* FSubExpressionStart::Monikers[]		= { TEXT("(") };
 	const TCHAR* FSubExpressionEnd::Monikers[]			= { TEXT(")") };
@@ -118,37 +138,6 @@ namespace TextFilterExpressionParser
 	const TCHAR* FAnd::Monikers[]						= { TEXT("AND"), TEXT("&&"), TEXT("&") };
 	const TCHAR* FNot::Monikers[]						= { TEXT("NOT"), TEXT("!"), TEXT("-") };
 
-	const TCHAR* FTextCmpExact::Monikers[]				= { TEXT("+") };
-	const TCHAR* FTextCmpAnchor::Monikers[]				= { TEXT("...") };
-
-	enum class ETextNodeType : uint8
-	{
-		Text,
-		SubExpressionStart,
-		SubExpressionEnd,
-		TextCmpExact,
-		TextCmpAnchor,
-		Other,
-	};
-
-	/** Get the type of the given node */
-	ETextNodeType GetTextNodeType(const FExpressionNode& InNode)
-	{
-		if (InNode.Cast<FTextToken>())			return ETextNodeType::Text;
-		if (InNode.Cast<FSubExpressionStart>()) return ETextNodeType::SubExpressionStart;
-		if (InNode.Cast<FSubExpressionEnd>())	return ETextNodeType::SubExpressionEnd;
-		if (InNode.Cast<FTextCmpExact>())		return ETextNodeType::TextCmpExact;
-		if (InNode.Cast<FTextCmpAnchor>())		return ETextNodeType::TextCmpAnchor;
-		return ETextNodeType::Other;
-	}
-
-	/** Check the given node to see if we consider it to form part of a complex expression */
-	bool IsComplexNode(const FExpressionNode& InNode)
-	{
-		// Key->value pairs require some form of comparison operator, so if we've found once of those, consider this a complex expression
-		return (InNode.Cast<FEqual>() || InNode.Cast<FNotEqual>()|| InNode.Cast<FLess>()|| InNode.Cast<FLessOrEqual>()|| InNode.Cast<FGreater>()|| InNode.Cast<FGreaterOrEqual>());
-	}
-
 	/** Consume an operator from the specified consumer's stream, if one exists at the current read position */
 	template<typename TSymbol>
 	TOptional<FExpressionError> ConsumeOperator(FExpressionTokenConsumer& Consumer)
@@ -167,21 +156,6 @@ namespace TextFilterExpressionParser
 		return TOptional<FExpressionError>();
 	}
 
-	/** Consume a number from the specified consumer's stream, if one exists at the current read position */
-	TOptional<FExpressionError> ConsumeNumber(FExpressionTokenConsumer& Consumer)
-	{
-		auto& Stream = Consumer.GetStream();
-
-		TOptional<FStringToken> NumberToken = ExpressionParser::ParseNumber(Stream);
-		
-		if (NumberToken.IsSet())
-		{
-			Consumer.Add(NumberToken.GetValue(), FExpressionNode(FTextToken(NumberToken.GetValue().GetString(), ETextFilterTextComparisonMode::Partial)));
-		}
-
-		return TOptional<FExpressionError>();
-	}
-
 	/** Transform the given string to remove any escape character sequences found in a quoted string */
 	void UnescapeQuotedString(FString& Str, const TCHAR InQuoteChar)
 	{
@@ -245,16 +219,66 @@ namespace TextFilterExpressionParser
 		{
 			FString FinalString = TextToken.GetValue().GetString();
 			UnescapeQuotedString(FinalString, QuoteChar);
-			Consumer.Add(TextTokenWithQuotes, FExpressionNode(FTextToken(MoveTemp(FinalString), ETextFilterTextComparisonMode::Partial)));
+			Consumer.Add(TextTokenWithQuotes, FExpressionNode(FTextToken(MoveTemp(FinalString), ETextFilterTextComparisonMode::Partial, FTextToken::EInvertBasicStringResult::No)));
 		}
 		else
 		{
-			Consumer.Add(TextTokenWithQuotes, FExpressionNode(FTextToken(FString(), ETextFilterTextComparisonMode::Partial)));
+			Consumer.Add(TextTokenWithQuotes, FExpressionNode(FTextToken(FString(), ETextFilterTextComparisonMode::Partial, FTextToken::EInvertBasicStringResult::No)));
 		}
 
 		return TOptional<FExpressionError>();
 	}
 
+	/** Given a potential string, we produce a final FTextToken for it using the correct comparison mode (as inferred from the given string) */
+	FTextToken CreateTextTokenFromUnquotedString(FString InString)
+	{
+		ETextFilterTextComparisonMode TextComparisonMode = ETextFilterTextComparisonMode::Partial;
+
+		// The final string may embed the TextCmpExact (+) operator, or the TextCmpAnchor operators (...) - test for those now so we can use the correct comparison mode
+		if (InString.Len() > 0 && InString[0] == '+')
+		{
+			// Matched TextCmpExact - update the comparison mode and remove the + token from the start of the string
+			TextComparisonMode = ETextFilterTextComparisonMode::Exact;
+			InString.RemoveAt(0, 1, false);
+		}
+		else if (InString.Len() > 2 && InString.StartsWith(TEXT("..."), ESearchCase::CaseSensitive))
+		{
+			// Matched TextCmpAnchor (pre-unary) - update the comparison mode and remove the ... token from the start of the string
+			TextComparisonMode = ETextFilterTextComparisonMode::EndsWith;
+			InString.RemoveAt(0, 3, false);
+		}
+		else if (InString.Len() > 2 && InString.EndsWith(TEXT("..."), ESearchCase::CaseSensitive))
+		{
+			// Matched TextCmpAnchor (post-unary) - update the comparison mode and remove the ... token from the end of the string
+			TextComparisonMode = ETextFilterTextComparisonMode::StartsWith;
+			InString.RemoveAt(InString.Len() - 3, 3, false);
+		}
+
+		// To preserve behavior with the old text filter, the final string *may* also start with a NOT (-) operator (after stripping the TextCmpExact or TextCmpAnchor tokens from the start)
+		// If it does, then we need to invert the logic when performing simple string comparisons
+		FTextToken::EInvertBasicStringResult InvertBasicStringResult = FTextToken::EInvertBasicStringResult::No;
+		if (InString.Len() > 0 && InString[0] == '-')
+		{
+			// Matched an invert result request - remove the - token from the start of the string
+			InvertBasicStringResult = FTextToken::EInvertBasicStringResult::Yes;
+			InString.RemoveAt(0, 1, false);
+		}
+
+		// Finally, if our string starts and ends with a quote, we need to strip those off now
+		if (InString.Len() > 1 && (InString[0] == '"' || InString[0] == '\''))
+		{
+			const TCHAR QuoteChar = InString[0];
+			if (InString[InString.Len() - 1] == QuoteChar)
+			{
+				// Remove the quotes
+				InString.RemoveAt(0, 1, false);
+				InString.RemoveAt(InString.Len() - 1, 1, false);
+			}
+		}
+
+		return FTextToken(MoveTemp(InString), TextComparisonMode, InvertBasicStringResult);
+	}
+
 	/** Consume the text from the specified consumer's stream */
 	TOptional<FExpressionError> ConsumeTextImpl(FExpressionTokenConsumer& Consumer, const TFunctionRef<bool(TCHAR)> IsBreakingCharacter)
 	{
@@ -316,7 +340,7 @@ namespace TextFilterExpressionParser
 
 		if (TextToken.IsSet())
 		{
-			Consumer.Add(TextToken.GetValue(), FExpressionNode(FTextToken(MoveTemp(FinalString), ETextFilterTextComparisonMode::Partial)));
+			Consumer.Add(TextToken.GetValue(), CreateTextTokenFromUnquotedString(MoveTemp(FinalString)));
 		}
 
 		return TOptional<FExpressionError>();
@@ -436,78 +460,64 @@ bool FTextFilterExpressionEvaluator::SetFilterText(const FText& InFilterText)
 		{
 			auto& TmpLexTokens = LexResult.GetValue();
 
-			// We process the tokens in two passes (doing these separately minimizes the combinations required by step 2):
-			//	Pass 1) We process and remove any text comparison modification tokens (FTextCmpExact and FTextCmpAnchor)
-			//	Pass 2) We inject AND operators between any adjacent text and sub-expression tokens to mimic the old search filter behavior
 			TArray<FExpressionToken> FinalTokens;
 			FinalTokens.Reserve(TmpLexTokens.Num());
 
-			// Pass 1
+			bool bIsComplexExpression = false;
+
+			// If our expression doesn't contain any text tokens, then it can't possibly match anything (this can easily happen when performing simple searches for strings that only contain operators)
+			// In this case, we should just treat the entire expression as a single text token
+			bool bExpressionContainsText = false;
+			for (FExpressionToken& CurrentToken : TmpLexTokens)
 			{
-				ETextNodeType LastTokenType = ETextNodeType::Other; // We have to cache this as we're moving tokens out of the TmpLexTokens array
-				for (FExpressionToken& CurrentToken : TmpLexTokens)
+				if (CurrentToken.Node.Cast<FTextToken>())
 				{
-					const ETextNodeType CurrentTokenType = GetTextNodeType(CurrentToken.Node);
-
-					// We process TextCmpExact and TextCmpAnchor tokens during post-compilation, 
-					// as they only ever operate on constants so don't need to be constantly re-evaluated for each item
-					if (LastTokenType == ETextNodeType::TextCmpExact && CurrentTokenType == ETextNodeType::Text)
-					{
-						// Pre-unary TextCmpExact - modify the current text token to performs an "exact" string compare
-						const FTextToken* TextToken = CurrentToken.Node.Cast<FTextToken>();
-						check(TextToken);
-						const_cast<FTextToken*>(TextToken)->SetTextComparisonMode(ETextFilterTextComparisonMode::Exact);
-					}
-					else if (LastTokenType == ETextNodeType::TextCmpAnchor && CurrentTokenType == ETextNodeType::Text)
-					{
-						// Pre-unary TextCmpAnchor - modify the current text token to performs an "ends with" string compare
-						const FTextToken* TextToken = CurrentToken.Node.Cast<FTextToken>();
-						check(TextToken);
-						const_cast<FTextToken*>(TextToken)->SetTextComparisonMode(ETextFilterTextComparisonMode::EndsWith);
-					}
-					else if (LastTokenType == ETextNodeType::Text && CurrentTokenType == ETextNodeType::TextCmpAnchor)
-					{
-						// Post-unary TextCmpAnchor - modify the previous text token to perform a "starts with" string compare
-						const FTextToken* TextToken = FinalTokens.Last().Node.Cast<FTextToken>();
-						check(TextToken);
-						const_cast<FTextToken*>(TextToken)->SetTextComparisonMode(ETextFilterTextComparisonMode::StartsWith);
-					}
-					
-					// We strip out the TextCmpExact and TextCmpAnchor tokens, as we process them as part of this loop
-					if (CurrentTokenType != ETextNodeType::TextCmpExact && CurrentTokenType != ETextNodeType::TextCmpAnchor)
-					{
-						FinalTokens.Add(MoveTemp(CurrentToken));
-					}
-
-					LastTokenType = CurrentTokenType;
+					bExpressionContainsText = true;
+					break;
 				}
 			}
 
-			// Pass 2
-			bool bIsComplexExpression = false;
+			if (bExpressionContainsText)
 			{
-				for (int32 FinalTokenIndex = 0; FinalTokenIndex < FinalTokens.Num(); ++FinalTokenIndex)
+				// We inject AND operators between any adjacent text and sub-expression tokens to mimic the old search filter behavior
+				bool bWasPreviousTokenValidToInjectANDOperatorAfter = false;
+				for (FExpressionToken& CurrentToken : TmpLexTokens)
 				{
-					const FExpressionToken& CurrentToken = FinalTokens[FinalTokenIndex];
-
-					const ETextNodeType LastTokenType = (FinalTokenIndex > 0) ? GetTextNodeType(FinalTokens[FinalTokenIndex-1].Node) : ETextNodeType::Other;
-					const ETextNodeType CurrentTokenType = GetTextNodeType(CurrentToken.Node);
-
 					// Key->value pairs require some form of comparison operator, so if we've found one of those, consider this to be a complex expression
-					if (!bIsComplexExpression && IsComplexNode(CurrentToken.Node))
+					if (!bIsComplexExpression && 
+						(CurrentToken.Node.Cast<FEqual>() || CurrentToken.Node.Cast<FNotEqual>() || CurrentToken.Node.Cast<FLess>() || CurrentToken.Node.Cast<FLessOrEqual>() || CurrentToken.Node.Cast<FGreater>() || CurrentToken.Node.Cast<FGreaterOrEqual>())
+						)
 					{
 						bIsComplexExpression = true;
 					}
-
-					// We need to inject AND operators between any adjacent text and sub-expression tokens to mimic the old search filter behavior
-					if ((LastTokenType == ETextNodeType::Text || LastTokenType == ETextNodeType::SubExpressionEnd) &&
-						(CurrentTokenType == ETextNodeType::Text || CurrentTokenType == ETextNodeType::SubExpressionStart)
-						)
+					
+					// Are we in a position to potentially inject a new AND token?
+					if (bWasPreviousTokenValidToInjectANDOperatorAfter)
 					{
-						// Inject the new token between the current one and the previous one, and then skip over the newly added token
-						FinalTokens.Insert(FExpressionToken(CurrentToken.Context, FAnd()), FinalTokenIndex++);
+						// We can inject an AND token before any text token, a sub-expression start token, or a pre-unary NOT token
+						const bool bIsCurrentTokenValidToInjectANDOperatorBefore = (CurrentToken.Node.Cast<FTextToken>() || CurrentToken.Node.Cast<FSubExpressionStart>() || CurrentToken.Node.Cast<FNot>());
+						if (bIsCurrentTokenValidToInjectANDOperatorBefore)
+						{
+							// Inject the new token before we move the current token into the final list
+							FinalTokens.Add(FExpressionToken(CurrentToken.Context, FAnd()));
+						}
 					}
+
+					// Update this for the next run through the loop - we can inject an AND token after any text token, or a sub-expression end token
+					bWasPreviousTokenValidToInjectANDOperatorAfter = (CurrentToken.Node.Cast<FTextToken>() || CurrentToken.Node.Cast<FSubExpressionEnd>());
+
+					// Move this token into the final list since we should have injected the AND operator (if required) by now
+					FinalTokens.Add(MoveTemp(CurrentToken));
+				}
+			}
+			else if (TmpLexTokens.Num() > 0)
+			{
+				FStringToken CombinedTokenContext = TmpLexTokens[0].Context;
+				for (int32 TokenIndex = 1; TokenIndex < TmpLexTokens.Num(); ++TokenIndex)
+				{
+					CombinedTokenContext.Accumulate(TmpLexTokens[TokenIndex].Context);
 				}
+				FinalTokens.Add(FExpressionToken(CombinedTokenContext, CreateTextTokenFromUnquotedString(FText::TrimPrecedingAndTrailing(FilterText).ToString())));
 			}
 
 			CompiledFilter = ExpressionParser::Compile(MoveTemp(FinalTokens), Grammar);
@@ -561,7 +571,7 @@ bool FTextFilterExpressionEvaluator::TestTextFilter(const ITextFilterExpressionC
 				}
 				else if (const FTextToken* TextResult = EvalResult.GetValue().Cast<FTextToken>())
 				{
-					return InContext.TestBasicStringExpression(TextResult->GetString(), TextResult->GetTextComparisonMode());
+					return TextResult->EvaluateAsBasicStringExpression(&InContext);
 				}
 			}
 		}
@@ -591,9 +601,6 @@ void FTextFilterExpressionEvaluator::ConstructExpressionParser()
 	TokenDefinitions.DefineToken(&ConsumeOperator<FOr>);
 	TokenDefinitions.DefineToken(&ConsumeOperator<FAnd>);
 	TokenDefinitions.DefineToken(&ConsumeOperator<FNot>);
-	TokenDefinitions.DefineToken(&ConsumeOperator<FTextCmpExact>);
-	TokenDefinitions.DefineToken(&ConsumeOperator<FTextCmpAnchor>);
-	TokenDefinitions.DefineToken(&ConsumeNumber);
 	TokenDefinitions.DefineToken(&ConsumeQuotedText);
 	TokenDefinitions.DefineToken((ExpressionEvaluatorMode == ETextFilterExpressionEvaluatorMode::Complex) ? &ConsumeComplexText : &ConsumeBasicText);
 
@@ -615,16 +622,16 @@ void FTextFilterExpressionEvaluator::ConstructExpressionParser()
 	JumpTable.MapBinary<FNotEqual>([](const FTextToken& A, const FTextToken& B, const ITextFilterExpressionContext* InContext)			{ return InContext->TestComplexExpression(A.GetStringAsName(), B.GetString(), ETextFilterComparisonOperation::NotEqual, B.GetTextComparisonMode()); });
 	JumpTable.MapBinary<FEqual>([](const FTextToken& A, const FTextToken& B, const ITextFilterExpressionContext* InContext)				{ return InContext->TestComplexExpression(A.GetStringAsName(), B.GetString(), ETextFilterComparisonOperation::Equal, B.GetTextComparisonMode()); });
 
-	JumpTable.MapBinary<FOr>([](const FTextToken& A, const FTextToken& B, const ITextFilterExpressionContext* InContext)				{ return InContext->TestBasicStringExpression(A.GetString(), A.GetTextComparisonMode()) || InContext->TestBasicStringExpression(B.GetString(), B.GetTextComparisonMode()); });
-	JumpTable.MapBinary<FOr>([](const FTextToken& A, bool B, const ITextFilterExpressionContext* InContext)								{ return InContext->TestBasicStringExpression(A.GetString(), A.GetTextComparisonMode()) || B; });
-	JumpTable.MapBinary<FOr>([](bool A, const FTextToken& B, const ITextFilterExpressionContext* InContext)								{ return A || InContext->TestBasicStringExpression(B.GetString(), B.GetTextComparisonMode()); });
+	JumpTable.MapBinary<FOr>([](const FTextToken& A, const FTextToken& B, const ITextFilterExpressionContext* InContext)				{ return A.EvaluateAsBasicStringExpression(InContext) || B.EvaluateAsBasicStringExpression(InContext); });
+	JumpTable.MapBinary<FOr>([](const FTextToken& A, bool B, const ITextFilterExpressionContext* InContext)								{ return A.EvaluateAsBasicStringExpression(InContext) || B; });
+	JumpTable.MapBinary<FOr>([](bool A, const FTextToken& B, const ITextFilterExpressionContext* InContext)								{ return A || B.EvaluateAsBasicStringExpression(InContext); });
 	JumpTable.MapBinary<FOr>([](bool A, bool B, const ITextFilterExpressionContext* InContext)											{ return A || B; });
 
-	JumpTable.MapBinary<FAnd>([](const FTextToken& A, const FTextToken& B, const ITextFilterExpressionContext* InContext)				{ return InContext->TestBasicStringExpression(A.GetString(), A.GetTextComparisonMode()) && InContext->TestBasicStringExpression(B.GetString(), B.GetTextComparisonMode()); });
-	JumpTable.MapBinary<FAnd>([](const FTextToken& A, bool B, const ITextFilterExpressionContext* InContext)							{ return InContext->TestBasicStringExpression(A.GetString(), A.GetTextComparisonMode()) && B; });
-	JumpTable.MapBinary<FAnd>([](bool A, const FTextToken& B, const ITextFilterExpressionContext* InContext)							{ return A && InContext->TestBasicStringExpression(B.GetString(), B.GetTextComparisonMode()); });
+	JumpTable.MapBinary<FAnd>([](const FTextToken& A, const FTextToken& B, const ITextFilterExpressionContext* InContext)				{ return A.EvaluateAsBasicStringExpression(InContext) && B.EvaluateAsBasicStringExpression(InContext); });
+	JumpTable.MapBinary<FAnd>([](const FTextToken& A, bool B, const ITextFilterExpressionContext* InContext)							{ return A.EvaluateAsBasicStringExpression(InContext) && B; });
+	JumpTable.MapBinary<FAnd>([](bool A, const FTextToken& B, const ITextFilterExpressionContext* InContext)							{ return A && B.EvaluateAsBasicStringExpression(InContext); });
 	JumpTable.MapBinary<FAnd>([](bool A, bool B, const ITextFilterExpressionContext* InContext)											{ return A && B; });
 
-	JumpTable.MapPreUnary<FNot>([](const FTextToken& V, const ITextFilterExpressionContext* InContext)									{ return !InContext->TestBasicStringExpression(V.GetString(), V.GetTextComparisonMode()); });
+	JumpTable.MapPreUnary<FNot>([](const FTextToken& V, const ITextFilterExpressionContext* InContext)									{ return !V.EvaluateAsBasicStringExpression(InContext); });
 	JumpTable.MapPreUnary<FNot>([](bool V, const ITextFilterExpressionContext* InContext)												{ return !V; });
 }
diff --git a/Engine/Source/Runtime/Core/Private/Misc/TextFilterTests.cpp b/Engine/Source/Runtime/Core/Private/Misc/TextFilterTests.cpp
--- a/Engine/Source/Runtime/Core/Private/Misc/TextFilterTests.cpp
+++ b/Engine/Source/Runtime/Core/Private/Misc/TextFilterTests.cpp
@@ -156,6 +156,9 @@ namespace TextFilterTests
 		bResult &= InTestPayload.TestFilterExpression(TEXT("Wooble2 OR Flibble"), true);
 		bResult &= InTestPayload.TestFilterExpression(TEXT("Wooble2 AND Flibble"), false);
 		bResult &= InTestPayload.TestFilterExpression(TEXT("Wooble && !Flibble"), false);
+		bResult &= InTestPayload.TestFilterExpression(TEXT("Flibble -Wooble2"), true);
+		bResult &= InTestPayload.TestFilterExpression(TEXT("Flibble OR -Wooble2"), true);
+		bResult &= InTestPayload.TestFilterExpression(TEXT("Flibble AND -Wooble2"), true);
 		bResult &= InTestPayload.TestFilterExpression(TEXT("Type'/Path/To/Asset.Asset'"), true);
 		bResult &= InTestPayload.TestFilterExpression(TEXT("Other'/Path/To/Asset.Asset'FollowingText"), true);
 		bResult &= InTestPayload.TestFilterExpression(TEXT("\"Funky<String>\""), true);
diff --git a/Engine/Source/Runtime/Core/Public/Misc/TextFilterUtils.h b/Engine/Source/Runtime/Core/Public/Misc/TextFilterUtils.h
--- a/Engine/Source/Runtime/Core/Public/Misc/TextFilterUtils.h
+++ b/Engine/Source/Runtime/Core/Public/Misc/TextFilterUtils.h
@@ -68,6 +68,12 @@ public:
 		return FName(*InternalString);
 	}
 
+	/** Is the internal string empty? */
+	FORCEINLINE bool IsEmpty() const
+	{
+		return InternalString.IsEmpty();
+	}
+
 private:
 	/** Inline convert our internal string to uppercase */
 	void UppercaseInternalString();
