Back out changelist 2593948

diff --git a/Engine/Source/Runtime/Core/Private/Android/AndroidFile.cpp b/Engine/Source/Runtime/Core/Private/Android/AndroidFile.cpp
--- a/Engine/Source/Runtime/Core/Private/Android/AndroidFile.cpp
+++ b/Engine/Source/Runtime/Core/Private/Android/AndroidFile.cpp
@@ -883,10 +883,6 @@ public:
 				MountOBB(*(OBBDir2 / PatchOBBName));
 			}
 		}
-
-		// make sure the base path directory exists (UE4Game)
-		mkdir(TCHAR_TO_UTF8(*GetFileBasePath()), 0766);
-
 		return true;
 	}
 
