Make sure the UE4Game and UE4Game/[GameName] directories are created
#ue-17385
#ue4
#android

diff --git a/Engine/Source/Runtime/Core/Private/Android/AndroidFile.cpp b/Engine/Source/Runtime/Core/Private/Android/AndroidFile.cpp
--- a/Engine/Source/Runtime/Core/Private/Android/AndroidFile.cpp
+++ b/Engine/Source/Runtime/Core/Private/Android/AndroidFile.cpp
@@ -50,6 +50,8 @@ FString GFontPathBase;
 // Is the OBB in an APK file or not
 bool GOBBinAPK;
 
+#define FILEBASE_DIRECTORY "/UE4Game/"
+
 extern jobject AndroidJNI_GetJavaAssetManager();
 extern AAssetManager * AndroidThunkCpp_GetAssetManager();
 
@@ -71,7 +73,7 @@ extern "C" void Java_com_epicgames_ue4_GameActivity_nativeSetObbInfo(JNIEnv* jen
 // Constructs the base path for any files which are not in OBB/pak data
 const FString &GetFileBasePath()
 {
-	static FString BasePath = GFilePathBase + FString("/UE4Game/") + FApp::GetGameName() + FString("/");
+	static FString BasePath = GFilePathBase + FString(FILEBASE_DIRECTORY) + FApp::GetGameName() + FString("/");
 	return BasePath;
 }
 
@@ -883,6 +885,12 @@ public:
 				MountOBB(*(OBBDir2 / PatchOBBName));
 			}
 		}
+
+		// make sure the base path directory exists (UE4Game and UE4Game/GameName)
+		FString FileBaseDir = GFilePathBase + FString(FILEBASE_DIRECTORY);
+		mkdir(TCHAR_TO_UTF8(*FileBaseDir), 0766);
+		mkdir(TCHAR_TO_UTF8(*GetFileBasePath()), 0766);
+
 		return true;
 	}
 
