Fix for FString::FindLastCharByPredicate when it's empty.
Some checks added.
Some tidy-ups.
FORCEINLINE added to wrapper function.

diff --git a/Engine/Source/Runtime/Core/Public/Containers/Array.h b/Engine/Source/Runtime/Core/Public/Containers/Array.h
--- a/Engine/Source/Runtime/Core/Public/Containers/Array.h
+++ b/Engine/Source/Runtime/Core/Public/Containers/Array.h
@@ -896,13 +896,12 @@ public:
 	 */
 	int32 FindLast(const ElementType& Item) const
 	{
-		const ElementType* RESTRICT End = GetData() + ArrayNum;
-		for (const ElementType* RESTRICT Data = End, *RESTRICT DataStart = Data - ArrayNum; Data != DataStart;)
+		for (const ElementType* RESTRICT Start = GetData(), *RESTRICT Data = Start + ArrayNum; Data != Start; )
 		{
 			--Data;
 			if (*Data == Item)
 			{
-				return static_cast<int32>(Data - DataStart);
+				return static_cast<int32>(Data - Start);
 			}
 		}
 		return INDEX_NONE;
@@ -919,13 +918,13 @@ public:
 	template <typename Predicate>
 	int32 FindLastByPredicate(Predicate Pred, int32 StartIndex) const
 	{
-		const ElementType* RESTRICT End = GetData() + StartIndex;
-		for (const ElementType* RESTRICT Data = End, *RESTRICT DataStart = Data - StartIndex; Data != DataStart;)
+		check(StartIndex >= 0 && StartIndex <= this->Num());
+		for (const ElementType* RESTRICT Start = GetData(), *RESTRICT Data = Start + StartIndex; Data != Start; )
 		{
 			--Data;
 			if (Pred(*Data))
 			{
-				return static_cast<int32>(Data - DataStart);
+				return static_cast<int32>(Data - Start);
 			}
 		}
 		return INDEX_NONE;
@@ -939,7 +938,7 @@ public:
 	* @returns Index of the found element. INDEX_NONE otherwise.
 	*/
 	template <typename Predicate>
-	int32 FindLastByPredicate(Predicate Pred) const
+	FORCEINLINE int32 FindLastByPredicate(Predicate Pred) const
 	{
 		return FindLastByPredicate(Pred, ArrayNum);
 	}
diff --git a/Engine/Source/Runtime/Core/Public/Containers/UnrealString.h b/Engine/Source/Runtime/Core/Public/Containers/UnrealString.h
--- a/Engine/Source/Runtime/Core/Public/Containers/UnrealString.h
+++ b/Engine/Source/Runtime/Core/Public/Containers/UnrealString.h
@@ -1071,6 +1071,7 @@ public:
 	template <typename Predicate>
 	FORCEINLINE int32 FindLastCharByPredicate(Predicate Pred, int32 StartIndex) const
 	{
+		check(StartIndex >= 0 && StartIndex <= this->Len());
 		return Data.FindLastByPredicate(Pred, StartIndex);
 	}
 
@@ -1085,7 +1086,7 @@ public:
 	template <typename Predicate>
 	FORCEINLINE int32 FindLastCharByPredicate(Predicate Pred) const
 	{
-		return Data.FindLastByPredicate(Pred, Data.Num() - 1);
+		return Data.FindLastByPredicate(Pred, this->Len());
 	}
 
 	/**
