Foreign archives now generate entries for the native translation of the native source as though it were the native source - the result is that the native translation will also be translated in the foreign culture. Localization resources now use the native translation of the native source as though it were the native source, specifically when looking up the foreign translation to use from archives - the result is that the native source will use the foreign translation of its native translation.

Given "Foo" as the native source and "Bar" as the native translation, the foreign archive will contain translations for "Foo" and for "Bar". The foreign localization resource will specify the foreign translation for "Bar" from the archive  as the translation for "Foo". In effect, foreign translations are now based on native translations rather than the actual native source text.

diff --git a/Engine/Source/Editor/UnrealEd/Private/Commandlets/GenerateGatherArchiveCommandlet.cpp b/Engine/Source/Editor/UnrealEd/Private/Commandlets/GenerateGatherArchiveCommandlet.cpp
--- a/Engine/Source/Editor/UnrealEd/Private/Commandlets/GenerateGatherArchiveCommandlet.cpp
+++ b/Engine/Source/Editor/UnrealEd/Private/Commandlets/GenerateGatherArchiveCommandlet.cpp
@@ -153,6 +153,48 @@ int32 UGenerateGatherArchiveCommandlet::Main( const FString& Params )
 
 	ManifestSerializer.DeserializeManifest( ManifestJsonObject.ToSharedRef(), InternationalizationManifest );
 
+	const FString NativeCulturePath = DestinationPath / *(NativeCulture);
+	TArray<FString> NativeArchiveFileNames;
+	IFileManager::Get().FindFiles(NativeArchiveFileNames, *(NativeCulturePath / TEXT("*.archive")), true, false);
+
+	TArray< TSharedPtr<FInternationalizationArchive> > NativeArchives;
+	for (const FString& NativeArchiveFileName : NativeArchiveFileNames)
+	{
+		// Read each archive file from the culture-named directory in the source path.
+		FString ArchiveFilePath = NativeCulturePath / NativeArchiveFileName;
+		ArchiveFilePath = FPaths::ConvertRelativePathToFull(ArchiveFilePath);
+		TSharedRef<FInternationalizationArchive> InternationalizationArchive = MakeShareable(new FInternationalizationArchive);
+
+		FJsonInternationalizationArchiveSerializer ArchiveSerializer;
+#if 0 // @todo Json: Serializing from FArchive is currently broken
+		FArchive* ArchiveFile = IFileManager::Get().CreateFileReader(*ArchiveFilePath);
+
+		if (ArchiveFile == nullptr)
+		{
+			UE_LOG(LogGenerateArchiveCommandlet, Error, TEXT("No archive found at %s."), *ArchiveFilePath);
+			continue;
+		}
+
+		ArchiveSerializer.DeserializeArchive(*ArchiveFile, InternationalizationArchive);
+#else
+		FString ArchiveContent;
+
+		if (!FFileHelper::LoadFileToString(ArchiveContent, *ArchiveFilePath))
+		{
+			UE_LOG(LogGenerateArchiveCommandlet, Error, TEXT("Failed to load file %s."), *ArchiveFilePath);
+			continue;
+		}
+
+		if (!ArchiveSerializer.DeserializeArchive(ArchiveContent, InternationalizationArchive))
+		{
+			UE_LOG(LogGenerateArchiveCommandlet, Error, TEXT("Failed to serialize archive from file %s."), *ArchiveFilePath);
+			continue;
+		}
+#endif
+
+		NativeArchives.Add(InternationalizationArchive);
+	}
+
 	for(int32 Culture = 0; Culture < CulturesToGenerate.Num(); Culture++)
 	{
 		TSharedRef< FInternationalizationArchive > InternationalizationArchive = MakeShareable( new FInternationalizationArchive );
@@ -258,6 +300,39 @@ int32 UGenerateGatherArchiveCommandlet::Main( const FString& Params )
 			}
 
 			ArchiveSerializer.DeserializeArchive( ExistingArchiveJsonObject.ToSharedRef(), OutputInternationalizationArchive );
+
+			if (CulturesToGenerate[Culture] != NativeCulture)
+			{
+				for(TManifestEntryByContextIdContainer::TConstIterator i = InternationalizationManifest->GetEntriesByContextIdIterator(); i; ++i)
+				{
+					// Gather relevant info from manifest entry.
+					const TSharedRef<FManifestEntry>& ManifestEntry = i.Value();
+					const FString& Namespace = ManifestEntry->Namespace;
+					const FLocItem& Source = ManifestEntry->Source;
+					const FString& SourceString = Source.Text;
+					const FString UnescapedSourceString = SourceString;
+					const uint32 SourceStringHash = FCrc::StrCrc32(*UnescapedSourceString);
+					const FString Key = i.Key();
+
+					for (const auto& Context : ManifestEntry->Contexts)
+					{
+						for (const auto& NativeArchive : NativeArchives)
+						{
+							const TSharedPtr<FArchiveEntry> NativeArchiveEntry = NativeArchive->FindEntryBySource(Namespace, Source, Context.KeyMetadataObj);
+							if (NativeArchiveEntry.IsValid() && !NativeArchiveEntry->Source.IsExactMatch(NativeArchiveEntry->Translation))
+							{
+								FLocItem NewSource = NativeArchiveEntry->Translation;
+								ConditionSource(NewSource);
+								FLocItem NewTranslation = NativeArchiveEntry->Translation;
+								ConditionTranslation(NewTranslation);
+								OutputInternationalizationArchive->AddEntry(Namespace, NewSource, NewTranslation, Context.KeyMetadataObj, Context.bIsOptional);
+								break;
+							}
+						}
+					}
+				}
+			}
+
 		}
 
 		if (InternationalizationArchive->GetFormatVersion() < FInternationalizationArchive::EFormatVersion::Latest)
diff --git a/Engine/Source/Editor/UnrealEd/Private/Commandlets/GenerateTextLocalizationResourceCommandlet.cpp b/Engine/Source/Editor/UnrealEd/Private/Commandlets/GenerateTextLocalizationResourceCommandlet.cpp
--- a/Engine/Source/Editor/UnrealEd/Private/Commandlets/GenerateTextLocalizationResourceCommandlet.cpp
+++ b/Engine/Source/Editor/UnrealEd/Private/Commandlets/GenerateTextLocalizationResourceCommandlet.cpp
@@ -67,6 +67,14 @@ int32 UGenerateTextLocalizationResourceCommandlet::Main(const FString& Params)
 		return -1;
 	}
 
+	// Get cultures to generate.
+	FString NativeCultureName;
+	if( !( GetStringFromConfig( *SectionName, TEXT("NativeCulture"), NativeCultureName, GatherTextConfigPath ) ) )
+	{
+		UE_LOG(LogGenerateTextLocalizationResourceCommandlet, Error, TEXT("No native culture specified."));
+		return -1;
+	}
+
 	// Get cultures to generate.
 	TArray<FString> CulturesToGenerate;
 	GetStringArrayFromConfig( *SectionName, TEXT("CulturesToGenerate"), CulturesToGenerate, GatherTextConfigPath );
@@ -140,7 +148,7 @@ int32 UGenerateTextLocalizationResourceCommandlet::Main(const FString& Params)
 		{
 			FJsonInternationalizationArchiveSerializer ArchiveSerializer;
 
-			if( !(FTextLocalizationResourceGenerator::Generate(SourcePath, InternationalizationManifest, CultureName, TextLocalizationResourceArchive, ArchiveSerializer)) )
+			if( !(FTextLocalizationResourceGenerator::Generate(SourcePath, InternationalizationManifest, NativeCultureName, CultureName, TextLocalizationResourceArchive, ArchiveSerializer)) )
 			{
 				IFileManager::Get().Delete( *TextLocalizationResourcePath );
 			}
diff --git a/Engine/Source/Editor/UnrealEd/Private/Commandlets/InternationalizationExportCommandlet.cpp b/Engine/Source/Editor/UnrealEd/Private/Commandlets/InternationalizationExportCommandlet.cpp
--- a/Engine/Source/Editor/UnrealEd/Private/Commandlets/InternationalizationExportCommandlet.cpp
+++ b/Engine/Source/Editor/UnrealEd/Private/Commandlets/InternationalizationExportCommandlet.cpp
@@ -1,4 +1,4 @@
-// Copyright 1998-2015 Epic Games, Inc. All Rights Reserved.
+﻿// Copyright 1998-2015 Epic Games, Inc. All Rights Reserved.
 
 #include "UnrealEd.h"
 #include "Culture.h"
@@ -136,6 +136,14 @@ UInternationalizationExportCommandlet::UInternationalizationExportCommandlet(con
 
 bool UInternationalizationExportCommandlet::DoExport( const FString& SourcePath, const FString& DestinationPath, const FString& Filename )
 {
+	// Get native culture.
+	FString NativeCultureName;
+	if( !GetStringFromConfig( *SectionName, TEXT("NativeCulture"), NativeCultureName, ConfigPath ) )
+	{
+		UE_LOG( LogInternationalizationExportCommandlet, Error, TEXT("No native culture specified.") );
+		return false;
+	}
+
 	// Get manifest name.
 	FString ManifestName;
 	if( !GetStringFromConfig( *SectionName, TEXT("ManifestName"), ManifestName, ConfigPath ) )
@@ -182,6 +190,25 @@ bool UInternationalizationExportCommandlet::DoExport( const FString& SourcePath,
 		ManifestSerializer.DeserializeManifest( ManifestJsonObject.ToSharedRef(), InternationalizationManifest );
 	}
 
+	TArray< TSharedPtr<FInternationalizationArchive> > NativeArchives;
+	{
+		const FString NativeCulturePath = DestinationPath / *(NativeCultureName);
+		TArray<FString> NativeArchiveFileNames;
+		IFileManager::Get().FindFiles(NativeArchiveFileNames, *(NativeCulturePath / TEXT("*.archive")), true, false);
+
+		for (const FString& NativeArchiveFileName : NativeArchiveFileNames)
+		{
+			// Read each archive file from the culture-named directory in the source path.
+			FString ArchiveFilePath = NativeCulturePath / NativeArchiveFileName;
+			ArchiveFilePath = FPaths::ConvertRelativePathToFull(ArchiveFilePath);
+			TSharedRef<FInternationalizationArchive> InternationalizationArchive = MakeShareable(new FInternationalizationArchive);
+			TSharedPtr< FJsonObject > ArchiveJsonObject = ReadJSONTextFile( ArchiveFilePath );
+			FJsonInternationalizationArchiveSerializer ArchiveSerializer;
+			ArchiveSerializer.DeserializeArchive( ArchiveJsonObject.ToSharedRef(), InternationalizationArchive );
+
+			NativeArchives.Add(InternationalizationArchive);
+		}
+	}
 
 	// Process the desired cultures
 	for(int32 Culture = 0; Culture < CulturesToGenerate.Num(); Culture++)
@@ -200,7 +227,6 @@ bool UInternationalizationExportCommandlet::DoExport( const FString& SourcePath,
 			TSharedRef< FInternationalizationArchive > InternationalizationArchive = MakeShareable( new FInternationalizationArchive );
 			ArchiveSerializer.DeserializeArchive( ArchiveJsonObject.ToSharedRef(), InternationalizationArchive );
 
-
 			{
 				FPortableObjectFormatDOM PortableObj;
 
@@ -224,7 +250,8 @@ bool UInternationalizationExportCommandlet::DoExport( const FString& SourcePath,
 
 						for( auto ContextIter = ManifestEntry->Contexts.CreateConstIterator(); ContextIter; ++ContextIter )
 						{
-							TSharedPtr<FArchiveEntry> ArchiveEntry = InternationalizationArchive->FindEntryBySource( Namespace, Source, ContextIter->KeyMetadataObj );
+							{
+								const TSharedPtr<FArchiveEntry> ArchiveEntry = InternationalizationArchive->FindEntryBySource( Namespace, Source, ContextIter->KeyMetadataObj );
 							if( ArchiveEntry.IsValid() )
 							{
 								const FString ConditionedArchiveSource = ConditionArchiveStrForPo(ArchiveEntry->Source.Text);
@@ -244,6 +271,45 @@ bool UInternationalizationExportCommandlet::DoExport( const FString& SourcePath,
 								PortableObj.AddEntry( PoEntry );
 							}
 						}
+
+							if (CultureName != NativeCultureName)
+							{
+								TSharedPtr<FArchiveEntry> NativeArchiveEntry;
+								for (const auto& NativeArchive : NativeArchives)
+								{
+									const TSharedPtr<FArchiveEntry> PotentialNativeArchiveEntry = NativeArchive->FindEntryBySource( Namespace, Source, ContextIter->KeyMetadataObj );
+									if (PotentialNativeArchiveEntry.IsValid())
+									{
+										NativeArchiveEntry = PotentialNativeArchiveEntry;
+										break;
+									}
+								}
+
+								if (NativeArchiveEntry.IsValid())
+								{
+									if (!NativeArchiveEntry->Source.IsExactMatch(NativeArchiveEntry->Translation))
+									{
+										const TSharedPtr<FArchiveEntry> ArchiveEntry = InternationalizationArchive->FindEntryBySource( Namespace, NativeArchiveEntry->Translation, NativeArchiveEntry->KeyMetadataObj );
+
+										const FString ConditionedArchiveSource = ConditionArchiveStrForPo(ArchiveEntry->Source.Text);
+										const FString ConditionedArchiveTranslation = ConditionArchiveStrForPo(ArchiveEntry->Translation.Text);
+
+										TSharedRef<FPortableObjectEntry> PoEntry = MakeShareable( new FPortableObjectEntry );
+										//@TODO: We support additional metadata entries that can be translated.  How do those fit in the PO file format?  Ex: isMature
+										PoEntry->MsgId = ConditionedArchiveSource;
+										//@TODO: Take into account optional entries and entries that differ by keymetadata.  Ex. Each optional entry needs a unique msgCtxt
+										PoEntry->MsgCtxt = Namespace;
+										PoEntry->MsgStr.Add( ConditionedArchiveTranslation );
+
+										FString PORefString = ConvertSrcLocationToPORef( ContextIter->SourceLocation );
+										PoEntry->AddReference( PORefString ); // Source location.
+										PoEntry->AddExtractedComment( ContextIter->Key ); // "Notes from Programmer" in the form of the Key.
+										PoEntry->AddExtractedComment( PORefString ); // "Notes from Programmer" in the form of the Source Location, since this comes in handy too and OneSky doesn't properly show references, only comments.
+										PortableObj.AddEntry( PoEntry );
+									}
+								}
+					}
+				}
 					}
 				}
 
diff --git a/Engine/Source/Runtime/Core/Private/Internationalization/TextLocalizationManager.cpp b/Engine/Source/Runtime/Core/Private/Internationalization/TextLocalizationManager.cpp
--- a/Engine/Source/Runtime/Core/Private/Internationalization/TextLocalizationManager.cpp
+++ b/Engine/Source/Runtime/Core/Private/Internationalization/TextLocalizationManager.cpp
@@ -48,82 +48,82 @@ void EndInitTextLocalization()
 		}
 		else
 #if !UE_BUILD_SHIPPING
-		// Use culture override specified on commandline.
-		if (FParse::Value(FCommandLine::Get(), TEXT("CULTURE="), RequestedCultureName))
-		{
-			//UE_LOG(LogInit, Log, TEXT("Overriding culture %s w/ command-line option".), *CultureName);
-		}
-		else
-#endif // !UE_BUILD_SHIPPING
-#if WITH_EDITOR
-		// See if we've been provided a culture override in the editor
-		if(GConfig->GetString( TEXT("Internationalization"), TEXT("Culture"), RequestedCultureName, GEditorSettingsIni ))
-		{
-			//UE_LOG(LogInit, Log, TEXT("Overriding culture %s w/ editor configuration."), *CultureName);
-		}
-		else
-#endif // WITH_EDITOR
-		// Use culture specified in engine configuration.
-		if(GConfig->GetString( TEXT("Internationalization"), TEXT("Culture"), RequestedCultureName, GEngineIni ))
-		{
-			//UE_LOG(LogInit, Log, TEXT("Overriding culture %s w/ engine configuration."), *CultureName);
-		}
-		else
-		{
-			RequestedCultureName = I18N.GetDefaultCulture()->GetName();
-		}
-
-		FString TargetCultureName = RequestedCultureName;
-		{
-			TArray<FString> LocalizationPaths;
-			if(ShouldLoadEditor)
-			{
-				LocalizationPaths += FPaths::GetEditorLocalizationPaths();
-			}
-			if(ShouldLoadGame)
+			// Use culture override specified on commandline.
+			if (FParse::Value(FCommandLine::Get(), TEXT("CULTURE="), RequestedCultureName))
 			{
-				LocalizationPaths += FPaths::GetGameLocalizationPaths();
+				//UE_LOG(LogInit, Log, TEXT("Overriding culture %s w/ command-line option".), *CultureName);
 			}
-			LocalizationPaths += FPaths::GetEngineLocalizationPaths();
-
-			// Validate the locale has data or fallback to one that does.
-			TArray< FCultureRef > AvailableCultures;
-			I18N.GetCulturesWithAvailableLocalization(LocalizationPaths, AvailableCultures, false);
-
-			FString ValidCultureName;
-			ValidCultureName.Empty();
-			FCulturePtr TargetCulture = I18N.GetCulture(TargetCultureName);
-			if (TargetCulture.IsValid())
-			{
-				TArray<FString> PrioritizedParentCultureNames = TargetCulture->GetPrioritizedParentCultureNames();
-				for (const FString& CultureName : PrioritizedParentCultureNames)
+			else
+#endif // !UE_BUILD_SHIPPING
+#if WITH_EDITOR
+				// See if we've been provided a culture override in the editor
+				if(GConfig->GetString( TEXT("Internationalization"), TEXT("Culture"), RequestedCultureName, GEditorSettingsIni ))
 				{
-					FCulturePtr ValidCulture = I18N.GetCulture(CultureName);
-					if (ValidCulture.IsValid() && AvailableCultures.Contains(ValidCulture.ToSharedRef()))
+					//UE_LOG(LogInit, Log, TEXT("Overriding culture %s w/ editor configuration."), *CultureName);
+				}
+				else
+#endif // WITH_EDITOR
+					// Use culture specified in engine configuration.
+					if(GConfig->GetString( TEXT("Internationalization"), TEXT("Culture"), RequestedCultureName, GEngineIni ))
 					{
-						ValidCultureName = CultureName;
-						break;
+						//UE_LOG(LogInit, Log, TEXT("Overriding culture %s w/ engine configuration."), *CultureName);
+					}
+					else
+					{
+						RequestedCultureName = I18N.GetDefaultCulture()->GetName();
 					}
-				}
-			}
 
-			if(!ValidCultureName.IsEmpty())
-			{
-				if(RequestedCultureName != ValidCultureName)
-				{
-					// Make the user aware that the localization data belongs to a parent culture.
-					UE_LOG(LogTextLocalizationManager, Log, TEXT("No specific translations for ('%s') exist, so ('%s') translations will be used."), *RequestedCultureName, *ValidCultureName);
-				}
-			}
-			else
-			{
-				// Fallback to English.
-				UE_LOG(LogTextLocalizationManager, Log, TEXT("No translations for ('%s') exist, falling back to 'en' for localization and internationalization data."), *RequestedCultureName);
-				TargetCultureName = TEXT("en");
-			}
-		}
+					FString TargetCultureName = RequestedCultureName;
+					{
+						TArray<FString> LocalizationPaths;
+						if(ShouldLoadEditor)
+						{
+							LocalizationPaths += FPaths::GetEditorLocalizationPaths();
+						}
+						if(ShouldLoadGame)
+						{
+							LocalizationPaths += FPaths::GetGameLocalizationPaths();
+						}
+						LocalizationPaths += FPaths::GetEngineLocalizationPaths();
+
+						// Validate the locale has data or fallback to one that does.
+						TArray< FCultureRef > AvailableCultures;
+						I18N.GetCulturesWithAvailableLocalization(LocalizationPaths, AvailableCultures, false);
+
+						FString ValidCultureName;
+						ValidCultureName.Empty();
+						FCulturePtr TargetCulture = I18N.GetCulture(TargetCultureName);
+						if (TargetCulture.IsValid())
+						{
+							TArray<FString> PrioritizedParentCultureNames = TargetCulture->GetPrioritizedParentCultureNames();
+							for (const FString& CultureName : PrioritizedParentCultureNames)
+							{
+								FCulturePtr ValidCulture = I18N.GetCulture(CultureName);
+								if (ValidCulture.IsValid() && AvailableCultures.Contains(ValidCulture.ToSharedRef()))
+								{
+									ValidCultureName = CultureName;
+									break;
+								}
+							}
+						}
+
+						if(!ValidCultureName.IsEmpty())
+						{
+							if(RequestedCultureName != ValidCultureName)
+							{
+								// Make the user aware that the localization data belongs to a parent culture.
+								UE_LOG(LogTextLocalizationManager, Log, TEXT("No specific translations for ('%s') exist, so ('%s') translations will be used."), *RequestedCultureName, *ValidCultureName);
+							}
+						}
+						else
+						{
+							// Fallback to English.
+							UE_LOG(LogTextLocalizationManager, Log, TEXT("No translations for ('%s') exist, falling back to 'en' for localization and internationalization data."), *RequestedCultureName);
+							TargetCultureName = TEXT("en");
+						}
+					}
 
-		I18N.SetCurrentCulture(TargetCultureName);
+					I18N.SetCurrentCulture(TargetCultureName);
 	}
 
 	FTextLocalizationManager::Get().LoadLocalizationResourcesForCulture(I18N.GetCurrentCulture()->GetName(), ShouldLoadEditor, ShouldLoadGame);
@@ -238,8 +238,8 @@ void FTextLocalizationManager::FLocalizationEntryTracker::DetectAndLogConflicts(
 
 				UE_LOG(LogTextLocalizationManager, Warning, TEXT("Loaded localization resources contain conflicting entries for (Namespace:%s, Key:%s):\n%s"), *NamespaceName, *KeyName, *CollidingEntryListString);
 			}
-			}
 		}
+	}
 }
 
 void FTextLocalizationManager::FDisplayStringLookupTable::Find(const FString& InNamespace, FKeysTable*& OutKeysTableForNamespace, const FString& InKey, FDisplayStringEntry*& OutDisplayStringEntry)
@@ -306,7 +306,7 @@ FTextDisplayStringRef FTextLocalizationManager::GetDisplayString(const FString&
 	{
 		static bool bHasParsedCommandLine = false;
 		if(!bHasParsedCommandLine && FCommandLine::IsInitialized())
-	{
+		{
 			bShouldLEETIFYUnlocalizedString = FParse::Param(FCommandLine::Get(), TEXT("LEETIFYUnlocalized"));
 			bHasParsedCommandLine = true;
 		}
@@ -356,7 +356,7 @@ FTextDisplayStringRef FTextLocalizationManager::GetDisplayString(const FString&
 	{
 		// Don't log warnings about unlocalized strings if the system hasn't been initialized - we simply don't have localization data yet.
 		if( bIsInitialized )
-	{
+		{
 			UE_LOG(LogTextLocalizationManager, Verbose, TEXT("An attempt was made to get a localized string (Namespace:%s, Key:%s, Source:%s), but it did not exist."), *Namespace, *Key, SourceString ? **SourceString : TEXT(""));
 		}
 
@@ -381,7 +381,7 @@ FTextDisplayStringRef FTextLocalizationManager::GetDisplayString(const FString&
 				if ( UnlocalizedString->IsEmpty() )
 				{
 					if ( !bIsInitialized )
-		{
+					{
 						*(UnlocalizedString) = AccessedStringBeforeLocLoadedErrorMsg;
 					}
 				}
@@ -406,7 +406,7 @@ FTextDisplayStringRef FTextLocalizationManager::GetDisplayString(const FString&
 
 				NamespaceKeyLookupTable.Add(NewEntry.DisplayString, FNamespaceKeyEntry(Namespace, Key));
 
-		}
+			}
 
 			return UnlocalizedString;
 	}
@@ -455,13 +455,13 @@ bool FTextLocalizationManager::AddDisplayString(const FTextDisplayStringRef& Dis
 	DisplayStringLookupTable.Find(Namespace, KeysTableForExistingNamespace, Key, ExistingDisplayStringEntry);
 
 	// If there are any existing entries, they may cause a conflict, unless they're exactly the same as what we would be adding.
-	if ( (ExistingDisplayStringEntry && ExistingDisplayStringEntry->DisplayString != DisplayString) || // Namespace and key mustn't be associated with a different display string
-		 (ReverseLiveTableEntry && (ReverseLiveTableEntry->Namespace != Namespace || ReverseLiveTableEntry->Key != Key)) ) // Display string mustn't be associated with a different namespace and key.
+	if ( (ExistingDisplayStringEntry && ExistingDisplayStringEntry->DisplayString != DisplayString) || // Namespace and key mustn't be associated with a different display string.
+		(ReverseLiveTableEntry && (ReverseLiveTableEntry->Namespace != Namespace || ReverseLiveTableEntry->Key != Key)) ) // Display string mustn't be associated with a different namespace and key.
 	{
 		return false;
 	}
 
-	// Add the necesary associations in both directions.
+	// Add the necessary associations in both directions.
 	FDisplayStringLookupTable::FKeysTable& KeysTableForNamespace = DisplayStringLookupTable.NamespacesTable.FindOrAdd(Namespace);
 	KeysTableForNamespace.Add(Key, FDisplayStringLookupTable::FDisplayStringEntry(false, TEXT(""), FCrc::StrCrc32(*DisplayString.Get()), DisplayString));
 	NamespaceKeyLookupTable.Add(DisplayString, FNamespaceKeyEntry(Namespace, Key));
@@ -527,6 +527,14 @@ void FTextLocalizationManager::LoadFromManifestAndArchives( const FString& Confi
 
 	FString SectionName = TEXT("RegenerateResources");
 
+	// Get native culture.
+	FString NativeCulture;
+	if( !( GConfig->GetString( *SectionName, TEXT("NativeCulture"), NativeCulture, ConfigFilePath ) ) )
+	{
+		UE_LOG(LogTextLocalizationManager, Error, TEXT("No native culture specified."));
+		return;
+	}
+
 	// Get source path.
 	FString SourcePath;
 	if( !( GConfig->GetString( *SectionName, TEXT("SourcePath"), SourcePath, ConfigFilePath ) ) )
@@ -574,7 +582,7 @@ void FTextLocalizationManager::LoadFromManifestAndArchives( const FString& Confi
 	FString ConfigFullPath = FPaths::ConvertRelativePathToFull(ConfigFilePath);
 	FString EngineFullPath = FPaths::ConvertRelativePathToFull(FPaths::EngineConfigDir());
 	bool IsEngineManifest = false;
-	
+
 	if (ConfigFullPath.StartsWith(EngineFullPath))
 	{
 		IsEngineManifest = true;
@@ -627,7 +635,7 @@ void FTextLocalizationManager::LoadFromManifestAndArchives( const FString& Confi
 #endif
 
 		// Write resource.
-		FTextLocalizationResourceGenerator::Generate(SourcePath, InternationalizationManifest, LocaleNames[i], &(MemoryWriter), ArchiveSerializer);
+		FTextLocalizationResourceGenerator::Generate(SourcePath, InternationalizationManifest, NativeCulture, LocaleNames[i], &(MemoryWriter), ArchiveSerializer);
 
 		MemoryWriter.Close();
 	}
@@ -654,7 +662,7 @@ void FTextLocalizationManager::LoadFromManifestAndArchives( const FString& Confi
 void FTextLocalizationManager::UpdateFromLocalizationResource(const FString& LocalizationResourceFilePath)
 {
 	TArray<FLocalizationEntryTracker> LocalizationEntryTrackers;
-	
+
 	FLocalizationEntryTracker LocalizationEntryTracker = LocalizationEntryTrackers[LocalizationEntryTrackers.Add(FLocalizationEntryTracker())];
 	LocalizationEntryTracker.LoadFromFile(LocalizationResourceFilePath);
 	LocalizationEntryTracker.DetectAndLogConflicts();
diff --git a/Engine/Source/Runtime/Core/Private/Internationalization/TextLocalizationResourceGenerator.cpp b/Engine/Source/Runtime/Core/Private/Internationalization/TextLocalizationResourceGenerator.cpp
--- a/Engine/Source/Runtime/Core/Private/Internationalization/TextLocalizationResourceGenerator.cpp
+++ b/Engine/Source/Runtime/Core/Private/Internationalization/TextLocalizationResourceGenerator.cpp
@@ -140,29 +140,73 @@ bool FTextLocalizationResourceGenerator::FLocalizationEntryTracker::WriteToArchi
 	return WasSuccessful;
 }
 
-bool FTextLocalizationResourceGenerator::Generate(const FString& SourcePath, const TSharedRef<FInternationalizationManifest>& InternationalizationManifest, const FString& CultureToGenerate, FArchive* const DestinationArchive, IInternationalizationArchiveSerializer& ArchiveSerializer)
+bool FTextLocalizationResourceGenerator::Generate(const FString& SourcePath, const TSharedRef<FInternationalizationManifest>& InternationalizationManifest, const FString& NativeCulture, const FString& CultureToGenerate, FArchive* const DestinationArchive, IInternationalizationArchiveSerializer& ArchiveSerializer)
 {
 	FLocalizationEntryTracker LocalizationEntryTracker;
 
-	FString CulturePath = SourcePath / *(CultureToGenerate);
+	// Find archives in the native culture-specific folder.
+	const FString NativeCulturePath = SourcePath / *(CultureToGenerate);
+	TArray<FString> NativeArchiveFileNames;
+	IFileManager::Get().FindFiles(NativeArchiveFileNames, *(NativeCulturePath / TEXT("*.archive")), true, false);
 
 	// Find archives in the culture-specific folder.
+	const FString CulturePath = SourcePath / *(CultureToGenerate);
 	TArray<FString> ArchiveFileNames;
 	IFileManager::Get().FindFiles(ArchiveFileNames, *(CulturePath / TEXT("*.archive")), true, false);
 
+	if(NativeArchiveFileNames.Num() == 0)
+	{
+		FString Message = FString::Printf(TEXT("No archives were found for native culture %s."), *(CultureToGenerate));
+		UE_LOG(LogTextLocalizationResourceGenerator, Warning, TEXT("%s"), *Message);
+	}
+
 	if(ArchiveFileNames.Num() == 0)
 	{
 		FString Message = FString::Printf(TEXT("No archives were found for culture %s."), *(CultureToGenerate));
 		UE_LOG(LogTextLocalizationResourceGenerator, Warning, TEXT("%s"), *Message);
 	}
 
-	// For each archive:
-	for(int32 ArchiveIndex = 0; ArchiveIndex < ArchiveFileNames.Num(); ++ArchiveIndex)
+	TArray< TSharedPtr<FInternationalizationArchive> > NativeArchives;
+	for (const FString& NativeArchiveFileName : NativeArchiveFileNames)
 	{
-		const FString ArchiveName = ArchiveFileNames[ArchiveIndex];
+		// Read each archive file from the culture-named directory in the source path.
+		FString ArchiveFilePath = CulturePath / NativeArchiveFileName;
+		ArchiveFilePath = FPaths::ConvertRelativePathToFull(ArchiveFilePath);
+		TSharedRef<FInternationalizationArchive> InternationalizationArchive = MakeShareable(new FInternationalizationArchive);
+#if 0 // @todo Json: Serializing from FArchive is currently broken
+		FArchive* ArchiveFile = IFileManager::Get().CreateFileReader(*ArchiveFilePath);
+
+		if (ArchiveFile == nullptr)
+		{
+			UE_LOG(LogTextLocalizationResourceGenerator, Error, TEXT("No archive found at %s."), *ArchiveFilePath);
+			continue;
+		}
+
+		ArchiveSerializer.DeserializeArchive(*ArchiveFile, InternationalizationArchive);
+#else
+		FString ArchiveContent;
 
+		if (!FFileHelper::LoadFileToString(ArchiveContent, *ArchiveFilePath))
+		{
+			UE_LOG(LogTextLocalizationResourceGenerator, Error, TEXT("Failed to load file %s."), *ArchiveFilePath);
+			continue;
+		}
+
+		if (!ArchiveSerializer.DeserializeArchive(ArchiveContent, InternationalizationArchive))
+		{
+			UE_LOG(LogTextLocalizationResourceGenerator, Error, TEXT("Failed to serialize archive from file %s."), *ArchiveFilePath);
+			continue;
+		}
+#endif
+
+		NativeArchives.Add(InternationalizationArchive);
+	}
+
+	// For each archive:
+	for (const FString& ArchiveFileName : ArchiveFileNames)
+	{
 		// Read each archive file from the culture-named directory in the source path.
-		FString ArchiveFilePath = CulturePath / ArchiveName;
+		FString ArchiveFilePath = CulturePath / ArchiveFileName;
 		ArchiveFilePath = FPaths::ConvertRelativePathToFull(ArchiveFilePath);
 		TSharedRef<FInternationalizationArchive> InternationalizationArchive = MakeShareable(new FInternationalizationArchive);
 
@@ -216,8 +260,23 @@ bool FTextLocalizationResourceGenerator::Generate(const FString& SourcePath, con
 			{
 				const FString& Key = ContextIt->Key;
 
+				// Get proper source string to use - if the native source is different than the native translation, the native translation becomes the source.
+				FLocItem SourceToUse = Source;
+				for (const auto& NativeArchive : NativeArchives)
+				{
+					if (NativeArchive.IsValid())
+					{
+						const TSharedPtr<FArchiveEntry> NativeArchiveEntry = NativeArchive->FindEntryBySource(Namespace, Source, ContextIt->KeyMetadataObj);
+						if (NativeArchiveEntry.IsValid() && !NativeArchiveEntry->Source.IsExactMatch(NativeArchiveEntry->Translation))
+						{
+							SourceToUse = NativeArchiveEntry->Translation;
+							break;
+						}
+					}
+				}
+				\
 				// Find matching archive entry.
-				TSharedPtr<FArchiveEntry> ArchiveEntry = InternationalizationArchive->FindEntryBySource(Namespace, Source, ContextIt->KeyMetadataObj);
+				TSharedPtr<FArchiveEntry> ArchiveEntry = InternationalizationArchive->FindEntryBySource(Namespace, SourceToUse, ContextIt->KeyMetadataObj);
 
 				if(ContextIt->bIsOptional && (!ArchiveEntry.IsValid() || (ArchiveEntry.IsValid() && ArchiveEntry->Translation.Text.IsEmpty())))
 				{
@@ -229,7 +288,7 @@ bool FTextLocalizationResourceGenerator::Generate(const FString& SourcePath, con
 				if( ArchiveEntry.IsValid() )
 				{
 					UnescapedTranslatedString = ArchiveEntry->Translation.Text;
-
+					
 					if(UnescapedTranslatedString.IsEmpty())
 					{
 						if ( !MissingArchiveTranslationKeyStrings.IsEmpty() )
@@ -241,6 +300,12 @@ bool FTextLocalizationResourceGenerator::Generate(const FString& SourcePath, con
 				}
 				else
 				{
+					// If we're not using the actual source string, then embed the fake source string as the foreign translation when we have no actual foreign translation.
+					if (!Source.IsExactMatch(SourceToUse))
+					{
+						UnescapedTranslatedString = SourceToUse.Text;
+					}
+
 					if ( !MissingArchiveEntryKeyStrings.IsEmpty() )
 					{
 						MissingArchiveEntryKeyStrings += TEXT(", ");
diff --git a/Engine/Source/Runtime/Core/Public/Internationalization/TextLocalizationResourceGenerator.h b/Engine/Source/Runtime/Core/Public/Internationalization/TextLocalizationResourceGenerator.h
--- a/Engine/Source/Runtime/Core/Public/Internationalization/TextLocalizationResourceGenerator.h
+++ b/Engine/Source/Runtime/Core/Public/Internationalization/TextLocalizationResourceGenerator.h
@@ -5,7 +5,7 @@
 class FTextLocalizationResourceGenerator
 {
 public:
-	CORE_API static bool Generate(const FString& SourcePath, const TSharedRef<FInternationalizationManifest>& InternationalizationManifest, const FString& CultureToGenerate, FArchive* const DestinationArchive, IInternationalizationArchiveSerializer& ArchiveSerializer);
+	CORE_API static bool Generate(const FString& SourcePath, const TSharedRef<FInternationalizationManifest>& InternationalizationManifest, const FString& NativeCulture, const FString& CultureToGenerate, FArchive* const DestinationArchive, IInternationalizationArchiveSerializer& ArchiveSerializer);
 
 private:
 
