UE-17494: GitHub 1304 : [Core] Use FMemory for allocations done by zlib.

#1304

diff --git a/Engine/Source/Runtime/Core/Private/Misc/Compression.cpp b/Engine/Source/Runtime/Core/Private/Misc/Compression.cpp
--- a/Engine/Source/Runtime/Core/Private/Misc/Compression.cpp
+++ b/Engine/Source/Runtime/Core/Private/Misc/Compression.cpp
@@ -34,13 +34,23 @@ static bool appCompressMemoryZLIB( void* CompressedBuffer, int32& CompressedSize
 	return bOperationSucceeded;
 }
 
+static void *zalloc(void *opaque, unsigned int size, unsigned int num)
+{
+	return FMemory::Malloc(size * num);
+}
+
+static void zfree(void *opaque, void *p)
+{
+	FMemory::Free(p);
+}
+
 static bool appCompressMemoryGZIP(void* CompressedBuffer, int32& CompressedSize, const void* UncompressedBuffer, int32 UncompressedSize)
 {
 	DECLARE_SCOPE_CYCLE_COUNTER( TEXT( "Compress Memory GZIP" ), STAT_appCompressMemoryGZIP, STATGROUP_Compression );
 
 	z_stream gzipstream;
-	gzipstream.zalloc = Z_NULL;
-	gzipstream.zfree = Z_NULL;
+	gzipstream.zalloc = &zalloc;
+	gzipstream.zfree = &zfree;
 	gzipstream.opaque = Z_NULL;
 
 	// Setup input buffer
@@ -97,8 +107,29 @@ bool appUncompressMemoryZLIB( void* UncompressedBuffer, int32 UncompressedSize,
 	unsigned long ZCompressedSize	= CompressedSize;
 	unsigned long ZUncompressedSize	= UncompressedSize;
 	
+	z_stream stream;
+	stream.zalloc = &zalloc;
+	stream.zfree = &zfree;
+	stream.opaque = Z_NULL;
+	stream.next_in = (uint8*)CompressedBuffer;
+	stream.avail_in = ZCompressedSize;
+	stream.next_out = (uint8*)UncompressedBuffer;
+	stream.avail_out = ZUncompressedSize;
+
+	int32 Result = inflateInit(&stream);
+
+	if(Result != Z_OK)
+		return false;
+
 	// Uncompress data.
-	const int32 Result = uncompress((uint8*)UncompressedBuffer, &ZUncompressedSize, (const uint8*)CompressedBuffer, ZCompressedSize);
+	Result = inflate(&stream, Z_FINISH);
+
+	if(Result == Z_STREAM_END)
+	{
+		ZUncompressedSize = stream.total_out;
+	}
+
+	Result = inflateEnd(&stream);
 	
 	// These warnings will be compiled out in shipping.
 	UE_CLOG(Result == Z_MEM_ERROR, LogCompression, Warning, TEXT("appUncompressMemoryZLIB failed: Error: Z_MEM_ERROR, not enough memory!"));
