Stats2 - Fixed startup stats not collecting all data

diff --git a/Engine/Source/Runtime/Core/Private/Stats/Stats2.cpp b/Engine/Source/Runtime/Core/Private/Stats/Stats2.cpp
--- a/Engine/Source/Runtime/Core/Private/Stats/Stats2.cpp
+++ b/Engine/Source/Runtime/Core/Private/Stats/Stats2.cpp
@@ -1162,13 +1162,13 @@ void FThreadStats::CheckForCollectingStartupStats()
 	{
 		DirectStatsCommand( TEXT( "stat group enable LinkerLoad" ) );
 		DirectStatsCommand( TEXT( "stat group enable AsyncLoad" ) );
-		DirectStatsCommand(TEXT("stat dumpsum -start -ms=250 -num=240"));
+		DirectStatsCommand( TEXT( "stat dumpsum -start -ms=250 -num=240" ), true );
 	}
 	if (FParse::Param( FCommandLine::Get(), TEXT( "LoadTimeFile" ) ))
 	{
 		DirectStatsCommand( TEXT( "stat group enable LinkerLoad" ) );
 		DirectStatsCommand( TEXT( "stat group enable AsyncLoad" ) );
-		DirectStatsCommand( TEXT( "stat startfile" ) );
+		DirectStatsCommand( TEXT( "stat startfile" ), true );
 	}
 
 	// Now we can safely enable malloc profiler.
