Optimizations for FPaths::GetCleanFilename, GetPath and GetBaseFilename.

diff --git a/Engine/Source/Runtime/Core/Private/Misc/Paths.cpp b/Engine/Source/Runtime/Core/Private/Misc/Paths.cpp
--- a/Engine/Source/Runtime/Core/Private/Misc/Paths.cpp
+++ b/Engine/Source/Runtime/Core/Private/Misc/Paths.cpp
@@ -12,6 +12,12 @@ DEFINE_LOG_CATEGORY_STATIC(LogPaths, Log, All);
 	Path helpers for retrieving game dir, engine dir, etc.
 -----------------------------------------------------------------------------*/
 
+namespace UE4Paths_Private
+{
+	auto IsSlashOrBackslash    = [](TCHAR C) { return C == TEXT('/') || C == TEXT('\\'); };
+	auto IsNotSlashOrBackslash = [](TCHAR C) { return C != TEXT('/') && C != TEXT('\\'); };
+}
+
 bool FPaths::ShouldSaveToUserDir()
 {
 	static bool bShouldSaveToUserDir = FApp::IsInstalled() || FParse::Param(FCommandLine::Get(), TEXT("SaveToUserDir"));
@@ -405,22 +411,13 @@ FString FPaths::GetExtension( const FString& InPath, bool bIncludeDot )
 
 FString FPaths::GetCleanFilename(const FString& InPath)
 {
-	int32 Pos = InPath.Find(TEXT("/"), ESearchCase::CaseSensitive, ESearchDir::FromEnd);
-	// in case we are using backslashes on a platform that doesn't use backslashes
-	Pos = FMath::Max(Pos, InPath.Find(TEXT("\\"), ESearchCase::CaseSensitive, ESearchDir::FromEnd));
+	static_assert(INDEX_NONE == -1, "INDEX_NONE assumed to be -1");
 
-	if ( Pos != INDEX_NONE)
-	{
-		// if it was a trailing one, cut it (account for trailing whitespace?) and try removing path again
-		if (Pos == InPath.Len() - 1)
-		{
-			return GetCleanFilename(InPath.Left(Pos));
-		}
+	int32 EndPos   = InPath.FindLastCharByPredicate(UE4Paths_Private::IsNotSlashOrBackslash) + 1;
+	int32 StartPos = InPath.FindLastCharByPredicate(UE4Paths_Private::IsSlashOrBackslash, EndPos) + 1;
 
-		return InPath.Mid(Pos + 1);
-	}
-
-	return InPath;
+	FString Result = InPath.Mid(StartPos, EndPos - StartPos);
+	return Result;
 }
 
 FString FPaths::GetBaseFilename( const FString& InPath, bool bRemovePath )
@@ -428,20 +425,18 @@ FString FPaths::GetBaseFilename( const FString& InPath, bool bRemovePath )
 	FString Wk = bRemovePath ? GetCleanFilename(InPath) : InPath;
 
 	// remove the extension
-	int32 Pos = Wk.Find(TEXT("."), ESearchCase::CaseSensitive, ESearchDir::FromEnd);
+	int32 ExtPos = Wk.Find(TEXT("."), ESearchCase::CaseSensitive, ESearchDir::FromEnd);
 	
 	// determine the position of the path/leaf separator
 	int32 LeafPos = INDEX_NONE;
 	if (!bRemovePath)
 	{
-		LeafPos = Wk.Find(TEXT("/"), ESearchCase::CaseSensitive, ESearchDir::FromEnd);
-		// in case we are using backslashes on a platform that doesn't use backslashes
-		LeafPos = FMath::Max(LeafPos, Wk.Find(TEXT("\\"), ESearchCase::CaseSensitive, ESearchDir::FromEnd));
+		LeafPos = Wk.FindLastCharByPredicate(UE4Paths_Private::IsSlashOrBackslash);
 	}
 
-	if (Pos != INDEX_NONE && (LeafPos == INDEX_NONE || Pos > LeafPos))
+	if (ExtPos != INDEX_NONE && (LeafPos == INDEX_NONE || ExtPos > LeafPos))
 	{
-		return Wk.Left(Pos);
+		Wk = Wk.Left(ExtPos);
 	}
 
 	return Wk;
@@ -449,15 +444,15 @@ FString FPaths::GetBaseFilename( const FString& InPath, bool bRemovePath )
 
 FString FPaths::GetPath(const FString& InPath)
 {
-	int32 Pos = InPath.Find(TEXT("/"), ESearchCase::CaseSensitive, ESearchDir::FromEnd);
-	// in case we are using backslashes on a platform that doesn't use backslashes
-	Pos = FMath::Max(Pos, InPath.Find(TEXT("\\"), ESearchCase::CaseSensitive, ESearchDir::FromEnd));
-	if ( Pos != INDEX_NONE )
+	int32 Pos = InPath.FindLastCharByPredicate(UE4Paths_Private::IsSlashOrBackslash);
+
+	FString Result;
+	if (Pos != INDEX_NONE)
 	{
-		return InPath.Left(Pos);
+		Result = InPath.Left(Pos);
 	}
 
-	return TEXT("");
+	return Result;
 }
 
 bool FPaths::FileExists(const FString& InPath)
