More text filter improvements

The "-" token as been removed from the list of general NOT operator aliases, and is now only available as TextCmpInvert. This means that negative numbers will now parse correctly, but also means that we can now handle complex expressions like "name=-Thing" as a user might expect (equivalent to "name!=thing").

The expression is now evaluated once after compilation (using a dummy context) so that syntax error information can be reported back to the user. Previously evaluation error information was being lost.

diff --git a/Engine/Source/Runtime/Core/Private/Math/BasicMathExpressionEvaluator.cpp b/Engine/Source/Runtime/Core/Private/Math/BasicMathExpressionEvaluator.cpp
--- a/Engine/Source/Runtime/Core/Private/Math/BasicMathExpressionEvaluator.cpp
+++ b/Engine/Source/Runtime/Core/Private/Math/BasicMathExpressionEvaluator.cpp
@@ -19,7 +19,7 @@ namespace ExpressionParser
 
 	TOptional<FStringToken> ParseNumber(const FTokenStream& InStream, FStringToken* Accumulate)
 	{
-		enum class EState { LeadIn, Integer, Dot, Fractional };
+		enum class EState { LeadIn, Sign, Integer, Dot, Fractional };
 
 		EState State = EState::LeadIn;
 		return InStream.ParseToken([&](TCHAR InC){
@@ -31,6 +31,34 @@ namespace ExpressionParser
 					State = EState::Integer;
 					return EParseState::Continue;
 				}
+				else if (InC == '+' || InC == '-')
+				{
+					State = EState::Sign;
+					return EParseState::Continue;
+				}
+				else if (InC == '.')
+				{
+					State = EState::Dot;
+					return EParseState::Continue;
+				}
+				else
+				{
+					// Not a number
+					return EParseState::Cancel;
+				}
+			}
+			else if (State == EState::Sign)
+			{
+				if (FChar::IsDigit(InC))
+				{
+					State = EState::Integer;
+					return EParseState::Continue;
+				}
+				else if (InC == '.')
+				{
+					State = EState::Dot;
+					return EParseState::Continue;
+				}
 				else
 				{
 					// Not a number
diff --git a/Engine/Source/Runtime/Core/Private/Misc/TextFilterExpressionEvaluator.cpp b/Engine/Source/Runtime/Core/Private/Misc/TextFilterExpressionEvaluator.cpp
--- a/Engine/Source/Runtime/Core/Private/Misc/TextFilterExpressionEvaluator.cpp
+++ b/Engine/Source/Runtime/Core/Private/Misc/TextFilterExpressionEvaluator.cpp
@@ -9,30 +9,30 @@ namespace TextFilterExpressionParser
 	class FTextToken
 	{
 	public:
-		enum class EInvertBasicStringResult : uint8
+		enum class EInvertResult : uint8
 		{
 			No,
 			Yes,
 		};
 
-		FTextToken(FTextFilterString InString, ETextFilterTextComparisonMode InTextComparisonMode, const EInvertBasicStringResult InInvertBasicStringResult)
+		FTextToken(FTextFilterString InString, ETextFilterTextComparisonMode InTextComparisonMode, const EInvertResult InInvertResult)
 			: String(MoveTemp(InString))
 			, TextComparisonMode(MoveTemp(InTextComparisonMode))
-			, InvertBasicStringResult(InInvertBasicStringResult)
+			, InvertResult(InInvertResult)
 		{
 		}
 
 		FTextToken(const FTextToken& Other)
 			: String(Other.String)
 			, TextComparisonMode(Other.TextComparisonMode)
-			, InvertBasicStringResult(Other.InvertBasicStringResult)
+			, InvertResult(Other.InvertResult)
 		{
 		}
 
 		FTextToken(FTextToken&& Other)
 			: String(MoveTemp(Other.String))
 			, TextComparisonMode(Other.TextComparisonMode)
-			, InvertBasicStringResult(Other.InvertBasicStringResult)
+			, InvertResult(Other.InvertResult)
 		{
 		}
 
@@ -40,7 +40,7 @@ namespace TextFilterExpressionParser
 		{
 			String = Other.String;
 			TextComparisonMode = Other.TextComparisonMode;
-			InvertBasicStringResult = Other.InvertBasicStringResult;
+			InvertResult = Other.InvertResult;
 			return *this;
 		}
 
@@ -48,7 +48,7 @@ namespace TextFilterExpressionParser
 		{
 			String = MoveTemp(Other.String);
 			TextComparisonMode = Other.TextComparisonMode;
-			InvertBasicStringResult = Other.InvertBasicStringResult;
+			InvertResult = Other.InvertResult;
 			return *this;
 		}
 
@@ -57,11 +57,6 @@ namespace TextFilterExpressionParser
 			return String;
 		}
 
-		FName GetStringAsName() const
-		{
-			return String.AsName();
-		}
-
 		bool EvaluateAsBasicStringExpression(const ITextFilterExpressionContext* InContext) const
 		{
 			if (String.IsEmpty())
@@ -71,7 +66,13 @@ namespace TextFilterExpressionParser
 			}
 
 			const bool bMatched = InContext->TestBasicStringExpression(String, TextComparisonMode);
-			return (InvertBasicStringResult == EInvertBasicStringResult::No) ? bMatched : !bMatched;
+			return (InvertResult == EInvertResult::No) ? bMatched : !bMatched;
+		}
+
+		bool EvaluateAsComplexExpression(const ITextFilterExpressionContext* InContext, const FTextFilterString& InKey, const ETextFilterComparisonOperation InComparisonOperation) const
+		{
+			const bool bMatched = InContext->TestComplexExpression(InKey.AsName(), String, InComparisonOperation, TextComparisonMode);
+			return (InvertResult == EInvertResult::No) ? bMatched : !bMatched;
 		}
 
 		ETextFilterTextComparisonMode GetTextComparisonMode() const
@@ -82,9 +83,9 @@ namespace TextFilterExpressionParser
 	private:
 		FTextFilterString String;
 		ETextFilterTextComparisonMode TextComparisonMode;
-		EInvertBasicStringResult InvertBasicStringResult;
+		EInvertResult InvertResult;
 	};
-}
+} // namespace TextFilterExpressionParser
 
 #define DEFINE_TEXT_EXPRESSION_OPERATOR_NODE(TYPE, MONIKER_COUNT, ...)	\
 	namespace TextFilterExpressionParser								\
@@ -109,7 +110,7 @@ DEFINE_TEXT_EXPRESSION_OPERATOR_NODE(FGreaterOrEqual,			2,	0x09D75C5E, 0xA29A419
 
 DEFINE_TEXT_EXPRESSION_OPERATOR_NODE(FOr,						3,	0xF4778B51, 0xF535414D, 0x9C0EB5F2, 0x2F2B0FD4)
 DEFINE_TEXT_EXPRESSION_OPERATOR_NODE(FAnd,						3,	0x7511397A, 0x02D24DC2, 0x86729800, 0xF454C320)
-DEFINE_TEXT_EXPRESSION_OPERATOR_NODE(FNot,						3,	0x03D78990, 0x41D04E26, 0x8E98AD2F, 0x74667868)
+DEFINE_TEXT_EXPRESSION_OPERATOR_NODE(FNot,						2,	0x03D78990, 0x41D04E26, 0x8E98AD2F, 0x74667868)
 
 DEFINE_EXPRESSION_NODE_TYPE(TextFilterExpressionParser::FTextToken,	0x09E49538, 0x633545E3, 0x84B5644F, 0x1F11628F);
 DEFINE_EXPRESSION_NODE_TYPE(bool,									0xC1CD5DCF, 0x2AB44958, 0xB3FF4F8F, 0xE665D121);
@@ -118,8 +119,7 @@ namespace TextFilterExpressionParser
 {
 	/**
 	 * This contains all the symbols that can define breaking points between text and an operator
-	 * Note: We don't include + and - in this list as these are valid to use inside names, and we should consume 
-	 *		 them as part of the text token as people usually put a space between their search terms
+	 * Note: We don't include + and - in this list as these are valid to use inside text and numbers, and should be consumed as part of the text token
 	 */
 	static const TCHAR BasicTextBreakingCharacters[]	= { '(', ')', '!', '&', '|', ' ' };						// ETextFilterExpressionEvaluatorMode::BasicString
 	static const TCHAR ComplexTextBreakingCharacters[]	= { '(', ')', '=', ':', '<', '>', '!', '&', '|', ' ' };	// ETextFilterExpressionEvaluatorMode::Complex
@@ -136,7 +136,7 @@ namespace TextFilterExpressionParser
 
 	const TCHAR* FOr::Monikers[]						= { TEXT("OR"), TEXT("||"), TEXT("|") };
 	const TCHAR* FAnd::Monikers[]						= { TEXT("AND"), TEXT("&&"), TEXT("&") };
-	const TCHAR* FNot::Monikers[]						= { TEXT("NOT"), TEXT("!"), TEXT("-") };
+	const TCHAR* FNot::Monikers[]						= { TEXT("NOT"), TEXT("!") };
 
 	/** Consume an operator from the specified consumer's stream, if one exists at the current read position */
 	template<typename TSymbol>
@@ -156,6 +156,21 @@ namespace TextFilterExpressionParser
 		return TOptional<FExpressionError>();
 	}
 
+	/** Consume a number from the specified consumer's stream, if one exists at the current read position */
+	TOptional<FExpressionError> ConsumeNumber(FExpressionTokenConsumer& Consumer)
+	{
+		auto& Stream = Consumer.GetStream();
+
+		TOptional<FStringToken> NumberToken = ExpressionParser::ParseNumber(Stream);
+		
+		if (NumberToken.IsSet())
+		{
+			Consumer.Add(NumberToken.GetValue(), FExpressionNode(FTextToken(NumberToken.GetValue().GetString(), ETextFilterTextComparisonMode::Partial, FTextToken::EInvertResult::No)));
+		}
+
+		return TOptional<FExpressionError>();
+	}
+
 	/** Transform the given string to remove any escape character sequences found in a quoted string */
 	void UnescapeQuotedString(FString& Str, const TCHAR InQuoteChar)
 	{
@@ -219,11 +234,11 @@ namespace TextFilterExpressionParser
 		{
 			FString FinalString = TextToken.GetValue().GetString();
 			UnescapeQuotedString(FinalString, QuoteChar);
-			Consumer.Add(TextTokenWithQuotes, FExpressionNode(FTextToken(MoveTemp(FinalString), ETextFilterTextComparisonMode::Partial, FTextToken::EInvertBasicStringResult::No)));
+			Consumer.Add(TextTokenWithQuotes, FExpressionNode(FTextToken(MoveTemp(FinalString), ETextFilterTextComparisonMode::Partial, FTextToken::EInvertResult::No)));
 		}
 		else
 		{
-			Consumer.Add(TextTokenWithQuotes, FExpressionNode(FTextToken(FString(), ETextFilterTextComparisonMode::Partial, FTextToken::EInvertBasicStringResult::No)));
+			Consumer.Add(TextTokenWithQuotes, FExpressionNode(FTextToken(FString(), ETextFilterTextComparisonMode::Partial, FTextToken::EInvertResult::No)));
 		}
 
 		return TOptional<FExpressionError>();
@@ -254,13 +269,12 @@ namespace TextFilterExpressionParser
 			InString.RemoveAt(InString.Len() - 3, 3, false);
 		}
 
-		// To preserve behavior with the old text filter, the final string *may* also start with a NOT (-) operator (after stripping the TextCmpExact or TextCmpAnchor tokens from the start)
-		// If it does, then we need to invert the logic when performing simple string comparisons
-		FTextToken::EInvertBasicStringResult InvertBasicStringResult = FTextToken::EInvertBasicStringResult::No;
+		// To preserve behavior with the old text filter, the final string may also contain a TextCmpInvert (-) operator (after stripping the TextCmpExact or TextCmpAnchor tokens from the start)
+		FTextToken::EInvertResult InvertResult = FTextToken::EInvertResult::No;
 		if (InString.Len() > 0 && InString[0] == '-')
 		{
-			// Matched an invert result request - remove the - token from the start of the string
-			InvertBasicStringResult = FTextToken::EInvertBasicStringResult::Yes;
+			// Matched TextCmpInvert - remove the - token from the start of the string
+			InvertResult = FTextToken::EInvertResult::Yes;
 			InString.RemoveAt(0, 1, false);
 		}
 
@@ -276,7 +290,7 @@ namespace TextFilterExpressionParser
 			}
 		}
 
-		return FTextToken(MoveTemp(InString), TextComparisonMode, InvertBasicStringResult);
+		return FTextToken(MoveTemp(InString), TextComparisonMode, InvertResult);
 	}
 
 	/** Consume the text from the specified consumer's stream */
@@ -381,6 +395,20 @@ namespace TextFilterExpressionParser
 
 #undef DEFINE_TEXT_EXPRESSION_OPERATOR_NODE
 
+/** Dummy context used for testing that the compiled expression is syntactically valid */
+class FDummyTextFilterExpressionContext : public ITextFilterExpressionContext
+{
+public:
+	virtual bool TestBasicStringExpression(const FTextFilterString& InValue, const ETextFilterTextComparisonMode InTextComparisonMode) const override
+	{
+		return false;
+	}
+	virtual bool TestComplexExpression(const FName& InKey, const FTextFilterString& InValue, const ETextFilterComparisonOperation InComparisonOperation, const ETextFilterTextComparisonMode InTextComparisonMode) const override
+	{
+		return false;
+	}
+};
+
 FTextFilterExpressionEvaluator::FTextFilterExpressionEvaluator(const ETextFilterExpressionEvaluatorMode InMode)
 	: ExpressionEvaluatorMode(InMode)
 	, FilterType(ETextFilterExpressionType::Empty)
@@ -527,6 +555,10 @@ bool FTextFilterExpressionEvaluator::SetFilterText(const FText& InFilterText)
 			{
 				FilterType = (bIsComplexExpression) ? ETextFilterExpressionType::Complex : ETextFilterExpressionType::BasicString;
 				FilterErrorText = FText::GetEmpty();
+
+				// Evaluate the compiled filter with a dummy context - this will let us know whether it's syntactically valid
+				FDummyTextFilterExpressionContext DummyContext;
+				EvaluateCompiledExpression(CompiledFilter.GetValue(), DummyContext, &FilterErrorText);
 			}
 			else
 			{
@@ -549,9 +581,6 @@ FText FTextFilterExpressionEvaluator::GetFilterErrorText() const
 
 bool FTextFilterExpressionEvaluator::TestTextFilter(const ITextFilterExpressionContext& InContext) const
 {
-	using namespace TextFilterExpressionParser;
-	using TextFilterExpressionParser::FTextToken;
-
 	if (FilterType == ETextFilterExpressionType::Empty)
 	{
 		return true;
@@ -559,22 +588,7 @@ bool FTextFilterExpressionEvaluator::TestTextFilter(const ITextFilterExpressionC
 
 	if (CompiledFilter.IsSet())
 	{
-		auto& CompiledResult = CompiledFilter.GetValue();
-		if (CompiledResult.IsValid())
-		{
-			auto EvalResult = ExpressionParser::Evaluate(CompiledResult.GetValue(), JumpTable, &InContext);
-			if (EvalResult.IsValid())
-			{
-				if (const bool* BoolResult = EvalResult.GetValue().Cast<bool>())
-				{
-					return *BoolResult;
-				}
-				else if (const FTextToken* TextResult = EvalResult.GetValue().Cast<FTextToken>())
-				{
-					return TextResult->EvaluateAsBasicStringExpression(&InContext);
-				}
-			}
-		}
+		return EvaluateCompiledExpression(CompiledFilter.GetValue(), InContext, nullptr);
 	}
 
 	return false;
@@ -601,6 +615,7 @@ void FTextFilterExpressionEvaluator::ConstructExpressionParser()
 	TokenDefinitions.DefineToken(&ConsumeOperator<FOr>);
 	TokenDefinitions.DefineToken(&ConsumeOperator<FAnd>);
 	TokenDefinitions.DefineToken(&ConsumeOperator<FNot>);
+	TokenDefinitions.DefineToken(&ConsumeNumber);
 	TokenDefinitions.DefineToken(&ConsumeQuotedText);
 	TokenDefinitions.DefineToken((ExpressionEvaluatorMode == ETextFilterExpressionEvaluatorMode::Complex) ? &ConsumeComplexText : &ConsumeBasicText);
 
@@ -615,12 +630,12 @@ void FTextFilterExpressionEvaluator::ConstructExpressionParser()
 	Grammar.DefineBinaryOperator<FAnd>(2);
 	Grammar.DefinePreUnaryOperator<FNot>();
 
-	JumpTable.MapBinary<FLessOrEqual>([](const FTextToken& A, const FTextToken& B, const ITextFilterExpressionContext* InContext)		{ return InContext->TestComplexExpression(A.GetStringAsName(), B.GetString(), ETextFilterComparisonOperation::LessOrEqual, B.GetTextComparisonMode()); });
-	JumpTable.MapBinary<FLess>([](const FTextToken& A, const FTextToken& B, const ITextFilterExpressionContext* InContext)				{ return InContext->TestComplexExpression(A.GetStringAsName(), B.GetString(), ETextFilterComparisonOperation::Less, B.GetTextComparisonMode()); });
-	JumpTable.MapBinary<FGreaterOrEqual>([](const FTextToken& A, const FTextToken& B, const ITextFilterExpressionContext* InContext)	{ return InContext->TestComplexExpression(A.GetStringAsName(), B.GetString(), ETextFilterComparisonOperation::GreaterOrEqual, B.GetTextComparisonMode()); });
-	JumpTable.MapBinary<FGreater>([](const FTextToken& A, const FTextToken& B, const ITextFilterExpressionContext* InContext)			{ return InContext->TestComplexExpression(A.GetStringAsName(), B.GetString(), ETextFilterComparisonOperation::Greater, B.GetTextComparisonMode()); });
-	JumpTable.MapBinary<FNotEqual>([](const FTextToken& A, const FTextToken& B, const ITextFilterExpressionContext* InContext)			{ return InContext->TestComplexExpression(A.GetStringAsName(), B.GetString(), ETextFilterComparisonOperation::NotEqual, B.GetTextComparisonMode()); });
-	JumpTable.MapBinary<FEqual>([](const FTextToken& A, const FTextToken& B, const ITextFilterExpressionContext* InContext)				{ return InContext->TestComplexExpression(A.GetStringAsName(), B.GetString(), ETextFilterComparisonOperation::Equal, B.GetTextComparisonMode()); });
+	JumpTable.MapBinary<FLessOrEqual>([](const FTextToken& A, const FTextToken& B, const ITextFilterExpressionContext* InContext)		{ return B.EvaluateAsComplexExpression(InContext, A.GetString(), ETextFilterComparisonOperation::LessOrEqual); });
+	JumpTable.MapBinary<FLess>([](const FTextToken& A, const FTextToken& B, const ITextFilterExpressionContext* InContext)				{ return B.EvaluateAsComplexExpression(InContext, A.GetString(), ETextFilterComparisonOperation::Less); });
+	JumpTable.MapBinary<FGreaterOrEqual>([](const FTextToken& A, const FTextToken& B, const ITextFilterExpressionContext* InContext)	{ return B.EvaluateAsComplexExpression(InContext, A.GetString(), ETextFilterComparisonOperation::GreaterOrEqual); });
+	JumpTable.MapBinary<FGreater>([](const FTextToken& A, const FTextToken& B, const ITextFilterExpressionContext* InContext)			{ return B.EvaluateAsComplexExpression(InContext, A.GetString(), ETextFilterComparisonOperation::Greater); });
+	JumpTable.MapBinary<FNotEqual>([](const FTextToken& A, const FTextToken& B, const ITextFilterExpressionContext* InContext)			{ return B.EvaluateAsComplexExpression(InContext, A.GetString(), ETextFilterComparisonOperation::NotEqual); });
+	JumpTable.MapBinary<FEqual>([](const FTextToken& A, const FTextToken& B, const ITextFilterExpressionContext* InContext)				{ return B.EvaluateAsComplexExpression(InContext, A.GetString(), ETextFilterComparisonOperation::Equal); });
 
 	JumpTable.MapBinary<FOr>([](const FTextToken& A, const FTextToken& B, const ITextFilterExpressionContext* InContext)				{ return A.EvaluateAsBasicStringExpression(InContext) || B.EvaluateAsBasicStringExpression(InContext); });
 	JumpTable.MapBinary<FOr>([](const FTextToken& A, bool B, const ITextFilterExpressionContext* InContext)								{ return A.EvaluateAsBasicStringExpression(InContext) || B; });
@@ -635,3 +650,31 @@ void FTextFilterExpressionEvaluator::ConstructExpressionParser()
 	JumpTable.MapPreUnary<FNot>([](const FTextToken& V, const ITextFilterExpressionContext* InContext)									{ return !V.EvaluateAsBasicStringExpression(InContext); });
 	JumpTable.MapPreUnary<FNot>([](bool V, const ITextFilterExpressionContext* InContext)												{ return !V; });
 }
+
+bool FTextFilterExpressionEvaluator::EvaluateCompiledExpression(const ExpressionParser::CompileResultType& InCompiledResult, const ITextFilterExpressionContext& InContext, FText* OutErrorText) const
+{
+	using namespace TextFilterExpressionParser;
+	using TextFilterExpressionParser::FTextToken;
+
+	if (InCompiledResult.IsValid())
+	{
+		auto EvalResult = ExpressionParser::Evaluate(InCompiledResult.GetValue(), JumpTable, &InContext);
+		if (EvalResult.IsValid())
+		{
+			if (const bool* BoolResult = EvalResult.GetValue().Cast<bool>())
+			{
+				return *BoolResult;
+			}
+			else if (const FTextToken* TextResult = EvalResult.GetValue().Cast<FTextToken>())
+			{
+				return TextResult->EvaluateAsBasicStringExpression(&InContext);
+			}
+		}
+		else if (OutErrorText)
+		{
+			*OutErrorText = EvalResult.GetError().Text;
+		}
+	}
+
+	return false;
+}
diff --git a/Engine/Source/Runtime/Core/Private/Misc/TextFilterTests.cpp b/Engine/Source/Runtime/Core/Private/Misc/TextFilterTests.cpp
--- a/Engine/Source/Runtime/Core/Private/Misc/TextFilterTests.cpp
+++ b/Engine/Source/Runtime/Core/Private/Misc/TextFilterTests.cpp
@@ -16,10 +16,14 @@ namespace TextFilterTests
 			BasicStrings.Add(TEXT("Type'/Path/To/Asset.Asset'"));
 			BasicStrings.Add(TEXT("Other'/Path/To/Asset.Asset'FollowingText"));
 			BasicStrings.Add(TEXT("Funky<String>"));
+			BasicStrings.Add(TEXT("My-Item"));
+			BasicStrings.Add(TEXT("My+Item"));
+			BasicStrings.Add(TEXT("My.Item"));
 
 			KeyValuePairs.Add("StringKey", TEXT("Test"));
 			KeyValuePairs.Add("IntKey", TEXT("123"));
 			KeyValuePairs.Add("FloatKey", TEXT("456.789"));
+			KeyValuePairs.Add("NegFloatKey", TEXT("-456.789"));
 		}
 
 		// ITextFilterExpressionContext API - used for testing FTextFilterExpressionEvaluator
@@ -87,7 +91,11 @@ namespace TextFilterTests
 
 			if (bActual != bExpected)
 			{
-				Test->AddError(FString::Printf(TEXT("Filter expression '%s' evaluated incorrectly: %d != %d"), InFilterExpression, bActual, bExpected));
+				Test->AddError(FString::Printf(TEXT("Filter expression '%s' evaluated incorrectly: Expected: %s, Actual: %s"), 
+					InFilterExpression, 
+					(bExpected) ? TEXT("true") : TEXT("false"), 
+					(bActual) ? TEXT("true") : TEXT("false"))
+					);
 				bResult = false;
 			}
 
@@ -162,6 +170,12 @@ namespace TextFilterTests
 		bResult &= InTestPayload.TestFilterExpression(TEXT("Type'/Path/To/Asset.Asset'"), true);
 		bResult &= InTestPayload.TestFilterExpression(TEXT("Other'/Path/To/Asset.Asset'FollowingText"), true);
 		bResult &= InTestPayload.TestFilterExpression(TEXT("\"Funky<String>\""), true);
+		bResult &= InTestPayload.TestFilterExpression(TEXT("My-Item"), true);
+		bResult &= InTestPayload.TestFilterExpression(TEXT("+My-Item"), true);
+		bResult &= InTestPayload.TestFilterExpression(TEXT("My+Item"), true);
+		bResult &= InTestPayload.TestFilterExpression(TEXT("+My+Item"), true);
+		bResult &= InTestPayload.TestFilterExpression(TEXT("My.Item"), true);
+		bResult &= InTestPayload.TestFilterExpression(TEXT("+My.Item"), true);
 
 		return bResult;
 	}
@@ -181,7 +195,7 @@ namespace TextFilterTests
 		bool bResult = true;
 
 		bResult &= TestAllCommonFilterExpressions(InTestPayload);
-		bResult &= InTestPayload.TestFilterExpression(TEXT("Funky<String>"), false);
+		bResult &= InTestPayload.TestFilterExpression(TEXT("'Funky<String>'"), true);
 		bResult &= InTestPayload.TestFilterExpression(TEXT("StringKey=Test"), true);
 		bResult &= InTestPayload.TestFilterExpression(TEXT("StringKey!=Test"), false);
 		bResult &= InTestPayload.TestFilterExpression(TEXT("IntKey=123"), true);
@@ -190,6 +204,9 @@ namespace TextFilterTests
 		bResult &= InTestPayload.TestFilterExpression(TEXT("FloatKey=456.789"), true);
 		bResult &= InTestPayload.TestFilterExpression(TEXT("FloatKey>456.789"), false);
 		bResult &= InTestPayload.TestFilterExpression(TEXT("FloatKey<456.789"), false);
+		bResult &= InTestPayload.TestFilterExpression(TEXT("NegFloatKey=-456.789"), true);
+		bResult &= InTestPayload.TestFilterExpression(TEXT("NegFloatKey>-456.789"), false);
+		bResult &= InTestPayload.TestFilterExpression(TEXT("NegFloatKey<-456.789"), false);
 		bResult &= InTestPayload.TestFilterExpression(TEXT("IntKey==300 || FloatKey==456.789"), true);
 		bResult &= InTestPayload.TestFilterExpression(TEXT("IntKey==300 && FloatKey==456.789"), false);
 		bResult &= InTestPayload.TestFilterExpression(TEXT("(IntKey==300 && FloatKey==456.789) OR StringKey==Test"), true);
diff --git a/Engine/Source/Runtime/Core/Public/Misc/TextFilterExpressionEvaluator.h b/Engine/Source/Runtime/Core/Public/Misc/TextFilterExpressionEvaluator.h
--- a/Engine/Source/Runtime/Core/Public/Misc/TextFilterExpressionEvaluator.h
+++ b/Engine/Source/Runtime/Core/Public/Misc/TextFilterExpressionEvaluator.h
@@ -62,6 +62,9 @@ private:
 	/** Common function to construct the expression parser */
 	void ConstructExpressionParser();
 
+	/** Evaluate the given compiled result, and optionally populate OutErrorText with any error information */
+	bool EvaluateCompiledExpression(const ExpressionParser::CompileResultType& InCompiledResult, const ITextFilterExpressionContext& InContext, FText* OutErrorText) const;
+
 	/** Defines whether or not the expression parser can evaluate complex expressions */
 	ETextFilterExpressionEvaluatorMode ExpressionEvaluatorMode;
 
