Fix hot reload not working correctly with UnrealGameSync; having a build changelist does not necessarily mean that the build is promoted any more. Renamed FEngineVersionBase::IsPromotedBuild() to FEngineVersionBase::HasChangelist(), and use FApp::GetEngineIsPromotedBuild() instead.

diff --git a/Engine/Source/Developer/DesktopPlatform/Private/DesktopPlatformBase.cpp b/Engine/Source/Developer/DesktopPlatform/Private/DesktopPlatformBase.cpp
--- a/Engine/Source/Developer/DesktopPlatform/Private/DesktopPlatformBase.cpp
+++ b/Engine/Source/Developer/DesktopPlatform/Private/DesktopPlatformBase.cpp
@@ -335,7 +335,7 @@ bool FDesktopPlatformBase::GetEngineIdentifierForProject(const FString& ProjectF
 	if(EngineVersionString.Len() > 0)
 	{
 		FEngineVersion EngineVersion;
-		if(FEngineVersion::Parse(EngineVersionString, EngineVersion) && EngineVersion.IsPromotedBuild() && EngineVersion.ToString(EVersionComponent::Minor) == TEXT("4.0"))
+		if(FEngineVersion::Parse(EngineVersionString, EngineVersion) && EngineVersion.HasChangelist() && EngineVersion.ToString(EVersionComponent::Minor) == TEXT("4.0"))
 		{
 			OutIdentifier = TEXT("4.0");
 			return true;
diff --git a/Engine/Source/Editor/LevelEditor/Private/LevelEditor.cpp b/Engine/Source/Editor/LevelEditor/Private/LevelEditor.cpp
--- a/Engine/Source/Editor/LevelEditor/Private/LevelEditor.cpp
+++ b/Engine/Source/Editor/LevelEditor/Private/LevelEditor.cpp
@@ -59,7 +59,7 @@ public:
 		FColor BadgeTextColor = FColor(128,128,128,255);
 		GConfig->GetColor(TEXT("LevelEditor"), TEXT("ProjectBadgeTextColor"), /*out*/ BadgeTextColor, GEditorPerProjectIni);
 
-		const FString EngineVersionString = GEngineVersion.ToString(GEngineVersion.IsPromotedBuild() ? EVersionComponent::Changelist : EVersionComponent::Patch);
+		const FString EngineVersionString = GEngineVersion.ToString(GEngineVersion.HasChangelist() ? EVersionComponent::Changelist : EVersionComponent::Patch);
 		
 		FFormatNamedArguments Args;
 		Args.Add(TEXT("Branch"), FText::FromString(OptionalBranchPrefix));
diff --git a/Engine/Source/Editor/LevelEditor/Private/LevelEditorActions.cpp b/Engine/Source/Editor/LevelEditor/Private/LevelEditorActions.cpp
--- a/Engine/Source/Editor/LevelEditor/Private/LevelEditorActions.cpp
+++ b/Engine/Source/Editor/LevelEditor/Private/LevelEditorActions.cpp
@@ -878,7 +878,7 @@ bool FLevelEditorActionCallbacks::Recompile_CanExecute()
 	// We're not able to recompile if a compile is already in progress!
 
 	IHotReloadInterface& HotReloadSupport = FModuleManager::LoadModuleChecked<IHotReloadInterface>(HotReloadModule);
-	return !HotReloadSupport.IsCurrentlyCompiling() && !(GEngineVersion.IsPromotedBuild() && FEngineBuildSettings::IsPerforceBuild());
+	return !HotReloadSupport.IsCurrentlyCompiling() && !(FApp::GetEngineIsPromotedBuild() && FEngineBuildSettings::IsPerforceBuild());
 }
 
 void FLevelEditorActionCallbacks::ConnectToSourceControl_Clicked()
diff --git a/Engine/Source/Runtime/Core/Private/Misc/EngineVersion.cpp b/Engine/Source/Runtime/Core/Private/Misc/EngineVersion.cpp
--- a/Engine/Source/Runtime/Core/Private/Misc/EngineVersion.cpp
+++ b/Engine/Source/Runtime/Core/Private/Misc/EngineVersion.cpp
@@ -39,7 +39,7 @@ bool FEngineVersionBase::IsEmpty() const
 	return Major == 0 && Minor == 0 && Patch == 0;
 }
 
-bool FEngineVersionBase::IsPromotedBuild() const
+bool FEngineVersionBase::HasChangelist() const
 {
 	return GetChangelist() != 0;
 }
@@ -71,7 +71,7 @@ EVersionComparison FEngineVersionBase::GetNewest(const FEngineVersionBase &First
 	}
 
 	// Compare changelists (only if they're both from the same vendor, and they're both valid)
-	if (First.IsLicenseeVersion() == Second.IsLicenseeVersion() && First.IsPromotedBuild() && Second.IsPromotedBuild() && First.GetChangelist() != Second.GetChangelist())
+	if (First.IsLicenseeVersion() == Second.IsLicenseeVersion() && First.HasChangelist() && Second.HasChangelist() && First.GetChangelist() != Second.GetChangelist())
 	{
 		Component = EVersionComponent::Changelist;
 		return (First.GetChangelist() > Second.GetChangelist()) ? EVersionComparison::First : EVersionComparison::Second;
@@ -109,7 +109,7 @@ void FEngineVersion::Empty()
 bool FEngineVersion::IsCompatibleWith(const FEngineVersionBase &Other) const
 {
 	// If this or the other is not a promoted build, always assume compatibility. 
-	if(!IsPromotedBuild() || !Other.IsPromotedBuild())
+	if(!HasChangelist() || !Other.HasChangelist())
 	{
 		return true;
 	}
diff --git a/Engine/Source/Runtime/Core/Public/Misc/EngineVersionBase.h b/Engine/Source/Runtime/Core/Public/Misc/EngineVersionBase.h
--- a/Engine/Source/Runtime/Core/Public/Misc/EngineVersionBase.h
+++ b/Engine/Source/Runtime/Core/Public/Misc/EngineVersionBase.h
@@ -60,8 +60,8 @@ public:
 	/** Returns whether the current version is empty. */
 	bool IsEmpty() const;
 
-	/** Returns whether the engine version is for a promoted build. */
-	bool IsPromotedBuild() const; 
+	/** Returns whether the engine version has a changelist component. */
+	bool HasChangelist() const; 
 
 	/** Returns the newest of two versions, and the component at which they differ */
 	static EVersionComparison GetNewest(const FEngineVersionBase &First, const FEngineVersionBase &Second, EVersionComponent *OutComponent);
diff --git a/Engine/Source/Runtime/CoreUObject/Private/UObject/LinkerLoad.cpp b/Engine/Source/Runtime/CoreUObject/Private/UObject/LinkerLoad.cpp
--- a/Engine/Source/Runtime/CoreUObject/Private/UObject/LinkerLoad.cpp
+++ b/Engine/Source/Runtime/CoreUObject/Private/UObject/LinkerLoad.cpp
@@ -945,7 +945,7 @@ FLinkerLoad::ELinkerStatus FLinkerLoad::SerializePackageFileSummary()
 			UE_LOG(LogLinker, Warning, TEXT("Asset '%s' has been saved with engine version newer than current and therefore can't be loaded. CurrEngineVersion: %s AssetEngineVersion: %s"), *Filename, *GEngineVersion.ToString(), *Summary.CompatibleWithEngineVersion.ToString() );
 			return LINKER_Failed;
 		}
-		else if( !FPlatformProperties::RequiresCookedData() && !Summary.SavedByEngineVersion.IsPromotedBuild() && GEngineVersion.IsPromotedBuild() )
+		else if( !FPlatformProperties::RequiresCookedData() && !Summary.SavedByEngineVersion.HasChangelist() && GEngineVersion.HasChangelist() )
 		{
 			// This warning can be disabled in ini with [Core.System] ZeroEngineVersionWarning=False
 			static struct FInitZeroEngineVersionWarning
diff --git a/Engine/Source/Runtime/CoreUObject/Private/UObject/PackageFileSummary.cpp b/Engine/Source/Runtime/CoreUObject/Private/UObject/PackageFileSummary.cpp
--- a/Engine/Source/Runtime/CoreUObject/Private/UObject/PackageFileSummary.cpp
+++ b/Engine/Source/Runtime/CoreUObject/Private/UObject/PackageFileSummary.cpp
@@ -168,7 +168,7 @@ FArchive& operator<<( FArchive& Ar, FPackageFileSummary& Sum )
 
 		if( Sum.GetFileVersionUE4() >= VER_UE4_ENGINE_VERSION_OBJECT )
 		{
-			if(Ar.IsCooking() || (Ar.IsSaving() && !GEngineVersion.IsPromotedBuild()))
+			if(Ar.IsCooking() || (Ar.IsSaving() && !GEngineVersion.HasChangelist()))
 			{
 				FEngineVersion EmptyEngineVersion;
 				Ar << EmptyEngineVersion;
@@ -191,7 +191,7 @@ FArchive& operator<<( FArchive& Ar, FPackageFileSummary& Sum )
 
 		if (Sum.GetFileVersionUE4() >= VER_UE4_PACKAGE_SUMMARY_HAS_COMPATIBLE_ENGINE_VERSION )
 		{
-			if(Ar.IsCooking() || (Ar.IsSaving() && !GEngineVersion.IsPromotedBuild()))
+			if(Ar.IsCooking() || (Ar.IsSaving() && !GEngineVersion.HasChangelist()))
 			{
 				FEngineVersion EmptyEngineVersion;
 				Ar << EmptyEngineVersion;
diff --git a/Engine/Source/Runtime/Engine/Private/UnrealNetwork.cpp b/Engine/Source/Runtime/Engine/Private/UnrealNetwork.cpp
--- a/Engine/Source/Runtime/Engine/Private/UnrealNetwork.cpp
+++ b/Engine/Source/Runtime/Engine/Private/UnrealNetwork.cpp
@@ -46,7 +46,7 @@ uint32 FNetworkVersion::GetLocalNetworkVersion(bool AllowOverrideDelegate /*=tru
 	uint32 LocalNetworkVersion = FCrc::MemCrc32( &InternalProtocolVersion, sizeof( InternalProtocolVersion ), VersionHash );
 
 #if 0//!(UE_BUILD_SHIPPING || UE_BUILD_TEST)	// DISABLED FOR NOW, MESSES UP COPIED BUILDS
-	if ( !GEngineVersion.IsPromotedBuild() )
+	if ( !GEngineVersion.HasChangelist() )
 	{
 		// Further hash with machine id if this is a non promoted build
 		const FString MachineId = FPlatformMisc::GetMachineId().ToString( EGuidFormats::Digits ).ToLower();
