PR #1316: Do not ask if there is no one to show a message (Contributed by [REDACTED_PERSON])

diff --git a/Engine/Source/Runtime/Core/Private/Windows/WindowsPlatformMisc.cpp b/Engine/Source/Runtime/Core/Private/Windows/WindowsPlatformMisc.cpp
--- a/Engine/Source/Runtime/Core/Private/Windows/WindowsPlatformMisc.cpp
+++ b/Engine/Source/Runtime/Core/Private/Windows/WindowsPlatformMisc.cpp
@@ -1958,6 +1958,12 @@ void FWindowsPlatformMisc::PromptForRemoteDebugging(bool bIsEnsure)
 			return;
 		}
 
+		if (FApp::IsUnattended())
+		{
+			// Do not ask if there is no one to show a message
+			return;
+		}
+
 		FCString::Sprintf(GErrorRemoteDebugPromptMessage, 
 			TEXT("Have a programmer remote debug this crash?\n")
 			TEXT("Hit NO to exit and submit error report as normal.\n")
