Add a separate place to update the changelist when using a custom licensee version number.

diff --git a/Engine/Source/Runtime/Core/Public/Modules/ModuleVersion.h b/Engine/Source/Runtime/Core/Public/Modules/ModuleVersion.h
--- a/Engine/Source/Runtime/Core/Public/Modules/ModuleVersion.h
+++ b/Engine/Source/Runtime/Core/Public/Modules/ModuleVersion.h
@@ -7,12 +7,16 @@
 // This number identifies a particular API revision, and is used to determine module compatibility. Hotfixes should retain the API version of the original release.
 // This define is parsed by the build tools, and should be a number or BUILT_FROM_CHANGELIST.
 #if BUILT_FROM_CHANGELIST > 0
-#define MODULE_API_VERSION BUILT_FROM_CHANGELIST
+	#if ENGINE_IS_LICENSEE_VERSION
+		#define MODULE_API_VERSION BUILT_FROM_CHANGELIST
+	#else
+		#define MODULE_API_VERSION BUILT_FROM_CHANGELIST /* Or hotfix compatibility changelist */
+	#endif
 #else
-#define MODULE_API_VERSION 0
+	#define MODULE_API_VERSION 0
 #endif
 
 // Check that the API version has been set manually for a hotfix release
-#if ENGINE_PATCH_VERSION > 0 && BUILT_FROM_CHANGELIST > 0 && MODULE_API_VERSION == BUILT_FROM_CHANGELIST
+#if ENGINE_IS_LICENSEE_VERSION == 0 && ENGINE_PATCH_VERSION > 0 && BUILT_FROM_CHANGELIST > 0 && MODULE_API_VERSION == BUILT_FROM_CHANGELIST
 	#error MODULE_API_VERSION must be manually defined for hotfix builds
 #endif
