Speculative fix for UE-16187

diff --git a/Engine/Source/Runtime/Core/Private/Mac/MacTextInputMethodSystem.cpp b/Engine/Source/Runtime/Core/Private/Mac/MacTextInputMethodSystem.cpp
--- a/Engine/Source/Runtime/Core/Private/Mac/MacTextInputMethodSystem.cpp
+++ b/Engine/Source/Runtime/Core/Private/Mac/MacTextInputMethodSystem.cpp
@@ -201,16 +201,17 @@ void FMacTextInputMethodSystem::DeactivateContext(const TSharedRef<ITextInputMet
 		UE_LOG(LogMacTextInputMethodSystem, Error, TEXT("Deactivating a context failed when its registration couldn't be found."));
 		return;
 	}
-	
+
 	TSharedPtr<ITextInputMethodChangeNotifier> Notifier(NotifierRef.Pin());
 	FTextInputMethodChangeNotifier* MacNotifier = (FTextInputMethodChangeNotifier*)Notifier.Get();
 	bool bDeactivatedContext = false;
-	if(MacNotifier->GetContextWindowNumber() != 0)
+	const int32 ContextWindowNumber = MacNotifier->GetContextWindowNumber();
+	if (ContextWindowNumber != 0)
 	{
-		FCocoaWindow* CocoaWindow = (FCocoaWindow*)[NSApp windowWithWindowNumber:MacNotifier->GetContextWindowNumber()];
-		if (CocoaWindow)
-		{
-			bDeactivatedContext = MainThreadReturn(^{
+		bDeactivatedContext = MainThreadReturn(^{
+			FCocoaWindow* CocoaWindow = (FCocoaWindow*)[NSApp windowWithWindowNumber:ContextWindowNumber];
+			if (CocoaWindow)
+			{
 				bool bSuccess = false;
 				if(CocoaWindow && [CocoaWindow openGLView])
 				{
@@ -223,13 +224,13 @@ void FMacTextInputMethodSystem::DeactivateContext(const TSharedRef<ITextInputMet
 					}
 				}
 				return bSuccess;
-			}, UE4IMEEventMode);
-		}
-		else
-		{
-			// If the window is no longer open, the context is already inactive (handled by MessageHandler->OnWindowActivationChanged() called in FMacApplication::OnWindowDestroyed())
-			bDeactivatedContext = true;
-		}
+			}
+			else
+			{
+				// If the window is no longer open, the context is already inactive (handled by MessageHandler->OnWindowActivationChanged() called in FMacApplication::OnWindowDestroyed())
+				return true;
+			}
+		}, UE4IMEEventMode);
 	}
 	if(!bDeactivatedContext)
 	{
