Fixed merge issues in text localization resource generation logic. Now properly looks up the native culture's archive files and doesn't attempt to find the native translation for the source's native translation.

diff --git a/Engine/Source/Runtime/Core/Private/Internationalization/TextLocalizationResourceGenerator.cpp b/Engine/Source/Runtime/Core/Private/Internationalization/TextLocalizationResourceGenerator.cpp
--- a/Engine/Source/Runtime/Core/Private/Internationalization/TextLocalizationResourceGenerator.cpp
+++ b/Engine/Source/Runtime/Core/Private/Internationalization/TextLocalizationResourceGenerator.cpp
@@ -145,7 +145,7 @@ bool FTextLocalizationResourceGenerator::Generate(const FString& SourcePath, con
 	FLocalizationEntryTracker LocalizationEntryTracker;
 
 	// Find archives in the native culture-specific folder.
-	const FString NativeCulturePath = SourcePath / *(CultureToGenerate);
+	const FString NativeCulturePath = SourcePath / *(NativeCulture);
 	TArray<FString> NativeArchiveFileNames;
 	IFileManager::Get().FindFiles(NativeArchiveFileNames, *(NativeCulturePath / TEXT("*.archive")), true, false);
 
@@ -170,7 +170,7 @@ bool FTextLocalizationResourceGenerator::Generate(const FString& SourcePath, con
 	for (const FString& NativeArchiveFileName : NativeArchiveFileNames)
 	{
 		// Read each archive file from the culture-named directory in the source path.
-		FString ArchiveFilePath = CulturePath / NativeArchiveFileName;
+		FString ArchiveFilePath = NativeCulturePath / NativeArchiveFileName;
 		ArchiveFilePath = FPaths::ConvertRelativePathToFull(ArchiveFilePath);
 		TSharedRef<FInternationalizationArchive> InternationalizationArchive = MakeShareable(new FInternationalizationArchive);
 #if 0 // @todo Json: Serializing from FArchive is currently broken
@@ -262,19 +262,21 @@ bool FTextLocalizationResourceGenerator::Generate(const FString& SourcePath, con
 
 				// Get proper source string to use - if the native source is different than the native translation, the native translation becomes the source.
 				FLocItem SourceToUse = Source;
-				for (const auto& NativeArchive : NativeArchives)
+				if (CultureToGenerate != NativeCulture)
 				{
-					if (NativeArchive.IsValid())
+					for (const auto& NativeArchive : NativeArchives)
 					{
-						const TSharedPtr<FArchiveEntry> NativeArchiveEntry = NativeArchive->FindEntryBySource(Namespace, Source, ContextIt->KeyMetadataObj);
-						if (NativeArchiveEntry.IsValid() && !NativeArchiveEntry->Source.IsExactMatch(NativeArchiveEntry->Translation))
+						if (NativeArchive.IsValid())
 						{
-							SourceToUse = NativeArchiveEntry->Translation;
-							break;
+							const TSharedPtr<FArchiveEntry> NativeArchiveEntry = NativeArchive->FindEntryBySource(Namespace, Source, ContextIt->KeyMetadataObj);
+							if (NativeArchiveEntry.IsValid() && !NativeArchiveEntry->Source.IsExactMatch(NativeArchiveEntry->Translation))
+							{
+								SourceToUse = NativeArchiveEntry->Translation;
+								break;
+							}
 						}
 					}
 				}
-				\
 				// Find matching archive entry.
 				TSharedPtr<FArchiveEntry> ArchiveEntry = InternationalizationArchive->FindEntryBySource(Namespace, SourceToUse, ContextIt->KeyMetadataObj);
 
