UE-17493: GitHub 1303 : [Core] Fix the functionality to track allocations that don't happen within the engine's wrappers.

#1303

diff --git a/Engine/Source/Runtime/Core/Private/Windows/WindowsPlatformMemory.cpp b/Engine/Source/Runtime/Core/Private/Windows/WindowsPlatformMemory.cpp
--- a/Engine/Source/Runtime/Core/Private/Windows/WindowsPlatformMemory.cpp
+++ b/Engine/Source/Runtime/Core/Private/Windows/WindowsPlatformMemory.cpp
@@ -7,6 +7,10 @@
 #include "GenericPlatformMemoryPoolStats.h"
 #include "MemoryMisc.h"
 
+#if ENABLE_WIN_ALLOC_TRACKING
+#include <crtdbg.h>
+#endif // ENABLE_WIN_ALLOC_TRACKING
+
 #if !FORCE_ANSI_ALLOCATOR
 #include "MallocBinned.h"
 #endif
@@ -17,19 +21,17 @@
 
 DECLARE_MEMORY_STAT(TEXT("Windows Specific Memory Stat"),	STAT_WindowsSpecificMemoryStat, STATGROUP_MemoryPlatform);
 
-/** Enable this to track down windows allocations not wrapped by our wrappers
+#if ENABLE_WIN_ALLOC_TRACKING
+// This allows tracking of allocations that don't happen within the engine's wrappers.
+// You will probably want to set conditional breakpoints here to capture specific allocations
+// which aren't related to static initialization, they will happen on the CRT anyway.
 int WindowsAllocHook(int nAllocType, void *pvData,
 				  size_t nSize, int nBlockUse, long lRequest,
 				  const unsigned char * szFileName, int nLine )
 {
-	if ((nAllocType == _HOOK_ALLOC || nAllocType == _HOOK_REALLOC) && nSize > 2048)
-	{
-		static int i = 0;
-		i++;
-	}
 	return true;
 }
-*/
+#endif // ENABLE_WIN_ALLOC_TRACKING
 
 #include "GenericPlatformMemoryPoolStats.h"
 
@@ -62,6 +64,13 @@ void FWindowsPlatformMemory::Init()
 
 FMalloc* FWindowsPlatformMemory::BaseAllocator()
 {
+#if ENABLE_WIN_ALLOC_TRACKING
+	// This allows tracking of allocations that don't happen within the engine's wrappers.
+	// This actually won't be compiled unless bDebugBuildsActuallyUseDebugCRT is set in the
+	// build configuration for UBT.
+	_CrtSetAllocHook(WindowsAllocHook);
+#endif // ENABLE_WIN_ALLOC_TRACKING
+
 #if FORCE_ANSI_ALLOCATOR
 	return new FMallocAnsi();
 #elif (WITH_EDITORONLY_DATA || IS_PROGRAM) && TBB_ALLOCATOR_ALLOWED
@@ -69,8 +78,6 @@ FMalloc* FWindowsPlatformMemory::BaseAllocator()
 #else
 	return new FMallocBinned((uint32)(GetConstants().PageSize&MAX_uint32), (uint64)MAX_uint32+1);
 #endif
-
-//	_CrtSetAllocHook(WindowsAllocHook); // Enable to track down windows allocs not handled by our wrapper
 }
 
 FPlatformMemoryStats FWindowsPlatformMemory::GetStats()
