UE-16149/Disabled runtime crash processing for launcher builds

diff --git a/Engine/Source/Runtime/Core/Private/Windows/WindowsPlatformCrashContext.cpp b/Engine/Source/Runtime/Core/Private/Windows/WindowsPlatformCrashContext.cpp
--- a/Engine/Source/Runtime/Core/Private/Windows/WindowsPlatformCrashContext.cpp
+++ b/Engine/Source/Runtime/Core/Private/Windows/WindowsPlatformCrashContext.cpp
@@ -5,6 +5,7 @@
 #include "ExceptionHandling.h"
 #include "WindowsPlatformCrashContext.h"
 #include "EngineVersion.h"
+#include "EngineBuildSettings.h"
 
 #include "AllowWindowsPlatformTypes.h"
 
@@ -492,23 +493,26 @@ int32 ReportCrash( LPEXCEPTION_POINTERS ExceptionInfo )
 	// Then try run time crash processing and broadcast information about a crash.
 	FCoreDelegates::OnHandleSystemError.Broadcast();
 
-	const SIZE_T StackTraceSize = 65535;
-	ANSICHAR* StackTrace = (ANSICHAR*) GMalloc->Malloc( StackTraceSize );
-	StackTrace[0] = 0;
-	// Walk the stack and dump it to the allocated memory. This process usually allocates a lot of memory.
-	FPlatformStackWalk::StackWalkAndDump( StackTrace, StackTraceSize, 0, ExceptionInfo->ContextRecord );
-
-	// @TODO yrx 2014-08-20 Make this constant.
-	if( ExceptionInfo->ExceptionRecord->ExceptionCode != 1 )
+	const bool bGenerateRuntimeCallstack = FEngineBuildSettings::IsInternalBuild() || FEngineBuildSettings::IsPerforceBuild() || FEngineBuildSettings::IsSourceDistribution();
+	if (bGenerateRuntimeCallstack)
 	{
-		CreateExceptionInfoString( ExceptionInfo->ExceptionRecord );
-		FCString::Strncat( GErrorHist, GErrorExceptionDescription, ARRAY_COUNT( GErrorHist ) );
-		FCString::Strncat( GErrorHist, TEXT( "\r\n\r\n" ), ARRAY_COUNT( GErrorHist ) );
-	}
+		const SIZE_T StackTraceSize = 65535;
+		ANSICHAR* StackTrace = (ANSICHAR*)GMalloc->Malloc( StackTraceSize );
+		StackTrace[0] = 0;
+		// Walk the stack and dump it to the allocated memory. This process usually allocates a lot of memory.
+		FPlatformStackWalk::StackWalkAndDump( StackTrace, StackTraceSize, 0, ExceptionInfo->ContextRecord );
+
+		if (ExceptionInfo->ExceptionRecord->ExceptionCode != 1)
+		{
+			CreateExceptionInfoString( ExceptionInfo->ExceptionRecord );
+			FCString::Strncat( GErrorHist, GErrorExceptionDescription, ARRAY_COUNT( GErrorHist ) );
+			FCString::Strncat( GErrorHist, TEXT( "\r\n\r\n" ), ARRAY_COUNT( GErrorHist ) );
+		}
 
-	FCString::Strncat( GErrorHist, ANSI_TO_TCHAR(StackTrace), ARRAY_COUNT(GErrorHist) );
+		FCString::Strncat( GErrorHist, ANSI_TO_TCHAR( StackTrace ), ARRAY_COUNT( GErrorHist ) );
 
-	GMalloc->Free( StackTrace );
+		GMalloc->Free( StackTrace );
+	}
 
 	return EXCEPTION_EXECUTE_HANDLER;
 }
