Compile fix for Clang

diff --git a/Engine/Source/Runtime/AssetRegistry/Public/AssetData.h b/Engine/Source/Runtime/AssetRegistry/Public/AssetData.h
--- a/Engine/Source/Runtime/AssetRegistry/Public/AssetData.h
+++ b/Engine/Source/Runtime/AssetRegistry/Public/AssetData.h
@@ -313,10 +313,10 @@ public:
 			}
 			else // if (Ar.IsLoading())
 			{
-				TMap<FName, FString> TagsAndValues;
-				Ar << TagsAndValues;
-				TagsAndValues.Remove(FiBName);
-				AssetData.TagsAndValues = MakeSharedMapView(MoveTemp(TagsAndValues));
+				TMap<FName, FString> LocalTagsAndValues;
+				Ar << LocalTagsAndValues;
+				LocalTagsAndValues.Remove(FiBName);
+				AssetData.TagsAndValues = MakeSharedMapView(MoveTemp(LocalTagsAndValues));
 			}
 		}
 		else
diff --git a/Engine/Source/Runtime/Core/Public/Containers/SharedMapView.h b/Engine/Source/Runtime/Core/Public/Containers/SharedMapView.h
--- a/Engine/Source/Runtime/Core/Public/Containers/SharedMapView.h
+++ b/Engine/Source/Runtime/Core/Public/Containers/SharedMapView.h
@@ -2,6 +2,8 @@
 
 #pragma once
 
+#include "Map.h"
+
 /** Immutable wrapper of shared pointer to a map */
 template <class FKey, class FValue>
 class TSharedMapView
@@ -85,14 +87,12 @@ public:
 	/** Range for iterator access - DO NOT USE DIRECTLY */
 	friend typename FMap::TConstIterator begin(const TSharedMapView& View)
 	{
-		using ::begin;
 		return begin(*View.Map);
 	}
 
 	/** Range for iterator access - DO NOT USE DIRECTLY */
 	friend typename FMap::TConstIterator end(const TSharedMapView& View)
 	{
-		using ::end;
 		return end(*View.Map);
 	}
 
