Fixed vulnerability when destructing a FFeedbackContext that had outstanding slow tasks

The issue here is that the derived destructor for FFeedbackContextEditor would destroy it's TWeakPtr<SWindow>, proceeded by the destruction of any remaining FScopedSlowTasks held in the LegacyAPIScopes array. This destruction may have caused FFeedbackContextEditor::FinalizeSlowTask to be called which could have resulted in the TWeakPtr being used after its RefereceCount was deleted, causing a memory overwrite. I can't see any code in the wild which could have caused this, but I've also added some additional safety checks to catch error conditions earlier.

diff --git a/Engine/Source/Editor/UnrealEd/Private/FeedbackContextEditor.cpp b/Engine/Source/Editor/UnrealEd/Private/FeedbackContextEditor.cpp
--- a/Engine/Source/Editor/UnrealEd/Private/FeedbackContextEditor.cpp
+++ b/Engine/Source/Editor/UnrealEd/Private/FeedbackContextEditor.cpp
@@ -29,7 +29,7 @@ public:
 		SLATE_EVENT( FOnCancelClickedDelegate, OnCancelClickedDelegate )
 
 		/** The feedback scope stack that we are presenting to the user */
-		SLATE_ARGUMENT( TWeakPtr<FScopedSlowTaskStack>, ScopeStack )
+		SLATE_ARGUMENT( TWeakPtr<FSlowTaskStack>, ScopeStack )
 
 	SLATE_END_ARGS()
 
@@ -157,11 +157,15 @@ private:
 		// Always show the first one
 		DynamicProgressIndices.Add(0);
 
-		for (int32 Index = 1; Index < ScopeStack->Num(); ++Index)
+		for (int32 Index = 1; Index < ScopeStack->Num() && DynamicProgressIndices.Num() <= MaxNumSecondaryBars - 1; ++Index)
 		{
 			const auto* Scope = (*ScopeStack)[Index];
 
-			if (Scope->bVisibleOnUI && !Scope->DefaultMessage.IsEmpty())
+			if (Scope->Visibility == ESlowTaskVisibility::ForceVisible)
+			{
+				DynamicProgressIndices.Add(Index);
+			}
+			else if (Scope->Visibility == ESlowTaskVisibility::Default && !Scope->DefaultMessage.IsEmpty())
 			{
 				const auto TimeOpen = FPlatformTime::Seconds() - Scope->StartTime;
 				const auto WorkDone = ScopeStack->GetProgressFraction(Index);
@@ -171,11 +175,6 @@ private:
 				{
 					DynamicProgressIndices.Add(Index);
 				}
-
-				if (DynamicProgressIndices.Num() == MaxNumSecondaryBars)
-				{
-					break;
-				}
 			}
 		}
 
@@ -319,7 +318,7 @@ private:
 	FOnCancelClickedDelegate OnCancelClickedDelegate;
 
 	/** The scope stack that we are reflecting */
-	TWeakPtr<FScopedSlowTaskStack> WeakStack;
+	TWeakPtr<FSlowTaskStack> WeakStack;
 
 	/** The vertical box containing the secondary progress bars */
 	TSharedPtr<SVerticalBox> SecondaryBars;
diff --git a/Engine/Source/Runtime/Core/Private/Misc/FeedbackContext.cpp b/Engine/Source/Runtime/Core/Private/Misc/FeedbackContext.cpp
--- a/Engine/Source/Runtime/Core/Private/Misc/FeedbackContext.cpp
+++ b/Engine/Source/Runtime/Core/Private/Misc/FeedbackContext.cpp
@@ -3,6 +3,11 @@
 #include "CorePrivatePCH.h"
 #include "FeedbackContext.h"
 
+FFeedbackContext::~FFeedbackContext()
+{
+	ensureMsgf(LegacyAPIScopes.Num() == 0, TEXT("EndSlowTask has not been called for %d outstanding tasks"), LegacyAPIScopes.Num());
+}
+
 void FFeedbackContext::RequestUpdateUI(bool bForceUpdate)
 {
 	// Only update a maximum of 5 times a second
@@ -20,6 +25,8 @@ void FFeedbackContext::RequestUpdateUI(bool bForceUpdate)
 
 void FFeedbackContext::UpdateUI()
 {
+	ensure(IsInGameThread());
+
 	if (ScopeStack->Num() != 0)
 	{
 		ProgressReported(ScopeStack->GetProgressFraction(0), (*ScopeStack)[0]->GetCurrentMessage());
@@ -29,16 +36,22 @@ void FFeedbackContext::UpdateUI()
 /**** Begin legacy API ****/
 void FFeedbackContext::BeginSlowTask( const FText& Task, bool ShowProgressDialog, bool bShowCancelButton )
 {
-	TUniquePtr<FScopedSlowTask> NewScope(new FScopedSlowTask(0, Task, true, *this));
+	ensure(IsInGameThread());
+
+	TUniquePtr<FSlowTask> NewScope(new FSlowTask(0, Task, true, *this));
 	if (ShowProgressDialog)
 	{
 		NewScope->MakeDialog(bShowCancelButton);
 	}
+
+	NewScope->Initialize();
 	LegacyAPIScopes.Add(MoveTemp(NewScope));
 }
 
 void FFeedbackContext::UpdateProgress( int32 Numerator, int32 Denominator )
 {
+	ensure(IsInGameThread());
+
 	if (LegacyAPIScopes.Num() != 0)
 	{
 		LegacyAPIScopes.Last()->TotalAmountOfWork = Denominator;
@@ -50,6 +63,8 @@ void FFeedbackContext::UpdateProgress( int32 Numerator, int32 Denominator )
 
 void FFeedbackContext::StatusUpdate( int32 Numerator, int32 Denominator, const FText& StatusText )
 {
+	ensure(IsInGameThread());
+
 	if (LegacyAPIScopes.Num() != 0)
 	{
 		if (Numerator > 0 && Denominator > 0)
@@ -63,6 +78,8 @@ void FFeedbackContext::StatusUpdate( int32 Numerator, int32 Denominator, const F
 
 void FFeedbackContext::StatusForceUpdate( int32 Numerator, int32 Denominator, const FText& StatusText )
 {
+	ensure(IsInGameThread());
+
 	if (LegacyAPIScopes.Num() != 0)
 	{
 		UpdateProgress(Numerator, Denominator);
@@ -73,12 +90,15 @@ void FFeedbackContext::StatusForceUpdate( int32 Numerator, int32 Denominator, co
 
 void FFeedbackContext::EndSlowTask()
 {
+	ensure(IsInGameThread());
+
 	check(LegacyAPIScopes.Num() != 0);
+	LegacyAPIScopes.Last()->Destroy();
 	LegacyAPIScopes.Pop();
 }
 /**** End legacy API ****/
 
-void FScopedSlowTask::MakeDialog(bool bShowCancelButton, bool bAllowInPIE)
+void FSlowTask::MakeDialog(bool bShowCancelButton, bool bAllowInPIE)
 {
 	const bool bIsDisabledByPIE = GIsPlayInEditorWorld && !bAllowInPIE;
 	const bool bIsDialogAllowed = bEnabled && !GIsSilent && !bIsDisabledByPIE && !IsRunningCommandlet() && IsInGameThread();
@@ -92,7 +112,7 @@ void FScopedSlowTask::MakeDialog(bool bShowCancelButton, bool bAllowInPIE)
 	}
 }
 
-float FScopedSlowTaskStack::GetProgressFraction(int32 Index) const
+float FSlowTaskStack::GetProgressFraction(int32 Index) const
 {
 	const int32 StartIndex = Num() - 1;
 	const int32 EndIndex = Index;
@@ -100,7 +120,7 @@ float FScopedSlowTaskStack::GetProgressFraction(int32 Index) const
 	float Progress = 0.f;
 	for (int32 CurrentIndex = StartIndex; CurrentIndex >= EndIndex; --CurrentIndex)
 	{
-		const FScopedSlowTask* Scope = (*this)[CurrentIndex];
+		const FSlowTask* Scope = (*this)[CurrentIndex];
 		
 		const float ThisScopeCompleted = float(Scope->CompletedWork) / Scope->TotalAmountOfWork;
 		const float ThisScopeCurrentFrame = float(Scope->CurrentFrameScope) / Scope->TotalAmountOfWork;
diff --git a/Engine/Source/Runtime/Core/Public/Misc/FeedbackContext.h b/Engine/Source/Runtime/Core/Public/Misc/FeedbackContext.h
--- a/Engine/Source/Runtime/Core/Public/Misc/FeedbackContext.h
+++ b/Engine/Source/Runtime/Core/Public/Misc/FeedbackContext.h
@@ -5,11 +5,11 @@
 #include "UniquePtr.h"
 
 
-struct FScopedSlowTask;
+struct FSlowTask;
 
 
 /** A stack of feedback scopes */
-struct FScopedSlowTaskStack : TArray<FScopedSlowTask*>
+struct FSlowTaskStack : TArray<FSlowTask*>
 {
 	/** Get the fraction of work that has been completed for the specified index in the stack (0=total progress) */
 	CORE_API float GetProgressFraction(int32 Index) const;
@@ -29,7 +29,7 @@ public:
 	virtual bool ReceivedUserCancel() { return false; };
 	
 	/** Public const access to the current state of the scope stack */
-	FORCEINLINE const FScopedSlowTaskStack& GetScopeStack() const
+	FORCEINLINE const FSlowTaskStack& GetScopeStack() const
 	{
 		return *ScopeStack;
 	}
@@ -82,20 +82,22 @@ public:
 
 	FFeedbackContext()
 		: TreatWarningsAsErrors(0)
-		, ScopeStack(MakeShareable(new FScopedSlowTaskStack))
+		, ScopeStack(MakeShareable(new FSlowTaskStack))
 	{}
 
+	virtual ~FFeedbackContext();
+
 private:
 	FFeedbackContext(const FFeedbackContext&);
 	FFeedbackContext& operator=(const FFeedbackContext&);
 
 protected:
 	
-	friend FScopedSlowTask;
+	friend FSlowTask;
 
 	/** Stack of pointers to feedback scopes that are currently open */
-	TSharedRef<FScopedSlowTaskStack> ScopeStack;
-	TArray<TUniquePtr<FScopedSlowTask>> LegacyAPIScopes;
+	TSharedRef<FSlowTaskStack> ScopeStack;
+	TArray<TUniquePtr<FSlowTask>> LegacyAPIScopes;
 
 	/** Ask that the UI be updated as a result of the scope stack changing */
 	void RequestUpdateUI(bool bForceUpdate = false);
@@ -105,30 +107,21 @@ protected:
 
 };
 
+/** Enum to specify a particular slow task section should be shown */
+enum class ESlowTaskVisibility
+{
+	/** Default visibility (inferred by some heuristic of remaining work/time open) */
+	Default,
+	/** Force this particular slow task to be visible on the UI */
+	ForceVisible,
+	/** Forcibly prevent this slow task from being shown, but still use it for work progress calculations */
+	Invisible,
+};
 
 /**
- * A scope block representing an amount of work divided up into sections.
- * Use one scope at the top of each function to give accurate feedback to the user of a slow operation's progress.
- *
- * Example Usage:
- *	void DoSlowWork()
- *	{
- *		FScopedSlowTask Progress(2.f, LOCTEXT("DoingSlowWork", "Doing Slow Work..."));
- *		// Optionally make this show a dialog if not already shown
- *		Progress.MakeDialog();
- *
- *		// Indicate that we are entering a frame representing 1 unit of work
- *		Progress.EnterProgressFrame(1.f);
- *		
- *		// DoFirstThing() can follow a similar pattern of creating a scope divided into frames. These contribute to their parent's progress frame proportionately.
- *		DoFirstThing();
- *		
- *		Progress.EnterProgressFrame(1.f);
- *		DoSecondThing();
- *	}
- *
+ * Data type used to store information about a currently running slow task. Direct use is not advised, use FScopedSlowTask instead
  */
-struct CORE_API FScopedSlowTask
+struct CORE_API FSlowTask
 {
 	/** Default message to display to the user when not overridden by a frame */
 	FText DefaultMessage;
@@ -145,8 +138,8 @@ struct CORE_API FScopedSlowTask
 	/** The amount of work the current frame is responsible for */
 	float CurrentFrameScope;
 
-	/** true if this scope should be presented to the user */
-	bool bVisibleOnUI;
+	/** The visibility of this slow task */
+	ESlowTaskVisibility Visibility;
 
 	/** The time that this scope was created */
 	double StartTime;
@@ -163,7 +156,7 @@ private:
 	FFeedbackContext& Context;
 
 	/** Prevent copying */
-	FScopedSlowTask(const FScopedSlowTask&);
+	FSlowTask(const FSlowTask&);
 
 public:
 
@@ -174,13 +167,13 @@ public:
 	 * @param		InDefaultMessage		A message to display to the user to describe the purpose of the scope
 	 * @param		bInVisible				When false, this scope will have no effect. Allows for proper scoped objects that are conditionally hidden.
 	 */
-	FORCEINLINE FScopedSlowTask(float InAmountOfWork, const FText& InDefaultMessage = FText(), bool bInEnabled = true, FFeedbackContext& InContext = *GWarn)
+	FORCEINLINE FSlowTask(float InAmountOfWork, const FText& InDefaultMessage = FText(), bool bInEnabled = true, FFeedbackContext& InContext = *GWarn)
 		: DefaultMessage(InDefaultMessage)
 		, FrameMessage()
 		, TotalAmountOfWork(InAmountOfWork)
 		, CompletedWork(0)
 		, CurrentFrameScope(0)
-		, bVisibleOnUI(true)
+		, Visibility(ESlowTaskVisibility::Default)
 		, StartTime(FPlatformTime::Seconds())
 		, bEnabled(bInEnabled && IsInGameThread())
 		, bCreatedDialog(false)		// only set to true if we create a dialog
@@ -191,15 +184,19 @@ public:
 		{
 			TotalAmountOfWork = CurrentFrameScope = 1.f;
 		}
+	}
 
+	/** Function that initializes the scope by adding it to its context's stack */
+	FORCEINLINE void Initialize()
+	{
 		if (bEnabled)
 		{
 			Context.ScopeStack->Push(this);
 		}
 	}
 
-	/** Destructor that finishes any remaining work and removes itself from the global scope stack */
-	FORCEINLINE ~FScopedSlowTask()
+	/** Function that finishes any remaining work and removes itself from the global scope stack */
+	FORCEINLINE void Destroy()
 	{
 		if (bEnabled)
 		{
@@ -209,7 +206,7 @@ public:
 				Context.FinalizeSlowTask();
 			}
 
-			FScopedSlowTaskStack& Stack = *Context.ScopeStack;
+			FSlowTaskStack& Stack = *Context.ScopeStack;
 			checkSlow(Stack.Num() != 0 && Stack.Last() == this);
 
 			auto* Task = Stack.Last();
@@ -266,13 +263,49 @@ public:
 	{
 		return FrameMessage.IsEmpty() ? DefaultMessage : FrameMessage;
 	}
+};
+
+
+/**
+ * A scope block representing an amount of work divided up into sections.
+ * Use one scope at the top of each function to give accurate feedback to the user of a slow operation's progress.
+ *
+ * Example Usage:
+ *	void DoSlowWork()
+ *	{
+ *		FScopedSlowTask Progress(2.f, LOCTEXT("DoingSlowWork", "Doing Slow Work..."));
+ *		// Optionally make this show a dialog if not already shown
+ *		Progress.MakeDialog();
+ *
+ *		// Indicate that we are entering a frame representing 1 unit of work
+ *		Progress.EnterProgressFrame(1.f);
+ *		
+ *		// DoFirstThing() can follow a similar pattern of creating a scope divided into frames. These contribute to their parent's progress frame proportionately.
+ *		DoFirstThing();
+ *		
+ *		Progress.EnterProgressFrame(1.f);
+ *		DoSecondThing();
+ *	}
+ *
+ */
+struct FScopedSlowTask : FSlowTask
+{
 
 	/**
-	 * Return whether or not this scope should be presented to the user on slow task dialogs
+	 * Construct this scope from an amount of work to do, and a message to display
+	 * @param		InAmountOfWork			Arbitrary number of work units to perform (can be a percentage or number of steps).
+	 *										0 indicates that no progress frames are to be entered in this scope (automatically enters a frame encompassing the entire scope)
+	 * @param		InDefaultMessage		A message to display to the user to describe the purpose of the scope
+	 * @param		bInEnabled				When false, this scope will have no effect. Allows for proper scoped objects that are conditionally disabled.
 	 */
-	bool IsPresentableToUser() const
+	FORCEINLINE FScopedSlowTask(float InAmountOfWork, const FText& InDefaultMessage = FText(), bool bInEnabled = true, FFeedbackContext& InContext = *GWarn)
+		: FSlowTask(InAmountOfWork, InDefaultMessage, bInEnabled, InContext)
 	{
-		// For now we only show scopes that attempted to open a slow task dialog
-		return bVisibleOnUI;
+		Initialize();
 	}
-};
+
+	FORCEINLINE ~FScopedSlowTask()
+	{
+		Destroy();
+	}
+};
\ No newline at end of file
diff --git a/Engine/Source/Runtime/CoreUObject/Private/UObject/LinkerLoad.cpp b/Engine/Source/Runtime/CoreUObject/Private/UObject/LinkerLoad.cpp
--- a/Engine/Source/Runtime/CoreUObject/Private/UObject/LinkerLoad.cpp
+++ b/Engine/Source/Runtime/CoreUObject/Private/UObject/LinkerLoad.cpp
@@ -2739,7 +2739,7 @@ void FLinkerLoad::LoadAllObjects( bool bForcePreload )
 {
 #if WITH_EDITOR
 	FScopedSlowTask SlowTask(ExportMap.Num(), NSLOCTEXT("Core", "LinkerLoad_LoadingObjects", "Loading Objects"), ShouldReportProgress());
-	SlowTask.bVisibleOnUI = false;
+	SlowTask.Visibility = ESlowTaskVisibility::Invisible;
 #endif
 
 #if USE_DEFERRED_DEPENDENCY_CHECK_VERIFICATION_TESTS
diff --git a/Engine/Source/Runtime/CoreUObject/Private/UObject/UObjectGlobals.cpp b/Engine/Source/Runtime/CoreUObject/Private/UObject/UObjectGlobals.cpp
--- a/Engine/Source/Runtime/CoreUObject/Private/UObject/UObjectGlobals.cpp
+++ b/Engine/Source/Runtime/CoreUObject/Private/UObject/UObjectGlobals.cpp
@@ -914,7 +914,7 @@ UPackage* LoadPackageInternal(UPackage* InOuter, const TCHAR* InLongPackageName,
 #endif
 
 	FScopedSlowTask SlowTask(100, FText::Format(NSLOCTEXT("Core", "LoadingPackage_Scope", "Loading Package '{0}'"), FText::FromString(FileToLoad)), ShouldReportProgress());
-	SlowTask.bVisibleOnUI = false;
+	SlowTask.Visibility = ESlowTaskVisibility::Invisible;
 	
 	SlowTask.EnterProgressFrame(10);
 
