Improved the error logging when the TSF IME handling fails to initialize

It now prints out a user friendly version of the HRESULT, as well as the raw HRESULT code.

diff --git a/Engine/Source/Runtime/Core/Private/Windows/WindowsTextInputMethodSystem.cpp b/Engine/Source/Runtime/Core/Private/Windows/WindowsTextInputMethodSystem.cpp
--- a/Engine/Source/Runtime/Core/Private/Windows/WindowsTextInputMethodSystem.cpp
+++ b/Engine/Source/Runtime/Core/Private/Windows/WindowsTextInputMethodSystem.cpp
@@ -363,7 +363,9 @@ bool FWindowsTextInputMethodSystem::InitializeTSF()
 		Result = ::CoCreateInstance(CLSID_TF_InputProcessorProfiles, nullptr, CLSCTX_INPROC_SERVER, IID_ITfInputProcessorProfiles, reinterpret_cast<void**>(&(RawPointerTSFInputProcessorProfiles)));
 		if(FAILED(Result))
 		{
-			UE_LOG(LogWindowsTextInputMethodSystem, Error, TEXT("Initialzation failed while creating the TSF input processor profiles."));
+			TCHAR ErrorMsg[1024];
+			FPlatformMisc::GetSystemErrorMessage(ErrorMsg, 1024, Result);
+			UE_LOG(LogWindowsTextInputMethodSystem, Error, TEXT("Initialization failed while creating the TSF input processor profiles. %s (0x%08x)"), ErrorMsg, Result);
 			return false;
 		}
 		TSFInputProcessorProfiles.Attach(RawPointerTSFInputProcessorProfiles);
@@ -372,7 +374,9 @@ bool FWindowsTextInputMethodSystem::InitializeTSF()
 		Result = TSFInputProcessorProfileManager.FromQueryInterface(IID_ITfInputProcessorProfileMgr, TSFInputProcessorProfiles);
 		if(FAILED(Result))
 		{
-			UE_LOG(LogWindowsTextInputMethodSystem, Error, TEXT("Initialzation failed while acquiring the TSF input processor profile manager."));
+			TCHAR ErrorMsg[1024];
+			FPlatformMisc::GetSystemErrorMessage(ErrorMsg, 1024, Result);
+			UE_LOG(LogWindowsTextInputMethodSystem, Error, TEXT("Initialization failed while acquiring the TSF input processor profile manager. %s (0x%08x)"), ErrorMsg, Result);
 			TSFInputProcessorProfiles.Reset();
 			return false;
 		}
@@ -385,7 +389,9 @@ bool FWindowsTextInputMethodSystem::InitializeTSF()
 		Result = ::CoCreateInstance(CLSID_TF_ThreadMgr, NULL, CLSCTX_INPROC_SERVER, IID_ITfThreadMgr, reinterpret_cast<void**>(&(RawPointerTSFThreadManager)));
 		if(FAILED(Result))
 		{
-			UE_LOG(LogWindowsTextInputMethodSystem, Warning, TEXT("Initialzation failed while creating the TSF thread manager."));
+			TCHAR ErrorMsg[1024];
+			FPlatformMisc::GetSystemErrorMessage(ErrorMsg, 1024, Result);
+			UE_LOG(LogWindowsTextInputMethodSystem, Warning, TEXT("Initialization failed while creating the TSF thread manager. %s (0x%08x)"), ErrorMsg, Result);
 			TSFInputProcessorProfiles.Reset();
 			TSFInputProcessorProfileManager.Reset();
 			return false;
@@ -396,7 +402,9 @@ bool FWindowsTextInputMethodSystem::InitializeTSF()
 		Result = TSFThreadManager->Activate(&(TSFClientId));
 		if(FAILED(Result))
 		{
-			UE_LOG(LogWindowsTextInputMethodSystem, Error, TEXT("Initialzation failed while activating the TSF thread manager."));
+			TCHAR ErrorMsg[1024];
+			FPlatformMisc::GetSystemErrorMessage(ErrorMsg, 1024, Result);
+			UE_LOG(LogWindowsTextInputMethodSystem, Error, TEXT("Initialization failed while activating the TSF thread manager. %s (0x%08x)"), ErrorMsg, Result);
 			TSFInputProcessorProfiles.Reset();
 			TSFInputProcessorProfileManager.Reset();
 			TSFThreadManager.Reset();
@@ -408,7 +416,9 @@ bool FWindowsTextInputMethodSystem::InitializeTSF()
 		Result = TSFSource.FromQueryInterface(IID_ITfSource, TSFThreadManager);
 		if(FAILED(Result))
 		{
-			UE_LOG(LogWindowsTextInputMethodSystem, Error, TEXT("Initialzation failed while acquiring the TSF source from TSF thread manager."));
+			TCHAR ErrorMsg[1024];
+			FPlatformMisc::GetSystemErrorMessage(ErrorMsg, 1024, Result);
+			UE_LOG(LogWindowsTextInputMethodSystem, Error, TEXT("Initialization failed while acquiring the TSF source from TSF thread manager. %s (0x%08x)"), ErrorMsg, Result);
 			TSFInputProcessorProfiles.Reset();
 			TSFInputProcessorProfileManager.Reset();
 			TSFThreadManager.Reset();
@@ -435,7 +445,9 @@ bool FWindowsTextInputMethodSystem::InitializeTSF()
 			Result = TSFSource->AdviseSink(IID_ITfInputProcessorProfileActivationSink, static_cast<ITfInputProcessorProfileActivationSink*>(TSFActivationProxy), &(TSFActivationProxy->TSFProfileCookie));
 			if(FAILED(Result))
 			{
-				UE_LOG(LogWindowsTextInputMethodSystem, Error, TEXT("Initialzation failed while advising the profile notification sink to the TSF source."));
+				TCHAR ErrorMsg[1024];
+				FPlatformMisc::GetSystemErrorMessage(ErrorMsg, 1024, Result);
+				UE_LOG(LogWindowsTextInputMethodSystem, Error, TEXT("Initialization failed while advising the profile notification sink to the TSF source. %s (0x%08x)"), ErrorMsg, Result);
 				TSFInputProcessorProfiles.Reset();
 				TSFInputProcessorProfileManager.Reset();
 				TSFThreadManager.Reset();
@@ -449,7 +461,9 @@ bool FWindowsTextInputMethodSystem::InitializeTSF()
 			Result = TSFSource->AdviseSink(IID_ITfActiveLanguageProfileNotifySink, static_cast<ITfActiveLanguageProfileNotifySink*>(TSFActivationProxy), &(TSFActivationProxy->TSFLanguageCookie));
 			if(FAILED(Result))
 			{
-				UE_LOG(LogWindowsTextInputMethodSystem, Error, TEXT("Initialzation failed while advising the language notification sink to the TSF source."));
+				TCHAR ErrorMsg[1024];
+				FPlatformMisc::GetSystemErrorMessage(ErrorMsg, 1024, Result);
+				UE_LOG(LogWindowsTextInputMethodSystem, Error, TEXT("Initialization failed while advising the language notification sink to the TSF source. %s (0x%08x)"), ErrorMsg, Result);
 				TSFInputProcessorProfiles.Reset();
 				TSFInputProcessorProfileManager.Reset();
 				TSFThreadManager.Reset();
@@ -463,7 +477,9 @@ bool FWindowsTextInputMethodSystem::InitializeTSF()
 	Result = TSFThreadManager->CreateDocumentMgr(&(TSFDisabledDocumentManager));
 	if(FAILED(Result))
 	{
-		UE_LOG(LogWindowsTextInputMethodSystem, Warning, TEXT("Initialzation failed while creating the TSF thread manager."));
+		TCHAR ErrorMsg[1024];
+		FPlatformMisc::GetSystemErrorMessage(ErrorMsg, 1024, Result);
+		UE_LOG(LogWindowsTextInputMethodSystem, Warning, TEXT("Initialization failed while creating the TSF thread manager. %s (0x%08x)"), ErrorMsg, Result);
 		TSFInputProcessorProfiles.Reset();
 		TSFInputProcessorProfileManager.Reset();
 		TSFThreadManager.Reset();
@@ -475,7 +491,9 @@ bool FWindowsTextInputMethodSystem::InitializeTSF()
 	Result = TSFThreadManager->SetFocus(TSFDisabledDocumentManager);
 	if(FAILED(Result))
 	{
-		UE_LOG(LogWindowsTextInputMethodSystem, Error, TEXT("Initialzation failed while activating the TSF thread manager."));
+		TCHAR ErrorMsg[1024];
+		FPlatformMisc::GetSystemErrorMessage(ErrorMsg, 1024, Result);
+		UE_LOG(LogWindowsTextInputMethodSystem, Error, TEXT("Initialization failed while activating the TSF thread manager. %s (0x%08x)"), ErrorMsg, Result);
 		TSFInputProcessorProfiles.Reset();
 		TSFInputProcessorProfileManager.Reset();
 		TSFThreadManager.Reset();
@@ -498,7 +516,9 @@ void FWindowsTextInputMethodSystem::Terminate()
 	Result = TSFSource.FromQueryInterface(IID_ITfSource, TSFThreadManager);
 	if(FAILED(Result) || !TSFSource)
 	{
-		UE_LOG(LogWindowsTextInputMethodSystem, Error, TEXT("Terminating failed while acquiring the TSF source from the TSF thread manager."));
+		TCHAR ErrorMsg[1024];
+		FPlatformMisc::GetSystemErrorMessage(ErrorMsg, 1024, Result);
+		UE_LOG(LogWindowsTextInputMethodSystem, Error, TEXT("Terminating failed while acquiring the TSF source from the TSF thread manager. %s (0x%08x)"), ErrorMsg, Result);
 	}
 
 	if(TSFSource && TSFActivationProxy)
@@ -509,7 +529,9 @@ void FWindowsTextInputMethodSystem::Terminate()
 			Result = TSFSource->UnadviseSink(TSFActivationProxy->TSFLanguageCookie);
 			if(FAILED(Result))
 			{
-				UE_LOG(LogWindowsTextInputMethodSystem, Error, TEXT("Terminating failed while unadvising the language notification sink from the TSF source."));
+				TCHAR ErrorMsg[1024];
+				FPlatformMisc::GetSystemErrorMessage(ErrorMsg, 1024, Result);
+				UE_LOG(LogWindowsTextInputMethodSystem, Error, TEXT("Terminating failed while unadvising the language notification sink from the TSF source. %s (0x%08x)"), ErrorMsg, Result);
 			}
 		}
 
@@ -519,7 +541,9 @@ void FWindowsTextInputMethodSystem::Terminate()
 			Result = TSFSource->UnadviseSink(TSFActivationProxy->TSFProfileCookie);
 			if(FAILED(Result))
 			{
-				UE_LOG(LogWindowsTextInputMethodSystem, Error, TEXT("Terminating failed while unadvising the profile notification sink from the TSF source."));
+				TCHAR ErrorMsg[1024];
+				FPlatformMisc::GetSystemErrorMessage(ErrorMsg, 1024, Result);
+				UE_LOG(LogWindowsTextInputMethodSystem, Error, TEXT("Terminating failed while unadvising the profile notification sink from the TSF source. %s (0x%08x)"), ErrorMsg, Result);
 			}
 		}
 	}
@@ -528,7 +552,9 @@ void FWindowsTextInputMethodSystem::Terminate()
 	Result = TSFThreadManager->Deactivate();
 	if(FAILED(Result))
 	{
-		UE_LOG(LogWindowsTextInputMethodSystem, Error, TEXT("Terminating failed while deactivating the TSF thread manager."));
+		TCHAR ErrorMsg[1024];
+		FPlatformMisc::GetSystemErrorMessage(ErrorMsg, 1024, Result);
+		UE_LOG(LogWindowsTextInputMethodSystem, Error, TEXT("Terminating failed while deactivating the TSF thread manager. %s (0x%08x)"), ErrorMsg, Result);
 	}
 
 	TSFThreadManager.Reset();
@@ -598,7 +624,9 @@ TSharedPtr<ITextInputMethodChangeNotifier> FWindowsTextInputMethodSystem::Regist
 	Result = TSFThreadManager->CreateDocumentMgr(&(TextStore->TSFDocumentManager));
 	if(FAILED(Result) || !TextStore->TSFDocumentManager)
 	{
-		UE_LOG(LogWindowsTextInputMethodSystem, Error, TEXT("Registering a context failed while creating the TSF document manager."));
+		TCHAR ErrorMsg[1024];
+		FPlatformMisc::GetSystemErrorMessage(ErrorMsg, 1024, Result);
+		UE_LOG(LogWindowsTextInputMethodSystem, Error, TEXT("Registering a context failed while creating the TSF document manager. %s (0x%08x)"), ErrorMsg, Result);
 		TextStore.Reset();
 		return nullptr;
 	}
@@ -606,7 +634,9 @@ TSharedPtr<ITextInputMethodChangeNotifier> FWindowsTextInputMethodSystem::Regist
 	Result = TextStore->TSFDocumentManager->CreateContext(TSFClientId, 0, static_cast<ITextStoreACP*>(TextStore), &(TextStore->TSFContext), &(TextStore->TSFEditCookie));	
 	if(FAILED(Result) || !TextStore->TSFContext)
 	{
-		UE_LOG(LogWindowsTextInputMethodSystem, Error, TEXT("Registering a context failed while creating the TSF context."));
+		TCHAR ErrorMsg[1024];
+		FPlatformMisc::GetSystemErrorMessage(ErrorMsg, 1024, Result);
+		UE_LOG(LogWindowsTextInputMethodSystem, Error, TEXT("Registering a context failed while creating the TSF context. %s (0x%08x)"), ErrorMsg, Result);
 		TextStore.Reset();
 		return nullptr;
 	}
@@ -614,7 +644,9 @@ TSharedPtr<ITextInputMethodChangeNotifier> FWindowsTextInputMethodSystem::Regist
 	Result = TextStore->TSFDocumentManager->Push(TextStore->TSFContext);
 	if(FAILED(Result))
 	{
-		UE_LOG(LogWindowsTextInputMethodSystem, Error, TEXT("Registering a context failed while pushing a TSF context onto the TSF document manager."));
+		TCHAR ErrorMsg[1024];
+		FPlatformMisc::GetSystemErrorMessage(ErrorMsg, 1024, Result);
+		UE_LOG(LogWindowsTextInputMethodSystem, Error, TEXT("Registering a context failed while pushing a TSF context onto the TSF document manager. %s (0x%08x)"), ErrorMsg, Result);
 		TextStore.Reset();
 		return nullptr;
 	}
@@ -622,11 +654,14 @@ TSharedPtr<ITextInputMethodChangeNotifier> FWindowsTextInputMethodSystem::Regist
 	Result = TextStore->TSFContextOwnerCompositionServices.FromQueryInterface(IID_ITfContextOwnerCompositionServices, TextStore->TSFContext);
 	if(FAILED(Result) || !TextStore->TSFContextOwnerCompositionServices)
 	{
-		UE_LOG(LogWindowsTextInputMethodSystem, Error, TEXT("Registering a context failed while getting the TSF context owner composition services."));
+		TCHAR ErrorMsg[1024];
+		FPlatformMisc::GetSystemErrorMessage(ErrorMsg, 1024, Result);
+		UE_LOG(LogWindowsTextInputMethodSystem, Error, TEXT("Registering a context failed while getting the TSF context owner composition services. %s (0x%08x)"), ErrorMsg, Result);
 		Result = TextStore->TSFDocumentManager->Pop(TF_POPF_ALL);
 		if(FAILED(Result))
 		{
-			UE_LOG(LogWindowsTextInputMethodSystem, Error, TEXT("Failed to pop a TSF context off from TSF document manager while recovering from failing."));
+			FPlatformMisc::GetSystemErrorMessage(ErrorMsg, 1024, Result);
+			UE_LOG(LogWindowsTextInputMethodSystem, Error, TEXT("Failed to pop a TSF context off from TSF document manager while recovering from failing. %s (0x%08x)"), ErrorMsg, Result);
 		}
 		TextStore.Reset();
 		return nullptr;
@@ -654,7 +689,9 @@ void FWindowsTextInputMethodSystem::UnregisterContext(const TSharedRef<ITextInpu
 	Result = TextStore->TSFDocumentManager->Pop(TF_POPF_ALL);
 	if(FAILED(Result))
 	{
-		UE_LOG(LogWindowsTextInputMethodSystem, Error, TEXT("Unregistering a context failed while popping a TSF context off from the TSF document manager."));
+		TCHAR ErrorMsg[1024];
+		FPlatformMisc::GetSystemErrorMessage(ErrorMsg, 1024, Result);
+		UE_LOG(LogWindowsTextInputMethodSystem, Error, TEXT("Unregistering a context failed while popping a TSF context off from the TSF document manager. %s (0x%08x)"), ErrorMsg, Result);
 	}
 
 	ContextToInternalContextMap.Remove(Context);
@@ -686,7 +723,9 @@ void FWindowsTextInputMethodSystem::ActivateContext(const TSharedRef<ITextInputM
 	Result = TSFThreadManager->AssociateFocus(InternalContext.WindowHandle, TextStore->TSFDocumentManager, &Unused);
 	if(FAILED(Result))
 	{
-		UE_LOG(LogWindowsTextInputMethodSystem, Error, TEXT("Activating a context failed while setting focus on a TSF document manager."));
+		TCHAR ErrorMsg[1024];
+		FPlatformMisc::GetSystemErrorMessage(ErrorMsg, 1024, Result);
+		UE_LOG(LogWindowsTextInputMethodSystem, Error, TEXT("Activating a context failed while setting focus on a TSF document manager. %s (0x%08x)"), ErrorMsg, Result);
 	}
 
 	UE_LOG(LogWindowsTextInputMethodSystem, Verbose, TEXT("Activated context %p!"), &(Context.Get()));
