#ue4 - Add FTransform::RotationEquals, TranslationEquals, and Scale3DEquals to check relevant component of another FTransform for equality.

- Renamed the old private versions to make them clearly private.
- We do it this way because there is a cost to conversion such as A.GetRotation().Equals(B.GetRotation()) in the vectorized FTransform, where it's more efficient to keep the components in the VectorRegister through the comparison.

diff --git a/Engine/Source/Runtime/Core/Private/Math/TransformVectorized.cpp b/Engine/Source/Runtime/Core/Private/Math/TransformVectorized.cpp
--- a/Engine/Source/Runtime/Core/Private/Math/TransformVectorized.cpp
+++ b/Engine/Source/Runtime/Core/Private/Math/TransformVectorized.cpp
@@ -224,21 +224,21 @@ bool FTransform::DebugEqualMatrix(const FMatrix& Matrix) const
 	if (!Equals(TestResult))
 	{
 		// see now which one isn't equal		
-		if (!Scale3DEquals(TestResult.Scale3D, ScalarRegister(0.01f)))
+		if (!Scale3DEquals(TestResult, 0.01f))
 		{
 			UE_LOG(LogTransform, Log, TEXT("Matrix(S)\t%s"), *TestResult.GetScale3D().ToString());
 			UE_LOG(LogTransform, Log, TEXT("VQS(S)\t%s"), *GetScale3D().ToString());
 		}
 
 		// see now which one isn't equal
-		if (!RotationEquals(TestResult.Rotation))
+		if (!RotationEquals(TestResult))
 		{
 			UE_LOG(LogTransform, Log, TEXT("Matrix(R)\t%s"), *TestResult.GetRotation().ToString());
 			UE_LOG(LogTransform, Log, TEXT("VQS(R)\t%s"), *GetRotation().ToString());
 		}
 
 		// see now which one isn't equal
-		if (!TranslationEquals(TestResult.Translation, ScalarRegister(0.01f)))
+		if (!TranslationEquals(TestResult, 0.01f))
 		{
 			UE_LOG(LogTransform, Log, TEXT("Matrix(T)\t%s"), *TestResult.GetTranslation().ToString());
 			UE_LOG(LogTransform, Log, TEXT("VQS(T)\t%s"), *GetTranslation().ToString());
diff --git a/Engine/Source/Runtime/Core/Public/Math/Transform.h b/Engine/Source/Runtime/Core/Public/Math/Transform.h
--- a/Engine/Source/Runtime/Core/Public/Math/Transform.h
+++ b/Engine/Source/Runtime/Core/Public/Math/Transform.h
@@ -612,17 +612,17 @@ public:
 
 private:
 
-	FORCEINLINE bool RotationEquals(const FQuat& InRotation, const float Tolerance = KINDA_SMALL_NUMBER) const
+	FORCEINLINE bool Private_RotationEquals(const FQuat& InRotation, const float Tolerance = KINDA_SMALL_NUMBER) const
 	{
 		return Rotation.Equals(InRotation, Tolerance);
 	}
 
-	FORCEINLINE bool TranslationEquals(const FVector& InTranslation, const float Tolerance = KINDA_SMALL_NUMBER) const
+	FORCEINLINE bool Private_TranslationEquals(const FVector& InTranslation, const float Tolerance = KINDA_SMALL_NUMBER) const
 	{
 		return Translation.Equals(InTranslation, Tolerance);
 	}
 
-	FORCEINLINE bool Scale3DEquals(const FVector& InScale3D, const float Tolerance = KINDA_SMALL_NUMBER) const
+	FORCEINLINE bool Private_Scale3DEquals(const FVector& InScale3D, const float Tolerance = KINDA_SMALL_NUMBER) const
 	{
 		return (Scale3D-InScale3D).SizeSquared() <= FMath::Square(Tolerance);
 	}
@@ -632,32 +632,52 @@ public:
 	// Test if A's rotation equals B's rotation, within a tolerance. Preferred over "A.GetRotation().Equals(B.GetRotation())" because it is faster on some platforms.
 	FORCEINLINE static bool AreRotationsEqual(const FTransform& A, const FTransform& B, float Tolerance=KINDA_SMALL_NUMBER)
 	{
-		return A.RotationEquals(B.Rotation, Tolerance);
+		return A.Private_RotationEquals(B.Rotation, Tolerance);
 	}
 
 	// Test if A's translation equals B's translation, within a tolerance. Preferred over "A.GetTranslation().Equals(B.GetTranslation())" because it is faster on some platforms.
 	FORCEINLINE static bool AreTranslationsEqual(const FTransform& A, const FTransform& B, float Tolerance=KINDA_SMALL_NUMBER)
 	{
-		return A.TranslationEquals(B.Translation, Tolerance);
+		return A.Private_TranslationEquals(B.Translation, Tolerance);
 	}
 
 	// Test if A's scale equals B's scale, within a tolerance. Preferred over "A.GetScale3D().Equals(B.GetScale3D())" because it is faster on some platforms.
 	FORCEINLINE static bool AreScale3DsEqual(const FTransform& A, const FTransform& B, float Tolerance=KINDA_SMALL_NUMBER)
 	{
-		return A.Scale3DEquals(B.Scale3D, Tolerance);
+		return A.Private_Scale3DEquals(B.Scale3D, Tolerance);
+	}
+
+
+
+	// Test if this Transform's rotation equals another's rotation, within a tolerance. Preferred over "GetRotation().Equals(Other.GetRotation())" because it is faster on some platforms.
+	FORCEINLINE bool RotationEquals(const FTransform& Other, float Tolerance=KINDA_SMALL_NUMBER) const
+	{
+		return AreRotationsEqual(*this, Other, Tolerance);
+	}
+
+	// Test if this Transform's translation equals another's translation, within a tolerance. Preferred over "GetTranslation().Equals(Other.GetTranslation())" because it is faster on some platforms.
+	FORCEINLINE bool TranslationEquals(const FTransform& Other, float Tolerance=KINDA_SMALL_NUMBER) const
+	{
+		return AreTranslationsEqual(*this, Other, Tolerance);
+	}
+
+	// Test if this Transform's scale equals another's scale, within a tolerance. Preferred over "GetScale3D().Equals(Other.GetScale3D())" because it is faster on some platforms.
+	FORCEINLINE bool Scale3DEquals(const FTransform& Other, float Tolerance=KINDA_SMALL_NUMBER) const
+	{
+		return AreScale3DsEqual(*this, Other, Tolerance);
 	}
 
 
 	// Test if all components of the transforms are equal, within a tolerance.
 	inline bool Equals(const FTransform& Other, float Tolerance=KINDA_SMALL_NUMBER) const
 	{
-		return RotationEquals(Other.Rotation, Tolerance) && TranslationEquals(Other.Translation, Tolerance) && Scale3DEquals(Other.Scale3D, Tolerance);
+		return Private_RotationEquals(Other.Rotation, Tolerance) && Private_TranslationEquals(Other.Translation, Tolerance) && Private_Scale3DEquals(Other.Scale3D, Tolerance);
 	}
 
 	// Test if rotation and translation components of the transforms are equal, within a tolerance.
 	inline bool EqualsNoScale(const FTransform& Other, float Tolerance=KINDA_SMALL_NUMBER) const
 	{
-		return RotationEquals(Other.Rotation, Tolerance) && TranslationEquals(Other.Translation, Tolerance);
+		return Private_RotationEquals(Other.Rotation, Tolerance) && Private_TranslationEquals(Other.Translation, Tolerance);
 	}
 
 	/**
diff --git a/Engine/Source/Runtime/Core/Public/Math/TransformVectorized.h b/Engine/Source/Runtime/Core/Public/Math/TransformVectorized.h
--- a/Engine/Source/Runtime/Core/Public/Math/TransformVectorized.h
+++ b/Engine/Source/Runtime/Core/Public/Math/TransformVectorized.h
@@ -639,7 +639,7 @@ public:
 
 private:
 
-	inline bool RotationEquals( const VectorRegister& InRotation, const ScalarRegister& Tolerance = ScalarRegister(GlobalVectorConstants::KindaSmallNumber)) const
+	inline bool Private_RotationEquals( const VectorRegister& InRotation, const ScalarRegister& Tolerance = ScalarRegister(GlobalVectorConstants::KindaSmallNumber)) const
 	{			
 		// !( (FMath::Abs(X-Q.X) > Tolerance) || (FMath::Abs(Y-Q.Y) > Tolerance) || (FMath::Abs(Z-Q.Z) > Tolerance) || (FMath::Abs(W-Q.W) > Tolerance) )
 		const VectorRegister RotationSub = VectorAbs(VectorSubtract(Rotation, InRotation));		
@@ -648,14 +648,14 @@ private:
 		return !VectorAnyGreaterThan(RotationSub, Tolerance.Value) || !VectorAnyGreaterThan(RotationAdd, Tolerance.Value);
 	}
 
-	inline bool TranslationEquals( const VectorRegister& InTranslation, const ScalarRegister& Tolerance = ScalarRegister(GlobalVectorConstants::KindaSmallNumber)) const
+	inline bool Private_TranslationEquals( const VectorRegister& InTranslation, const ScalarRegister& Tolerance = ScalarRegister(GlobalVectorConstants::KindaSmallNumber)) const
 	{			
 		// !( (FMath::Abs(X-V.X) > Tolerance) || (FMath::Abs(Y-V.Y) > Tolerance) || (FMath::Abs(Z-V.Z) > Tolerance) )
 		const VectorRegister TranslationDiff = VectorAbs(VectorSubtract(Translation, InTranslation));		
 		return !VectorAnyGreaterThan(TranslationDiff, Tolerance.Value);
 	}
 
-	inline bool Scale3DEquals( const VectorRegister& InScale3D, const ScalarRegister& Tolerance = ScalarRegister(GlobalVectorConstants::KindaSmallNumber)) const
+	inline bool Private_Scale3DEquals( const VectorRegister& InScale3D, const ScalarRegister& Tolerance = ScalarRegister(GlobalVectorConstants::KindaSmallNumber)) const
 	{			
 		const VectorRegister ScaleDiff = VectorSubtract(Scale3D, InScale3D);		
 		// d = dot3(ScaleDiff.xyz, ScaleDiff.xyz), VectorRegister( d, d, d, d )
@@ -670,19 +670,38 @@ public:
 	// Test if A's rotation equals B's rotation, within a tolerance. Preferred over "A.GetRotation().Equals(B.GetRotation())" because it is faster on some platforms.
 	FORCEINLINE static bool AreRotationsEqual(const FTransform& A, const FTransform& B, float Tolerance=KINDA_SMALL_NUMBER)
 	{
-		return A.RotationEquals(B.Rotation, ScalarRegister(Tolerance));
+		return A.Private_RotationEquals(B.Rotation, ScalarRegister(Tolerance));
 	}
 
 	// Test if A's translation equals B's translation, within a tolerance. Preferred over "A.GetTranslation().Equals(B.GetTranslation())" because it avoids VectorRegister->FVector conversion.
 	FORCEINLINE static bool AreTranslationsEqual(const FTransform& A, const FTransform& B, float Tolerance=KINDA_SMALL_NUMBER)
 	{
-		return A.TranslationEquals(B.Translation, ScalarRegister(Tolerance));
+		return A.Private_TranslationEquals(B.Translation, ScalarRegister(Tolerance));
 	}
 
 	// Test if A's scale equals B's scale, within a tolerance. Preferred over "A.GetScale3D().Equals(B.GetScale3D())" because it avoids VectorRegister->FVector conversion.
 	FORCEINLINE static bool AreScale3DsEqual(const FTransform& A, const FTransform& B, float Tolerance=KINDA_SMALL_NUMBER)
 	{
-		return A.Scale3DEquals(B.Scale3D, ScalarRegister(Tolerance));
+		return A.Private_Scale3DEquals(B.Scale3D, ScalarRegister(Tolerance));
+	}
+
+
+	// Test if this Transform's rotation equals another's rotation, within a tolerance. Preferred over "GetRotation().Equals(Other.GetRotation())" because it is faster on some platforms.
+	FORCEINLINE bool RotationEquals(const FTransform& Other, float Tolerance=KINDA_SMALL_NUMBER) const
+	{
+		return AreRotationsEqual(*this, Other, Tolerance);
+	}
+
+	// Test if this Transform's translation equals another's translation, within a tolerance. Preferred over "GetTranslation().Equals(Other.GetTranslation())" because it avoids VectorRegister->FVector conversion.
+	FORCEINLINE bool TranslationEquals(const FTransform& Other, float Tolerance=KINDA_SMALL_NUMBER) const
+	{
+		return AreTranslationsEqual(*this, Other, Tolerance);
+	}
+
+	// Test if this Transform's scale equals another's scale, within a tolerance. Preferred over "GetScale3D().Equals(Other.GetScale3D())" because it avoids VectorRegister->FVector conversion.
+	FORCEINLINE bool Scale3DEquals(const FTransform& Other, float Tolerance=KINDA_SMALL_NUMBER) const
+	{
+		return AreScale3DsEqual(*this, Other, Tolerance);
 	}
 
 
@@ -690,14 +709,14 @@ public:
 	inline bool Equals(const FTransform& Other, float Tolerance=KINDA_SMALL_NUMBER) const
 	{
 		const ScalarRegister ToleranceRegister(Tolerance);
-		return RotationEquals(Other.Rotation, ToleranceRegister) && TranslationEquals(Other.Translation, ToleranceRegister) && Scale3DEquals(Other.Scale3D, ToleranceRegister);
+		return Private_RotationEquals(Other.Rotation, ToleranceRegister) && Private_TranslationEquals(Other.Translation, ToleranceRegister) && Private_Scale3DEquals(Other.Scale3D, ToleranceRegister);
 	}
 
 	// Test if rotation and translation components of the transforms are equal, within a tolerance.
 	inline bool EqualsNoScale(const FTransform& Other, float Tolerance=KINDA_SMALL_NUMBER) const
 	{
 		const ScalarRegister ToleranceRegister(Tolerance);
-		return RotationEquals(Other.Rotation, ToleranceRegister) && TranslationEquals(Other.Translation, ToleranceRegister);
+		return Private_RotationEquals(Other.Rotation, ToleranceRegister) && Private_TranslationEquals(Other.Translation, ToleranceRegister);
 	}
 
 	FORCEINLINE static void Multiply(FTransform* OutTransform, const FTransform* A, const FTransform* B);
