Fixed build fallout from CL# 2606453

diff --git a/Engine/Source/Runtime/Core/Private/Internationalization/ICUText.cpp b/Engine/Source/Runtime/Core/Private/Internationalization/ICUText.cpp
--- a/Engine/Source/Runtime/Core/Private/Internationalization/ICUText.cpp
+++ b/Engine/Source/Runtime/Core/Private/Internationalization/ICUText.cpp
@@ -225,7 +225,7 @@ bool FText::FSortPredicate::operator()(const FText& A, const FText& B) const
 
 bool FText::IsLetter( const TCHAR Char )
 {
-	icu::UnicodeString PatternString(TEXT("[\\p{L}]"));
+	icu::UnicodeString PatternString = ICUUtilities::ConvertString(TEXT("[\\p{L}]"));
 	UErrorCode ICUStatus = U_ZERO_ERROR;
 	icu::UnicodeSet Uniscode(PatternString, ICUStatus);
 	return Uniscode.contains(Char) != 0;
