Make sure the UE4Game and UE4Game/[ProjectName] directories are created (do not use GetFileBasePath() since FApp::GetGameName() isn't valid yet and it will cache the wrong base path)
#ue-17385
#ue-17980
#ue4
#android

diff --git a/Engine/Source/Runtime/Core/Private/Android/AndroidFile.cpp b/Engine/Source/Runtime/Core/Private/Android/AndroidFile.cpp
--- a/Engine/Source/Runtime/Core/Private/Android/AndroidFile.cpp
+++ b/Engine/Source/Runtime/Core/Private/Android/AndroidFile.cpp
@@ -847,6 +847,7 @@ public:
 #endif
 		if (!IPhysicalPlatformFile::Initialize(Inner, CmdLine))
 		{
+			FPlatformMisc::LowLevelOutputDebugStringf(TEXT("FAndroidPlatformFile::Initialize failed"));
 			return false;
 		}
 		if (GOBBinAPK)
@@ -886,10 +887,10 @@ public:
 			}
 		}
 
-		// make sure the base path directory exists (UE4Game and UE4Game/GameName)
+		// make sure the base path directory exists (UE4Game and UE4Game/ProjectName)
 		FString FileBaseDir = GFilePathBase + FString(FILEBASE_DIRECTORY);
 		mkdir(TCHAR_TO_UTF8(*FileBaseDir), 0766);
-		mkdir(TCHAR_TO_UTF8(*GetFileBasePath()), 0766);
+		mkdir(TCHAR_TO_UTF8(*(FileBaseDir + GAndroidProjectName)), 0766);
 
 		return true;
 	}
@@ -1327,6 +1328,9 @@ public:
 		bool Found = false;
 		if (IsLocal(LocalPath))
 		{
+#if LOG_ANDROID_FILE
+			FPlatformMisc::LowLevelOutputDebugStringf(TEXT("FAndroidPlatformFile::DirectoryExists('%s') => Check IsLocal: '%s'"), Directory, *(LocalPath + "/"));
+#endif
 			struct stat FileInfo;
 			if (stat(TCHAR_TO_UTF8(*LocalPath), &FileInfo) != -1)
 			{
