Slate/Core - Adding back TAttribute == / !=. but marking them as deprecated.

diff --git a/Engine/Source/Runtime/Core/Public/Misc/Attribute.h b/Engine/Source/Runtime/Core/Public/Misc/Attribute.h
--- a/Engine/Source/Runtime/Core/Public/Misc/Attribute.h
+++ b/Engine/Source/Runtime/Core/Public/Misc/Attribute.h
@@ -333,6 +333,18 @@ public:
 	{
 		return Get() != InOther;
 	}
+
+	DEPRECATED(4.9, "Comparing an attribute directly to an attribute is deprecated.  Please call Get() on the attributes before comparing them.  If you were actually testing if the Attributes were the same functionally, call IdenticalTo")
+	bool operator==( const TAttribute& InOther ) const
+	{
+		return Get() == InOther.Get();
+	}
+
+	DEPRECATED(4.9, "Comparing an attribute directly to an attribute is deprecated.  Please call Get() on the attributes before comparing them.  If you were actually testing if the Attributes were the same functionally, call IdenticalTo")
+	bool operator!=( const TAttribute& InOther ) const
+	{
+		return Get() != InOther.Get();
+	}
 	
 private:
 
