TSet fixes when ElementType != KeyType.

diff --git a/Engine/Source/Runtime/Core/Public/Containers/Set.h b/Engine/Source/Runtime/Core/Public/Containers/Set.h
--- a/Engine/Source/Runtime/Core/Public/Containers/Set.h
+++ b/Engine/Source/Runtime/Core/Public/Containers/Set.h
@@ -735,7 +735,7 @@ public:
 		TSet Result;
 		for(TConstIterator SetIt(*this);SetIt;++SetIt)
 		{
-			if(OtherSet.Contains(*SetIt))
+			if(OtherSet.Contains(KeyFuncs::GetSetKey(*SetIt)))
 			{
 				Result.Add(*SetIt);
 			}
@@ -764,7 +764,7 @@ public:
 		TSet Result;
 		for(TConstIterator SetIt(*this);SetIt;++SetIt)
 		{
-			if(!OtherSet.Contains(*SetIt))
+			if(!OtherSet.Contains(KeyFuncs::GetSetKey(*SetIt)))
 			{
 				Result.Add(*SetIt);
 			}
@@ -784,7 +784,7 @@ public:
 		bool bIncludesSet = true;
 		for(TConstIterator OtherSetIt(OtherSet); OtherSetIt; ++OtherSetIt)
 		{
-			if (!Contains(*OtherSetIt))
+			if (!Contains(KeyFuncs::GetSetKey(*OtherSetIt)))
 			{
 				bIncludesSet = false;
 				break;
