Backout 2600433 for now - Device Clipboard Support for Crashes. Issues may arise from this change with stalling after a crash.

diff --git a/Engine/Source/Runtime/Core/Private/IOS/IOSPlatformOutputDevices.cpp b/Engine/Source/Runtime/Core/Private/IOS/IOSPlatformOutputDevices.cpp
--- a/Engine/Source/Runtime/Core/Private/IOS/IOSPlatformOutputDevices.cpp
+++ b/Engine/Source/Runtime/Core/Private/IOS/IOSPlatformOutputDevices.cpp
@@ -81,9 +81,6 @@ void FOutputDeviceIOSError::HandleError()
 	GIsRunning = 0;
 	GIsCriticalError = 1;
 	GLogConsole = NULL;
-    
-    GErrorHist[ARRAY_COUNT(GErrorHist)-1] = 0;
-    FPlatformMisc::ClipboardCopy(GErrorHist);
 
 	GLog->Flush();
 }
