Build Fix.

diff --git a/Engine/Source/Runtime/Core/Public/Misc/Attribute.h b/Engine/Source/Runtime/Core/Public/Misc/Attribute.h
--- a/Engine/Source/Runtime/Core/Public/Misc/Attribute.h
+++ b/Engine/Source/Runtime/Core/Public/Misc/Attribute.h
@@ -316,10 +316,10 @@ public:
 	 * @param InOther The other attribute to compare with.
 	 * @return true if the attributes are not equal, false otherwise.
 	 */
-	DEPRECATED(4.9, "Comparing an attribute directly to a literal value is deprecated.  Please call Get() on the attribute before comparing it." )
-	bool operator==( const ObjectType& InOther ) const
+	DEPRECATED(4.9, "Comparing an attribute directly to an attribute is deprecated.  Please call Get() on the attributes before comparing them.  If you were actually testing if the Attributes were the same functionally, call IdenticalTo")
+	bool operator==( const TAttribute& InOther ) const
 	{
-		return Get() == InOther;
+		return Get() == InOther.Get();
 	}
 
 	/**
@@ -328,18 +328,6 @@ public:
 	 * @param InOther The other attribute to compare with.
 	 * @return true if the attributes are not equal, false otherwise.
 	 */
-	DEPRECATED(4.9, "Comparing an attribute directly to a literal value is deprecated.  Please call Get() on the attribute before comparing it.")
-	bool operator!=( const ObjectType& InOther ) const
-	{
-		return Get() != InOther;
-	}
-
-	DEPRECATED(4.9, "Comparing an attribute directly to an attribute is deprecated.  Please call Get() on the attributes before comparing them.  If you were actually testing if the Attributes were the same functionally, call IdenticalTo")
-	bool operator==( const TAttribute& InOther ) const
-	{
-		return Get() == InOther.Get();
-	}
-
 	DEPRECATED(4.9, "Comparing an attribute directly to an attribute is deprecated.  Please call Get() on the attributes before comparing them.  If you were actually testing if the Attributes were the same functionally, call IdenticalTo")
 	bool operator!=( const TAttribute& InOther ) const
 	{
