Ensures that a custom version doesn't change during a reregistration.

diff --git a/Engine/Source/Runtime/Core/Private/Serialization/CustomVersion.cpp b/Engine/Source/Runtime/Core/Private/Serialization/CustomVersion.cpp
--- a/Engine/Source/Runtime/Core/Private/Serialization/CustomVersion.cpp
+++ b/Engine/Source/Runtime/Core/Private/Serialization/CustomVersion.cpp
@@ -133,12 +133,21 @@ FCustomVersionRegistration::FCustomVersionRegistration(FGuid InKey, int32 InVers
 	// Check if this tag hasn't already been registered
 	if (FCustomVersion* ExistingRegistration = Versions.FindByKey(InKey))
 	{
-		// We don't allow the version number to decrease between registrations
-		check(InVersion >= ExistingRegistration->Version);
+		// We don't allow the registration details to change across registrations - this code path only exists to support hotreload
+
+		// If you hit this then you've probably either:
+		// * Changed registration details during hotreload.
+		// * Accidentally copy-and-pasted an FCustomVersionRegistration object.
+		ensureMsgf(
+			ExistingRegistration->Version == InVersion && ExistingRegistration->FriendlyName == InFriendlyName,
+			TEXT("Custom version registrations cannot change between hotreloads - \"%s\" version %d is being reregistered as \"%s\" version %d"),
+			*ExistingRegistration->FriendlyName,
+			ExistingRegistration->Version,
+			InFriendlyName,
+			InVersion
+		);
 
 		++ExistingRegistration->ReferenceCount;
-		ExistingRegistration->Version      = InVersion;
-		ExistingRegistration->FriendlyName = InFriendlyName;
 	}
 	else
 	{
