Deprecated FRotator::operator-() and created a new method GetInverse().  Componentwise negation is a meaningless operation on a Rotator.

diff --git a/Engine/Source/Runtime/Core/Private/Math/UnrealMath.cpp b/Engine/Source/Runtime/Core/Private/Math/UnrealMath.cpp
--- a/Engine/Source/Runtime/Core/Private/Math/UnrealMath.cpp
+++ b/Engine/Source/Runtime/Core/Private/Math/UnrealMath.cpp
@@ -272,6 +272,12 @@ CORE_API FVector FRotator::Vector() const
 }
 
 
+FRotator FRotator::GetInverse() const
+{
+	return Quaternion().Inverse().Rotator();
+}
+
+
 FQuat FRotator::Quaternion() const
 {
 	SCOPE_CYCLE_COUNTER(STAT_MathConvertRotatorToQuat);
diff --git a/Engine/Source/Runtime/Core/Public/Math/Rotator.h b/Engine/Source/Runtime/Core/Public/Math/Rotator.h
--- a/Engine/Source/Runtime/Core/Public/Math/Rotator.h
+++ b/Engine/Source/Runtime/Core/Public/Math/Rotator.h
@@ -106,6 +106,7 @@ public:
 	 *
 	 * @return A negated copy of the rotator.
 	 */
+	DEPRECATED(4.9, "The unary negation operator has been deprecated as componentwise negation is not meaningful for a Rotator. To get the inverse, please use FRotator::GetInverse()")
 	FORCEINLINE FRotator operator-() const;
 
 	// Binary comparison operators.
@@ -182,6 +183,11 @@ public:
 	 */
 	FRotator Add( float DeltaPitch, float DeltaYaw, float DeltaRoll );
 
+	/**
+	 * Returns the inverse of the rotator.
+	 */
+	CORE_API FRotator GetInverse() const;
+
 	/**
 	 * Get the rotation, snapped to specified degree segments.
 	 *
diff --git a/Engine/Source/Runtime/Engine/Classes/Kismet/KismetMathLibrary.h b/Engine/Source/Runtime/Engine/Classes/Kismet/KismetMathLibrary.h
--- a/Engine/Source/Runtime/Engine/Classes/Kismet/KismetMathLibrary.h
+++ b/Engine/Source/Runtime/Engine/Classes/Kismet/KismetMathLibrary.h
@@ -735,7 +735,7 @@ class ENGINE_API UKismetMathLibrary : public UBlueprintFunctionLibrary
 	static FRotator ComposeRotators(FRotator A, FRotator B);
 
 	/** Negate a rotator*/
-	UFUNCTION(BlueprintPure, Category="Math|Rotator", meta=(Keywords="rotate rotation"))
+	UFUNCTION(BlueprintPure, Category="Math|Rotator", meta=(DisplayName="InvertRotator", Keywords="rotate rotation"))
 	static FRotator NegateRotator(FRotator A);
 
 	/** Get the reference frame direction vectors (axes) described by this rotation */
diff --git a/Engine/Source/Runtime/Engine/Private/KismetMathLibrary.cpp b/Engine/Source/Runtime/Engine/Private/KismetMathLibrary.cpp
--- a/Engine/Source/Runtime/Engine/Private/KismetMathLibrary.cpp
+++ b/Engine/Source/Runtime/Engine/Private/KismetMathLibrary.cpp
@@ -969,8 +969,7 @@ FRotator UKismetMathLibrary::ComposeRotators(FRotator A, FRotator B)
 
 FRotator UKismetMathLibrary::NegateRotator( FRotator A )
 {
-	FQuat AQuat = FQuat(A);
-	return FRotator(AQuat.Inverse());
+	return A.GetInverse();
 }
 
 void UKismetMathLibrary::GetAxes(FRotator A, FVector& X, FVector& Y, FVector& Z)
