UE-16970 - StrategyGame crash on Android devices bad name index
- General fix/change to the IO interface so that it can work with multiple threads calling in to it by using the offset version of the read/write functions and tracking the offset itself.

diff --git a/Engine/Source/Runtime/Core/Private/Android/AndroidFile.cpp b/Engine/Source/Runtime/Core/Private/Android/AndroidFile.cpp
--- a/Engine/Source/Runtime/Core/Private/Android/AndroidFile.cpp
+++ b/Engine/Source/Runtime/Core/Private/Android/AndroidFile.cpp
@@ -133,6 +133,7 @@ public:
 	TSharedPtr<FileReference> File;
 	int64 Start;
 	int64 Length;
+	int64 CurrentOffset;
 
 	FORCEINLINE void CheckValid()
 	{
@@ -142,7 +143,7 @@ public:
 	// Invalid handle.
 	FFileHandleAndroid()
 		: File(MakeShareable(new FileReference()))
-		, Start(0), Length(0)
+		, Start(0), Length(0), CurrentOffset(0)
 	{
 	}
 
@@ -151,6 +152,7 @@ public:
 		int64 start, int64 length)
 		: File(base.File)
 		, Start(base.Start + start), Length(length)
+		, CurrentOffset(0)
 	{
 		CheckValid();
 		LogInfo();
@@ -159,7 +161,7 @@ public:
 	// Handle that covers the entire file content.
 	FFileHandleAndroid(const FString & path, int32 filehandle)
 		: File(MakeShareable(new FileReference(path, filehandle)))
-		, Start(0), Length(0)
+		, Start(0), Length(0), CurrentOffset(0)
 	{
 		CheckValid();
 #if UE_ANDROID_FILE_64
@@ -175,7 +177,7 @@ public:
 	// Handle that covers the entire content of an asset.
 	FFileHandleAndroid(const FString & path, AAsset * asset)
 		: File(MakeShareable(new FileReference(path, asset)))
-		, Start(0), Length(0)
+		, Start(0), Length(0), CurrentOffset(0)
 	{
 #if UE_ANDROID_FILE_64
 		File->Handle = AAsset_openFileDescriptor64(File->Asset, &Start, &Length);
@@ -197,11 +199,7 @@ public:
 	virtual int64 Tell() override
 	{
 		CheckValid();
-#if UE_ANDROID_FILE_64
-		int64 pos = lseek64(File->Handle, 0, SEEK_CUR);
-#else
-		int64 pos = lseek(File->Handle, 0, SEEK_CUR);
-#endif
+		int64 pos = CurrentOffset;
 		check(pos != -1);
 		return pos - Start; // We are treating 'tell' as a virtual location from file Start
 	}
@@ -210,16 +208,10 @@ public:
 	{
 		CheckValid();
 		// we need to offset all positions by the Start offset
-		NewPosition += Start;
+		CurrentOffset = NewPosition += Start;
 		check(NewPosition >= 0);
 		
-#if UE_ANDROID_FILE_64
-		int64 pos = lseek64(File->Handle, NewPosition, SEEK_SET);
-#else
-		check(NewPosition <= std::numeric_limits<off_t>::max());
-		int64 pos = lseek(File->Handle, static_cast<off_t>(NewPosition), SEEK_SET);
-#endif
-		return pos != -1;
+		return true;
 	}
 
     virtual bool SeekFromEnd(int64 NewPositionRelativeToEnd = 0) override
@@ -227,34 +219,35 @@ public:
 		CheckValid();
 		check(NewPositionRelativeToEnd <= 0);
 		// We need to convert this to a virtual offset inside the file we are interested in
-		int64 position = Start + (Length - NewPositionRelativeToEnd);
-#if UE_ANDROID_FILE_64
-		int64 pos = lseek64(File->Handle, position, SEEK_SET);
-#else
-		check(position <= std::numeric_limits<off_t>::max());
-		int64 pos = lseek(File->Handle, static_cast<off_t>(position), SEEK_SET);
-#endif
-		return pos != -1;
+		CurrentOffset = Start + (Length - NewPositionRelativeToEnd);
+		return true;
 	}
 
+	
+
 	virtual bool Read(uint8* Destination, int64 BytesToRead) override
 	{
 		CheckValid();
 #if LOG_ANDROID_FILE
 		FPlatformMisc::LowLevelOutputDebugStringf(
-			TEXT("FFileHandleAndroid:Read => Path = %s, BytesToRead = %d"),
+			TEXT("(%d/%d) FFileHandleAndroid:Read => Path = %s, BytesToRead = %d"),
+			FAndroidTLS::GetCurrentThreadId(), File->Handle,
 			*(File->Path), int32(BytesToRead));
 #endif
 		check(BytesToRead >= 0);
 		check(Destination);
+		
 		while (BytesToRead > 0)
 		{
+
 			int64 ThisSize = FMath::Min<int64>(READWRITE_SIZE, BytesToRead);
-			ThisSize = read(File->Handle, Destination, ThisSize);
+			
+			ThisSize = pread(File->Handle, Destination, ThisSize, CurrentOffset);
 #if LOG_ANDROID_FILE
 			FPlatformMisc::LowLevelOutputDebugStringf(
-				TEXT("FFileHandleAndroid:Read => Path = %s, ThisSize = %d"),
-				*(File->Path), int32(ThisSize));
+				TEXT("(%d/%d) FFileHandleAndroid:Read => Path = %s, ThisSize = %d, destination = %X"),
+				FAndroidTLS::GetCurrentThreadId(), File->Handle,
+				*(File->Path), int32(ThisSize), Destination);
 #endif
 			if (ThisSize < 0)
 			{
@@ -264,9 +257,11 @@ public:
 			{
 				break;
 			}
+			CurrentOffset += ThisSize;
 			Destination += ThisSize;
 			BytesToRead -= ThisSize;
 		}
+
 		return BytesToRead == 0;
 	}
 
@@ -278,15 +273,17 @@ public:
 			// Can't write to assets.
 			return false;
 		}
+		
 		while (BytesToWrite)
 		{
 			check(BytesToWrite >= 0);
 			int64 ThisSize = FMath::Min<int64>(READWRITE_SIZE, BytesToWrite);
 			check(Source);
-			if (write(File->Handle, Source, ThisSize) != ThisSize)
+			if (pwrite(File->Handle, Source, ThisSize, CurrentOffset) != ThisSize)
 			{
 				return false;
 			}
+			CurrentOffset += ThisSize;
 			Source += ThisSize;
 			BytesToWrite -= ThisSize;
 		}
