PR #1281: Only dump the script stack on assertion failure in the game thread (Contributed by [REDACTED_PERSON])

diff --git a/Engine/Source/Runtime/Core/Private/Misc/OutputDevice.cpp b/Engine/Source/Runtime/Core/Private/Misc/OutputDevice.cpp
--- a/Engine/Source/Runtime/Core/Private/Misc/OutputDevice.cpp
+++ b/Engine/Source/Runtime/Core/Private/Misc/OutputDevice.cpp
@@ -185,7 +185,7 @@ void VARARGS FDebug::LogAssertFailedMessage(const ANSICHAR* Expr, const ANSICHAR
 {
 #if !(UE_BUILD_SHIPPING || UE_BUILD_TEST)
 	// Walk the script stack, if any
-	if( GScriptStack.Num() > 0 )
+	if( IsInGameThread() && GScriptStack.Num() > 0 )
 	{
 		FString ScriptStack = TEXT( "\n\nScript Stack:\n" );
 		while( GScriptStack.Num() )
