Avoid [UE-15946] "Crash handler on Mac can deadlock" by switching to vfork as suggested by Keli.H.

diff --git a/Engine/Source/Runtime/Core/Private/Mac/MacPlatformMisc.cpp b/Engine/Source/Runtime/Core/Private/Mac/MacPlatformMisc.cpp
--- a/Engine/Source/Runtime/Core/Private/Mac/MacPlatformMisc.cpp
+++ b/Engine/Source/Runtime/Core/Private/Mac/MacPlatformMisc.cpp
@@ -1696,14 +1696,14 @@ void FMacCrashContext::GenerateCrashInfoAndLaunchReporter() const
 		}
 		
 		// try launching the tool and wait for its exit, if at all
-		// Use fork() & execl() as they are async-signal safe, CreateProc can fail in Cocoa
+		// Use vfork() & execl() as they are async-signal safe, CreateProc can fail in Cocoa
 		int32 ReturnCode = 0;
-		pid_t ForkPID = fork();
+		FCStringAnsi::Strncpy(FilePath, CrashInfoFolder, PATH_MAX);
+		FCStringAnsi::Strcat(FilePath, PATH_MAX, "/");
+		pid_t ForkPID = vfork();
 		if(ForkPID == 0)
 		{
 			// Child
-			FCStringAnsi::Strncpy(FilePath, CrashInfoFolder, PATH_MAX);
-			FCStringAnsi::Strcat(FilePath, PATH_MAX, "/");
 			execl(GMacAppInfo.CrashReportClient, "CrashReportClient", FilePath, NULL);
 		}
 		// We no longer wait here because on return the OS will scribble & crash again due to the behaviour of the XPC function
