UE-17490: GitHub 1301 : Added ConvertDigitToInt function

#1303

diff --git a/Engine/Source/Runtime/Core/Public/Misc/Char.h b/Engine/Source/Runtime/Core/Public/Misc/Char.h
--- a/Engine/Source/Runtime/Core/Public/Misc/Char.h
+++ b/Engine/Source/Runtime/Core/Public/Misc/Char.h
@@ -105,6 +105,10 @@ public:
 	static inline bool IsAlnum(CharType Char);
 	static inline bool IsDigit(CharType Char);
 	static inline bool IsHexDigit(CharType Char);
+	static inline int32 ConvertCharDigitToInt(CharType Char)
+	{
+		return static_cast<int32>(Char) - static_cast<int32>('0');
+	}
 	static inline bool IsWhitespace(CharType Char);
 	static inline bool IsIdentifier(CharType Char)
 	{
