Remove EControllerButtons

diff --git a/Engine/Plugins/Runtime/Steam/SteamController/Source/SteamController/Private/SteamController.cpp b/Engine/Plugins/Runtime/Steam/SteamController/Source/SteamController/Private/SteamController.cpp
--- a/Engine/Plugins/Runtime/Steam/SteamController/Source/SteamController/Private/SteamController.cpp
+++ b/Engine/Plugins/Runtime/Steam/SteamController/Source/SteamController/Private/SteamController.cpp
@@ -65,6 +65,16 @@ bool LoadSteamModule()
 	return true;
 }
 
+namespace SteamControllerKeyNames
+{
+	const FGamepadKeyNames::Type Touch0("Steam_Touch_0");
+	const FGamepadKeyNames::Type Touch1("Steam_Touch_1");
+	const FGamepadKeyNames::Type Touch2("Steam_Touch_2");
+	const FGamepadKeyNames::Type Touch3("Steam_Touch_3");
+	const FGamepadKeyNames::Type BackLeft("Steam_Back_Left");
+	const FGamepadKeyNames::Type BackRight("Steam_Back_Right");
+}
+
 class FSteamController : public IInputDevice
 {
 
@@ -107,32 +117,32 @@ public:
 		ButtonRepeatDelay = 0.1f;
 
 		// set up mapping
-		Buttons[0] = EControllerButtons::FaceButtonBottom;
-		Buttons[1] = EControllerButtons::FaceButtonRight;
-		Buttons[2] = EControllerButtons::FaceButtonLeft;
-		Buttons[3] = EControllerButtons::FaceButtonTop;
-		Buttons[4] = EControllerButtons::LeftShoulder;
-		Buttons[5] = EControllerButtons::RightShoulder;
-		Buttons[6] = EControllerButtons::SpecialRight;
-		Buttons[7] = EControllerButtons::SpecialLeft;
-		Buttons[8] = EControllerButtons::LeftThumb;
-		Buttons[9] = EControllerButtons::RightThumb;
-		Buttons[10] = EControllerButtons::LeftTriggerThreshold;
-		Buttons[11] = EControllerButtons::RightTriggerThreshold;
-		Buttons[12] = EControllerButtons::Touch0;
-		Buttons[13] = EControllerButtons::Touch1;
-		Buttons[14] = EControllerButtons::Touch2;
-		Buttons[15] = EControllerButtons::Touch3;
-		Buttons[16] = EControllerButtons::LeftStickUp;
-		Buttons[17] = EControllerButtons::LeftStickDown;
-		Buttons[18] = EControllerButtons::LeftStickLeft;
-		Buttons[19] = EControllerButtons::LeftStickRight;
-		Buttons[20] = EControllerButtons::RightStickUp;
-		Buttons[21] = EControllerButtons::RightStickDown;
-		Buttons[22] = EControllerButtons::RightStickLeft;
-		Buttons[23] = EControllerButtons::RightStickRight;
-		Buttons[24] = EControllerButtons::BackLeft;
-		Buttons[25] = EControllerButtons::BackRight;
+		Buttons[0] = FGamepadKeyNames::FaceButtonBottom;
+		Buttons[1] = FGamepadKeyNames::FaceButtonRight;
+		Buttons[2] = FGamepadKeyNames::FaceButtonLeft;
+		Buttons[3] = FGamepadKeyNames::FaceButtonTop;
+		Buttons[4] = FGamepadKeyNames::LeftShoulder;
+		Buttons[5] = FGamepadKeyNames::RightShoulder;
+		Buttons[6] = FGamepadKeyNames::SpecialRight;
+		Buttons[7] = FGamepadKeyNames::SpecialLeft;
+		Buttons[8] = FGamepadKeyNames::LeftThumb;
+		Buttons[9] = FGamepadKeyNames::RightThumb;
+		Buttons[10] = FGamepadKeyNames::LeftTriggerThreshold;
+		Buttons[11] = FGamepadKeyNames::RightTriggerThreshold;
+		Buttons[12] = SteamControllerKeyNames::Touch0;
+		Buttons[13] = SteamControllerKeyNames::Touch1;
+		Buttons[14] = SteamControllerKeyNames::Touch2;
+		Buttons[15] = SteamControllerKeyNames::Touch3;
+		Buttons[16] = FGamepadKeyNames::LeftStickUp;
+		Buttons[17] = FGamepadKeyNames::LeftStickDown;
+		Buttons[18] = FGamepadKeyNames::LeftStickLeft;
+		Buttons[19] = FGamepadKeyNames::LeftStickRight;
+		Buttons[20] = FGamepadKeyNames::RightStickUp;
+		Buttons[21] = FGamepadKeyNames::RightStickDown;
+		Buttons[22] = FGamepadKeyNames::RightStickLeft;
+		Buttons[23] = FGamepadKeyNames::RightStickRight;
+		Buttons[24] = SteamControllerKeyNames::BackLeft;
+		Buttons[25] = SteamControllerKeyNames::BackLeft;
 	}
 
 	virtual ~FSteamController()
@@ -193,25 +203,25 @@ public:
 
 						if (ControllerStates[ControllerIndex].LeftXAnalog != ControllerState.sLeftPadX)
 						{
-							MessageHandler->OnControllerAnalog(EControllerButtons::LeftAnalogX, ControllerIndex, ShortToNormalizedFloat(ControllerState.sLeftPadX));
+							MessageHandler->OnControllerAnalog(FGamepadKeyNames::LeftAnalogX, ControllerIndex, ShortToNormalizedFloat(ControllerState.sLeftPadX));
 							ControllerStates[ControllerIndex].LeftXAnalog = ControllerState.sLeftPadX;
 						}
 
 						if (ControllerStates[ControllerIndex].LeftYAnalog != ControllerState.sLeftPadY)
 						{
-							MessageHandler->OnControllerAnalog(EControllerButtons::LeftAnalogY, ControllerIndex, ShortToNormalizedFloat(ControllerState.sLeftPadY));
+							MessageHandler->OnControllerAnalog(FGamepadKeyNames::LeftAnalogY, ControllerIndex, ShortToNormalizedFloat(ControllerState.sLeftPadY));
 							ControllerStates[ControllerIndex].LeftYAnalog = ControllerState.sLeftPadY;
 						}
 
 						if (ControllerStates[ControllerIndex].RightXAnalog != ControllerState.sRightPadX)
 						{
-							MessageHandler->OnControllerAnalog(EControllerButtons::RightAnalogX, ControllerIndex, ShortToNormalizedFloat(ControllerState.sRightPadX));
+							MessageHandler->OnControllerAnalog(FGamepadKeyNames::RightAnalogX, ControllerIndex, ShortToNormalizedFloat(ControllerState.sRightPadX));
 							ControllerStates[ControllerIndex].RightXAnalog = ControllerState.sRightPadX;
 						}
 
 						if (ControllerStates[ControllerIndex].RightYAnalog != ControllerState.sRightPadY)
 						{
-							MessageHandler->OnControllerAnalog(EControllerButtons::RightAnalogY, ControllerIndex, ShortToNormalizedFloat(ControllerState.sRightPadY));
+							MessageHandler->OnControllerAnalog(FGamepadKeyNames::RightAnalogY, ControllerIndex, ShortToNormalizedFloat(ControllerState.sRightPadY));
 							ControllerStates[ControllerIndex].RightYAnalog = ControllerState.sRightPadY;
 						}
 
@@ -358,7 +368,7 @@ private:
 	float ButtonRepeatDelay;
 
 	/** Mapping of controller buttons */
-	EControllerButtons::Type Buttons[MAX_NUM_CONTROLLER_BUTTONS];
+	FGamepadKeyNames::Type Buttons[MAX_NUM_CONTROLLER_BUTTONS];
 
 	/** handler to send all messages to */
 	TSharedRef<FGenericApplicationMessageHandler> MessageHandler;
diff --git a/Engine/Plugins/Runtime/Steam/SteamVR/Source/SteamVRController/Private/SteamVRController.cpp b/Engine/Plugins/Runtime/Steam/SteamVR/Source/SteamVRController/Private/SteamVRController.cpp
--- a/Engine/Plugins/Runtime/Steam/SteamVR/Source/SteamVRController/Private/SteamVRController.cpp
+++ b/Engine/Plugins/Runtime/Steam/SteamVR/Source/SteamVRController/Private/SteamVRController.cpp
@@ -14,7 +14,13 @@ DEFINE_LOG_CATEGORY_STATIC(LogSteamVRController, Log, All);
 //
 #define TOUCHPAD_DEADZONE  0.0f
 
-
+namespace SteamVRControllerKeyNames
+{
+	const FGamepadKeyNames::Type Touch0("Steam_Touch_0");
+	const FGamepadKeyNames::Type Touch1("Steam_Touch_1");
+	const FGamepadKeyNames::Type BackLeft("Steam_Back_Left");
+	const FGamepadKeyNames::Type BackRight("Steam_Back_Right");
+}
 
 class FSteamVRController : public IInputDevice
 {
@@ -90,27 +96,27 @@ public:
 		InitialButtonRepeatDelay = 0.2f;
 		ButtonRepeatDelay = 0.1f;
 
-		Buttons[ (int32)EControllerHand::Left ][ ESteamVRControllerButton::System ] = EControllerButtons::SpecialLeft;
-		Buttons[ (int32)EControllerHand::Left ][ ESteamVRControllerButton::ApplicationMenu ] = EControllerButtons::LeftShoulder;
-		Buttons[ (int32)EControllerHand::Left ][ ESteamVRControllerButton::TouchPadPress ] = EControllerButtons::LeftThumb;
-		Buttons[ (int32)EControllerHand::Left ][ ESteamVRControllerButton::TouchPadTouch ] = EControllerButtons::Touch0;
-		Buttons[ (int32)EControllerHand::Left ][ ESteamVRControllerButton::TriggerPress ] = EControllerButtons::LeftTriggerThreshold;
-		Buttons[ (int32)EControllerHand::Left ][ ESteamVRControllerButton::Grip ] = EControllerButtons::BackLeft;
-		Buttons[ (int32)EControllerHand::Left ][ ESteamVRControllerButton::TouchPadUp ] = EControllerButtons::LeftStickUp;
-		Buttons[ (int32)EControllerHand::Left ][ ESteamVRControllerButton::TouchPadDown ] = EControllerButtons::LeftStickDown;
-		Buttons[ (int32)EControllerHand::Left ][ ESteamVRControllerButton::TouchPadLeft ] = EControllerButtons::LeftStickLeft;
-		Buttons[ (int32)EControllerHand::Left ][ ESteamVRControllerButton::TouchPadRight ] = EControllerButtons::LeftStickRight;
-
-		Buttons[ (int32)EControllerHand::Right ][ ESteamVRControllerButton::System ] = EControllerButtons::SpecialRight;
-		Buttons[ (int32)EControllerHand::Right ][ ESteamVRControllerButton::ApplicationMenu ] = EControllerButtons::RightShoulder;
-		Buttons[ (int32)EControllerHand::Right ][ ESteamVRControllerButton::TouchPadPress ] = EControllerButtons::RightThumb;
-		Buttons[ (int32)EControllerHand::Right ][ ESteamVRControllerButton::TouchPadTouch ] = EControllerButtons::Touch1;
-		Buttons[ (int32)EControllerHand::Right ][ ESteamVRControllerButton::TriggerPress ] = EControllerButtons::RightTriggerThreshold;
-		Buttons[ (int32)EControllerHand::Right ][ ESteamVRControllerButton::Grip ] = EControllerButtons::BackRight;
-		Buttons[ (int32)EControllerHand::Right ][ ESteamVRControllerButton::TouchPadUp ] = EControllerButtons::RightStickUp;
-		Buttons[ (int32)EControllerHand::Right ][ ESteamVRControllerButton::TouchPadDown ] = EControllerButtons::RightStickDown;
-		Buttons[ (int32)EControllerHand::Right ][ ESteamVRControllerButton::TouchPadLeft ] = EControllerButtons::RightStickLeft;
-		Buttons[ (int32)EControllerHand::Right ][ ESteamVRControllerButton::TouchPadRight ] = EControllerButtons::RightStickRight;
+		Buttons[ (int32)EControllerHand::Left ][ ESteamVRControllerButton::System ] = FGamepadKeyNames::SpecialLeft;
+		Buttons[ (int32)EControllerHand::Left ][ ESteamVRControllerButton::ApplicationMenu ] = FGamepadKeyNames::LeftShoulder;
+		Buttons[ (int32)EControllerHand::Left ][ ESteamVRControllerButton::TouchPadPress ] = FGamepadKeyNames::LeftThumb;
+		Buttons[ (int32)EControllerHand::Left ][ ESteamVRControllerButton::TouchPadTouch ] = SteamVRControllerKeyNames::Touch0;
+		Buttons[ (int32)EControllerHand::Left ][ ESteamVRControllerButton::TriggerPress ] = FGamepadKeyNames::LeftTriggerThreshold;
+		Buttons[ (int32)EControllerHand::Left ][ ESteamVRControllerButton::Grip ] = SteamVRControllerKeyNames::BackLeft;
+		Buttons[ (int32)EControllerHand::Left ][ ESteamVRControllerButton::TouchPadUp ] = FGamepadKeyNames::LeftStickUp;
+		Buttons[ (int32)EControllerHand::Left ][ ESteamVRControllerButton::TouchPadDown ] = FGamepadKeyNames::LeftStickDown;
+		Buttons[ (int32)EControllerHand::Left ][ ESteamVRControllerButton::TouchPadLeft ] = FGamepadKeyNames::LeftStickLeft;
+		Buttons[ (int32)EControllerHand::Left ][ ESteamVRControllerButton::TouchPadRight ] = FGamepadKeyNames::LeftStickRight;
+
+		Buttons[ (int32)EControllerHand::Right ][ ESteamVRControllerButton::System ] = FGamepadKeyNames::SpecialRight;
+		Buttons[ (int32)EControllerHand::Right ][ ESteamVRControllerButton::ApplicationMenu ] = FGamepadKeyNames::RightShoulder;
+		Buttons[ (int32)EControllerHand::Right ][ ESteamVRControllerButton::TouchPadPress ] = FGamepadKeyNames::RightThumb;
+		Buttons[ (int32)EControllerHand::Right ][ ESteamVRControllerButton::TouchPadTouch ] = SteamVRControllerKeyNames::Touch1;
+		Buttons[ (int32)EControllerHand::Right ][ ESteamVRControllerButton::TriggerPress ] = FGamepadKeyNames::RightTriggerThreshold;
+		Buttons[ (int32)EControllerHand::Right ][ ESteamVRControllerButton::Grip ] = SteamVRControllerKeyNames::BackRight;
+		Buttons[ (int32)EControllerHand::Right ][ ESteamVRControllerButton::TouchPadUp ] = FGamepadKeyNames::RightStickUp;
+		Buttons[ (int32)EControllerHand::Right ][ ESteamVRControllerButton::TouchPadDown ] = FGamepadKeyNames::RightStickDown;
+		Buttons[ (int32)EControllerHand::Right ][ ESteamVRControllerButton::TouchPadLeft ] = FGamepadKeyNames::RightStickLeft;
+		Buttons[ (int32)EControllerHand::Right ][ ESteamVRControllerButton::TouchPadRight ] = FGamepadKeyNames::RightStickRight;
 	}
 
 	virtual ~FSteamVRController()
@@ -199,19 +205,19 @@ public:
 
 						if ( ControllerState.TouchPadXAnalog != VRControllerState.rAxis[TOUCHPAD_AXIS].x)
 						{
-							MessageHandler->OnControllerAnalog(EControllerButtons::LeftAnalogX, ControllerIndex, VRControllerState.rAxis[TOUCHPAD_AXIS].x);
+							MessageHandler->OnControllerAnalog(FGamepadKeyNames::LeftAnalogX, ControllerIndex, VRControllerState.rAxis[TOUCHPAD_AXIS].x);
 							ControllerState.TouchPadXAnalog = VRControllerState.rAxis[TOUCHPAD_AXIS].x;
 						}
 
 						if ( ControllerState.TouchPadYAnalog != VRControllerState.rAxis[TOUCHPAD_AXIS].y)
 						{
-							MessageHandler->OnControllerAnalog(EControllerButtons::LeftAnalogY, ControllerIndex, VRControllerState.rAxis[TOUCHPAD_AXIS].y);
+							MessageHandler->OnControllerAnalog(FGamepadKeyNames::LeftAnalogY, ControllerIndex, VRControllerState.rAxis[TOUCHPAD_AXIS].y);
 							ControllerState.TouchPadYAnalog = VRControllerState.rAxis[TOUCHPAD_AXIS].y;
 						}
 
 						if ( ControllerState.TriggerAnalog != VRControllerState.rAxis[TRIGGER_AXIS].x)
 						{
-							MessageHandler->OnControllerAnalog(EControllerButtons::LeftTriggerAnalog, ControllerIndex, VRControllerState.rAxis[TRIGGER_AXIS].x);
+							MessageHandler->OnControllerAnalog(FGamepadKeyNames::LeftTriggerAnalog, ControllerIndex, VRControllerState.rAxis[TRIGGER_AXIS].x);
 							ControllerState.TriggerAnalog = VRControllerState.rAxis[TRIGGER_AXIS].x;
 						}
 
@@ -401,7 +407,7 @@ private:
 	float ButtonRepeatDelay;
 
 	/** Mapping of controller buttons */
-	EControllerButtons::Type Buttons[ (int32)EControllerHand::TotalHandCount ][ ESteamVRControllerButton::TotalButtonCount ];
+	FGamepadKeyNames::Type Buttons[ (int32)EControllerHand::TotalHandCount ][ ESteamVRControllerButton::TotalButtonCount ];
 
 	/** handler to send all messages to */
 	TSharedRef<FGenericApplicationMessageHandler> MessageHandler;
diff --git a/Engine/Source/Runtime/Core/Private/Android/AndroidInputInterface.cpp b/Engine/Source/Runtime/Core/Private/Android/AndroidInputInterface.cpp
--- a/Engine/Source/Runtime/Core/Private/Android/AndroidInputInterface.cpp
+++ b/Engine/Source/Runtime/Core/Private/Android/AndroidInputInterface.cpp
@@ -12,7 +12,7 @@ FCriticalSection FAndroidInputInterface::TouchInputCriticalSection;
 FAndroidControllerData FAndroidInputInterface::OldControllerData[MAX_NUM_CONTROLLERS];
 FAndroidControllerData FAndroidInputInterface::NewControllerData[MAX_NUM_CONTROLLERS];
 
-EControllerButtons::Type FAndroidInputInterface::ButtonMapping[MAX_NUM_CONTROLLER_BUTTONS];
+FGamepadKeyNames::Type FAndroidInputInterface::ButtonMapping[MAX_NUM_CONTROLLER_BUTTONS];
 float FAndroidInputInterface::InitialButtonRepeatDelay;
 float FAndroidInputInterface::ButtonRepeatDelay;
 
@@ -33,29 +33,37 @@ FAndroidInputInterface::~FAndroidInputInterface()
 {
 }
 
+namespace AndroidKeyNames
+{
+	const FGamepadKeyNames::Type Android_Back("Android_Back");
+	const FGamepadKeyNames::Type Android_Volume_Up("Android_Volume_Up");
+	const FGamepadKeyNames::Type Android_Volume_Down("Android_Volume_Down");
+	const FGamepadKeyNames::Type Android_Menu("Android_Menu");
+}
+
 FAndroidInputInterface::FAndroidInputInterface( const TSharedRef< FGenericApplicationMessageHandler >& InMessageHandler )
 	: MessageHandler( InMessageHandler )
 {
-	ButtonMapping[0] = EControllerButtons::FaceButtonBottom;
-	ButtonMapping[1] = EControllerButtons::FaceButtonRight;
-	ButtonMapping[2] = EControllerButtons::FaceButtonLeft;
-	ButtonMapping[3] = EControllerButtons::FaceButtonTop;
-	ButtonMapping[4] = EControllerButtons::LeftShoulder;
-	ButtonMapping[5] = EControllerButtons::RightShoulder;
-	ButtonMapping[6] = EControllerButtons::SpecialRight;
-	ButtonMapping[7] = EControllerButtons::SpecialLeft;
-	ButtonMapping[8] = EControllerButtons::LeftThumb;
-	ButtonMapping[9] = EControllerButtons::RightThumb;
-	ButtonMapping[10] = EControllerButtons::LeftTriggerThreshold;
-	ButtonMapping[11] = EControllerButtons::RightTriggerThreshold;
-	ButtonMapping[12] = EControllerButtons::DPadUp;
-	ButtonMapping[13] = EControllerButtons::DPadDown;
-	ButtonMapping[14] = EControllerButtons::DPadLeft;
-	ButtonMapping[15] = EControllerButtons::DPadRight;
-	ButtonMapping[16] = EControllerButtons::AndroidBack;  // Technically just an alias for SpecialLeft
-	ButtonMapping[17] = EControllerButtons::AndroidVolumeUp;
-	ButtonMapping[18] = EControllerButtons::AndroidVolumeDown;
-	ButtonMapping[19] = EControllerButtons::AndroidMenu;  // Technically just an alias for SpecialRightNewControllerData[deviceId].ButtonStates[6] = buttonDown; 
+	ButtonMapping[0] = FGamepadKeyNames::FaceButtonBottom;
+	ButtonMapping[1] = FGamepadKeyNames::FaceButtonRight;
+	ButtonMapping[2] = FGamepadKeyNames::FaceButtonLeft;
+	ButtonMapping[3] = FGamepadKeyNames::FaceButtonTop;
+	ButtonMapping[4] = FGamepadKeyNames::LeftShoulder;
+	ButtonMapping[5] = FGamepadKeyNames::RightShoulder;
+	ButtonMapping[6] = FGamepadKeyNames::SpecialRight;
+	ButtonMapping[7] = FGamepadKeyNames::SpecialLeft;
+	ButtonMapping[8] = FGamepadKeyNames::LeftThumb;
+	ButtonMapping[9] = FGamepadKeyNames::RightThumb;
+	ButtonMapping[10] = FGamepadKeyNames::LeftTriggerThreshold;
+	ButtonMapping[11] = FGamepadKeyNames::RightTriggerThreshold;
+	ButtonMapping[12] = FGamepadKeyNames::DPadUp;
+	ButtonMapping[13] = FGamepadKeyNames::DPadDown;
+	ButtonMapping[14] = FGamepadKeyNames::DPadLeft;
+	ButtonMapping[15] = FGamepadKeyNames::DPadRight;
+	ButtonMapping[16] = AndroidKeyNames::Android_Back;  // Technically just an alias for SpecialLeft
+	ButtonMapping[17] = AndroidKeyNames::Android_Volume_Up;
+	ButtonMapping[18] = AndroidKeyNames::Android_Volume_Down;
+	ButtonMapping[19] = AndroidKeyNames::Android_Menu;  // Technically just an alias for SpecialRightNewControllerData[deviceId].ButtonStates[6] = buttonDown; 
 
 	InitialButtonRepeatDelay = 0.2f;
 	ButtonRepeatDelay = 0.1f;
@@ -639,39 +647,39 @@ void FAndroidInputInterface::SendControllerEvents()
 
 		if (NewControllerState.LXAnalog != OldControllerState.LXAnalog)
 		{
-			MessageHandler->OnControllerAnalog(EControllerButtons::LeftAnalogX, NewControllerState.DeviceId, NewControllerState.LXAnalog);
+			MessageHandler->OnControllerAnalog(FGamepadKeyNames::LeftAnalogX, NewControllerState.DeviceId, NewControllerState.LXAnalog);
 		}
 		if (NewControllerState.LYAnalog != OldControllerState.LYAnalog)
 		{
 			//LOGD("    Sending updated LeftAnalogY value of %f", NewControllerState.LYAnalog);
-			MessageHandler->OnControllerAnalog(EControllerButtons::LeftAnalogY, NewControllerState.DeviceId, NewControllerState.LYAnalog);
+			MessageHandler->OnControllerAnalog(FGamepadKeyNames::LeftAnalogY, NewControllerState.DeviceId, NewControllerState.LYAnalog);
 		}
 		if (NewControllerState.RXAnalog != OldControllerState.RXAnalog)
 		{
 			//LOGD("    Sending updated RightAnalogX value of %f", NewControllerState.RXAnalog);
-			MessageHandler->OnControllerAnalog(EControllerButtons::RightAnalogX, NewControllerState.DeviceId, NewControllerState.RXAnalog);
+			MessageHandler->OnControllerAnalog(FGamepadKeyNames::RightAnalogX, NewControllerState.DeviceId, NewControllerState.RXAnalog);
 		}
 		if (NewControllerState.RYAnalog != OldControllerState.RYAnalog)
 		{
 			//LOGD("    Sending updated RightAnalogY value of %f", NewControllerState.RYAnalog);
-			MessageHandler->OnControllerAnalog(EControllerButtons::RightAnalogY, NewControllerState.DeviceId, NewControllerState.RYAnalog);
+			MessageHandler->OnControllerAnalog(FGamepadKeyNames::RightAnalogY, NewControllerState.DeviceId, NewControllerState.RYAnalog);
 		}
 		if (NewControllerState.LTAnalog != OldControllerState.LTAnalog)
 		{
 			//LOGD("    Sending updated LeftTriggerAnalog value of %f", NewControllerState.LTAnalog);
-			MessageHandler->OnControllerAnalog(EControllerButtons::LeftTriggerAnalog, NewControllerState.DeviceId, NewControllerState.LTAnalog);
+			MessageHandler->OnControllerAnalog(FGamepadKeyNames::LeftTriggerAnalog, NewControllerState.DeviceId, NewControllerState.LTAnalog);
 
 			// Handle the trigger theshold "virtual" button state
-			//check(ButtonMapping[10] == EControllerButtons::LeftTriggerThreshold);
+			//check(ButtonMapping[10] == FGamepadKeyNames::LeftTriggerThreshold);
 			NewControllerState.ButtonStates[10] = NewControllerState.LTAnalog >= 0.1f;
 		}
 		if (NewControllerState.RTAnalog != OldControllerState.RTAnalog)
 		{
 			//LOGD("    Sending updated RightTriggerAnalog value of %f", NewControllerState.RTAnalog);
-			MessageHandler->OnControllerAnalog(EControllerButtons::RightTriggerAnalog, NewControllerState.DeviceId, NewControllerState.RTAnalog);
+			MessageHandler->OnControllerAnalog(FGamepadKeyNames::RightTriggerAnalog, NewControllerState.DeviceId, NewControllerState.RTAnalog);
 
 			// Handle the trigger theshold "virtual" button state
-			//check(ButtonMapping[11] == EControllerButtons::RightTriggerThreshold);
+			//check(ButtonMapping[11] == FGamepadKeyNames::RightTriggerThreshold);
 			NewControllerState.ButtonStates[11] = NewControllerState.RTAnalog >= 0.1f;
 		}
 
diff --git a/Engine/Source/Runtime/Core/Private/GenericPlatform/GenericApplication.cpp b/Engine/Source/Runtime/Core/Private/GenericPlatform/GenericApplication.cpp
--- a/Engine/Source/Runtime/Core/Private/GenericPlatform/GenericApplication.cpp
+++ b/Engine/Source/Runtime/Core/Private/GenericPlatform/GenericApplication.cpp
@@ -2,7 +2,43 @@
 
 #include "CorePrivatePCH.h"
 #include "GenericApplication.h"
+#include "GenericApplicationMessageHandler.h"
 
+const FGamepadKeyNames::Type FGamepadKeyNames::Invalid(NAME_None);
+
+const FGamepadKeyNames::Type FGamepadKeyNames::LeftAnalogX("Gamepad_LeftX");
+const FGamepadKeyNames::Type FGamepadKeyNames::LeftAnalogY("Gamepad_LeftY");
+const FGamepadKeyNames::Type FGamepadKeyNames::RightAnalogX("Gamepad_RightX");
+const FGamepadKeyNames::Type FGamepadKeyNames::RightAnalogY("Gamepad_RightY");
+const FGamepadKeyNames::Type FGamepadKeyNames::LeftTriggerAnalog("Gamepad_LeftTriggerAxis");
+const FGamepadKeyNames::Type FGamepadKeyNames::RightTriggerAnalog("Gamepad_RightTriggerAxis");
+
+const FGamepadKeyNames::Type FGamepadKeyNames::LeftThumb("Gamepad_LeftThumbstick");
+const FGamepadKeyNames::Type FGamepadKeyNames::RightThumb("Gamepad_RightThumbstick");
+const FGamepadKeyNames::Type FGamepadKeyNames::SpecialLeft("Gamepad_Special_Left");
+const FGamepadKeyNames::Type FGamepadKeyNames::SpecialRight("Gamepad_Special_Right");
+const FGamepadKeyNames::Type FGamepadKeyNames::FaceButtonBottom("Gamepad_FaceButton_Bottom");
+const FGamepadKeyNames::Type FGamepadKeyNames::FaceButtonRight("Gamepad_FaceButton_Right");
+const FGamepadKeyNames::Type FGamepadKeyNames::FaceButtonLeft("Gamepad_FaceButton_Left");
+const FGamepadKeyNames::Type FGamepadKeyNames::FaceButtonTop("Gamepad_FaceButton_Top");
+const FGamepadKeyNames::Type FGamepadKeyNames::LeftShoulder("Gamepad_LeftShoulder");
+const FGamepadKeyNames::Type FGamepadKeyNames::RightShoulder("Gamepad_RightShoulder");
+const FGamepadKeyNames::Type FGamepadKeyNames::LeftTriggerThreshold("Gamepad_LeftTrigger");
+const FGamepadKeyNames::Type FGamepadKeyNames::RightTriggerThreshold("Gamepad_RightTrigger");
+const FGamepadKeyNames::Type FGamepadKeyNames::DPadUp("Gamepad_DPad_Up");
+const FGamepadKeyNames::Type FGamepadKeyNames::DPadDown("Gamepad_DPad_Down");
+const FGamepadKeyNames::Type FGamepadKeyNames::DPadRight("Gamepad_DPad_Right");
+const FGamepadKeyNames::Type FGamepadKeyNames::DPadLeft("Gamepad_DPad_Left");
+
+const FGamepadKeyNames::Type FGamepadKeyNames::LeftStickUp("Gamepad_LeftStick_Up");
+const FGamepadKeyNames::Type FGamepadKeyNames::LeftStickDown("Gamepad_LeftStick_Down");
+const FGamepadKeyNames::Type FGamepadKeyNames::LeftStickRight("Gamepad_LeftStick_Right");
+const FGamepadKeyNames::Type FGamepadKeyNames::LeftStickLeft("Gamepad_LeftStick_Left");
+
+const FGamepadKeyNames::Type FGamepadKeyNames::RightStickUp("Gamepad_RightStick_Up");
+const FGamepadKeyNames::Type FGamepadKeyNames::RightStickDown("Gamepad_RightStick_Down");
+const FGamepadKeyNames::Type FGamepadKeyNames::RightStickRight("Gamepad_RightStick_Right");
+const FGamepadKeyNames::Type FGamepadKeyNames::RightStickLeft("Gamepad_RightStick_Left");
 
 void FDisplayMetrics::PrintToLog() const
 {
diff --git a/Engine/Source/Runtime/Core/Private/HTML5/HTML5InputInterface.cpp b/Engine/Source/Runtime/Core/Private/HTML5/HTML5InputInterface.cpp
--- a/Engine/Source/Runtime/Core/Private/HTML5/HTML5InputInterface.cpp
+++ b/Engine/Source/Runtime/Core/Private/HTML5/HTML5InputInterface.cpp
@@ -14,12 +14,12 @@ DEFINE_LOG_CATEGORY_STATIC(LogHTML5Input, Log, All);
 #include "emscripten.h"
 #include "html5.h"
 
-static EControllerButtons::Type AxisMapping[4] =
+static FGamepadKeyNames::Type AxisMapping[4] =
 {
-	EControllerButtons::LeftAnalogX,
-	EControllerButtons::LeftAnalogY, 
-	EControllerButtons::RightAnalogX,
-	EControllerButtons::RightAnalogY
+	FGamepadKeyNames::LeftAnalogX,
+	FGamepadKeyNames::LeftAnalogY, 
+	FGamepadKeyNames::RightAnalogX,
+	FGamepadKeyNames::RightAnalogY
 };
 
 // Axis Mapping, reversed or not. 
@@ -32,24 +32,24 @@ static int Reversed[4] =
 };
 
 // all are digital except Left and Right Trigger Analog. 
-static EControllerButtons::Type  ButtonMapping[16] = 
+static FGamepadKeyNames::Type ButtonMapping[16] = 
 {
-	EControllerButtons::FaceButtonBottom,
-	EControllerButtons::FaceButtonRight,
-	EControllerButtons::FaceButtonLeft,
-	EControllerButtons::FaceButtonTop,
-	EControllerButtons::LeftShoulder,
-	EControllerButtons::RightShoulder,
-	EControllerButtons::LeftTriggerThreshold, 
-	EControllerButtons::RightTriggerThreshold,
-	EControllerButtons::SpecialLeft, 
-	EControllerButtons::SpecialRight,
-	EControllerButtons::LeftStickDown,
-	EControllerButtons::RightStickDown,
-	EControllerButtons::DPadUp,
-	EControllerButtons::DPadDown,
-	EControllerButtons::DPadLeft,
-	EControllerButtons::DPadRight
+	FGamepadKeyNames::FaceButtonBottom,
+	FGamepadKeyNames::FaceButtonRight,
+	FGamepadKeyNames::FaceButtonLeft,
+	FGamepadKeyNames::FaceButtonTop,
+	FGamepadKeyNames::LeftShoulder,
+	FGamepadKeyNames::RightShoulder,
+	FGamepadKeyNames::LeftTriggerThreshold, 
+	FGamepadKeyNames::RightTriggerThreshold,
+	FGamepadKeyNames::SpecialLeft, 
+	FGamepadKeyNames::SpecialRight,
+	FGamepadKeyNames::LeftStickDown,
+	FGamepadKeyNames::RightStickDown,
+	FGamepadKeyNames::DPadUp,
+	FGamepadKeyNames::DPadDown,
+	FGamepadKeyNames::DPadLeft,
+	FGamepadKeyNames::DPadRight
 };
 
 #endif 
diff --git a/Engine/Source/Runtime/Core/Private/Linux/LinuxApplication.cpp b/Engine/Source/Runtime/Core/Private/Linux/LinuxApplication.cpp
--- a/Engine/Source/Runtime/Core/Private/Linux/LinuxApplication.cpp
+++ b/Engine/Source/Runtime/Core/Private/Linux/LinuxApplication.cpp
@@ -413,231 +413,231 @@ void FLinuxApplication::ProcessDeferredMessage( SDL_Event Event )
 	case SDL_CONTROLLERAXISMOTION:
 		{
 			SDL_ControllerAxisEvent caxisEvent = Event.caxis;
-			EControllerButtons::Type analog;
+			FGamepadKeyNames::Type analog = FGamepadKeyNames::Invalid;
 			float value = ShortToNormalFloat(caxisEvent.value);
 
 			switch (caxisEvent.axis)
 			{
 			case SDL_CONTROLLER_AXIS_LEFTX:
-				analog = EControllerButtons::LeftAnalogX;
+				analog = FGamepadKeyNames::LeftAnalogX;
 				if(caxisEvent.value > GAMECONTROLLER_LEFT_THUMB_DEADZONE)
 				{
 					if(!ControllerStates[caxisEvent.which].analogOverThreshold[0])
 					{
-						MessageHandler->OnControllerButtonPressed(EControllerButtons::LeftStickRight, caxisEvent.which, false);
+						MessageHandler->OnControllerButtonPressed(FGamepadKeyNames::LeftStickRight, caxisEvent.which, false);
 						ControllerStates[caxisEvent.which].analogOverThreshold[0] = true;
 					}
 				}
 				else if(ControllerStates[caxisEvent.which].analogOverThreshold[0])
 				{
-					MessageHandler->OnControllerButtonReleased(EControllerButtons::LeftStickRight, caxisEvent.which, false);
+					MessageHandler->OnControllerButtonReleased(FGamepadKeyNames::LeftStickRight, caxisEvent.which, false);
 					ControllerStates[caxisEvent.which].analogOverThreshold[0] = false;
 				}
 				if(caxisEvent.value < -GAMECONTROLLER_LEFT_THUMB_DEADZONE)
 				{
 					if(!ControllerStates[caxisEvent.which].analogOverThreshold[1])
 					{
-						MessageHandler->OnControllerButtonPressed(EControllerButtons::LeftStickLeft, caxisEvent.which, false);
+						MessageHandler->OnControllerButtonPressed(FGamepadKeyNames::LeftStickLeft, caxisEvent.which, false);
 						ControllerStates[caxisEvent.which].analogOverThreshold[1] = true;
 					}
 				}
 				else if(ControllerStates[caxisEvent.which].analogOverThreshold[1])
 				{
-					MessageHandler->OnControllerButtonReleased(EControllerButtons::LeftStickLeft, caxisEvent.which, false);
+					MessageHandler->OnControllerButtonReleased(FGamepadKeyNames::LeftStickLeft, caxisEvent.which, false);
 					ControllerStates[caxisEvent.which].analogOverThreshold[1] = false;
 				}
 				break;
 			case SDL_CONTROLLER_AXIS_LEFTY:
-				analog = EControllerButtons::LeftAnalogY;
+				analog = FGamepadKeyNames::LeftAnalogY;
 				value *= -1;
 				if(caxisEvent.value > GAMECONTROLLER_LEFT_THUMB_DEADZONE)
 				{
 					if(!ControllerStates[caxisEvent.which].analogOverThreshold[2])
 					{
-						MessageHandler->OnControllerButtonPressed(EControllerButtons::LeftStickDown, caxisEvent.which, false);
+						MessageHandler->OnControllerButtonPressed(FGamepadKeyNames::LeftStickDown, caxisEvent.which, false);
 						ControllerStates[caxisEvent.which].analogOverThreshold[2] = true;
 					}
 				}
 				else if(ControllerStates[caxisEvent.which].analogOverThreshold[2])
 				{
-					MessageHandler->OnControllerButtonReleased(EControllerButtons::LeftStickDown, caxisEvent.which, false);
+					MessageHandler->OnControllerButtonReleased(FGamepadKeyNames::LeftStickDown, caxisEvent.which, false);
 					ControllerStates[caxisEvent.which].analogOverThreshold[2] = false;
 				}
 				if(caxisEvent.value < -GAMECONTROLLER_LEFT_THUMB_DEADZONE)
 				{
 					if(!ControllerStates[caxisEvent.which].analogOverThreshold[3])
 					{
-						MessageHandler->OnControllerButtonPressed(EControllerButtons::LeftStickUp, caxisEvent.which, false);
+						MessageHandler->OnControllerButtonPressed(FGamepadKeyNames::LeftStickUp, caxisEvent.which, false);
 						ControllerStates[caxisEvent.which].analogOverThreshold[3] = true;
 					}
 				}
 				else if(ControllerStates[caxisEvent.which].analogOverThreshold[3])
 				{
-					MessageHandler->OnControllerButtonReleased(EControllerButtons::LeftStickUp, caxisEvent.which, false);
+					MessageHandler->OnControllerButtonReleased(FGamepadKeyNames::LeftStickUp, caxisEvent.which, false);
 					ControllerStates[caxisEvent.which].analogOverThreshold[3] = false;
 				}
 				break;
 			case SDL_CONTROLLER_AXIS_RIGHTX:
-				analog = EControllerButtons::RightAnalogX;
+				analog = FGamepadKeyNames::RightAnalogX;
 				if(caxisEvent.value > GAMECONTROLLER_RIGHT_THUMB_DEADZONE)
 				{
 					if(!ControllerStates[caxisEvent.which].analogOverThreshold[4])
 					{
-						MessageHandler->OnControllerButtonPressed(EControllerButtons::RightStickRight, caxisEvent.which, false);
+						MessageHandler->OnControllerButtonPressed(FGamepadKeyNames::RightStickRight, caxisEvent.which, false);
 						ControllerStates[caxisEvent.which].analogOverThreshold[4] = true;
 					}
 				}
 				else if(ControllerStates[caxisEvent.which].analogOverThreshold[4])
 				{
-					MessageHandler->OnControllerButtonReleased(EControllerButtons::RightStickRight, caxisEvent.which, false);
+					MessageHandler->OnControllerButtonReleased(FGamepadKeyNames::RightStickRight, caxisEvent.which, false);
 					ControllerStates[caxisEvent.which].analogOverThreshold[4] = false;
 				}
 				if(caxisEvent.value < -GAMECONTROLLER_RIGHT_THUMB_DEADZONE)
 				{
 					if(!ControllerStates[caxisEvent.which].analogOverThreshold[5])
 					{
-						MessageHandler->OnControllerButtonPressed(EControllerButtons::RightStickLeft, caxisEvent.which, false);
+						MessageHandler->OnControllerButtonPressed(FGamepadKeyNames::RightStickLeft, caxisEvent.which, false);
 						ControllerStates[caxisEvent.which].analogOverThreshold[5] = true;
 					}
 				}
 				else if(ControllerStates[caxisEvent.which].analogOverThreshold[5])
 				{
-					MessageHandler->OnControllerButtonReleased(EControllerButtons::RightStickLeft, caxisEvent.which, false);
+					MessageHandler->OnControllerButtonReleased(FGamepadKeyNames::RightStickLeft, caxisEvent.which, false);
 					ControllerStates[caxisEvent.which].analogOverThreshold[5] = false;
 				}
 				break;
 			case SDL_CONTROLLER_AXIS_RIGHTY:
-				analog = EControllerButtons::RightAnalogY;
+				analog = FGamepadKeyNames::RightAnalogY;
 				value *= -1;
 				if(caxisEvent.value > GAMECONTROLLER_RIGHT_THUMB_DEADZONE)
 				{
 					if(!ControllerStates[caxisEvent.which].analogOverThreshold[6])
 					{
-						MessageHandler->OnControllerButtonPressed(EControllerButtons::RightStickDown, caxisEvent.which, false);
+						MessageHandler->OnControllerButtonPressed(FGamepadKeyNames::RightStickDown, caxisEvent.which, false);
 						ControllerStates[caxisEvent.which].analogOverThreshold[6] = true;
 					}
 				}
 				else if(ControllerStates[caxisEvent.which].analogOverThreshold[6])
 				{
-					MessageHandler->OnControllerButtonReleased(EControllerButtons::RightStickDown, caxisEvent.which, false);
+					MessageHandler->OnControllerButtonReleased(FGamepadKeyNames::RightStickDown, caxisEvent.which, false);
 					ControllerStates[caxisEvent.which].analogOverThreshold[6] = false;
 				}
 				if(caxisEvent.value < -GAMECONTROLLER_RIGHT_THUMB_DEADZONE)
 				{
 					if(!ControllerStates[caxisEvent.which].analogOverThreshold[7])
 					{
-						MessageHandler->OnControllerButtonPressed(EControllerButtons::RightStickUp, caxisEvent.which, false);
+						MessageHandler->OnControllerButtonPressed(FGamepadKeyNames::RightStickUp, caxisEvent.which, false);
 						ControllerStates[caxisEvent.which].analogOverThreshold[7] = true;
 					}
 				}
 				else if(ControllerStates[caxisEvent.which].analogOverThreshold[7])
 				{
-					MessageHandler->OnControllerButtonReleased(EControllerButtons::RightStickUp, caxisEvent.which, false);
+					MessageHandler->OnControllerButtonReleased(FGamepadKeyNames::RightStickUp, caxisEvent.which, false);
 					ControllerStates[caxisEvent.which].analogOverThreshold[7] = false;
 				}
 				break;
 			case SDL_CONTROLLER_AXIS_TRIGGERLEFT:
-				analog = EControllerButtons::LeftTriggerAnalog;
+				analog = FGamepadKeyNames::LeftTriggerAnalog;
 				if(caxisEvent.value > GAMECONTROLLER_TRIGGER_THRESHOLD)
 				{
 					if(!ControllerStates[caxisEvent.which].analogOverThreshold[8])
 					{
-						MessageHandler->OnControllerButtonPressed(EControllerButtons::LeftTriggerThreshold, caxisEvent.which, false);
+						MessageHandler->OnControllerButtonPressed(FGamepadKeyNames::LeftTriggerThreshold, caxisEvent.which, false);
 						ControllerStates[caxisEvent.which].analogOverThreshold[8] = true;
 					}
 				}
 				else if(ControllerStates[caxisEvent.which].analogOverThreshold[8])
 				{
-					MessageHandler->OnControllerButtonReleased(EControllerButtons::LeftTriggerThreshold, caxisEvent.which, false);
+					MessageHandler->OnControllerButtonReleased(FGamepadKeyNames::LeftTriggerThreshold, caxisEvent.which, false);
 					ControllerStates[caxisEvent.which].analogOverThreshold[8] = false;
 				}
 				break;
 			case SDL_CONTROLLER_AXIS_TRIGGERRIGHT:
-				analog = EControllerButtons::LeftTriggerAnalog;
+				analog = FGamepadKeyNames::LeftTriggerAnalog;
 				if(caxisEvent.value > GAMECONTROLLER_TRIGGER_THRESHOLD)
 				{
 					if(!ControllerStates[caxisEvent.which].analogOverThreshold[9])
 					{
-						MessageHandler->OnControllerButtonPressed(EControllerButtons::RightTriggerThreshold, caxisEvent.which, false);
+						MessageHandler->OnControllerButtonPressed(FGamepadKeyNames::RightTriggerThreshold, caxisEvent.which, false);
 						ControllerStates[caxisEvent.which].analogOverThreshold[9] = true;
 					}
 				}
 				else if(ControllerStates[caxisEvent.which].analogOverThreshold[9])
 				{
-					MessageHandler->OnControllerButtonReleased(EControllerButtons::RightTriggerThreshold, caxisEvent.which, false);
+					MessageHandler->OnControllerButtonReleased(FGamepadKeyNames::RightTriggerThreshold, caxisEvent.which, false);
 					ControllerStates[caxisEvent.which].analogOverThreshold[9] = false;
 				}
 				break;
-			default:
-				analog = EControllerButtons::Invalid;
-				break;
 			}
 
-			MessageHandler->OnControllerAnalog(analog, caxisEvent.which, value);
+			if (analog != FGamepadKeyNames::Invalid)
+			{
+				MessageHandler->OnControllerAnalog(analog, caxisEvent.which, value);
+			}
 		}
 		break;
 	case SDL_CONTROLLERBUTTONDOWN:
 	case SDL_CONTROLLERBUTTONUP:
 		{
 			SDL_ControllerButtonEvent cbuttonEvent = Event.cbutton;
-			EControllerButtons::Type button;
+			FGamepadKeyNames::Type button = FGamepadKeyNames::Invalid;
 
 			switch (cbuttonEvent.button)
 			{
 			case SDL_CONTROLLER_BUTTON_A:
-				button = EControllerButtons::FaceButtonBottom;
+				button = FGamepadKeyNames::FaceButtonBottom;
 				break;
 			case SDL_CONTROLLER_BUTTON_B:
-				button = EControllerButtons::FaceButtonRight;
+				button = FGamepadKeyNames::FaceButtonRight;
 				break;
 			case SDL_CONTROLLER_BUTTON_X:
-				button = EControllerButtons::FaceButtonLeft;
+				button = FGamepadKeyNames::FaceButtonLeft;
 				break;
 			case SDL_CONTROLLER_BUTTON_Y:
-				button = EControllerButtons::FaceButtonTop;
+				button = FGamepadKeyNames::FaceButtonTop;
 				break;
 			case SDL_CONTROLLER_BUTTON_BACK:
-				button = EControllerButtons::SpecialLeft;
+				button = FGamepadKeyNames::SpecialLeft;
 				break;
 			case SDL_CONTROLLER_BUTTON_START:
-				button = EControllerButtons::SpecialRight;
+				button = FGamepadKeyNames::SpecialRight;
 				break;
 			case SDL_CONTROLLER_BUTTON_LEFTSTICK:
-				button = EControllerButtons::LeftStickDown;
+				button = FGamepadKeyNames::LeftStickDown;
 				break;
 			case SDL_CONTROLLER_BUTTON_RIGHTSTICK:
-				button = EControllerButtons::RightStickDown;
+				button = FGamepadKeyNames::RightStickDown;
 				break;
 			case SDL_CONTROLLER_BUTTON_LEFTSHOULDER:
-				button = EControllerButtons::LeftShoulder;
+				button = FGamepadKeyNames::LeftShoulder;
 				break;
 			case SDL_CONTROLLER_BUTTON_RIGHTSHOULDER:
-				button = EControllerButtons::RightShoulder;
+				button = FGamepadKeyNames::RightShoulder;
 				break;
 			case SDL_CONTROLLER_BUTTON_DPAD_UP:
-				button = EControllerButtons::DPadUp;
+				button = FGamepadKeyNames::DPadUp;
 				break;
 			case SDL_CONTROLLER_BUTTON_DPAD_DOWN:
-				button = EControllerButtons::DPadDown;
+				button = FGamepadKeyNames::DPadDown;
 				break;
 			case SDL_CONTROLLER_BUTTON_DPAD_LEFT:
-				button = EControllerButtons::DPadLeft;
+				button = FGamepadKeyNames::DPadLeft;
 				break;
 			case SDL_CONTROLLER_BUTTON_DPAD_RIGHT:
-				button = EControllerButtons::DPadRight;
-				break;
-			default:
-				button = EControllerButtons::Invalid;
+				button = FGamepadKeyNames::DPadRight;
 				break;
 			}
 
-			if(cbuttonEvent.type == SDL_CONTROLLERBUTTONDOWN)
-			{
-				MessageHandler->OnControllerButtonPressed(button, cbuttonEvent.which, false);
-			}
-			else
+			if (button != FGamepadKeyNames::Invalid)
 			{
-				MessageHandler->OnControllerButtonReleased(button, cbuttonEvent.which, false);
+				if (cbuttonEvent.type == SDL_CONTROLLERBUTTONDOWN)
+				{
+					MessageHandler->OnControllerButtonPressed(button, cbuttonEvent.which, false);
+				}
+				else
+				{
+					MessageHandler->OnControllerButtonReleased(button, cbuttonEvent.which, false);
+				}
 			}
 		}
 		break;
diff --git a/Engine/Source/Runtime/Core/Private/Mac/HIDInputInterface.cpp b/Engine/Source/Runtime/Core/Private/Mac/HIDInputInterface.cpp
--- a/Engine/Source/Runtime/Core/Private/Mac/HIDInputInterface.cpp
+++ b/Engine/Source/Runtime/Core/Private/Mac/HIDInputInterface.cpp
@@ -35,30 +35,30 @@ HIDInputInterface::HIDInputInterface(const TSharedRef<FGenericApplicationMessage
 	InitialButtonRepeatDelay = 0.2f;
 	ButtonRepeatDelay = 0.1f;
 
-	Buttons[0] = EControllerButtons::FaceButtonBottom;
-	Buttons[1] = EControllerButtons::FaceButtonRight;
-	Buttons[2] = EControllerButtons::FaceButtonLeft;
-	Buttons[3] = EControllerButtons::FaceButtonTop;
-	Buttons[4] = EControllerButtons::LeftShoulder;
-	Buttons[5] = EControllerButtons::RightShoulder;
-	Buttons[6] = EControllerButtons::SpecialRight;
-	Buttons[7] = EControllerButtons::SpecialLeft;
-	Buttons[8] = EControllerButtons::LeftThumb;
-	Buttons[9] = EControllerButtons::RightThumb;
-	Buttons[10] = EControllerButtons::LeftTriggerThreshold;
-	Buttons[11] = EControllerButtons::RightTriggerThreshold;
-	Buttons[12] = EControllerButtons::DPadUp;
-	Buttons[13] = EControllerButtons::DPadDown;
-	Buttons[14] = EControllerButtons::DPadLeft;
-	Buttons[15] = EControllerButtons::DPadRight;
-	Buttons[16] = EControllerButtons::LeftStickUp;
-	Buttons[17] = EControllerButtons::LeftStickDown;
-	Buttons[18] = EControllerButtons::LeftStickLeft;
-	Buttons[19] = EControllerButtons::LeftStickRight;
-	Buttons[20] = EControllerButtons::RightStickUp;
-	Buttons[21] = EControllerButtons::RightStickDown;
-	Buttons[22] = EControllerButtons::RightStickLeft;
-	Buttons[23] = EControllerButtons::RightStickRight;
+	Buttons[0] = FGamepadKeyNames::FaceButtonBottom;
+	Buttons[1] = FGamepadKeyNames::FaceButtonRight;
+	Buttons[2] = FGamepadKeyNames::FaceButtonLeft;
+	Buttons[3] = FGamepadKeyNames::FaceButtonTop;
+	Buttons[4] = FGamepadKeyNames::LeftShoulder;
+	Buttons[5] = FGamepadKeyNames::RightShoulder;
+	Buttons[6] = FGamepadKeyNames::SpecialRight;
+	Buttons[7] = FGamepadKeyNames::SpecialLeft;
+	Buttons[8] = FGamepadKeyNames::LeftThumb;
+	Buttons[9] = FGamepadKeyNames::RightThumb;
+	Buttons[10] = FGamepadKeyNames::LeftTriggerThreshold;
+	Buttons[11] = FGamepadKeyNames::RightTriggerThreshold;
+	Buttons[12] = FGamepadKeyNames::DPadUp;
+	Buttons[13] = FGamepadKeyNames::DPadDown;
+	Buttons[14] = FGamepadKeyNames::DPadLeft;
+	Buttons[15] = FGamepadKeyNames::DPadRight;
+	Buttons[16] = FGamepadKeyNames::LeftStickUp;
+	Buttons[17] = FGamepadKeyNames::LeftStickDown;
+	Buttons[18] = FGamepadKeyNames::LeftStickLeft;
+	Buttons[19] = FGamepadKeyNames::LeftStickRight;
+	Buttons[20] = FGamepadKeyNames::RightStickUp;
+	Buttons[21] = FGamepadKeyNames::RightStickDown;
+	Buttons[22] = FGamepadKeyNames::RightStickLeft;
+	Buttons[23] = FGamepadKeyNames::RightStickRight;
 
 	// Init HID Manager
 	HIDManager = IOHIDManagerCreate(kCFAllocatorDefault, 0L);
@@ -427,37 +427,37 @@ void HIDInputInterface::SendControllerEvents()
 
 						if (Element.Usage == ControllerState.Device.LeftAnalogXMapping && ControllerState.LeftAnalogX != NewValue)
 						{
-							MessageHandler->OnControllerAnalog(EControllerButtons::LeftAnalogX, ControllerState.ControllerId, FloatValue);
-							CurrentButtonStates[18] = FloatValue < -0.2f;	// EControllerButtons::LeftStickLeft
-							CurrentButtonStates[19] = FloatValue > 0.2f;	// EControllerButtons::LeftStickRight
+							MessageHandler->OnControllerAnalog(FGamepadKeyNames::LeftAnalogX, ControllerState.ControllerId, FloatValue);
+							CurrentButtonStates[18] = FloatValue < -0.2f;	// FGamepadKeyNames::LeftStickLeft
+							CurrentButtonStates[19] = FloatValue > 0.2f;	// FGamepadKeyNames::LeftStickRight
 							ControllerState.LeftAnalogX = NewValue;
 						}
 						else if (Element.Usage == ControllerState.Device.LeftAnalogYMapping && ControllerState.LeftAnalogY != NewValue)
 						{
-							MessageHandler->OnControllerAnalog(EControllerButtons::LeftAnalogY, ControllerState.ControllerId, -FloatValue);
-							CurrentButtonStates[16] = FloatValue < -0.2f;	// EControllerButtons::LeftStickUp
-							CurrentButtonStates[17] = FloatValue > 0.2f;	// EControllerButtons::LeftStickDown
+							MessageHandler->OnControllerAnalog(FGamepadKeyNames::LeftAnalogY, ControllerState.ControllerId, -FloatValue);
+							CurrentButtonStates[16] = FloatValue < -0.2f;	// FGamepadKeyNames::LeftStickUp
+							CurrentButtonStates[17] = FloatValue > 0.2f;	// FGamepadKeyNames::LeftStickDown
 							ControllerState.LeftAnalogY = NewValue;
 						}
 						else if (Element.Usage == ControllerState.Device.RightAnalogXMapping && ControllerState.RightAnalogX != NewValue)
 						{
-							MessageHandler->OnControllerAnalog(EControllerButtons::RightAnalogX, ControllerState.ControllerId, FloatValue);
-							CurrentButtonStates[22] = FloatValue < -0.2f;	// EControllerButtons::RightStickLeft
-							CurrentButtonStates[23] = FloatValue > 0.2f;	// EControllerButtons::RightStickDown
+							MessageHandler->OnControllerAnalog(FGamepadKeyNames::RightAnalogX, ControllerState.ControllerId, FloatValue);
+							CurrentButtonStates[22] = FloatValue < -0.2f;	// FGamepadKeyNames::RightStickLeft
+							CurrentButtonStates[23] = FloatValue > 0.2f;	// FGamepadKeyNames::RightStickDown
 							ControllerState.RightAnalogX = NewValue;
 						}
 						else if (Element.Usage == ControllerState.Device.RightAnalogYMapping && ControllerState.RightAnalogY != NewValue)
 						{
-							MessageHandler->OnControllerAnalog(EControllerButtons::RightAnalogY, ControllerState.ControllerId, -FloatValue);
-							CurrentButtonStates[20] = FloatValue < -0.2f;	// EControllerButtons::RightStickUp
-							CurrentButtonStates[21] = FloatValue > 0.2f;	// EControllerButtons::RightStickRight
+							MessageHandler->OnControllerAnalog(FGamepadKeyNames::RightAnalogY, ControllerState.ControllerId, -FloatValue);
+							CurrentButtonStates[20] = FloatValue < -0.2f;	// FGamepadKeyNames::RightStickUp
+							CurrentButtonStates[21] = FloatValue > 0.2f;	// FGamepadKeyNames::RightStickRight
 							ControllerState.RightAnalogY = NewValue;
 						}
 						else if (Element.Usage == ControllerState.Device.LeftTriggerAnalogMapping)
 						{
 							if (ControllerState.LeftTriggerAnalog != NewValue)
 							{
-								MessageHandler->OnControllerAnalog(EControllerButtons::LeftTriggerAnalog, ControllerState.ControllerId, FloatValue);
+								MessageHandler->OnControllerAnalog(FGamepadKeyNames::LeftTriggerAnalog, ControllerState.ControllerId, FloatValue);
 								ControllerState.LeftTriggerAnalog = NewValue;
 							}
 							CurrentButtonStates[10] = FloatValue > 0.01f;
@@ -466,7 +466,7 @@ void HIDInputInterface::SendControllerEvents()
 						{
 							if (ControllerState.RightTriggerAnalog != NewValue)
 							{
-								MessageHandler->OnControllerAnalog(EControllerButtons::RightTriggerAnalog, ControllerState.ControllerId, FloatValue);
+								MessageHandler->OnControllerAnalog(FGamepadKeyNames::RightTriggerAnalog, ControllerState.ControllerId, FloatValue);
 								ControllerState.RightTriggerAnalog = NewValue;
 							}
 							CurrentButtonStates[11] = FloatValue > 0.01f;
diff --git a/Engine/Source/Runtime/Core/Private/Mac/HIDInputInterface.h b/Engine/Source/Runtime/Core/Private/Mac/HIDInputInterface.h
--- a/Engine/Source/Runtime/Core/Private/Mac/HIDInputInterface.h
+++ b/Engine/Source/Runtime/Core/Private/Mac/HIDInputInterface.h
@@ -103,7 +103,7 @@ private:
 private:
 
 	/** Names of all the buttons */
-	EControllerButtons::Type Buttons[MAX_NUM_CONTROLLER_BUTTONS];
+	FGamepadKeyNames::Type Buttons[MAX_NUM_CONTROLLER_BUTTONS];
 
 	/** Controller states */
 	FControllerState ControllerStates[MAX_NUM_HIDINPUT_CONTROLLERS];
diff --git a/Engine/Source/Runtime/Core/Private/WinRT/WinRTInputInterface.cpp b/Engine/Source/Runtime/Core/Private/WinRT/WinRTInputInterface.cpp
--- a/Engine/Source/Runtime/Core/Private/WinRT/WinRTInputInterface.cpp
+++ b/Engine/Source/Runtime/Core/Private/WinRT/WinRTInputInterface.cpp
@@ -84,22 +84,22 @@ FWinRTInputInterface::FWinRTInputInterface( const TSharedRef< FGenericApplicatio
 	X360ToXboxControllerMapping[14] = 14;	// Dpad left
 	X360ToXboxControllerMapping[15] = 15;	// Dpad right
 
-	Buttons[0] = EControllerButtons::FaceButtonBottom;
-	Buttons[1] = EControllerButtons::FaceButtonRight;
-	Buttons[2] = EControllerButtons::FaceButtonLeft;
-	Buttons[3] = EControllerButtons::FaceButtonTop;
-	Buttons[4] = EControllerButtons::LeftShoulder;
-	Buttons[5] = EControllerButtons::RightShoulder;
-	Buttons[6] = EControllerButtons::SpecialRight;
-	Buttons[7] = EControllerButtons::SpecialLeft;
-	Buttons[8] = EControllerButtons::LeftThumb;
-	Buttons[9] = EControllerButtons::RightThumb;
-	Buttons[10] = EControllerButtons::LeftTriggerThreshold;
-	Buttons[11] = EControllerButtons::RightTriggerThreshold;
-	Buttons[12] = EControllerButtons::DPadUp;
-	Buttons[13] = EControllerButtons::DPadDown;
-	Buttons[14] = EControllerButtons::DPadLeft;
-	Buttons[15] = EControllerButtons::DPadRight;
+	Buttons[0] = FGamepadKeyNames::FaceButtonBottom;
+	Buttons[1] = FGamepadKeyNames::FaceButtonRight;
+	Buttons[2] = FGamepadKeyNames::FaceButtonLeft;
+	Buttons[3] = FGamepadKeyNames::FaceButtonTop;
+	Buttons[4] = FGamepadKeyNames::LeftShoulder;
+	Buttons[5] = FGamepadKeyNames::RightShoulder;
+	Buttons[6] = FGamepadKeyNames::SpecialRight;
+	Buttons[7] = FGamepadKeyNames::SpecialLeft;
+	Buttons[8] = FGamepadKeyNames::LeftThumb;
+	Buttons[9] = FGamepadKeyNames::RightThumb;
+	Buttons[10] = FGamepadKeyNames::LeftTriggerThreshold;
+	Buttons[11] = FGamepadKeyNames::RightTriggerThreshold;
+	Buttons[12] = FGamepadKeyNames::DPadUp;
+	Buttons[13] = FGamepadKeyNames::DPadDown;
+	Buttons[14] = FGamepadKeyNames::DPadLeft;
+	Buttons[15] = FGamepadKeyNames::DPadRight;
 }
 
 void FWinRTInputInterface::SetMessageHandler( const TSharedRef< FGenericApplicationMessageHandler >& InMessageHandler )
@@ -398,37 +398,37 @@ void FWinRTInputInterface::SendControllerEvents()
 			// Axis, convert range -32768..+32767 set up to +/- 1.
 			if( ControllerState.LeftXAnalog != XInputState.Gamepad.sThumbLX )
 			{
-				MessageHandler->OnControllerAnalog( EControllerButtons::LeftAnalogX, ControllerState.ControllerId, XInputState.Gamepad.sThumbLX / 32768.f );
+				MessageHandler->OnControllerAnalog( FGamepadKeyNames::LeftAnalogX, ControllerState.ControllerId, XInputState.Gamepad.sThumbLX / 32768.f );
 				ControllerState.LeftXAnalog = XInputState.Gamepad.sThumbLX;
 			}
 
 			if( ControllerState.LeftYAnalog != XInputState.Gamepad.sThumbLY )
 			{
-				MessageHandler->OnControllerAnalog( EControllerButtons::LeftAnalogY, ControllerState.ControllerId, XInputState.Gamepad.sThumbLY / 32768.f );
+				MessageHandler->OnControllerAnalog( FGamepadKeyNames::LeftAnalogY, ControllerState.ControllerId, XInputState.Gamepad.sThumbLY / 32768.f );
 				ControllerState.LeftYAnalog = XInputState.Gamepad.sThumbLY;
 			}
 
 			if( ControllerState.RightXAnalog != XInputState.Gamepad.sThumbRX )
 			{
-				MessageHandler->OnControllerAnalog( EControllerButtons::RightAnalogX, ControllerState.ControllerId, XInputState.Gamepad.sThumbRX / 32768.f );
+				MessageHandler->OnControllerAnalog( FGamepadKeyNames::RightAnalogX, ControllerState.ControllerId, XInputState.Gamepad.sThumbRX / 32768.f );
 				ControllerState.RightXAnalog = XInputState.Gamepad.sThumbRX;
 			}
 
 			if( ControllerState.RightYAnalog != XInputState.Gamepad.sThumbRY )
 			{
-				MessageHandler->OnControllerAnalog( EControllerButtons::RightAnalogY, ControllerState.ControllerId, XInputState.Gamepad.sThumbRY / 32768.f );
+				MessageHandler->OnControllerAnalog( FGamepadKeyNames::RightAnalogY, ControllerState.ControllerId, XInputState.Gamepad.sThumbRY / 32768.f );
 				ControllerState.RightYAnalog = XInputState.Gamepad.sThumbRY;
 			}
 
 			if( ControllerState.LeftTriggerAnalog != XInputState.Gamepad.bLeftTrigger )
 			{
-				MessageHandler->OnControllerAnalog( EControllerButtons::LeftTriggerAnalog, ControllerState.ControllerId, XInputState.Gamepad.bLeftTrigger / 255.f );
+				MessageHandler->OnControllerAnalog( FGamepadKeyNames::LeftTriggerAnalog, ControllerState.ControllerId, XInputState.Gamepad.bLeftTrigger / 255.f );
 				ControllerState.LeftTriggerAnalog = XInputState.Gamepad.bLeftTrigger;
 			}
 
 			if( ControllerState.RightTriggerAnalog != XInputState.Gamepad.bRightTrigger )
 			{
-				MessageHandler->OnControllerAnalog( EControllerButtons::RightTriggerAnalog, ControllerState.ControllerId, XInputState.Gamepad.bRightTrigger / 255.f );
+				MessageHandler->OnControllerAnalog( FGamepadKeyNames::RightTriggerAnalog, ControllerState.ControllerId, XInputState.Gamepad.bRightTrigger / 255.f );
 				ControllerState.RightTriggerAnalog = XInputState.Gamepad.bRightTrigger;
 			}
 
diff --git a/Engine/Source/Runtime/Core/Private/WinRT/WinRTInputInterface.h b/Engine/Source/Runtime/Core/Private/WinRT/WinRTInputInterface.h
--- a/Engine/Source/Runtime/Core/Private/WinRT/WinRTInputInterface.h
+++ b/Engine/Source/Runtime/Core/Private/WinRT/WinRTInputInterface.h
@@ -98,7 +98,7 @@ private:
 	uint8 X360ToXboxControllerMapping[MAX_NUM_CONTROLLER_BUTTONS];
 
 	/** Names of all the buttons */
-	EControllerButtons::Type Buttons[MAX_NUM_CONTROLLER_BUTTONS];
+	FGamepadKeyNames::Type Buttons[MAX_NUM_CONTROLLER_BUTTONS];
 
 	/** Controller states */
 	FControllerState ControllerStates[MAX_NUM_XINPUT_CONTROLLERS];
diff --git a/Engine/Source/Runtime/Core/Private/Windows/XInputInterface.cpp b/Engine/Source/Runtime/Core/Private/Windows/XInputInterface.cpp
--- a/Engine/Source/Runtime/Core/Private/Windows/XInputInterface.cpp
+++ b/Engine/Source/Runtime/Core/Private/Windows/XInputInterface.cpp
@@ -59,30 +59,30 @@ XInputInterface::XInputInterface( const TSharedRef< FGenericApplicationMessageHa
 	X360ToXboxControllerMapping[22] = 22;	// Right stick left
 	X360ToXboxControllerMapping[23] = 23;	// Right stick right
 
-	Buttons[0] = EControllerButtons::FaceButtonBottom;
-	Buttons[1] = EControllerButtons::FaceButtonRight;
-	Buttons[2] = EControllerButtons::FaceButtonLeft;
-	Buttons[3] = EControllerButtons::FaceButtonTop;
-	Buttons[4] = EControllerButtons::LeftShoulder;
-	Buttons[5] = EControllerButtons::RightShoulder;
-	Buttons[6] = EControllerButtons::SpecialRight;
-	Buttons[7] = EControllerButtons::SpecialLeft;
-	Buttons[8] = EControllerButtons::LeftThumb;
-	Buttons[9] = EControllerButtons::RightThumb;
-	Buttons[10] = EControllerButtons::LeftTriggerThreshold;
-	Buttons[11] = EControllerButtons::RightTriggerThreshold;
-	Buttons[12] = EControllerButtons::DPadUp;
-	Buttons[13] = EControllerButtons::DPadDown;
-	Buttons[14] = EControllerButtons::DPadLeft;
-	Buttons[15] = EControllerButtons::DPadRight;
-	Buttons[16] = EControllerButtons::LeftStickUp;
-	Buttons[17] = EControllerButtons::LeftStickDown;
-	Buttons[18] = EControllerButtons::LeftStickLeft;
-	Buttons[19] = EControllerButtons::LeftStickRight;
-	Buttons[20] = EControllerButtons::RightStickUp;
-	Buttons[21] = EControllerButtons::RightStickDown;
-	Buttons[22] = EControllerButtons::RightStickLeft;
-	Buttons[23] = EControllerButtons::RightStickRight;
+	Buttons[0] = FGamepadKeyNames::FaceButtonBottom;
+	Buttons[1] = FGamepadKeyNames::FaceButtonRight;
+	Buttons[2] = FGamepadKeyNames::FaceButtonLeft;
+	Buttons[3] = FGamepadKeyNames::FaceButtonTop;
+	Buttons[4] = FGamepadKeyNames::LeftShoulder;
+	Buttons[5] = FGamepadKeyNames::RightShoulder;
+	Buttons[6] = FGamepadKeyNames::SpecialRight;
+	Buttons[7] = FGamepadKeyNames::SpecialLeft;
+	Buttons[8] = FGamepadKeyNames::LeftThumb;
+	Buttons[9] = FGamepadKeyNames::RightThumb;
+	Buttons[10] = FGamepadKeyNames::LeftTriggerThreshold;
+	Buttons[11] = FGamepadKeyNames::RightTriggerThreshold;
+	Buttons[12] = FGamepadKeyNames::DPadUp;
+	Buttons[13] = FGamepadKeyNames::DPadDown;
+	Buttons[14] = FGamepadKeyNames::DPadLeft;
+	Buttons[15] = FGamepadKeyNames::DPadRight;
+	Buttons[16] = FGamepadKeyNames::LeftStickUp;
+	Buttons[17] = FGamepadKeyNames::LeftStickDown;
+	Buttons[18] = FGamepadKeyNames::LeftStickLeft;
+	Buttons[19] = FGamepadKeyNames::LeftStickRight;
+	Buttons[20] = FGamepadKeyNames::RightStickUp;
+	Buttons[21] = FGamepadKeyNames::RightStickDown;
+	Buttons[22] = FGamepadKeyNames::RightStickLeft;
+	Buttons[23] = FGamepadKeyNames::RightStickRight;
 }
 
 
@@ -140,37 +140,37 @@ void XInputInterface::SendControllerEvents()
 
 				if( ControllerState.LeftXAnalog != XInputState.Gamepad.sThumbLX )
 				{
-					MessageHandler->OnControllerAnalog( EControllerButtons::LeftAnalogX, ControllerState.ControllerId, ShortToNormalizedFloat(XInputState.Gamepad.sThumbLX) );
+					MessageHandler->OnControllerAnalog( FGamepadKeyNames::LeftAnalogX, ControllerState.ControllerId, ShortToNormalizedFloat(XInputState.Gamepad.sThumbLX) );
 					ControllerState.LeftXAnalog = XInputState.Gamepad.sThumbLX;
 				}
 
 				if( ControllerState.LeftYAnalog != XInputState.Gamepad.sThumbLY )
 				{
-					MessageHandler->OnControllerAnalog( EControllerButtons::LeftAnalogY, ControllerState.ControllerId, ShortToNormalizedFloat(XInputState.Gamepad.sThumbLY) );
+					MessageHandler->OnControllerAnalog( FGamepadKeyNames::LeftAnalogY, ControllerState.ControllerId, ShortToNormalizedFloat(XInputState.Gamepad.sThumbLY) );
 					ControllerState.LeftYAnalog = XInputState.Gamepad.sThumbLY;
 				}
 
 				if( ControllerState.RightXAnalog != XInputState.Gamepad.sThumbRX )
 				{
-					MessageHandler->OnControllerAnalog( EControllerButtons::RightAnalogX, ControllerState.ControllerId, ShortToNormalizedFloat(XInputState.Gamepad.sThumbRX) );
+					MessageHandler->OnControllerAnalog( FGamepadKeyNames::RightAnalogX, ControllerState.ControllerId, ShortToNormalizedFloat(XInputState.Gamepad.sThumbRX) );
 					ControllerState.RightXAnalog = XInputState.Gamepad.sThumbRX;
 				}
 
 				if( ControllerState.RightYAnalog != XInputState.Gamepad.sThumbRY )
 				{
-					MessageHandler->OnControllerAnalog( EControllerButtons::RightAnalogY, ControllerState.ControllerId, ShortToNormalizedFloat(XInputState.Gamepad.sThumbRY) );
+					MessageHandler->OnControllerAnalog( FGamepadKeyNames::RightAnalogY, ControllerState.ControllerId, ShortToNormalizedFloat(XInputState.Gamepad.sThumbRY) );
 					ControllerState.RightYAnalog = XInputState.Gamepad.sThumbRY;
 				}
 
 				if( ControllerState.LeftTriggerAnalog != XInputState.Gamepad.bLeftTrigger )
 				{
-					MessageHandler->OnControllerAnalog( EControllerButtons::LeftTriggerAnalog, ControllerState.ControllerId, XInputState.Gamepad.bLeftTrigger / 255.f );
+					MessageHandler->OnControllerAnalog( FGamepadKeyNames::LeftTriggerAnalog, ControllerState.ControllerId, XInputState.Gamepad.bLeftTrigger / 255.f );
 					ControllerState.LeftTriggerAnalog = XInputState.Gamepad.bLeftTrigger;
 				}
 
 				if( ControllerState.RightTriggerAnalog != XInputState.Gamepad.bRightTrigger )
 				{
-					MessageHandler->OnControllerAnalog( EControllerButtons::RightTriggerAnalog, ControllerState.ControllerId, XInputState.Gamepad.bRightTrigger / 255.f );
+					MessageHandler->OnControllerAnalog( FGamepadKeyNames::RightTriggerAnalog, ControllerState.ControllerId, XInputState.Gamepad.bRightTrigger / 255.f );
 					ControllerState.RightTriggerAnalog = XInputState.Gamepad.bRightTrigger;
 				}
 
diff --git a/Engine/Source/Runtime/Core/Private/Windows/XInputInterface.h b/Engine/Source/Runtime/Core/Private/Windows/XInputInterface.h
--- a/Engine/Source/Runtime/Core/Private/Windows/XInputInterface.h
+++ b/Engine/Source/Runtime/Core/Private/Windows/XInputInterface.h
@@ -106,7 +106,7 @@ private:
 	float ButtonRepeatDelay;
 
 	/**  */
-	EControllerButtons::Type Buttons[MAX_NUM_CONTROLLER_BUTTONS];
+	FGamepadKeyNames::Type Buttons[MAX_NUM_CONTROLLER_BUTTONS];
 
 	/**  */
 	TSharedRef<FGenericApplicationMessageHandler> MessageHandler;
diff --git a/Engine/Source/Runtime/Core/Public/Android/AndroidInputInterface.h b/Engine/Source/Runtime/Core/Public/Android/AndroidInputInterface.h
--- a/Engine/Source/Runtime/Core/Public/Android/AndroidInputInterface.h
+++ b/Engine/Source/Runtime/Core/Public/Android/AndroidInputInterface.h
@@ -202,7 +202,7 @@ private:
 	static FAndroidControllerData OldControllerData[MAX_NUM_CONTROLLERS];
 	static FAndroidControllerData NewControllerData[MAX_NUM_CONTROLLERS];
 
-	static EControllerButtons::Type ButtonMapping[MAX_NUM_CONTROLLER_BUTTONS];
+	static FGamepadKeyNames::Type ButtonMapping[MAX_NUM_CONTROLLER_BUTTONS];
 
 	static float InitialButtonRepeatDelay;
 	static float ButtonRepeatDelay;
diff --git a/Engine/Source/Runtime/Core/Public/GenericPlatform/GenericApplicationMessageHandler.h b/Engine/Source/Runtime/Core/Public/GenericPlatform/GenericApplicationMessageHandler.h
--- a/Engine/Source/Runtime/Core/Public/GenericPlatform/GenericApplicationMessageHandler.h
+++ b/Engine/Source/Runtime/Core/Public/GenericPlatform/GenericApplicationMessageHandler.h
@@ -24,71 +24,46 @@ namespace EMouseButtons
 	};
 }
 
-
-namespace EControllerButtons
+struct CORE_API FGamepadKeyNames
 {
-	enum Type
-	{
-		LeftAnalogY = 0,
-		LeftAnalogX,
-		
-		RightAnalogY,
-		RightAnalogX,
-
-		LeftTriggerAnalog,
-		RightTriggerAnalog,
-
-		FaceButtonBottom,
-		FaceButtonRight,
-		FaceButtonLeft,
-		FaceButtonTop,
-		LeftShoulder,
-		RightShoulder,
-		SpecialLeft,
-		SpecialRight,
-		LeftThumb,
-		RightThumb,
-		LeftTriggerThreshold,
-		RightTriggerThreshold,
-		DPadUp,
-		DPadDown,
-		DPadLeft,
-		DPadRight,
-
-		LeftStickUp,
-		LeftStickDown,
-		LeftStickLeft,
-		LeftStickRight,
-
-		RightStickUp,
-		RightStickDown,
-		RightStickLeft,
-		RightStickRight,
-
-		// Steam Controller
-		Touch0,
-		Touch1,
-		Touch2,
-		Touch3,
-		BackLeft,
-		BackRight,
-
-		// global speech commands
-		GlobalMenu,
-		GlobalView,
-		GlobalPause,
-		GlobalPlay,
-		GlobalBack,
-
-		AndroidBack,
-		AndroidVolumeUp,
-		AndroidVolumeDown,
-		AndroidMenu,
-
-		Invalid,
-	};
-}
-
+	typedef FName Type;
+
+	static const FName Invalid;
+
+	static const FName LeftAnalogX;
+	static const FName LeftAnalogY;
+	static const FName RightAnalogX;
+	static const FName RightAnalogY;
+	static const FName LeftTriggerAnalog;
+	static const FName RightTriggerAnalog;
+
+	static const FName LeftThumb;
+	static const FName RightThumb;
+	static const FName SpecialLeft;
+	static const FName SpecialRight;
+	static const FName FaceButtonBottom;
+	static const FName FaceButtonRight;
+	static const FName FaceButtonLeft;
+	static const FName FaceButtonTop;
+	static const FName LeftShoulder;
+	static const FName RightShoulder;
+	static const FName LeftTriggerThreshold;
+	static const FName RightTriggerThreshold;
+	static const FName DPadUp;
+	static const FName DPadDown;
+	static const FName DPadRight;
+	static const FName DPadLeft;
+
+	static const FName LeftStickUp;
+	static const FName LeftStickDown;
+	static const FName LeftStickRight;
+	static const FName LeftStickLeft;
+
+	static const FName RightStickUp;
+	static const FName RightStickDown;
+	static const FName RightStickRight;
+	static const FName RightStickLeft;
+};
 
 namespace EWindowActivation
 {
@@ -254,21 +229,21 @@ public:
 		return false;
 	}
 
-	virtual bool OnControllerAnalog( EControllerButtons::Type Button, int32 ControllerId, float AnalogValue )
+	virtual bool OnControllerAnalog( FGamepadKeyNames::Type KeyName, int32 ControllerId, float AnalogValue )
 	{
 		return false;
 	}
 
-	virtual bool OnControllerButtonPressed( EControllerButtons::Type Button, int32 ControllerId, bool IsRepeat )
+	virtual bool OnControllerButtonPressed( FGamepadKeyNames::Type KeyName, int32 ControllerId, bool IsRepeat )
 	{
 		return false;
 	}
 
-	virtual bool OnControllerButtonReleased( EControllerButtons::Type Button, int32 ControllerId, bool IsRepeat )
+	virtual bool OnControllerButtonReleased( FGamepadKeyNames::Type KeyName, int32 ControllerId, bool IsRepeat )
 	{
 		return false;
 	}
-    
+
     virtual void OnBeginGesture()
     {
     }
diff --git a/Engine/Source/Runtime/InputCore/Classes/InputCoreTypes.h b/Engine/Source/Runtime/InputCore/Classes/InputCoreTypes.h
--- a/Engine/Source/Runtime/InputCore/Classes/InputCoreTypes.h
+++ b/Engine/Source/Runtime/InputCore/Classes/InputCoreTypes.h
@@ -307,7 +307,7 @@ struct INPUTCORE_API EKeys
 	// On Macs, the Delete key is the virtual key for BackSpace.
 	static const FKey Platform_Delete;
 
-	// Gameplay Keys
+	// Gamepad Keys
 	static const FKey Gamepad_LeftX;
 	static const FKey Gamepad_LeftY;
 	static const FKey Gamepad_RightX;
diff --git a/Engine/Source/Runtime/InputCore/Private/InputCoreTypes.cpp b/Engine/Source/Runtime/InputCore/Private/InputCoreTypes.cpp
--- a/Engine/Source/Runtime/InputCore/Private/InputCoreTypes.cpp
+++ b/Engine/Source/Runtime/InputCore/Private/InputCoreTypes.cpp
@@ -157,40 +157,40 @@ const FKey EKeys::Platform_Delete = EKeys::BackSpace;
 const FKey EKeys::Platform_Delete = EKeys::Delete;
 #endif
 
-const FKey EKeys::Gamepad_LeftX("Gamepad_LeftX");
-const FKey EKeys::Gamepad_LeftY("Gamepad_LeftY");
-const FKey EKeys::Gamepad_RightX("Gamepad_RightX");
-const FKey EKeys::Gamepad_RightY("Gamepad_RightY");
-const FKey EKeys::Gamepad_LeftTriggerAxis("Gamepad_LeftTriggerAxis");
-const FKey EKeys::Gamepad_RightTriggerAxis("Gamepad_RightTriggerAxis");
-
-const FKey EKeys::Gamepad_LeftThumbstick("Gamepad_LeftThumbstick");
-const FKey EKeys::Gamepad_RightThumbstick("Gamepad_RightThumbstick");
-const FKey EKeys::Gamepad_Special_Left("Gamepad_Special_Left");
-const FKey EKeys::Gamepad_Special_Right("Gamepad_Special_Right");
-const FKey EKeys::Gamepad_FaceButton_Bottom("Gamepad_FaceButton_Bottom");
-const FKey EKeys::Gamepad_FaceButton_Right("Gamepad_FaceButton_Right");
-const FKey EKeys::Gamepad_FaceButton_Left("Gamepad_FaceButton_Left");
-const FKey EKeys::Gamepad_FaceButton_Top("Gamepad_FaceButton_Top");
-const FKey EKeys::Gamepad_LeftShoulder("Gamepad_LeftShoulder");
-const FKey EKeys::Gamepad_RightShoulder("Gamepad_RightShoulder");
-const FKey EKeys::Gamepad_LeftTrigger("Gamepad_LeftTrigger");
-const FKey EKeys::Gamepad_RightTrigger("Gamepad_RightTrigger");
-const FKey EKeys::Gamepad_DPad_Up("Gamepad_DPad_Up");
-const FKey EKeys::Gamepad_DPad_Down("Gamepad_DPad_Down");
-const FKey EKeys::Gamepad_DPad_Right("Gamepad_DPad_Right");
-const FKey EKeys::Gamepad_DPad_Left("Gamepad_DPad_Left");
+const FKey EKeys::Gamepad_LeftX(FGamepadKeyNames::LeftAnalogX);
+const FKey EKeys::Gamepad_LeftY(FGamepadKeyNames::LeftAnalogY);
+const FKey EKeys::Gamepad_RightX(FGamepadKeyNames::RightAnalogX);
+const FKey EKeys::Gamepad_RightY(FGamepadKeyNames::RightAnalogY);
+const FKey EKeys::Gamepad_LeftTriggerAxis(FGamepadKeyNames::LeftTriggerAnalog);
+const FKey EKeys::Gamepad_RightTriggerAxis(FGamepadKeyNames::RightTriggerAnalog);
+
+const FKey EKeys::Gamepad_LeftThumbstick(FGamepadKeyNames::LeftThumb);
+const FKey EKeys::Gamepad_RightThumbstick(FGamepadKeyNames::RightThumb);
+const FKey EKeys::Gamepad_Special_Left(FGamepadKeyNames::SpecialLeft);
+const FKey EKeys::Gamepad_Special_Right(FGamepadKeyNames::SpecialRight);
+const FKey EKeys::Gamepad_FaceButton_Bottom(FGamepadKeyNames::FaceButtonBottom);
+const FKey EKeys::Gamepad_FaceButton_Right(FGamepadKeyNames::FaceButtonRight);
+const FKey EKeys::Gamepad_FaceButton_Left(FGamepadKeyNames::FaceButtonLeft);
+const FKey EKeys::Gamepad_FaceButton_Top(FGamepadKeyNames::FaceButtonTop);
+const FKey EKeys::Gamepad_LeftShoulder(FGamepadKeyNames::LeftShoulder);
+const FKey EKeys::Gamepad_RightShoulder(FGamepadKeyNames::RightShoulder);
+const FKey EKeys::Gamepad_LeftTrigger(FGamepadKeyNames::LeftTriggerThreshold);
+const FKey EKeys::Gamepad_RightTrigger(FGamepadKeyNames::RightTriggerThreshold);
+const FKey EKeys::Gamepad_DPad_Up(FGamepadKeyNames::DPadUp);
+const FKey EKeys::Gamepad_DPad_Down(FGamepadKeyNames::DPadDown);
+const FKey EKeys::Gamepad_DPad_Right(FGamepadKeyNames::DPadRight);
+const FKey EKeys::Gamepad_DPad_Left(FGamepadKeyNames::DPadLeft);
 
 // Virtual key codes used for input axis button press/release emulation
-const FKey EKeys::Gamepad_LeftStick_Up("Gamepad_LeftStick_Up");
-const FKey EKeys::Gamepad_LeftStick_Down("Gamepad_LeftStick_Down");
-const FKey EKeys::Gamepad_LeftStick_Right("Gamepad_LeftStick_Right");
-const FKey EKeys::Gamepad_LeftStick_Left("Gamepad_LeftStick_Left");
-
-const FKey EKeys::Gamepad_RightStick_Up("Gamepad_RightStick_Up");
-const FKey EKeys::Gamepad_RightStick_Down("Gamepad_RightStick_Down");
-const FKey EKeys::Gamepad_RightStick_Right("Gamepad_RightStick_Right");
-const FKey EKeys::Gamepad_RightStick_Left("Gamepad_RightStick_Left");
+const FKey EKeys::Gamepad_LeftStick_Up(FGamepadKeyNames::LeftStickUp);
+const FKey EKeys::Gamepad_LeftStick_Down(FGamepadKeyNames::LeftStickDown);
+const FKey EKeys::Gamepad_LeftStick_Right(FGamepadKeyNames::LeftStickRight);
+const FKey EKeys::Gamepad_LeftStick_Left(FGamepadKeyNames::LeftStickLeft);
+
+const FKey EKeys::Gamepad_RightStick_Up(FGamepadKeyNames::RightStickUp);
+const FKey EKeys::Gamepad_RightStick_Down(FGamepadKeyNames::RightStickDown);
+const FKey EKeys::Gamepad_RightStick_Right(FGamepadKeyNames::RightStickRight);
+const FKey EKeys::Gamepad_RightStick_Left(FGamepadKeyNames::RightStickLeft);
 
 // const FKey EKeys::Vector axes (FVector("Vector axes (FVector"); not float)
 const FKey EKeys::Tilt("Tilt");
diff --git a/Engine/Source/Runtime/Slate/Private/Framework/Application/SlateApplication.cpp b/Engine/Source/Runtime/Slate/Private/Framework/Application/SlateApplication.cpp
--- a/Engine/Source/Runtime/Slate/Private/Framework/Application/SlateApplication.cpp
+++ b/Engine/Source/Runtime/Slate/Private/Framework/Application/SlateApplication.cpp
@@ -4773,74 +4773,6 @@ bool FSlateApplication::OnCursorSet()
 	return true;
 }
 
-FKey TranslateControllerButtonToKey( EControllerButtons::Type Button )
-{
-	FKey Key = EKeys::Invalid;
-
-	switch ( Button )
-	{
-	case EControllerButtons::LeftAnalogY: Key = EKeys::Gamepad_LeftY; break;
-	case EControllerButtons::LeftAnalogX: Key = EKeys::Gamepad_LeftX; break;
-
-	case EControllerButtons::RightAnalogY: Key = EKeys::Gamepad_RightY; break;
-	case EControllerButtons::RightAnalogX: Key = EKeys::Gamepad_RightX; break;
-
-	case EControllerButtons::LeftTriggerAnalog: Key = EKeys::Gamepad_LeftTriggerAxis; break;
-	case EControllerButtons::RightTriggerAnalog: Key = EKeys::Gamepad_RightTriggerAxis; break;
-
-	case EControllerButtons::FaceButtonBottom: Key = EKeys::Gamepad_FaceButton_Bottom; break;
-	case EControllerButtons::FaceButtonRight: Key = EKeys::Gamepad_FaceButton_Right; break;
-	case EControllerButtons::FaceButtonLeft: Key = EKeys::Gamepad_FaceButton_Left; break;
-	case EControllerButtons::FaceButtonTop: Key = EKeys::Gamepad_FaceButton_Top; break;
-
-	case EControllerButtons::LeftShoulder: Key = EKeys::Gamepad_LeftShoulder; break;
-	case EControllerButtons::RightShoulder: Key = EKeys::Gamepad_RightShoulder; break;
-	case EControllerButtons::SpecialLeft: Key = EKeys::Gamepad_Special_Left; break;
-	case EControllerButtons::SpecialRight: Key = EKeys::Gamepad_Special_Right; break;
-	case EControllerButtons::LeftThumb: Key = EKeys::Gamepad_LeftThumbstick; break;
-	case EControllerButtons::RightThumb: Key = EKeys::Gamepad_RightThumbstick; break;
-	case EControllerButtons::LeftTriggerThreshold: Key = EKeys::Gamepad_LeftTrigger; break;
-	case EControllerButtons::RightTriggerThreshold: Key = EKeys::Gamepad_RightTrigger; break;
-
-	case EControllerButtons::DPadUp: Key = EKeys::Gamepad_DPad_Up; break;
-	case EControllerButtons::DPadDown: Key = EKeys::Gamepad_DPad_Down; break;
-	case EControllerButtons::DPadLeft: Key = EKeys::Gamepad_DPad_Left; break;
-	case EControllerButtons::DPadRight: Key = EKeys::Gamepad_DPad_Right; break;
-
-	case EControllerButtons::LeftStickUp: Key = EKeys::Gamepad_LeftStick_Up; break;
-	case EControllerButtons::LeftStickDown: Key = EKeys::Gamepad_LeftStick_Down; break;
-	case EControllerButtons::LeftStickLeft: Key = EKeys::Gamepad_LeftStick_Left; break;
-	case EControllerButtons::LeftStickRight: Key = EKeys::Gamepad_LeftStick_Right; break;
-
-	case EControllerButtons::RightStickUp: Key = EKeys::Gamepad_RightStick_Up; break;
-	case EControllerButtons::RightStickDown: Key = EKeys::Gamepad_RightStick_Down; break;
-	case EControllerButtons::RightStickLeft: Key = EKeys::Gamepad_RightStick_Left; break;
-	case EControllerButtons::RightStickRight: Key = EKeys::Gamepad_RightStick_Right; break;
-
-	case EControllerButtons::Touch0: Key = EKeys::Steam_Touch_0; break;
-	case EControllerButtons::Touch1: Key = EKeys::Steam_Touch_1; break;
-	case EControllerButtons::Touch2: Key = EKeys::Steam_Touch_2; break;
-	case EControllerButtons::Touch3: Key = EKeys::Steam_Touch_3; break;
-	case EControllerButtons::BackLeft: Key = EKeys::Steam_Back_Left; break;
-	case EControllerButtons::BackRight: Key = EKeys::Steam_Back_Right; break;
-
-	case EControllerButtons::GlobalMenu: Key = EKeys::Global_Menu; break;
-	case EControllerButtons::GlobalView: Key = EKeys::Global_View; break;
-	case EControllerButtons::GlobalPause: Key = EKeys::Global_Pause; break;
-	case EControllerButtons::GlobalPlay: Key = EKeys::Global_Play; break;
-	case EControllerButtons::GlobalBack: Key = EKeys::Global_Back; break;
-
-	case EControllerButtons::AndroidBack: Key = EKeys::Android_Back; break;
-	case EControllerButtons::AndroidVolumeUp: Key = EKeys::Android_Volume_Up; break;
-	case EControllerButtons::AndroidVolumeDown: Key = EKeys::Android_Volume_Down; break;
-	case EControllerButtons::AndroidMenu: Key = EKeys::Android_Menu; break;
-
-	case EControllerButtons::Invalid: Key = EKeys::Invalid; break;
-	}
-
-	return Key;
-}
-
 bool FSlateApplication::AttemptNavigation(const FNavigationEvent& NavigationEvent, const FNavigationReply& NavigationReply, const FArrangedWidget& BoundaryWidget)
 {
 	// Get the controller focus target for this user
@@ -4900,9 +4832,11 @@ bool FSlateApplication::AttemptNavigation(const FNavigationEvent& NavigationEven
 	return false;
 }
 
-bool FSlateApplication::OnControllerAnalog( EControllerButtons::Type Button, int32 ControllerId, float AnalogValue )
+bool FSlateApplication::OnControllerAnalog( FGamepadKeyNames::Type KeyName, int32 ControllerId, float AnalogValue )
 {
-	FKey Key = TranslateControllerButtonToKey(Button);
+	FKey Key(KeyName);
+	check(Key.IsValid());
+
 	int32 UserIndex = GetUserIndexForController(ControllerId);
 	
 	FAnalogInputEvent AnalogInputEvent(Key, PlatformApplication->GetModifierKeys(), UserIndex, false, 0, 0, AnalogValue);
@@ -4910,9 +4844,11 @@ bool FSlateApplication::OnControllerAnalog( EControllerButtons::Type Button, int
 	return ProcessAnalogInputEvent(AnalogInputEvent);
 }
 
-bool FSlateApplication::OnControllerButtonPressed( EControllerButtons::Type Button, int32 ControllerId, bool IsRepeat )
+bool FSlateApplication::OnControllerButtonPressed( FGamepadKeyNames::Type KeyName, int32 ControllerId, bool IsRepeat )
 {
-	FKey Key = TranslateControllerButtonToKey(Button);
+	FKey Key(KeyName);
+	check(Key.IsValid());
+
 	int32 UserIndex = GetUserIndexForController(ControllerId);
 
 	FKeyEvent KeyEvent(Key, PlatformApplication->GetModifierKeys(), UserIndex, IsRepeat, 0, 0);
@@ -4920,9 +4856,11 @@ bool FSlateApplication::OnControllerButtonPressed( EControllerButtons::Type Butt
 	return ProcessKeyDownEvent(KeyEvent);
 }
 
-bool FSlateApplication::OnControllerButtonReleased( EControllerButtons::Type Button, int32 ControllerId, bool IsRepeat )
+bool FSlateApplication::OnControllerButtonReleased( FGamepadKeyNames::Type KeyName, int32 ControllerId, bool IsRepeat )
 {
-	FKey Key = TranslateControllerButtonToKey(Button);
+	FKey Key(KeyName);
+	check(Key.IsValid());
+
 	int32 UserIndex = GetUserIndexForController(ControllerId);
 	
 	FKeyEvent KeyEvent(Key, PlatformApplication->GetModifierKeys(),UserIndex, IsRepeat,  0, 0);
diff --git a/Engine/Source/Runtime/Slate/Private/Widgets/Input/SVirtualJoystick.cpp b/Engine/Source/Runtime/Slate/Private/Widgets/Input/SVirtualJoystick.cpp
--- a/Engine/Source/Runtime/Slate/Private/Widgets/Input/SVirtualJoystick.cpp
+++ b/Engine/Source/Runtime/Slate/Private/Widgets/Input/SVirtualJoystick.cpp
@@ -381,8 +381,8 @@ void SVirtualJoystick::Tick( const FGeometry& AllottedGeometry, const double InC
 			// Assume that joystick size is all equal
 			float JoystickInputSize = Control.ThumbPosition.Size() * 2.f / Control.CorrectedVisualSize.X;
 			FVector2D NormalizedOffset = Control.ThumbPosition.GetSafeNormal() * Control.CorrectedInputScale * JoystickInputSize;
-			EControllerButtons::Type XAxis = ControlIndex == 0 ? EControllerButtons::LeftAnalogX : EControllerButtons::RightAnalogX;
-			EControllerButtons::Type YAxis = ControlIndex == 0 ? EControllerButtons::LeftAnalogY : EControllerButtons::RightAnalogY;
+			const FGamepadKeyNames::Type XAxis = ControlIndex == 0 ? FGamepadKeyNames::LeftAnalogX : FGamepadKeyNames::RightAnalogX;
+			const FGamepadKeyNames::Type YAxis = ControlIndex == 0 ? FGamepadKeyNames::LeftAnalogY : FGamepadKeyNames::RightAnalogY;
 
 	//		UE_LOG(LogTemp, Log, TEXT("Joysticking %f,%f"), NormalizedOffset.X, -NormalizedOffset.Y);
 			FSlateApplication::Get().SetAllUserFocusToGameViewport();
diff --git a/Engine/Source/Runtime/Slate/Public/Framework/Application/SlateApplication.h b/Engine/Source/Runtime/Slate/Public/Framework/Application/SlateApplication.h
--- a/Engine/Source/Runtime/Slate/Public/Framework/Application/SlateApplication.h
+++ b/Engine/Source/Runtime/Slate/Public/Framework/Application/SlateApplication.h
@@ -1106,9 +1106,9 @@ public:
 	virtual bool OnMouseMove() override;
 	virtual bool OnRawMouseMove( const int32 X, const int32 Y ) override;
 	virtual bool OnCursorSet() override;
-	virtual bool OnControllerAnalog( EControllerButtons::Type Button, int32 ControllerId, float AnalogValue ) override;
-	virtual bool OnControllerButtonPressed( EControllerButtons::Type Button, int32 ControllerId, bool IsRepeat ) override;
-	virtual bool OnControllerButtonReleased( EControllerButtons::Type Button, int32 ControllerId, bool IsRepeat ) override;
+	virtual bool OnControllerAnalog( FGamepadKeyNames::Type KeyName, int32 ControllerId, float AnalogValue ) override;
+	virtual bool OnControllerButtonPressed( FGamepadKeyNames::Type KeyName, int32 ControllerId, bool IsRepeat ) override;
+	virtual bool OnControllerButtonReleased( FGamepadKeyNames::Type KeyName, int32 ControllerId, bool IsRepeat ) override;
 	virtual bool OnTouchGesture( EGestureEvent::Type GestureType, const FVector2D& Delta, float WheelDelta ) override;
 	virtual bool OnTouchStarted( const TSharedPtr< FGenericWindow >& PlatformWindow, const FVector2D& Location, int32 TouchIndex, int32 ControllerId ) override;
 	virtual bool OnTouchMoved( const FVector2D& Location, int32 TouchIndex, int32 ControllerId ) override;
