Pass bIsEnsureParameter properly to PromptForRemoteDebugging.

diff --git a/Engine/Source/Runtime/Core/Public/Android/AndroidMisc.h b/Engine/Source/Runtime/Core/Public/Android/AndroidMisc.h
--- a/Engine/Source/Runtime/Core/Public/Android/AndroidMisc.h
+++ b/Engine/Source/Runtime/Core/Public/Android/AndroidMisc.h
@@ -83,12 +83,12 @@ struct CORE_API FAndroidMisc : public FGenericPlatformMisc
 	}
 
 	/** Prompts for remote debugging if debugger is not attached. Regardless of result, breaks into debugger afterwards. Returns false for use in conditionals. */
-	static FORCEINLINE bool DebugBreakAndPromptForRemoteReturningFalse()
+	static FORCEINLINE bool DebugBreakAndPromptForRemoteReturningFalse(bool bIsEnsure = false)
 	{
 #if !UE_BUILD_SHIPPING
 		if (!IsDebuggerPresent())
 		{
-			PromptForRemoteDebugging(false);
+			PromptForRemoteDebugging(bIsEnsure);
 		}
 
 		DebugBreak();
diff --git a/Engine/Source/Runtime/Core/Public/HTML5/HTML5PlatformMisc.h b/Engine/Source/Runtime/Core/Public/HTML5/HTML5PlatformMisc.h
--- a/Engine/Source/Runtime/Core/Public/HTML5/HTML5PlatformMisc.h
+++ b/Engine/Source/Runtime/Core/Public/HTML5/HTML5PlatformMisc.h
@@ -68,12 +68,12 @@ struct CORE_API FHTML5Misc : public FGenericPlatformMisc
 	}
 
 	/** Prompts for remote debugging if debugger is not attached. Regardless of result, breaks into debugger afterwards. Returns false for use in conditionals. */
-	static FORCEINLINE bool DebugBreakAndPromptForRemoteReturningFalse()
+	static FORCEINLINE bool DebugBreakAndPromptForRemoteReturningFalse(bool bIsEnsure = false)
 	{
 #if !UE_BUILD_SHIPPING
 		if (!IsDebuggerPresent())
 		{
-			PromptForRemoteDebugging(false);
+			PromptForRemoteDebugging(bIsEnsure);
 		}
 
 		DebugBreak();
diff --git a/Engine/Source/Runtime/Core/Public/Linux/LinuxPlatformMisc.h b/Engine/Source/Runtime/Core/Public/Linux/LinuxPlatformMisc.h
--- a/Engine/Source/Runtime/Core/Public/Linux/LinuxPlatformMisc.h
+++ b/Engine/Source/Runtime/Core/Public/Linux/LinuxPlatformMisc.h
@@ -46,12 +46,12 @@ struct CORE_API FLinuxPlatformMisc : public FGenericPlatformMisc
 	}
 
 	/** Prompts for remote debugging if debugger is not attached. Regardless of result, breaks into debugger afterwards. Returns false for use in conditionals. */
-	static FORCEINLINE bool DebugBreakAndPromptForRemoteReturningFalse()
+	static FORCEINLINE bool DebugBreakAndPromptForRemoteReturningFalse(bool bIsEnsure = false)
 	{
 #if !UE_BUILD_SHIPPING
 		if (!IsDebuggerPresent())
 		{
-			PromptForRemoteDebugging(false);
+			PromptForRemoteDebugging(bIsEnsure);
 		}
 
 		DebugBreak();
diff --git a/Engine/Source/Runtime/Core/Public/Mac/MacPlatformMisc.h b/Engine/Source/Runtime/Core/Public/Mac/MacPlatformMisc.h
--- a/Engine/Source/Runtime/Core/Public/Mac/MacPlatformMisc.h
+++ b/Engine/Source/Runtime/Core/Public/Mac/MacPlatformMisc.h
@@ -63,12 +63,12 @@ struct CORE_API FMacPlatformMisc : public FGenericPlatformMisc
 	}
 
 	/** Prompts for remote debugging if debugger is not attached. Regardless of result, breaks into debugger afterwards. Returns false for use in conditionals. */
-	static FORCEINLINE bool DebugBreakAndPromptForRemoteReturningFalse()
+	static FORCEINLINE bool DebugBreakAndPromptForRemoteReturningFalse(bool bIsEnsure = false)
 	{
 #if !UE_BUILD_SHIPPING
 		if (!IsDebuggerPresent())
 		{
-			PromptForRemoteDebugging(false);
+			PromptForRemoteDebugging(bIsEnsure);
 		}
 
 		DebugBreak();
diff --git a/Engine/Source/Runtime/Core/Public/Misc/AssertionMacros.h b/Engine/Source/Runtime/Core/Public/Misc/AssertionMacros.h
--- a/Engine/Source/Runtime/Core/Public/Misc/AssertionMacros.h
+++ b/Engine/Source/Runtime/Core/Public/Misc/AssertionMacros.h
@@ -126,21 +126,21 @@
 			return Result;
 		}
 
-		FORCEINLINE bool OptionallyDebugBreakAndPromptForRemoteReturningFalse(bool bBreak)
+		FORCEINLINE bool OptionallyDebugBreakAndPromptForRemoteReturningFalse(bool bBreak, bool bIsEnsure = false)
 		{
 			if (bBreak)
 			{
-				FPlatformMisc::DebugBreakAndPromptForRemoteReturningFalse();
+				FPlatformMisc::DebugBreakAndPromptForRemoteReturningFalse(bIsEnsure);
 			}
 			return false;
 		}
 	}
 
-	#define ensure(         InExpression                ) ((InExpression) != 0 || FDebug::OptionallyLogFormattedEnsureMessageReturningFalse(true,                                          #InExpression, __FILE__, __LINE__, TEXT("")               ) || FPlatformMisc::DebugBreakAndPromptForRemoteReturningFalse())
-	#define ensureMsg(      InExpression, InMsg         ) ((InExpression) != 0 || FDebug::OptionallyLogFormattedEnsureMessageReturningFalse(true,                                          #InExpression, __FILE__, __LINE__, InMsg                  ) || FPlatformMisc::DebugBreakAndPromptForRemoteReturningFalse())
-	#define ensureMsgf(     InExpression, InFormat, ... ) ((InExpression) != 0 || FDebug::OptionallyLogFormattedEnsureMessageReturningFalse(true,                                          #InExpression, __FILE__, __LINE__, InFormat, ##__VA_ARGS__) || FPlatformMisc::DebugBreakAndPromptForRemoteReturningFalse())
-	#define ensureOnce(     InExpression                ) ((InExpression) != 0 || FDebug::OptionallyLogFormattedEnsureMessageReturningFalse(UE4Asserts_Private::TrueOnFirstCallOnly([]{}), #InExpression, __FILE__, __LINE__, TEXT("")               ) || UE4Asserts_Private::OptionallyDebugBreakAndPromptForRemoteReturningFalse(UE4Asserts_Private::TrueOnFirstCallOnly([]{})))
-	#define ensureOnceMsgf( InExpression, InFormat, ... ) ((InExpression) != 0 || FDebug::OptionallyLogFormattedEnsureMessageReturningFalse(UE4Asserts_Private::TrueOnFirstCallOnly([]{}), #InExpression, __FILE__, __LINE__, InFormat, ##__VA_ARGS__) || UE4Asserts_Private::OptionallyDebugBreakAndPromptForRemoteReturningFalse(UE4Asserts_Private::TrueOnFirstCallOnly([]{})))
+	#define ensure(         InExpression                ) ((InExpression) != 0 || FDebug::OptionallyLogFormattedEnsureMessageReturningFalse(true,                                          #InExpression, __FILE__, __LINE__, TEXT("")               ) || FPlatformMisc::DebugBreakAndPromptForRemoteReturningFalse(true))
+	#define ensureMsg(      InExpression, InMsg         ) ((InExpression) != 0 || FDebug::OptionallyLogFormattedEnsureMessageReturningFalse(true,                                          #InExpression, __FILE__, __LINE__, InMsg                  ) || FPlatformMisc::DebugBreakAndPromptForRemoteReturningFalse(true))
+	#define ensureMsgf(     InExpression, InFormat, ... ) ((InExpression) != 0 || FDebug::OptionallyLogFormattedEnsureMessageReturningFalse(true,                                          #InExpression, __FILE__, __LINE__, InFormat, ##__VA_ARGS__) || FPlatformMisc::DebugBreakAndPromptForRemoteReturningFalse(true))
+	#define ensureOnce(     InExpression                ) ((InExpression) != 0 || FDebug::OptionallyLogFormattedEnsureMessageReturningFalse(UE4Asserts_Private::TrueOnFirstCallOnly([]{}), #InExpression, __FILE__, __LINE__, TEXT("")               ) || UE4Asserts_Private::OptionallyDebugBreakAndPromptForRemoteReturningFalse(UE4Asserts_Private::TrueOnFirstCallOnly([]{}), true))
+	#define ensureOnceMsgf( InExpression, InFormat, ... ) ((InExpression) != 0 || FDebug::OptionallyLogFormattedEnsureMessageReturningFalse(UE4Asserts_Private::TrueOnFirstCallOnly([]{}), #InExpression, __FILE__, __LINE__, InFormat, ##__VA_ARGS__) || UE4Asserts_Private::OptionallyDebugBreakAndPromptForRemoteReturningFalse(UE4Asserts_Private::TrueOnFirstCallOnly([]{}), true))
 
 #else	// DO_CHECK
 
diff --git a/Engine/Source/Runtime/Core/Public/Windows/WindowsPlatformMisc.h b/Engine/Source/Runtime/Core/Public/Windows/WindowsPlatformMisc.h
--- a/Engine/Source/Runtime/Core/Public/Windows/WindowsPlatformMisc.h
+++ b/Engine/Source/Runtime/Core/Public/Windows/WindowsPlatformMisc.h
@@ -65,12 +65,12 @@ struct CORE_API FWindowsPlatformMisc
 	}
 
 	/** Prompts for remote debugging if debugger is not attached. Regardless of result, breaks into debugger afterwards. Returns false for use in conditionals. */
-	static FORCEINLINE bool DebugBreakAndPromptForRemoteReturningFalse()
+	static FORCEINLINE bool DebugBreakAndPromptForRemoteReturningFalse(bool bIsEnsure = false)
 	{
 #if !UE_BUILD_SHIPPING
 		if (!IsDebuggerPresent())
 		{
-			PromptForRemoteDebugging(false);
+			PromptForRemoteDebugging(bIsEnsure);
 		}
 
 		DebugBreak();
