Fixes to string operator commutativity.
Fix for incorrect comments.

diff --git a/Engine/Source/Runtime/Core/Public/Containers/UnrealString.h b/Engine/Source/Runtime/Core/Public/Containers/UnrealString.h
--- a/Engine/Source/Runtime/Core/Public/Containers/UnrealString.h
+++ b/Engine/Source/Runtime/Core/Public/Containers/UnrealString.h
@@ -758,7 +758,6 @@ public:
 		return *this += Str;
 	}
 
-
 	/**
 	 * Concatenate this path with given path ensuring the / character is used between them
 	 * 
@@ -772,174 +771,308 @@ public:
 
 	/**
 	 * Concatenate this path with given path ensuring the / character is used between them
-	 * 
-	 * @param Str path array of TCHAR to be concatenated onto the end of this
+	 *
+	 * @param Lhs Path to concatenate onto.
+	 * @param Rhs Path to concatenate.
 	 * @return new FString of the path
 	 */
-	FORCEINLINE FString operator/( const TCHAR* Str ) const
+	FORCEINLINE friend FString operator/(const FString& Lhs, const TCHAR* Rhs)
 	{
-		return FString( *this ) /= Str;
+		return FString(Lhs) /= Rhs;
 	}
 
 	/**
 	 * Concatenate this path with given path ensuring the / character is used between them
-	 * 
-	 * @param Str path FString to be concatenated onto the end of this
+	 *
+	 * @param Lhs Path to concatenate onto.
+	 * @param Rhs Path to concatenate.
 	 * @return new FString of the path
 	 */
-	FORCEINLINE FString operator/( const FString& Str ) const
+	FORCEINLINE friend FString operator/(FString&& Lhs, const TCHAR* Rhs)
 	{
-		return operator/( *Str );
+		return FString(MoveTemp(Lhs)) /= Rhs;
 	}
 
 	/**
-	 * Lexicographically test whether this string is <= the Other given string
-	 * 
-	 * @param Other array of CharType to compare against
-	 * @return true if length of this string is lexicographically <= the other, otherwise false
+	 * Concatenate this path with given path ensuring the / character is used between them
+	 *
+	 * @param Lhs Path to concatenate onto.
+	 * @param Rhs Path to concatenate.
+	 * @return new FString of the path
+	 */
+	FORCEINLINE friend FString operator/(const FString& Lhs, const FString& Rhs)
+	{
+		return FString(Lhs) /= *Rhs;
+	}
+
+	/**
+	 * Concatenate this path with given path ensuring the / character is used between them
+	 *
+	 * @param Lhs Path to concatenate onto.
+	 * @param Rhs Path to concatenate.
+	 * @return new FString of the path
+	 */
+	FORCEINLINE friend FString operator/(FString&& Lhs, const FString& Rhs)
+	{
+		return FString(MoveTemp(Lhs)) /= *Rhs;
+	}
+
+	/**
+	 * Concatenate this path with given path ensuring the / character is used between them
+	 *
+	 * @param Lhs Path to concatenate onto.
+	 * @param Rhs Path to concatenate.
+	 * @return new FString of the path
+	 */
+	FORCEINLINE friend FString operator/(const TCHAR* Lhs, const FString& Rhs)
+	{
+		return FString(Lhs) /= *Rhs;
+	}
+
+	/**
+	 * Lexicographically test whether the left string is <= the right string
+	 *
+	 * @param Lhs String to compare against.
+	 * @param Rhs String to compare against.
+	 * @return true if the left string is lexicographically <= the right string, otherwise false
+	 * @note case insensitive
+	 */
+	FORCEINLINE friend bool operator<=(const FString& Lhs, const FString& Rhs)
+	{
+		return FPlatformString::Stricmp(*Lhs, *Rhs) <= 0;
+	}
+
+	/**
+	 * Lexicographically test whether the left string is <= the right string
+	 *
+	 * @param Lhs String to compare against.
+	 * @param Rhs String to compare against.
+	 * @return true if the left string is lexicographically <= the right string, otherwise false
 	 * @note case insensitive
 	 */
 	template <typename CharType>
-	FORCEINLINE bool operator<=(const CharType* Other) const
+	FORCEINLINE friend bool operator<=(const FString& Lhs, const CharType* Rhs)
 	{
-		return !(FPlatformString::Stricmp(**this, Other) > 0);
+		return FPlatformString::Stricmp(*Lhs, Rhs) <= 0;
 	}
 
 	/**
-	 * Lexicographically test whether this string is < the Other given string
-	 * 
-	 * @param Other array of CharType to compare against
-	 * @return true if length of this string is lexicographically < the other, otherwise false
+	 * Lexicographically test whether the left string is <= the right string
+	 *
+	 * @param Lhs String to compare against.
+	 * @param Rhs String to compare against.
+	 * @return true if the left string is lexicographically <= the right string, otherwise false
 	 * @note case insensitive
 	 */
 	template <typename CharType>
-	FORCEINLINE bool operator<(const CharType* Other) const
+	FORCEINLINE friend bool operator<=(const CharType* Lhs, const FString& Rhs)
 	{
-		return FPlatformString::Stricmp(**this, Other) < 0;
+		return FPlatformString::Stricmp(Lhs, *Rhs) <= 0;
 	}
 
 	/**
-	 * Lexicographically test whether this string is < the Other given string
-	 * 
-	 * @param Other FString to compare against
-	 * @return true if length of this string is lexicographically < the other, otherwise false
+	 * Lexicographically test whether the left string is < the right string
+	 *
+	 * @param Lhs String to compare against.
+	 * @param Rhs String to compare against.
+	 * @return true if the left string is lexicographically < the right string, otherwise false
 	 * @note case insensitive
 	 */
-	FORCEINLINE bool operator<( const FString& Other ) const
+	FORCEINLINE friend bool operator<(const FString& Lhs, const FString& Rhs)
 	{
-		return FCString::Stricmp( **this, *Other ) < 0;
+		return FPlatformString::Stricmp(*Lhs, *Rhs) < 0;
 	}
 
 	/**
-	 * Lexicographically test whether this string is >= the Other given string
-	 * 
-	 * @param Other array of CharType to compare against
-	 * @return true if length of this string is lexicographically >= the other, otherwise false
+	 * Lexicographically test whether the left string is < the right string
+	 *
+	 * @param Lhs String to compare against.
+	 * @param Rhs String to compare against.
+	 * @return true if the left string is lexicographically < the right string, otherwise false
 	 * @note case insensitive
 	 */
 	template <typename CharType>
-	FORCEINLINE bool operator>=(const CharType* Other) const
+	FORCEINLINE friend bool operator<(const FString& Lhs, const CharType* Rhs)
 	{
-		return !(FPlatformString::Stricmp(**this, Other) < 0);
+		return FPlatformString::Stricmp(*Lhs, Rhs) < 0;
 	}
 
 	/**
-	 * Lexicographically test whether this string is > the Other given string
-	 * 
-	 * @param Other array of CharType to compare against
-	 * @return true if length of this string is lexicographically > the other, otherwise false
+	 * Lexicographically test whether the left string is < the right string
+	 *
+	 * @param Lhs String to compare against.
+	 * @param Rhs String to compare against.
+	 * @return true if the left string is lexicographically < the right string, otherwise false
 	 * @note case insensitive
 	 */
 	template <typename CharType>
-	FORCEINLINE bool operator>(const CharType* Other) const
+	FORCEINLINE friend bool operator<(const CharType* Lhs, const FString& Rhs)
 	{
-		return FPlatformString::Stricmp(**this, Other) > 0;
+		return FPlatformString::Stricmp(Lhs, *Rhs) < 0;
 	}
 
 	/**
-	 * Lexicographically test whether this string is > the Other given string
-	 * 
-	 * @param Other FString to compare against
-	 * @return true if length of this string is lexicographically > the other, otherwise false
+	 * Lexicographically test whether the left string is >= the right string
+	 *
+	 * @param Lhs String to compare against.
+	 * @param Rhs String to compare against.
+	 * @return true if the left string is lexicographically >= the right string, otherwise false
 	 * @note case insensitive
 	 */
-	FORCEINLINE bool operator>( const FString& Other ) const
+	FORCEINLINE friend bool operator>=(const FString& Lhs, const FString& Rhs)
 	{
-		return FCString::Stricmp( **this, *Other ) > 0;
+		return FPlatformString::Stricmp(*Lhs, *Rhs) >= 0;
 	}
 
 	/**
-	 * Lexicographically test whether this string is equivalent to the Other given string
-	 * 
-	 * @param Other array of CharType to compare against
-	 * @return true if length of this string is lexicographically equivalent to the other, otherwise false
+	 * Lexicographically test whether the left string is >= the right string
+	 *
+	 * @param Lhs String to compare against.
+	 * @param Rhs String to compare against.
+	 * @return true if the left string is lexicographically >= the right string, otherwise false
 	 * @note case insensitive
 	 */
 	template <typename CharType>
-	FORCEINLINE bool operator==(const CharType* Other) const
+	FORCEINLINE friend bool operator>=(const FString& Lhs, const CharType* Rhs)
 	{
-		return FPlatformString::Stricmp(**this, Other) == 0;
+		return FPlatformString::Stricmp(*Lhs, Rhs) >= 0;
 	}
 
 	/**
-	 * Lexicographically test whether this string is equivalent to the Other given string
-	 * 
-	 * @param Other FString to compare against
-	 * @return true if length of this string is lexicographically equivalent to the other, otherwise false
+	 * Lexicographically test whether the left string is >= the right string
+	 *
+	 * @param Lhs String to compare against.
+	 * @param Rhs String to compare against.
+	 * @return true if the left string is lexicographically >= the right string, otherwise false
 	 * @note case insensitive
 	 */
-	FORCEINLINE bool operator==( const FString& Other ) const
+	template <typename CharType>
+	FORCEINLINE friend bool operator>=(const CharType* Lhs, const FString& Rhs)
 	{
-		return FCString::Stricmp( **this, *Other )==0;
+		return FPlatformString::Stricmp(Lhs, *Rhs) >= 0;
 	}
 
 	/**
-	 * Lexicographically test whether this string is not equivalent to the Other given string
-	 * 
-	 * @param Other array of CharType to compare against
-	 * @return true if length of this string is lexicographically not equivalent to the other, otherwise false
+	 * Lexicographically test whether the left string is > the right string
+	 *
+	 * @param Lhs String to compare against.
+	 * @param Rhs String to compare against.
+	 * @return true if the left string is lexicographically > the right string, otherwise false
+	 * @note case insensitive
+	 */
+	FORCEINLINE friend bool operator>(const FString& Lhs, const FString& Rhs)
+	{
+		return FPlatformString::Stricmp(*Lhs, *Rhs) > 0;
+	}
+
+	/**
+	 * Lexicographically test whether the left string is > the right string
+	 *
+	 * @param Lhs String to compare against.
+	 * @param Rhs String to compare against.
+	 * @return true if the left string is lexicographically > the right string, otherwise false
 	 * @note case insensitive
 	 */
 	template <typename CharType>
-	FORCEINLINE bool operator!=(const CharType* Other) const
+	FORCEINLINE friend bool operator>(const FString& Lhs, const CharType* Rhs)
 	{
-		return FPlatformString::Stricmp(**this, Other) != 0;
+		return FPlatformString::Stricmp(*Lhs, Rhs) > 0;
 	}
 
 	/**
-	 * Lexicographically test whether this string is not equivalent to the Other given string
-	 * 
-	 * @param Other FString to compare against
-	 * @return true if length of this string is lexicographically not equivalent to the other, otherwise false
+	 * Lexicographically test whether the left string is > the right string
+	 *
+	 * @param Lhs String to compare against.
+	 * @param Rhs String to compare against.
+	 * @return true if the left string is lexicographically > the right string, otherwise false
 	 * @note case insensitive
 	 */
-	FORCEINLINE bool operator!=( const FString& Other ) const
+	template <typename CharType>
+	FORCEINLINE friend bool operator>(const CharType* Lhs, const FString& Rhs)
 	{
-		return FCString::Stricmp( **this, *Other )!=0;
+		return FPlatformString::Stricmp(Lhs, *Rhs) > 0;
 	}
 
 	/**
-	 * Lexicographically test whether this string is <= the Other given string
-	 * 
-	 * @param Other FString to compare against
-	 * @return true if length of this string is lexicographically <= the other, otherwise false
+	 * Lexicographically test whether the left string is == the right string
+	 *
+	 * @param Lhs String to compare against.
+	 * @param Rhs String to compare against.
+	 * @return true if the left string is lexicographically == the right string, otherwise false
 	 * @note case insensitive
 	 */
-	FORCEINLINE bool operator<=(const FString& Str) const
+	FORCEINLINE friend bool operator==(const FString& Lhs, const FString& Rhs)
 	{
-		return !(FCString::Stricmp(**this, *Str) > 0);
+		return FPlatformString::Stricmp(*Lhs, *Rhs) == 0;
 	}
 
 	/**
-	 * Lexicographically test whether this string is >= the Other given string
-	 * 
-	 * @param Other FString to compare against
-	 * @return true if length of this string is lexicographically >= the other, otherwise false
+	 * Lexicographically test whether the left string is == the right string
+	 *
+	 * @param Lhs String to compare against.
+	 * @param Rhs String to compare against.
+	 * @return true if the left string is lexicographically == the right string, otherwise false
 	 * @note case insensitive
 	 */
-	FORCEINLINE bool operator>=(const FString& Str) const
+	template <typename CharType>
+	FORCEINLINE friend bool operator==(const FString& Lhs, const CharType* Rhs)
+	{
+		return FPlatformString::Stricmp(*Lhs, Rhs) == 0;
+	}
+
+	/**
+	 * Lexicographically test whether the left string is == the right string
+	 *
+	 * @param Lhs String to compare against.
+	 * @param Rhs String to compare against.
+	 * @return true if the left string is lexicographically == the right string, otherwise false
+	 * @note case insensitive
+	 */
+	template <typename CharType>
+	FORCEINLINE friend bool operator==(const CharType* Lhs, const FString& Rhs)
+	{
+		return FPlatformString::Stricmp(Lhs, *Rhs) == 0;
+	}
+
+	/**
+	 * Lexicographically test whether the left string is != the right string
+	 *
+	 * @param Lhs String to compare against.
+	 * @param Rhs String to compare against.
+	 * @return true if the left string is lexicographically != the right string, otherwise false
+	 * @note case insensitive
+	 */
+	FORCEINLINE friend bool operator!=(const FString& Lhs, const FString& Rhs)
+	{
+		return FPlatformString::Stricmp(*Lhs, *Rhs) != 0;
+	}
+
+	/**
+	 * Lexicographically test whether the left string is != the right string
+	 *
+	 * @param Lhs String to compare against.
+	 * @param Rhs String to compare against.
+	 * @return true if the left string is lexicographically != the right string, otherwise false
+	 * @note case insensitive
+	 */
+	template <typename CharType>
+	FORCEINLINE friend bool operator!=(const FString& Lhs, const CharType* Rhs)
+	{
+		return FPlatformString::Stricmp(*Lhs, Rhs) != 0;
+	}
+
+	/**
+	 * Lexicographically test whether the left string is != the right string
+	 *
+	 * @param Lhs String to compare against.
+	 * @param Rhs String to compare against.
+	 * @return true if the left string is lexicographically != the right string, otherwise false
+	 * @note case insensitive
+	 */
+	template <typename CharType>
+	FORCEINLINE friend bool operator!=(const CharType* Lhs, const FString& Rhs)
 	{
-		return !(FCString::Stricmp(**this, *Str) < 0);
+		return FPlatformString::Stricmp(Lhs, *Rhs) != 0;
 	}
 
 	/** Get the length of the sting, excluding terminating character */
@@ -1040,7 +1173,7 @@ public:
 	 * Searches the string for a character
 	 *
 	 * @param InChar the character to search for
-	 * @param Index out the position the character was found at, not updated if return is false
+	 * @param Index out the position the character was found at, INDEX_NONE if return is false
 	 * @return true if character was found in this string, otherwise false
 	 */
 	FORCEINLINE bool FindChar( TCHAR InChar, int32& Index ) const
@@ -1052,7 +1185,7 @@ public:
 	 * Searches the string for the last occurrence of a character
 	 *
 	 * @param InChar the character to search for
-	 * @param Index out the position the character was found at, not updated if return is false
+	 * @param Index out the position the character was found at, INDEX_NONE if return is false
 	 * @return true if character was found in this string, otherwise false
 	 */
 	FORCEINLINE bool FindLastChar( TCHAR InChar, int32& Index ) const
