CIS fix, adding missing IOS file.

diff --git a/Engine/Source/Runtime/Core/Public/IOS/IOSPlatformMisc.h b/Engine/Source/Runtime/Core/Public/IOS/IOSPlatformMisc.h
--- a/Engine/Source/Runtime/Core/Public/IOS/IOSPlatformMisc.h
+++ b/Engine/Source/Runtime/Core/Public/IOS/IOSPlatformMisc.h
@@ -62,12 +62,12 @@ struct CORE_API FIOSPlatformMisc : public FGenericPlatformMisc
 	}
 
 	/** Prompts for remote debugging if debugger is not attached. Regardless of result, breaks into debugger afterwards. Returns false for use in conditionals. */
-	FORCEINLINE static bool DebugBreakAndPromptForRemoteReturningFalse()
+	FORCEINLINE static bool DebugBreakAndPromptForRemoteReturningFalse(bool bIsEnsure = false)
 	{
 #if !UE_BUILD_SHIPPING
 		if (!IsDebuggerPresent())
 		{
-			PromptForRemoteDebugging(false);
+			PromptForRemoteDebugging(bIsEnsure);
 		}
 
 		DebugBreak();
