Add memory stats on exception for generic platform.

diff --git a/Engine/Source/Runtime/Core/Private/GenericPlatform/GenericPlatformMisc.cpp b/Engine/Source/Runtime/Core/Private/GenericPlatform/GenericPlatformMisc.cpp
--- a/Engine/Source/Runtime/Core/Private/GenericPlatform/GenericPlatformMisc.cpp
+++ b/Engine/Source/Runtime/Core/Private/GenericPlatform/GenericPlatformMisc.cpp
@@ -264,6 +264,20 @@ void FGenericPlatformMisc::HandleIOFailure( const TCHAR* Filename )
 	UE_LOG(LogGenericPlatformMisc, Fatal,TEXT("I/O failure operating on '%s'"), Filename ? Filename : TEXT("Unknown file"));
 }
 
+void FGenericPlatformMisc::RaiseException(uint32 ExceptionCode)
+{
+	/** This is the last place to gather memory stats before exception. */
+	FGenericCrashContext::CrashMemoryStats = FPlatformMemory::GetStats();
+
+#if HACK_HEADER_GENERATOR && !PLATFORM_EXCEPTIONS_DISABLED
+	// We want Unreal Header Tool to throw an exception but in normal runtime code 
+	// we don't support exception handling
+	throw(ExceptionCode);
+#else	
+	*((uint32*)3) = ExceptionCode;
+#endif
+}
+
 bool FGenericPlatformMisc::SetStoredValue(const FString& InStoreId, const FString& InSectionName, const FString& InKeyName, const FString& InValue)
 {
 	check(!InStoreId.IsEmpty());
diff --git a/Engine/Source/Runtime/Core/Public/GenericPlatform/GenericPlatformMisc.h b/Engine/Source/Runtime/Core/Public/GenericPlatform/GenericPlatformMisc.h
--- a/Engine/Source/Runtime/Core/Public/GenericPlatform/GenericPlatformMisc.h
+++ b/Engine/Source/Runtime/Core/Public/GenericPlatform/GenericPlatformMisc.h
@@ -371,16 +371,7 @@ struct CORE_API FGenericPlatformMisc
 		return 0;
 	}
 
-	FORCEINLINE static void RaiseException( uint32 ExceptionCode )
-	{
-#if HACK_HEADER_GENERATOR && !PLATFORM_EXCEPTIONS_DISABLED
-		// We want Unreal Header Tool to throw an exception but in normal runtime code 
-		// we don't support exception handling
-		throw( ExceptionCode );
-#else	
-		*((uint32*)3) = ExceptionCode;
-#endif
-	}
+	static void RaiseException( uint32 ExceptionCode );
 
 protected:
 
