UE-6778/Inconsistent documentation on how to use AsyncTasks

diff --git a/Engine/Source/Runtime/Core/Public/Async/AsyncWork.h b/Engine/Source/Runtime/Core/Public/Async/AsyncWork.h
--- a/Engine/Source/Runtime/Core/Public/Async/AsyncWork.h
+++ b/Engine/Source/Runtime/Core/Public/Async/AsyncWork.h
@@ -35,13 +35,15 @@
 		}
 	};
 
-	start an example job
-
-	(new FAutoDeleteAsyncTask<ExampleAutoDeleteAsyncTask>(5)->StartBackgroundTask();
-
-	do an example job now, on this thread
+	
+	void Example()
+	{
+		// start an example job
+		(new FAutoDeleteAsyncTask<ExampleAutoDeleteAsyncTask>(5)->StartBackgroundTask();
 
-	(new FAutoDeleteAsyncTask<ExampleAutoDeleteAsyncTask>(5)->StartSynchronousTask();
+		// do an example job now, on this thread
+		(new FAutoDeleteAsyncTask<ExampleAutoDeleteAsyncTask>(5)->StartSynchronousTask();
+	}
 
 **/
 template<typename TTask>
@@ -216,31 +218,31 @@ public:
 		}
 	};
 
-	start an example job
-
-	FAutoDeleteAsyncTask<ExampleAutoDeleteAsyncTask> MyTask(5);
-	
-	MyTask.StartBackgroundTask();
-
-	-- or -- 
-
-	MyTask.StartSynchronousTask();
+	void Example()
+	{
 
-	to just do it now on this thread
+		//start an example job
 
-	Check if the task is done:
+		FAsyncTask<ExampleAsyncTask>* MyTask = new FAsyncTask<ExampleAsyncTask>( 5 );
+		MyTask->StartBackgroundTask();
 
-	if (MyTask.IsDone())
-	{
-	}
+		//--or --
 
-	Spinning on IsDone is not acceptable (see EnsureCompletion), but it is ok to check once a frame.
+		MyTask->StartSynchronousTask();
 
-	Ensure the task is done, doing the task on the current thread if it has not been started, waiting until completion in all cases.
+		//to just do it now on this thread
+		//Check if the task is done :
 
-	Task.EnsureCompletion();
+		if (MyTask->IsDone())
+		{
+		}
 
+		//Spinning on IsDone is not acceptable( see EnsureCompletion ), but it is ok to check once a frame.
+		//Ensure the task is done, doing the task on the current thread if it has not been started, waiting until completion in all cases.
 
+		MyTask->EnsureCompletion();
+		delete Task;
+	}
 **/
 template<typename TTask>
 class FAsyncTask
@@ -378,7 +380,7 @@ class FAsyncTask
 	}
 
 	/** 
-	* Internal call to intialize internal variables
+	* Internal call to initialize internal variables
 	**/
 	void Init()
 	{
