UE-18173/Crash report client submits user names for github builds (Only internal builds will send the user name)

diff --git a/Engine/Source/Programs/CrashReporter/CrashReportClient/Private/CrashDescription.cpp b/Engine/Source/Programs/CrashReporter/CrashReportClient/Private/CrashDescription.cpp
--- a/Engine/Source/Programs/CrashReporter/CrashReportClient/Private/CrashDescription.cpp
+++ b/Engine/Source/Programs/CrashReporter/CrashReportClient/Private/CrashDescription.cpp
@@ -217,19 +217,22 @@ void FCrashDescription::UpdateIDs()
 	{
 		// The Epic ID can be looked up from this ID.
 		EpicAccountId = FPlatformMisc::GetEpicAccountId();
-
-		// Don't add real user name if running from the launcher.
-		const bool bAddUserName = FEngineBuildSettings::IsInternalBuild() || FEngineBuildSettings::IsPerforceBuild() || FEngineBuildSettings::IsSourceDistribution();
-		if (bAddUserName)
-		{
-			// Remove periods from user names to match AutoReporter user names
-			// The name prefix is read by CrashRepository.AddNewCrash in the website code
-			UserName = FString( FPlatformProcess::UserName() ).Replace( TEXT( "." ), TEXT( "" ) );
-		}
 	}
 	else
 	{
 		EpicAccountId.Empty();
+	}
+
+	// Add real user name only for internal builds.
+	const bool bSendUserName = FEngineBuildSettings::IsInternalBuild();
+	if (bSendUserName)
+	{
+		// Remove periods from user names to match AutoReporter user names
+		// The name prefix is read by CrashRepository.AddNewCrash in the website code
+		UserName = FString( FPlatformProcess::UserName() ).Replace( TEXT( "." ), TEXT( "" ) );
+	}
+	else
+	{
 		UserName.Empty();
 	}
 
diff --git a/Engine/Source/Programs/CrashReporter/CrashReportClient/Private/GenericErrorReport.cpp b/Engine/Source/Programs/CrashReporter/CrashReportClient/Private/GenericErrorReport.cpp
--- a/Engine/Source/Programs/CrashReporter/CrashReportClient/Private/GenericErrorReport.cpp
+++ b/Engine/Source/Programs/CrashReporter/CrashReportClient/Private/GenericErrorReport.cpp
@@ -110,27 +110,22 @@ bool FGenericErrorReport::SetUserComment(const FText& UserComment, bool bAllowTo
 		DynamicSignaturesNode->AppendChildNode(TEXT("Parameter3"), UserComment.ToString());
 	}
 	
-	FString MachineIDandUserID;
+
 	// Set global user name ID: will be added to the report
 	extern FCrashDescription& GetCrashDescription();
 
-	MachineIDandUserID = FString::Printf( TEXT( "!MachineId:%s!EpicAccountId:%s" ), *GetCrashDescription().MachineId, *GetCrashDescription().EpicAccountId );
-
-	const bool bSendName = FEngineBuildSettings::IsInternalBuild() || FEngineBuildSettings::IsPerforceBuild() || FEngineBuildSettings::IsSourceDistribution();
-	if (bSendName)
-	{
-		MachineIDandUserID += FString::Printf( TEXT( "!Name:%s" ), *GetCrashDescription().UserName );
-	}
+	// @see FCrashDescription::UpdateIDs
+	const FString EpicMachindAndUserNameIDs = FString::Printf( TEXT( "!MachineId:%s!EpicAccountId:%s!Name:%s" ), *GetCrashDescription().MachineId, *GetCrashDescription().EpicAccountId, *GetCrashDescription().UserName );
 
 	// Add or update a user ID.
 	FXmlNode* Parameter4Node = DynamicSignaturesNode->FindChildNode(TEXT("Parameter4"));
 	if( Parameter4Node )
 	{
-		Parameter4Node->SetContent(MachineIDandUserID);
+		Parameter4Node->SetContent(EpicMachindAndUserNameIDs);
 	}
 	else
 	{
-		DynamicSignaturesNode->AppendChildNode(TEXT("Parameter4"), MachineIDandUserID);
+		DynamicSignaturesNode->AppendChildNode(TEXT("Parameter4"), EpicMachindAndUserNameIDs);
 	}
 
 	// Add or update bAllowToBeContacted
diff --git a/Engine/Source/Runtime/Core/Private/GenericPlatform/GenericPlatformContext.cpp b/Engine/Source/Runtime/Core/Private/GenericPlatform/GenericPlatformContext.cpp
--- a/Engine/Source/Runtime/Core/Private/GenericPlatform/GenericPlatformContext.cpp
+++ b/Engine/Source/Runtime/Core/Private/GenericPlatform/GenericPlatformContext.cpp
@@ -117,10 +117,10 @@ void FGenericCrashContext::SerializeContentToBuffer()
 
 	AddCrashProperty( TEXT( "IsUE4Release" ), (int32)NCachedCrashContextProperties::bIsUE4Release);
 
-	const bool bSetUserName = NCachedCrashContextProperties::bIsInternalBuild || NCachedCrashContextProperties::bIsPerforceBuild || NCachedCrashContextProperties::bIsSourceDistribution;
 	// Remove periods from user names to match AutoReporter user names
 	// The name prefix is read by CrashRepository.AddNewCrash in the website code
-	AddCrashProperty( TEXT( "UserName" ), bSetUserName ? *NCachedCrashContextProperties::UserName.Replace( TEXT( "." ), TEXT( "" ) ) : TEXT( "" ) );
+	const bool bSendUserName = NCachedCrashContextProperties::bIsInternalBuild;
+	AddCrashProperty( TEXT( "UserName" ), bSendUserName ? *NCachedCrashContextProperties::UserName.Replace( TEXT( "." ), TEXT( "" ) ) : TEXT( "" ) );
 
 	AddCrashProperty(TEXT("BaseDir"), *NCachedCrashContextProperties::BaseDir);
 	AddCrashProperty(TEXT("RootDir"), *NCachedCrashContextProperties::RootDir);
@@ -132,13 +132,6 @@ void FGenericCrashContext::SerializeContentToBuffer()
 	AddCrashProperty( TEXT( "UserDescription" ), TEXT( "" ) );
 	AddCrashProperty( TEXT( "ErrorMessage" ), (const TCHAR*)GErrorMessage ); // GErrorMessage may be broken.
 
-	// @TODO yrx 2014-10-08 Move to the crash report client.
-	/*if( CanUseUnsafeAPI() )
-	{
-		AddCrashProperty( TEXT( "TimeOfCrash" ), FDateTime::UtcNow().GetTicks() );
-	}
-	*/
-
 	// Add misc stats.
 	AddCrashProperty(TEXT("MiscNumberOfCores"), NCachedCrashContextProperties::NumberOfCores);
 	AddCrashProperty(TEXT("MiscNumberOfCoresIncludingHyperthreads"), NCachedCrashContextProperties::NumberOfCoresIncludingHyperthreads);
