Fixed non-unity compilaton

diff --git a/Engine/Source/Runtime/Core/Private/Windows/WindowsPlatformMisc.cpp b/Engine/Source/Runtime/Core/Private/Windows/WindowsPlatformMisc.cpp
--- a/Engine/Source/Runtime/Core/Private/Windows/WindowsPlatformMisc.cpp
+++ b/Engine/Source/Runtime/Core/Private/Windows/WindowsPlatformMisc.cpp
@@ -5,6 +5,7 @@
 #include "SecureHash.h"
 #include "WindowsApplication.h"
 #include "EngineVersion.h"
+#include "WindowsPlatformCrashContext.h"
 
 #include "GenericPlatformChunkInstall.h"
 
