Fix for FSphere::TransformBy(FTransform) which wasn't multiplying W by the scale.

diff --git a/Engine/Source/Runtime/Core/Private/Math/Sphere.cpp b/Engine/Source/Runtime/Core/Private/Math/Sphere.cpp
--- a/Engine/Source/Runtime/Core/Private/Math/Sphere.cpp
+++ b/Engine/Source/Runtime/Core/Private/Math/Sphere.cpp
@@ -59,7 +59,7 @@ FSphere FSphere::TransformBy(const FTransform& M) const
 	FSphere	Result;
 
 	Result.Center = M.TransformPosition(this->Center);
-	Result.W = M.GetMaximumAxisScale();
+	Result.W *= M.GetMaximumAxisScale();
 
 	return Result;
 }
