Better comment

diff --git a/Engine/Source/Runtime/Core/Public/Stats/Stats2.h b/Engine/Source/Runtime/Core/Public/Stats/Stats2.h
--- a/Engine/Source/Runtime/Core/Public/Stats/Stats2.h
+++ b/Engine/Source/Runtime/Core/Public/Stats/Stats2.h
@@ -66,7 +66,7 @@ struct CORE_API FStats
 
 	/**
 	* @return true, if the command line has the StatsForCommandlets token which enables stats in the commandlets.
-	* !!!CAUTION!!! You need to manually advanced frame in order to maintain the data integrity and not to leak the memory.
+	* !!!CAUTION!!! You need to manually advance stats frame in order to maintain the data integrity and not to leak the memory.
 	*/
 	static bool HasStatsForCommandletsToken();
 
