Compilation fix

diff --git a/Engine/Source/Runtime/Core/Private/Stats/Stats2.cpp b/Engine/Source/Runtime/Core/Private/Stats/Stats2.cpp
--- a/Engine/Source/Runtime/Core/Private/Stats/Stats2.cpp
+++ b/Engine/Source/Runtime/Core/Private/Stats/Stats2.cpp
@@ -199,7 +199,7 @@ void FStats::TickCommandletStats()
 {
 	if (HasStatsForCommandletsToken())
 	{
-		FThreadStats* ThreadStats = FThreadStats::GetThreadStats();
+		//FThreadStats* ThreadStats = FThreadStats::GetThreadStats();
 		//check( ThreadStats->ScopeCount == 0 && TEXT( "FStats::TickCommandletStats must be called outside any scope counters" ) );
 
 		FTaskGraphInterface::Get().ProcessThreadUntilIdle( ENamedThreads::GameThread );
