Merging a couple of changes required to build for iOS with Xcode 7

diff --git a/Engine/Source/Runtime/Core/Public/IOS/IOSPlatformCompilerPreSetup.h b/Engine/Source/Runtime/Core/Public/IOS/IOSPlatformCompilerPreSetup.h
--- a/Engine/Source/Runtime/Core/Public/IOS/IOSPlatformCompilerPreSetup.h
+++ b/Engine/Source/Runtime/Core/Public/IOS/IOSPlatformCompilerPreSetup.h
@@ -102,4 +102,7 @@
 // Rationale: we don't want to suppress those as there are plans to address them (e.g. UE-12341), but breaking builds due to these warnings is very expensive
 // since they cannot be caught by all compilers that we support. They are deemed to be relatively safe to be ignored, at least until all SDKs/toolchains start supporting them.
 #pragma clang diagnostic warning "-Wreorder"
-#pragma clang diagnostic warning "-Wparentheses-equality"
\ No newline at end of file
+#pragma clang diagnostic warning "-Wparentheses-equality"
+#pragma clang diagnostic ignored "-Winconsistent-missing-override"
+#pragma clang diagnostic ignored "-Wundefined-bool-conversion"
+#pragma clang diagnostic ignored "-Wunused-local-typedef"
diff --git a/Engine/Source/Runtime/Online/IOS/OnlineSubsystemIOS/Private/OnlineSessionInterfaceIOS.cpp b/Engine/Source/Runtime/Online/IOS/OnlineSubsystemIOS/Private/OnlineSessionInterfaceIOS.cpp
--- a/Engine/Source/Runtime/Online/IOS/OnlineSubsystemIOS/Private/OnlineSessionInterfaceIOS.cpp
+++ b/Engine/Source/Runtime/Online/IOS/OnlineSubsystemIOS/Private/OnlineSessionInterfaceIOS.cpp
@@ -112,7 +112,7 @@
 - (void)initSessionWithName:(NSString*) sessionName
 {
     UE_LOG(LogOnline, Display, TEXT("- (void)initSessionWithName:(NSString*) sessionName"));
-    self.PeerID = [[MCPeerID alloc] initWithDisplayName:nil];
+    self.PeerID = [[MCPeerID alloc] initWithDisplayName:@""];
     self.Session = [[MCSession alloc] initWithPeer:self.PeerID];
     self.Session.delegate = self;
 }
