Stats - Better stats performance for commandlet stats

diff --git a/Engine/Source/Runtime/Core/Private/Stats/Stats2.cpp b/Engine/Source/Runtime/Core/Private/Stats/Stats2.cpp
--- a/Engine/Source/Runtime/Core/Private/Stats/Stats2.cpp
+++ b/Engine/Source/Runtime/Core/Private/Stats/Stats2.cpp
@@ -211,7 +211,8 @@ void FStats::TickCommandletStats()
 
 bool FStats::HasStatsForCommandletsToken()
 {
-	return FParse::Param( FCommandLine::Get(), TEXT( "StatsForCommandlets" ) );
+	static bool bHasStatsForCommandletsToken = FParse::Param( FCommandLine::Get(), TEXT( "StatsForCommandlets" ) );
+	return bHasStatsForCommandletsToken;
 }
 
 /* Todo
@@ -843,8 +844,9 @@ public:
 		}
 		else
 		{
-			// For regular stats we won't process more than every 5ms or every 16 packets.
-			bShouldProcess = bReadyToProcess && (FPlatformTime::Seconds() - LastTime > 0.005f || IncomingData.Packets.Num() > MaxIncomingPackets);
+			// For regular stats we won't process more than every 5ms or every 16 packets. 
+			// Commandlet stats are flushed as soon as.
+			bShouldProcess = bReadyToProcess && (FPlatformTime::Seconds() - LastTime > 0.005f || IncomingData.Packets.Num() > MaxIncomingPackets || FStats::HasStatsForCommandletsToken());
 		}
 
 		if( bShouldProcess )
