PR #1122: Linux controller fixes (UE-15024).

- Contributed by [REDACTED_PERSON],  with changes.
- Axis events aren't sent immediately as they are received anymore (which could happen multiple times per frame), but are instead collected and sent once.

diff --git a/Engine/Source/Runtime/Core/Private/Linux/LinuxApplication.cpp b/Engine/Source/Runtime/Core/Private/Linux/LinuxApplication.cpp
--- a/Engine/Source/Runtime/Core/Private/Linux/LinuxApplication.cpp
+++ b/Engine/Source/Runtime/Core/Private/Linux/LinuxApplication.cpp
@@ -47,15 +47,23 @@ FLinuxApplication* FLinuxApplication::CreateLinuxApplication()
 
 	LinuxApplication = new FLinuxApplication();
 
-	SDLControllerState* ControllerState = LinuxApplication->ControllerStates;
+	int32 ControllerIndex = 0;
+
 	for (int i = 0; i < SDL_NumJoysticks(); ++i)
 	{
 		if (SDL_IsGameController(i))
 		{
-			ControllerState->controller = SDL_GameControllerOpen(i);
-			if (ControllerState++->controller == nullptr)
+			auto Controller = SDL_GameControllerOpen(i);
+			if (Controller == nullptr)
+			{
+				UE_LOG(LogLoad, Warning, TEXT("Could not open gamecontroller %i: %s\n"), i, ANSI_TO_TCHAR(SDL_GetError()) );
+			}
+			else
 			{
-				UE_LOG(LogLoad, Warning, TEXT("Could not open gamecontroller %i: %s\n"), i, UTF8_TO_TCHAR(SDL_GetError()));
+				SDL_JoystickID Id = SDL_JoystickInstanceID(SDL_GameControllerGetJoystick(Controller));
+				LinuxApplication->ControllerStates.Add(Id);
+				LinuxApplication->ControllerStates[Id].Controller = Controller;
+				LinuxApplication->ControllerStates[Id].ControllerIndex = ControllerIndex++;
 			}
 		}
 	}
@@ -78,8 +86,6 @@ FLinuxApplication::FLinuxApplication()
 	bUsingHighPrecisionMouseInput = false;
 	bAllowedToDeferMessageProcessing = true;
 	MouseCaptureWindow = NULL;
-	ControllerStates = new SDLControllerState[SDL_NumJoysticks()];
-	FMemory::Memset(ControllerStates, 0, sizeof(SDLControllerState) * SDL_NumJoysticks());
 
 	fMouseWheelScrollAccel = 1.0f;
 	if (GConfig)
@@ -95,17 +101,18 @@ FLinuxApplication::~FLinuxApplication()
 		GConfig->GetFloat(TEXT("X11.Tweaks"), TEXT("MouseWheelScrollAcceleration"), fMouseWheelScrollAccel, GEngineIni);
 		GConfig->Flush(false, GEngineIni);
 	}
-	delete [] ControllerStates;
 }
 
 void FLinuxApplication::DestroyApplication()
 {
-	for (int i = 0; i < SDL_NumJoysticks(); ++i) {
-		if(ControllerStates[i].controller != NULL)
+	for(auto ControllerIt = ControllerStates.CreateConstIterator(); ControllerIt; ++ControllerIt)
+	{
+		if(ControllerIt.Value().Controller != nullptr)
 		{
-			SDL_GameControllerClose(ControllerStates[i].controller);
+			SDL_GameControllerClose(ControllerIt.Value().Controller);
 		}
 	}
+	ControllerStates.Empty();
 }
 
 TSharedRef< FGenericWindow > FLinuxApplication::MakeWindow()
@@ -413,166 +420,176 @@ void FLinuxApplication::ProcessDeferredMessage( SDL_Event Event )
 	case SDL_CONTROLLERAXISMOTION:
 		{
 			SDL_ControllerAxisEvent caxisEvent = Event.caxis;
-			FGamepadKeyNames::Type analog = FGamepadKeyNames::Invalid;
-			float value = ShortToNormalFloat(caxisEvent.value);
+			FGamepadKeyNames::Type Axis = FGamepadKeyNames::Invalid;
+			float AxisValue = ShortToNormalFloat(caxisEvent.value);
+
+			if (!ControllerStates.Contains(caxisEvent.which))
+			{
+				break;
+			}
+
+			SDLControllerState &ControllerState = ControllerStates[caxisEvent.which];
 
 			switch (caxisEvent.axis)
 			{
 			case SDL_CONTROLLER_AXIS_LEFTX:
-				analog = FGamepadKeyNames::LeftAnalogX;
+				Axis = FGamepadKeyNames::LeftAnalogX;
 				if(caxisEvent.value > GAMECONTROLLER_LEFT_THUMB_DEADZONE)
 				{
-					if(!ControllerStates[caxisEvent.which].analogOverThreshold[0])
+					if(!ControllerState.AnalogOverThreshold[0])
 					{
-						MessageHandler->OnControllerButtonPressed(FGamepadKeyNames::LeftStickRight, caxisEvent.which, false);
-						ControllerStates[caxisEvent.which].analogOverThreshold[0] = true;
+						MessageHandler->OnControllerButtonPressed(FGamepadKeyNames::LeftStickRight, ControllerState.ControllerIndex, false);
+						ControllerState.AnalogOverThreshold[0] = true;
 					}
 				}
-				else if(ControllerStates[caxisEvent.which].analogOverThreshold[0])
+				else if(ControllerState.AnalogOverThreshold[0])
 				{
-					MessageHandler->OnControllerButtonReleased(FGamepadKeyNames::LeftStickRight, caxisEvent.which, false);
-					ControllerStates[caxisEvent.which].analogOverThreshold[0] = false;
+					MessageHandler->OnControllerButtonReleased(FGamepadKeyNames::LeftStickRight, ControllerState.ControllerIndex, false);
+					ControllerState.AnalogOverThreshold[0] = false;
 				}
 				if(caxisEvent.value < -GAMECONTROLLER_LEFT_THUMB_DEADZONE)
 				{
-					if(!ControllerStates[caxisEvent.which].analogOverThreshold[1])
+					if(!ControllerState.AnalogOverThreshold[1])
 					{
-						MessageHandler->OnControllerButtonPressed(FGamepadKeyNames::LeftStickLeft, caxisEvent.which, false);
-						ControllerStates[caxisEvent.which].analogOverThreshold[1] = true;
+						MessageHandler->OnControllerButtonPressed(FGamepadKeyNames::LeftStickLeft, ControllerState.ControllerIndex, false);
+						ControllerState.AnalogOverThreshold[1] = true;
 					}
 				}
-				else if(ControllerStates[caxisEvent.which].analogOverThreshold[1])
+				else if(ControllerState.AnalogOverThreshold[1])
 				{
-					MessageHandler->OnControllerButtonReleased(FGamepadKeyNames::LeftStickLeft, caxisEvent.which, false);
-					ControllerStates[caxisEvent.which].analogOverThreshold[1] = false;
+					MessageHandler->OnControllerButtonReleased(FGamepadKeyNames::LeftStickLeft, ControllerState.ControllerIndex, false);
+					ControllerState.AnalogOverThreshold[1] = false;
 				}
 				break;
 			case SDL_CONTROLLER_AXIS_LEFTY:
-				analog = FGamepadKeyNames::LeftAnalogY;
-				value *= -1;
+				Axis = FGamepadKeyNames::LeftAnalogY;
+				AxisValue *= -1;
 				if(caxisEvent.value > GAMECONTROLLER_LEFT_THUMB_DEADZONE)
 				{
-					if(!ControllerStates[caxisEvent.which].analogOverThreshold[2])
+					if(!ControllerState.AnalogOverThreshold[2])
 					{
-						MessageHandler->OnControllerButtonPressed(FGamepadKeyNames::LeftStickDown, caxisEvent.which, false);
-						ControllerStates[caxisEvent.which].analogOverThreshold[2] = true;
+						MessageHandler->OnControllerButtonPressed(FGamepadKeyNames::LeftStickDown, ControllerState.ControllerIndex, false);
+						ControllerState.AnalogOverThreshold[2] = true;
 					}
 				}
-				else if(ControllerStates[caxisEvent.which].analogOverThreshold[2])
+				else if(ControllerState.AnalogOverThreshold[2])
 				{
-					MessageHandler->OnControllerButtonReleased(FGamepadKeyNames::LeftStickDown, caxisEvent.which, false);
-					ControllerStates[caxisEvent.which].analogOverThreshold[2] = false;
+					MessageHandler->OnControllerButtonReleased(FGamepadKeyNames::LeftStickDown, ControllerState.ControllerIndex, false);
+					ControllerState.AnalogOverThreshold[2] = false;
 				}
 				if(caxisEvent.value < -GAMECONTROLLER_LEFT_THUMB_DEADZONE)
 				{
-					if(!ControllerStates[caxisEvent.which].analogOverThreshold[3])
+					if(!ControllerState.AnalogOverThreshold[3])
 					{
-						MessageHandler->OnControllerButtonPressed(FGamepadKeyNames::LeftStickUp, caxisEvent.which, false);
-						ControllerStates[caxisEvent.which].analogOverThreshold[3] = true;
+						MessageHandler->OnControllerButtonPressed(FGamepadKeyNames::LeftStickUp, ControllerState.ControllerIndex, false);
+						ControllerState.AnalogOverThreshold[3] = true;
 					}
 				}
-				else if(ControllerStates[caxisEvent.which].analogOverThreshold[3])
+				else if(ControllerState.AnalogOverThreshold[3])
 				{
-					MessageHandler->OnControllerButtonReleased(FGamepadKeyNames::LeftStickUp, caxisEvent.which, false);
-					ControllerStates[caxisEvent.which].analogOverThreshold[3] = false;
+					MessageHandler->OnControllerButtonReleased(FGamepadKeyNames::LeftStickUp, ControllerState.ControllerIndex, false);
+					ControllerState.AnalogOverThreshold[3] = false;
 				}
 				break;
 			case SDL_CONTROLLER_AXIS_RIGHTX:
-				analog = FGamepadKeyNames::RightAnalogX;
+				Axis = FGamepadKeyNames::RightAnalogX;
 				if(caxisEvent.value > GAMECONTROLLER_RIGHT_THUMB_DEADZONE)
 				{
-					if(!ControllerStates[caxisEvent.which].analogOverThreshold[4])
+					if(!ControllerState.AnalogOverThreshold[4])
 					{
-						MessageHandler->OnControllerButtonPressed(FGamepadKeyNames::RightStickRight, caxisEvent.which, false);
-						ControllerStates[caxisEvent.which].analogOverThreshold[4] = true;
+						MessageHandler->OnControllerButtonPressed(FGamepadKeyNames::RightStickRight, ControllerState.ControllerIndex, false);
+						ControllerState.AnalogOverThreshold[4] = true;
 					}
 				}
-				else if(ControllerStates[caxisEvent.which].analogOverThreshold[4])
+				else if(ControllerState.AnalogOverThreshold[4])
 				{
-					MessageHandler->OnControllerButtonReleased(FGamepadKeyNames::RightStickRight, caxisEvent.which, false);
-					ControllerStates[caxisEvent.which].analogOverThreshold[4] = false;
+					MessageHandler->OnControllerButtonReleased(FGamepadKeyNames::RightStickRight, ControllerState.ControllerIndex, false);
+					ControllerState.AnalogOverThreshold[4] = false;
 				}
 				if(caxisEvent.value < -GAMECONTROLLER_RIGHT_THUMB_DEADZONE)
 				{
-					if(!ControllerStates[caxisEvent.which].analogOverThreshold[5])
+					if(!ControllerState.AnalogOverThreshold[5])
 					{
-						MessageHandler->OnControllerButtonPressed(FGamepadKeyNames::RightStickLeft, caxisEvent.which, false);
-						ControllerStates[caxisEvent.which].analogOverThreshold[5] = true;
+						MessageHandler->OnControllerButtonPressed(FGamepadKeyNames::RightStickLeft, ControllerState.ControllerIndex, false);
+						ControllerState.AnalogOverThreshold[5] = true;
 					}
 				}
-				else if(ControllerStates[caxisEvent.which].analogOverThreshold[5])
+				else if(ControllerState.AnalogOverThreshold[5])
 				{
-					MessageHandler->OnControllerButtonReleased(FGamepadKeyNames::RightStickLeft, caxisEvent.which, false);
-					ControllerStates[caxisEvent.which].analogOverThreshold[5] = false;
+					MessageHandler->OnControllerButtonReleased(FGamepadKeyNames::RightStickLeft, ControllerState.ControllerIndex, false);
+					ControllerState.AnalogOverThreshold[5] = false;
 				}
 				break;
 			case SDL_CONTROLLER_AXIS_RIGHTY:
-				analog = FGamepadKeyNames::RightAnalogY;
-				value *= -1;
+				Axis = FGamepadKeyNames::RightAnalogY;
+				AxisValue *= -1;
 				if(caxisEvent.value > GAMECONTROLLER_RIGHT_THUMB_DEADZONE)
 				{
-					if(!ControllerStates[caxisEvent.which].analogOverThreshold[6])
+					if(!ControllerState.AnalogOverThreshold[6])
 					{
-						MessageHandler->OnControllerButtonPressed(FGamepadKeyNames::RightStickDown, caxisEvent.which, false);
-						ControllerStates[caxisEvent.which].analogOverThreshold[6] = true;
+						MessageHandler->OnControllerButtonPressed(FGamepadKeyNames::RightStickDown, ControllerState.ControllerIndex, false);
+						ControllerState.AnalogOverThreshold[6] = true;
 					}
 				}
-				else if(ControllerStates[caxisEvent.which].analogOverThreshold[6])
+				else if(ControllerState.AnalogOverThreshold[6])
 				{
-					MessageHandler->OnControllerButtonReleased(FGamepadKeyNames::RightStickDown, caxisEvent.which, false);
-					ControllerStates[caxisEvent.which].analogOverThreshold[6] = false;
+					MessageHandler->OnControllerButtonReleased(FGamepadKeyNames::RightStickDown, ControllerState.ControllerIndex, false);
+					ControllerState.AnalogOverThreshold[6] = false;
 				}
 				if(caxisEvent.value < -GAMECONTROLLER_RIGHT_THUMB_DEADZONE)
 				{
-					if(!ControllerStates[caxisEvent.which].analogOverThreshold[7])
+					if(!ControllerState.AnalogOverThreshold[7])
 					{
-						MessageHandler->OnControllerButtonPressed(FGamepadKeyNames::RightStickUp, caxisEvent.which, false);
-						ControllerStates[caxisEvent.which].analogOverThreshold[7] = true;
+						MessageHandler->OnControllerButtonPressed(FGamepadKeyNames::RightStickUp, ControllerState.ControllerIndex, false);
+						ControllerState.AnalogOverThreshold[7] = true;
 					}
 				}
-				else if(ControllerStates[caxisEvent.which].analogOverThreshold[7])
+				else if(ControllerState.AnalogOverThreshold[7])
 				{
-					MessageHandler->OnControllerButtonReleased(FGamepadKeyNames::RightStickUp, caxisEvent.which, false);
-					ControllerStates[caxisEvent.which].analogOverThreshold[7] = false;
+					MessageHandler->OnControllerButtonReleased(FGamepadKeyNames::RightStickUp, ControllerState.ControllerIndex, false);
+					ControllerState.AnalogOverThreshold[7] = false;
 				}
 				break;
 			case SDL_CONTROLLER_AXIS_TRIGGERLEFT:
-				analog = FGamepadKeyNames::LeftTriggerAnalog;
+				Axis = FGamepadKeyNames::LeftTriggerAnalog;
 				if(caxisEvent.value > GAMECONTROLLER_TRIGGER_THRESHOLD)
 				{
-					if(!ControllerStates[caxisEvent.which].analogOverThreshold[8])
+					if(!ControllerState.AnalogOverThreshold[8])
 					{
-						MessageHandler->OnControllerButtonPressed(FGamepadKeyNames::LeftTriggerThreshold, caxisEvent.which, false);
-						ControllerStates[caxisEvent.which].analogOverThreshold[8] = true;
+						MessageHandler->OnControllerButtonPressed(FGamepadKeyNames::LeftTriggerThreshold, ControllerState.ControllerIndex, false);
+						ControllerState.AnalogOverThreshold[8] = true;
 					}
 				}
-				else if(ControllerStates[caxisEvent.which].analogOverThreshold[8])
+				else if(ControllerState.AnalogOverThreshold[8])
 				{
-					MessageHandler->OnControllerButtonReleased(FGamepadKeyNames::LeftTriggerThreshold, caxisEvent.which, false);
-					ControllerStates[caxisEvent.which].analogOverThreshold[8] = false;
+					MessageHandler->OnControllerButtonReleased(FGamepadKeyNames::LeftTriggerThreshold, ControllerState.ControllerIndex, false);
+					ControllerState.AnalogOverThreshold[8] = false;
 				}
 				break;
 			case SDL_CONTROLLER_AXIS_TRIGGERRIGHT:
-				analog = FGamepadKeyNames::LeftTriggerAnalog;
+				Axis = FGamepadKeyNames::RightTriggerAnalog;
 				if(caxisEvent.value > GAMECONTROLLER_TRIGGER_THRESHOLD)
 				{
-					if(!ControllerStates[caxisEvent.which].analogOverThreshold[9])
+					if(!ControllerState.AnalogOverThreshold[9])
 					{
-						MessageHandler->OnControllerButtonPressed(FGamepadKeyNames::RightTriggerThreshold, caxisEvent.which, false);
-						ControllerStates[caxisEvent.which].analogOverThreshold[9] = true;
+						MessageHandler->OnControllerButtonPressed(FGamepadKeyNames::RightTriggerThreshold, ControllerState.ControllerIndex, false);
+						ControllerState.AnalogOverThreshold[9] = true;
 					}
 				}
-				else if(ControllerStates[caxisEvent.which].analogOverThreshold[9])
+				else if(ControllerState.AnalogOverThreshold[9])
 				{
-					MessageHandler->OnControllerButtonReleased(FGamepadKeyNames::RightTriggerThreshold, caxisEvent.which, false);
-					ControllerStates[caxisEvent.which].analogOverThreshold[9] = false;
+					MessageHandler->OnControllerButtonReleased(FGamepadKeyNames::RightTriggerThreshold, ControllerState.ControllerIndex, false);
+					ControllerState.AnalogOverThreshold[9] = false;
 				}
 				break;
+			default:
+				break;
 			}
 
-			if (analog != FGamepadKeyNames::Invalid)
+			if (Axis != FGamepadKeyNames::Invalid)
 			{
-				MessageHandler->OnControllerAnalog(analog, caxisEvent.which, value);
+				float & ExistingAxisEventValue = ControllerState.AxisEvents.FindOrAdd(Axis);
+				ExistingAxisEventValue = AxisValue;
 			}
 		}
 		break;
@@ -580,63 +597,70 @@ void FLinuxApplication::ProcessDeferredMessage( SDL_Event Event )
 	case SDL_CONTROLLERBUTTONUP:
 		{
 			SDL_ControllerButtonEvent cbuttonEvent = Event.cbutton;
-			FGamepadKeyNames::Type button = FGamepadKeyNames::Invalid;
+			FGamepadKeyNames::Type Button = FGamepadKeyNames::Invalid;
+
+			if (!ControllerStates.Contains(cbuttonEvent.which))
+			{
+				break;
+			}
 
 			switch (cbuttonEvent.button)
 			{
 			case SDL_CONTROLLER_BUTTON_A:
-				button = FGamepadKeyNames::FaceButtonBottom;
+				Button = FGamepadKeyNames::FaceButtonBottom;
 				break;
 			case SDL_CONTROLLER_BUTTON_B:
-				button = FGamepadKeyNames::FaceButtonRight;
+				Button = FGamepadKeyNames::FaceButtonRight;
 				break;
 			case SDL_CONTROLLER_BUTTON_X:
-				button = FGamepadKeyNames::FaceButtonLeft;
+				Button = FGamepadKeyNames::FaceButtonLeft;
 				break;
 			case SDL_CONTROLLER_BUTTON_Y:
-				button = FGamepadKeyNames::FaceButtonTop;
+				Button = FGamepadKeyNames::FaceButtonTop;
 				break;
 			case SDL_CONTROLLER_BUTTON_BACK:
-				button = FGamepadKeyNames::SpecialLeft;
+				Button = FGamepadKeyNames::SpecialLeft;
 				break;
 			case SDL_CONTROLLER_BUTTON_START:
-				button = FGamepadKeyNames::SpecialRight;
+				Button = FGamepadKeyNames::SpecialRight;
 				break;
 			case SDL_CONTROLLER_BUTTON_LEFTSTICK:
-				button = FGamepadKeyNames::LeftStickDown;
+				Button = FGamepadKeyNames::LeftStickDown;
 				break;
 			case SDL_CONTROLLER_BUTTON_RIGHTSTICK:
-				button = FGamepadKeyNames::RightStickDown;
+				Button = FGamepadKeyNames::RightStickDown;
 				break;
 			case SDL_CONTROLLER_BUTTON_LEFTSHOULDER:
-				button = FGamepadKeyNames::LeftShoulder;
+				Button = FGamepadKeyNames::LeftShoulder;
 				break;
 			case SDL_CONTROLLER_BUTTON_RIGHTSHOULDER:
-				button = FGamepadKeyNames::RightShoulder;
+				Button = FGamepadKeyNames::RightShoulder;
 				break;
 			case SDL_CONTROLLER_BUTTON_DPAD_UP:
-				button = FGamepadKeyNames::DPadUp;
+				Button = FGamepadKeyNames::DPadUp;
 				break;
 			case SDL_CONTROLLER_BUTTON_DPAD_DOWN:
-				button = FGamepadKeyNames::DPadDown;
+				Button = FGamepadKeyNames::DPadDown;
 				break;
 			case SDL_CONTROLLER_BUTTON_DPAD_LEFT:
-				button = FGamepadKeyNames::DPadLeft;
+				Button = FGamepadKeyNames::DPadLeft;
 				break;
 			case SDL_CONTROLLER_BUTTON_DPAD_RIGHT:
-				button = FGamepadKeyNames::DPadRight;
+				Button = FGamepadKeyNames::DPadRight;
+				break;
+			default:
 				break;
 			}
 
-			if (button != FGamepadKeyNames::Invalid)
+			if (Button != FGamepadKeyNames::Invalid)
 			{
-				if (cbuttonEvent.type == SDL_CONTROLLERBUTTONDOWN)
+				if(cbuttonEvent.type == SDL_CONTROLLERBUTTONDOWN)
 				{
-					MessageHandler->OnControllerButtonPressed(button, cbuttonEvent.which, false);
+					MessageHandler->OnControllerButtonPressed(Button, ControllerStates[cbuttonEvent.which].ControllerIndex, false);
 				}
 				else
 				{
-					MessageHandler->OnControllerButtonReleased(button, cbuttonEvent.which, false);
+					MessageHandler->OnControllerButtonReleased(Button, ControllerStates[cbuttonEvent.which].ControllerIndex, false);
 				}
 			}
 		}
@@ -937,6 +961,15 @@ void FLinuxApplication::ProcessDeferredEvents( const float TimeDelta )
 
 void FLinuxApplication::PollGameDeviceState( const float TimeDelta )
 {
+	for(auto ControllerIt = ControllerStates.CreateIterator(); ControllerIt; ++ControllerIt)
+	{
+		for(auto Event = ControllerIt.Value().AxisEvents.CreateConstIterator(); Event; ++Event)
+		{
+			MessageHandler->OnControllerAnalog(Event.Key(), ControllerIt.Value().ControllerIndex, Event.Value());
+		}
+		ControllerIt.Value().AxisEvents.Empty();
+	}
+
 	// initialize any externally-implemented input devices (we delay load initialize the array so any plugins have had time to load)
 	if (!bHasLoadedInputPlugins)
 	{
diff --git a/Engine/Source/Runtime/Core/Public/Linux/LinuxApplication.h b/Engine/Source/Runtime/Core/Public/Linux/LinuxApplication.h
--- a/Engine/Source/Runtime/Core/Public/Linux/LinuxApplication.h
+++ b/Engine/Source/Runtime/Core/Public/Linux/LinuxApplication.h
@@ -6,8 +6,6 @@
 #include "LinuxWindow.h"
 #include "LinuxCursor.h"
 
-typedef SDL_GameController* SDL_HController;
-
 class FLinuxWindow;
 class FGenericApplicationMessageHandler;
 
@@ -120,10 +118,24 @@ private:
 
 	struct SDLControllerState
 	{
-		SDL_HController controller;
+		/** SDL controller */
+		SDL_GameController* Controller;
+
+		/** Tracks whether the "button" was previously pressed so we don't generate extra events */
+		bool AnalogOverThreshold[10];
+
+		/** The player index of the controller, because the joystick index includes devices that are not controllers. */
+		int32 ControllerIndex;
+
+		/** Store axis values from events here to be handled once per frame. */
+		TMap<FGamepadKeyNames::Type, float> AxisEvents;
 
-		// We need to remember if the "button" was previously pressed so we don't generate extra events
-		bool analogOverThreshold[10];
+		SDLControllerState()
+			:	Controller(nullptr)
+			,	ControllerIndex(-1)
+		{
+			FMemory::Memzero(AnalogOverThreshold);
+		}
 	};
 
 	bool bUsingHighPrecisionMouseInput;
@@ -157,7 +169,7 @@ private:
 
 	SDL_HWindow MouseCaptureWindow;
 
-	SDLControllerState *ControllerStates;
+	TMap<SDL_JoystickID, SDLControllerState> ControllerStates;
 
 	float fMouseWheelScrollAccel;
 
