Fixed code comment

diff --git a/Engine/Source/Runtime/Core/Public/Containers/CircularQueue.h b/Engine/Source/Runtime/Core/Public/Containers/CircularQueue.h
--- a/Engine/Source/Runtime/Core/Public/Containers/CircularQueue.h
+++ b/Engine/Source/Runtime/Core/Public/Containers/CircularQueue.h
@@ -14,7 +14,7 @@
  * The number of items that can be enqueued is one less than the queue's capacity,
  * because one item will be used for detecting full and empty states.
  *
- * @param ElementType - The type of elements held in the queue.
+ * @param ElementType The type of elements held in the queue.
  */
 template<typename ElementType> class TCircularQueue
 {
