Stats - Cleaned a bit stats file related functionality

diff --git a/Engine/Source/Runtime/Core/Public/Stats/StatsData.h b/Engine/Source/Runtime/Core/Public/Stats/StatsData.h
--- a/Engine/Source/Runtime/Core/Public/Stats/StatsData.h
+++ b/Engine/Source/Runtime/Core/Public/Stats/StatsData.h
@@ -484,20 +484,11 @@ class CORE_API FStatsThreadState
 	/** Internal method to scan the messages to update the current frame. */
 	void ScanForAdvance(FStatPacketArray& NewData);
 
-	// @TODO yrx 2014-11-28 Remove from here
-	/** Internal method to scan the messages to accumulate any non-frame stats. **/
-	void ProcessMetaDataForLoad(TArray<FStatMessage>& Data);
-
 public:
 	/** Internal method to add meta data packets to the data structures. **/
 	void ProcessMetaDataOnly(TArray<FStatMessage>& Data);
 
-	// @TODO yrx 2014-11-28 Remove from here
-	/** Marks this stats state as loaded. */
-	void MarkAsLoaded()
-	{
-		bWasLoaded =  true;
-	}
+
 
 	/** Toggles tracking the most memory expensive stats. */
 	void ToggleFindMemoryExtensiveStats();
@@ -546,14 +537,6 @@ private:
 	/** Number of frames to keep in the history. **/
 	int32 HistoryFrames;
 
-	// @TODO yrx 2014-11-28 
-	/** Largest frame seen. Loaded stats only. **/
-	int64 MaxFrameSeen;
-
-	// @TODO yrx 2014-11-28 
-	/** First frame seen. Loaded stats only. **/
-	int64 MinFrameSeen;
-
 	/** Used to track which packets have been sent to listeners **/
 	int64 LastFullFrameMetaAndNonFrame;
 
@@ -581,10 +564,6 @@ private:
 	/** Maximum number of stat messages seen so far for active session. */
 	int32 MaxNumStatMessages;
 
-	// @TODO yrx 2014-11-28 
-	/** Valid frame computation is different if we just loaded these stats **/
-	bool bWasLoaded;
-
 	/** If true, stats each frame will dump to the log a list of most memory expensive stats. */
 	bool bFindMemoryExtensiveStats;
 
@@ -593,10 +572,6 @@ public:
 	/** Constructor used by GetLocalState(), also used by the profiler to hold a previewing stats thread state. We don't keep many frames by default **/
 	FStatsThreadState(int32 InHistoryFrames = STAT_FRAME_SLOP + 2);
 
-	// @TODO yrx 2014-11-28 
-	/** Constructor to load stats from a file **/
-	FStatsThreadState(FString const& Filename);
-
 	/** Delegate we fire every time we have a new complete frame of data. **/
 	mutable FOnNewFrameHistory NewFrameDelegate;
 
@@ -672,12 +647,40 @@ public:
 	/** Adds missing stats to the group so it doesn't jitter. **/
 	void AddMissingStats(TArray<FStatMessage>& Dest, TSet<FName> const& EnabledItems) const;
 
-	// @TODO yrx 2014-11-28 REmove from here, it's file related.
-	/** Adds a frame worth of messages */
-	void AddMessages(TArray<FStatMessage>& InMessages);
-
 	/** Singleton to get the stats being collected by this executable. Can be only accessed from the stats thread. **/
 	static FStatsThreadState& GetLocalState();
+
+	/*-----------------------------------------------------------------------------
+		Stats file related functionality
+
+		#YRX_Stats: 2015-07-07 Maybe move to FStatsLoadedState
+	-----------------------------------------------------------------------------*/
+public:
+	/** Constructor to load stats from a file **/
+	FStatsThreadState( FString const& Filename );
+
+	/** Adds a frame worth of messages */
+	void AddMessages( TArray<FStatMessage>& InMessages );
+
+	/** Marks this stats state as loaded. */
+	void MarkAsLoaded()
+	{
+		bWasLoaded = true;
+	}
+
+protected:
+	/** Internal method to scan the messages to accumulate any non-frame stats. **/
+	void ProcessMetaDataForLoad( TArray<FStatMessage>& Data );
+
+	/** Largest frame seen. Loaded stats only. **/
+	int64 MaxFrameSeen;
+
+	/** First frame seen. Loaded stats only. **/
+	int64 MinFrameSeen;
+
+	/** Valid frame computation is different if we just loaded these stats **/
+	bool bWasLoaded;
+
 };
 
 //@todo split header
