UECORE-234/Enable stats for commandlets (added option to dump collected stats to the log file)

diff --git a/Engine/Source/Runtime/Core/Private/Stats/Stats2.cpp b/Engine/Source/Runtime/Core/Private/Stats/Stats2.cpp
--- a/Engine/Source/Runtime/Core/Private/Stats/Stats2.cpp
+++ b/Engine/Source/Runtime/Core/Private/Stats/Stats2.cpp
@@ -197,7 +197,7 @@ void FStats::AdvanceFrame( bool bDiscardCallstack, const FOnAdvanceRenderingThre
 
 void FStats::TickCommandletStats()
 {
-	if (HasStatsForCommandletsToken())
+	if (EnabledForCommandlet())
 	{
 		//FThreadStats* ThreadStats = FThreadStats::GetThreadStats();
 		//check( ThreadStats->ScopeCount == 0 && TEXT( "FStats::TickCommandletStats must be called outside any scope counters" ) );
@@ -209,12 +209,24 @@ void FStats::TickCommandletStats()
 	}
 }
 
-bool FStats::HasStatsForCommandletsToken()
+bool FStats::EnabledForCommandlet()
 {
-	static bool bHasStatsForCommandletsToken = FParse::Param( FCommandLine::Get(), TEXT( "StatsForCommandlets" ) );
+	static bool bHasStatsForCommandletsToken = HasLoadTimeStatsForCommandletToken() || HasLoadTimeFileForCommandletToken();
 	return bHasStatsForCommandletsToken;
 }
 
+bool FStats::HasLoadTimeStatsForCommandletToken()
+{
+	static bool bHasLoadTimeStatsForCommandletToken = FParse::Param( FCommandLine::Get(), TEXT( "LoadTimeStatsForCommandlet" ) );
+	return bHasLoadTimeStatsForCommandletToken;
+}
+
+bool FStats::HasLoadTimeFileForCommandletToken()
+{
+	static bool bHasLoadTimeFileForCommandletToken = FParse::Param( FCommandLine::Get(), TEXT( "LoadTimeFileForCommandlet" ) );
+	return bHasLoadTimeFileForCommandletToken;
+}
+
 /* Todo
 
 ST_int64 combined with IsPackedCCAndDuration is fairly bogus, ST_uint32_pair should probably be a first class type. IsPackedCCAndDuration can stay, it says what kind of data is in the pair. remove FromPackedCallCountDuration_Duration et al
@@ -846,7 +858,7 @@ public:
 		{
 			// For regular stats we won't process more than every 5ms or every 16 packets. 
 			// Commandlet stats are flushed as soon as.
-			bShouldProcess = bReadyToProcess && (FPlatformTime::Seconds() - LastTime > 0.005f || IncomingData.Packets.Num() > MaxIncomingPackets || FStats::HasStatsForCommandletsToken());
+			bShouldProcess = bReadyToProcess && (FPlatformTime::Seconds() - LastTime > 0.005f || IncomingData.Packets.Num() > MaxIncomingPackets || FStats::EnabledForCommandlet());
 		}
 
 		if( bShouldProcess )
@@ -1020,7 +1032,7 @@ FThreadStats::FThreadStats( EConstructor ):
 void FThreadStats::CheckEnable()
 {
 	bool bOldMasterEnable(bMasterEnable);
-	bool bNewMasterEnable( WillEverCollectData() && (!IsRunningCommandlet() || FStats::HasStatsForCommandletsToken()) && IsThreadingReady() && (MasterEnableCounter.GetValue()) );
+	bool bNewMasterEnable( WillEverCollectData() && (!IsRunningCommandlet() || FStats::EnabledForCommandlet()) && IsThreadingReady() && (MasterEnableCounter.GetValue()) );
 	if (bMasterEnable != bNewMasterEnable)
 	{
 		MasterDisableChangeTagLockAdd();
@@ -1184,12 +1196,18 @@ void FThreadStats::CheckForCollectingStartupStats()
 		DirectStatsCommand( TEXT( "stat group enable AsyncLoad" ) );
 		DirectStatsCommand( TEXT( "stat dumpsum -start -ms=250 -num=240" ), true );
 	}
-	if (FParse::Param( FCommandLine::Get(), TEXT( "LoadTimeFile" ) ))
+	else if (FParse::Param( FCommandLine::Get(), TEXT( "LoadTimeFile" ) ) || FStats::HasLoadTimeFileForCommandletToken())
 	{
 		DirectStatsCommand( TEXT( "stat group enable LinkerLoad" ) );
 		DirectStatsCommand( TEXT( "stat group enable AsyncLoad" ) );
 		DirectStatsCommand( TEXT( "stat startfile" ), true );
 	}
+	else if (FStats::HasLoadTimeStatsForCommandletToken())
+	{
+		DirectStatsCommand( TEXT( "stat group enable LinkerLoad" ) );
+		DirectStatsCommand( TEXT( "stat group enable AsyncLoad" ) );
+		DirectStatsCommand( TEXT( "stat dumpsum -start" ), true );
+	}
 
 	// Now we can safely enable malloc profiler.
 	if (FStatsMallocProfilerProxy::HasMemoryProfilerToken())
@@ -1259,6 +1277,12 @@ void FThreadStats::StopThread()
 	// Nothing to stop if it was never started
 	if (IsThreadingReady())
 	{	
+		if (FStats::HasLoadTimeStatsForCommandletToken())
+		{
+			// Dump all the collected stats to the log, if any.
+			DirectStatsCommand( TEXT( "stat dumpsum -stop -ms=100" ), true );
+		}
+
 		// If we are writing stats data, stop it now.
 		DirectStatsCommand(TEXT("stat stopfile"), true);
 
diff --git a/Engine/Source/Runtime/Core/Public/Stats/Stats2.h b/Engine/Source/Runtime/Core/Public/Stats/Stats2.h
--- a/Engine/Source/Runtime/Core/Public/Stats/Stats2.h
+++ b/Engine/Source/Runtime/Core/Public/Stats/Stats2.h
@@ -65,10 +65,21 @@ struct CORE_API FStats
 	static void TickCommandletStats();
 
 	/**
-	* @return true, if the command line has the StatsForCommandlets token which enables stats in the commandlets.
+	* @return true, if the command line has the LoadTimeStatsForCommandlet or LoadTimeFileForCommandlet token which enables stats in the commandlets.
 	* !!!CAUTION!!! You need to manually advance stats frame in order to maintain the data integrity and not to leak the memory.
 	*/
-	static bool HasStatsForCommandletsToken();
+	static bool EnabledForCommandlet();
+
+	/**
+	* @return true, if the command line has the LoadTimeStatsForCommandlet token which enables LoadTimeStats equivalent for commandlets.
+	* All collected stats will be dumped to the log file at the end of running the specified commandlet.
+	*/
+	static bool HasLoadTimeStatsForCommandletToken();
+
+	/**
+	* @return true, if the command line has the LoadTimeFileForCommandlet token which enables LoadTimeFile equivalent for commandlets.
+	*/
+	static bool HasLoadTimeFileForCommandletToken();
 
 	/** Current game thread stats frame. */
 	static int32 GameThreadStatsFrame;
diff --git a/Engine/Source/Runtime/Launch/Private/LaunchEngineLoop.cpp b/Engine/Source/Runtime/Launch/Private/LaunchEngineLoop.cpp
--- a/Engine/Source/Runtime/Launch/Private/LaunchEngineLoop.cpp
+++ b/Engine/Source/Runtime/Launch/Private/LaunchEngineLoop.cpp
@@ -926,7 +926,7 @@ int32 FEngineLoop::PreInit( const TCHAR* CmdLine )
 		{
 #if STATS
 			// Leave the stats enabled.
-			if (!FStats::HasStatsForCommandletsToken())
+			if (!FStats::EnabledForCommandlet())
 			{
 				FThreadStats::MasterDisableForever();
 			}
@@ -987,7 +987,7 @@ int32 FEngineLoop::PreInit( const TCHAR* CmdLine )
 	{
 #if STATS
 		// Leave the stats enabled.
-		if (!FStats::HasStatsForCommandletsToken())
+		if (!FStats::EnabledForCommandlet())
 		{
 			FThreadStats::MasterDisableForever();
 		}
