Added named pipe writing (GitHub PR# 1300)

https://jira.ol.epicgames.net/browse/UE-17489

This still needs an implementation for iOS and Android, if supported.

diff --git a/Engine/Source/Runtime/Core/Private/GenericPlatform/GenericPlatformProcess.cpp b/Engine/Source/Runtime/Core/Private/GenericPlatform/GenericPlatformProcess.cpp
--- a/Engine/Source/Runtime/Core/Private/GenericPlatform/GenericPlatformProcess.cpp
+++ b/Engine/Source/Runtime/Core/Private/GenericPlatform/GenericPlatformProcess.cpp
@@ -479,6 +479,12 @@ bool FGenericPlatformProcess::ReadPipeToArray(void* ReadPipe, TArray<uint8> & Ou
 	return false;
 }
 
+bool FGenericPlatformProcess::WritePipe(void* WritePipe, const FString& Message, FString* OutWritten)
+{
+	UE_LOG(LogHAL, Fatal, TEXT("FGenericPlatformProcess::WriteToPipe not implemented on this platform"));
+	return false;
+}
+
 bool FGenericPlatformProcess::SupportsMultithreading()
 {
 	static bool bSupportsMultithreading = !FParse::Param(FCommandLine::Get(), TEXT("nothreading"));
diff --git a/Engine/Source/Runtime/Core/Private/Linux/LinuxPlatformProcess.cpp b/Engine/Source/Runtime/Core/Private/Linux/LinuxPlatformProcess.cpp
--- a/Engine/Source/Runtime/Core/Private/Linux/LinuxPlatformProcess.cpp
+++ b/Engine/Source/Runtime/Core/Private/Linux/LinuxPlatformProcess.cpp
@@ -412,6 +412,34 @@ bool FLinuxPlatformProcess::ReadPipeToArray(void* ReadPipe, TArray<uint8> & Outp
 	return false;
 }
 
+bool FLinuxPlatformProcess::WritePipe(void* WritePipe, const FString& Message, FString* OutWritten)
+{
+	// if there is not a message or WritePipe is null
+	if ((Message.Len() == 0) || (WritePipe == nullptr))
+	{
+		return false;
+	}
+
+	// convert input to UTF8CHAR
+	uint32 BytesAvailable = Message.Len();
+	UTF8CHAR* Buffer = new UTF8CHAR[BytesAvailable + 1];
+
+	if (!FString::ToBlob(Message, Buffer, BytesAvailable))
+	{
+		return false;
+	}
+
+	// write to pipe
+	uint32 BytesWritten = write(*(int*)WritePipe, Buffer, BytesAvailable);
+
+	if (OutWritten != nullptr)
+	{
+		OutWritten->FromBlob(Buffer, BytesWritten);
+	}
+
+	return (BytesWritten == BytesAvailable);
+}
+
 FRunnableThread* FLinuxPlatformProcess::CreateRunnableThread()
 {
 	return new FRunnableThreadLinux();
diff --git a/Engine/Source/Runtime/Core/Private/Mac/MacPlatformProcess.cpp b/Engine/Source/Runtime/Core/Private/Mac/MacPlatformProcess.cpp
--- a/Engine/Source/Runtime/Core/Private/Mac/MacPlatformProcess.cpp
+++ b/Engine/Source/Runtime/Core/Private/Mac/MacPlatformProcess.cpp
@@ -979,6 +979,34 @@ bool FMacPlatformProcess::ReadPipeToArray(void* ReadPipe, TArray<uint8>& Output)
 	return false;
 }
 
+bool FMacPlatformProcess::WritePipe(void* WritePipe, const FString& Message, FString* OutWritten)
+{
+	// if there is not a message or WritePipe is nullptr
+	if ((Message.Len() == 0) || (WritePipe == nullptr))
+	{
+		return false;
+	}
+
+	// convert input to UTF8CHAR
+	uint32 BytesAvailable = Message.Len();
+	UTF8CHAR* Buffer = new UTF8CHAR[BytesAvailable + 1];
+
+	if (!FString::ToBlob(Message, Buffer, BytesAvailable))
+	{
+		return false;
+	}
+
+	// Write to pipe
+	uint32 BytesWritten = write(*(int*)WritePipe, Buffer, BytesAvailable);
+
+	if (OutWritten != nullptr)
+	{
+		OutWritten->FromBlob(Buffer, BytesWritten);
+	}
+
+	return (BytesWritten == BytesAvailable);
+}
+
 bool FMacPlatformProcess::IsApplicationRunning(const TCHAR* ProcName)
 {
 	const FString ProcString = FPaths::GetCleanFilename(ProcName);
diff --git a/Engine/Source/Runtime/Core/Private/Windows/WindowsPlatformProcess.cpp b/Engine/Source/Runtime/Core/Private/Windows/WindowsPlatformProcess.cpp
--- a/Engine/Source/Runtime/Core/Private/Windows/WindowsPlatformProcess.cpp
+++ b/Engine/Source/Runtime/Core/Private/Windows/WindowsPlatformProcess.cpp
@@ -1140,6 +1140,35 @@ bool FWindowsPlatformProcess::ReadPipeToArray(void* ReadPipe, TArray<uint8> & Ou
 	return false;
 }
 
+bool FWindowsPlatformProcess::WritePipe(void* WritePipe, const FString& Message, FString* OutWritten)
+{
+	// If there is not a message or WritePipe is null
+	if (Message.Len() == 0 || WritePipe == nullptr)
+	{
+		return false;
+	}
+
+	// Convert input to UTF8CHAR
+	uint32 BytesAvailable = Message.Len();
+	UTF8CHAR* Buffer = new UTF8CHAR[BytesAvailable + 1];
+
+	if (!FString::ToBlob(Message, Buffer, BytesAvailable))
+	{
+		return false;
+	}
+
+	// Write to pipe
+	uint32 BytesWritten = 0;
+	bool bIsWritten = !!WriteFile(WritePipe, Buffer, BytesAvailable, (::DWORD*)&BytesWritten, nullptr);
+
+	if (OutWritten)
+	{
+		OutWritten->FromBlob(Buffer, BytesWritten);
+	}
+
+	return bIsWritten;
+}
+
 #include "AllowWindowsPlatformTypes.h"
 
 FWindowsPlatformProcess::FWindowsSemaphore::FWindowsSemaphore(const FString & InName, HANDLE InSemaphore)
diff --git a/Engine/Source/Runtime/Core/Public/GenericPlatform/GenericPlatformProcess.h b/Engine/Source/Runtime/Core/Public/GenericPlatform/GenericPlatformProcess.h
--- a/Engine/Source/Runtime/Core/Public/GenericPlatform/GenericPlatformProcess.h
+++ b/Engine/Source/Runtime/Core/Public/GenericPlatform/GenericPlatformProcess.h
@@ -505,6 +505,17 @@ struct CORE_API FGenericPlatformProcess
 	 */
 	static bool ReadPipeToArray(void* ReadPipe, TArray<uint8> & Output);
 
+	/**
+	* Sends the message to process through pipe
+	*
+	* @param WritePipe Pipe for writing.
+	* @param Message The message to be written.
+	* @param OutWritten Optional parameter to know how much of the string written.
+	* @return True if all bytes written successfully.
+	* @see CreatePipe, ClosePipe, ReadPipe
+	*/
+	static bool WritePipe(void* WritePipe, const FString& Message, FString* OutWritten = nullptr);
+
 	/**
 	 * Gets whether this platform can use multiple threads.
 	 *
diff --git a/Engine/Source/Runtime/Core/Public/Linux/LinuxPlatformProcess.h b/Engine/Source/Runtime/Core/Public/Linux/LinuxPlatformProcess.h
--- a/Engine/Source/Runtime/Core/Public/Linux/LinuxPlatformProcess.h
+++ b/Engine/Source/Runtime/Core/Public/Linux/LinuxPlatformProcess.h
@@ -262,6 +262,7 @@ struct CORE_API FLinuxPlatformProcess : public FGenericPlatformProcess
 	static bool CreatePipe( void*& ReadPipe, void*& WritePipe );
 	static FString ReadPipe( void* ReadPipe );
 	static bool ReadPipeToArray(void* ReadPipe, TArray<uint8> & Output);
+	static bool WritePipe(void* WritePipe, const FString& Message, FString* OutWritten = nullptr);
 	static class FRunnableThread* CreateRunnableThread();
 	static void LaunchURL(const TCHAR* URL, const TCHAR* Parms, FString* Error);
 	static FProcHandle CreateProc(const TCHAR* URL, const TCHAR* Parms, bool bLaunchDetached, bool bLaunchHidden, bool bLaunchReallyHidden, uint32* OutProcessID, int32 PriorityModifier, const TCHAR* OptionalWorkingDirectory, void* PipeWrite);
diff --git a/Engine/Source/Runtime/Core/Public/Mac/MacPlatformProcess.h b/Engine/Source/Runtime/Core/Public/Mac/MacPlatformProcess.h
--- a/Engine/Source/Runtime/Core/Public/Mac/MacPlatformProcess.h
+++ b/Engine/Source/Runtime/Core/Public/Mac/MacPlatformProcess.h
@@ -127,6 +127,7 @@ struct CORE_API FMacPlatformProcess : public FGenericPlatformProcess
 	static bool CreatePipe( void*& ReadPipe, void*& WritePipe );
 	static FString ReadPipe( void* ReadPipe );
 	static bool ReadPipeToArray(void* ReadPipe, TArray<uint8> & Output);
+	static bool WritePipe(void* WritePipe, const FString& Message, FString* OutWritten = nullptr);
 	static bool ExecProcess(const TCHAR* URL, const TCHAR* Params, int32* OutReturnCode, FString* OutStdOut, FString* OutStdErr);
 
 	// Mac specific
diff --git a/Engine/Source/Runtime/Core/Public/Windows/WindowsPlatformProcess.h b/Engine/Source/Runtime/Core/Public/Windows/WindowsPlatformProcess.h
--- a/Engine/Source/Runtime/Core/Public/Windows/WindowsPlatformProcess.h
+++ b/Engine/Source/Runtime/Core/Public/Windows/WindowsPlatformProcess.h
@@ -168,6 +168,7 @@ public:
 	static bool CreatePipe( void*& ReadPipe, void*& WritePipe );
 	static FString ReadPipe( void* ReadPipe );
 	static bool ReadPipeToArray(void* ReadPipe, TArray<uint8> & Output);
+	static bool WritePipe(void* WritePipe, const FString& Message, FString* OutWritten = nullptr);
 	static FSemaphore* NewInterprocessSynchObject(const FString& Name, bool bCreate, uint32 MaxLocks = 1);
 	static bool DeleteInterprocessSynchObject(FSemaphore * Object);
 	static bool Daemonize();
