PR1280 - increased size of game thread in development and debug
#ios

diff --git a/Engine/Source/Runtime/Core/Private/IOS/IOSAppDelegate.cpp b/Engine/Source/Runtime/Core/Private/IOS/IOSAppDelegate.cpp
--- a/Engine/Source/Runtime/Core/Private/IOS/IOSAppDelegate.cpp
+++ b/Engine/Source/Runtime/Core/Private/IOS/IOSAppDelegate.cpp
@@ -13,7 +13,11 @@
 #include <AVFoundation/AVAudioSession.h>
 
 // this is the size of the game thread stack, it must be a multiple of 4k
-#define GAME_THREAD_STACK_SIZE 1024 * 1024 
+#if (UE_BUILD_SHIPPING || UE_BUILD_TEST)
+#define GAME_THREAD_STACK_SIZE 1024 * 1024
+#else
+#define GAME_THREAD_STACK_SIZE 4 * 1024 * 1024
+#endif
 
 DEFINE_LOG_CATEGORY(LogIOSAudioSession);
 DECLARE_LOG_CATEGORY_EXTERN(LogEngine, Log, All);
