Removed the dedicated quoted string parsing from FTextFilterExpressionEvaluator

This wasn't handling searches like '"Thing"...' correctly, and the non-quoted string parsing will already parse quoted strings correctly since it enters "quoted string parsing mode" once it finds an opening quote.

diff --git a/Engine/Source/Runtime/Core/Private/Misc/TextFilterExpressionEvaluator.cpp b/Engine/Source/Runtime/Core/Private/Misc/TextFilterExpressionEvaluator.cpp
--- a/Engine/Source/Runtime/Core/Private/Misc/TextFilterExpressionEvaluator.cpp
+++ b/Engine/Source/Runtime/Core/Private/Misc/TextFilterExpressionEvaluator.cpp
@@ -182,68 +182,6 @@ namespace TextFilterExpressionParser
 		Str.ReplaceInline(TEXT("\\\\"), TEXT("\\"));
 	}
 
-	/** Consume quoted text from the specified consumer's stream, if one exists at the current read position */
-	TOptional<FExpressionError> ConsumeQuotedText(FExpressionTokenConsumer& Consumer)
-	{
-		auto& Stream = Consumer.GetStream();
-
-		const TCHAR QuoteChar = Stream.PeekChar();
-		
-		if (QuoteChar != '"' && QuoteChar != '\'')
-		{
-			return TOptional<FExpressionError>();
-		}
-
-		FStringToken TextTokenWithQuotes = Stream.ParseSymbol().GetValue();
-
-		int32 NumConsecutiveSlashes = 0;
-		TOptional<FStringToken> TextToken = Stream.ParseToken([&](TCHAR InC)
-		{
-			if (InC == QuoteChar)
-			{
-				if (NumConsecutiveSlashes%2 == 0)
-				{
-					// Stop consuming, excluding this closing quote
-					return EParseState::StopBefore;
-				}
-			}
-			else if (InC == '\\')
-			{
-				NumConsecutiveSlashes++;
-			}
-			else
-			{
-				NumConsecutiveSlashes = 0;
-			}
-			
-			return EParseState::Continue;
-		}, &TextTokenWithQuotes);
-
-		if (!Stream.ParseSymbol(QuoteChar, &TextTokenWithQuotes).IsSet())
-		{
-			FString QuoteCharStr;
-			QuoteCharStr.AppendChar(QuoteChar);
-
-			return FExpressionError(FText::Format(NSLOCTEXT("TextFilterExpressionParser", "SyntaxError_UnmatchedQuote", "Syntax error: Reached end of expression before matching closing quote for '{0}' at line {1}:{2}"),
-				FText::FromString(QuoteCharStr),
-				FText::AsNumber(TextTokenWithQuotes.GetLineNumber()),
-				FText::AsNumber(TextTokenWithQuotes.GetCharacterIndex())
-			));
-		}
-		else if (TextToken.IsSet())
-		{
-			FString FinalString = TextToken.GetValue().GetString();
-			UnescapeQuotedString(FinalString, QuoteChar);
-			Consumer.Add(TextTokenWithQuotes, FExpressionNode(FTextToken(MoveTemp(FinalString), ETextFilterTextComparisonMode::Partial, FTextToken::EInvertResult::No)));
-		}
-		else
-		{
-			Consumer.Add(TextTokenWithQuotes, FExpressionNode(FTextToken(FString(), ETextFilterTextComparisonMode::Partial, FTextToken::EInvertResult::No)));
-		}
-
-		return TOptional<FExpressionError>();
-	}
-
 	/** Given a potential string, we produce a final FTextToken for it using the correct comparison mode (as inferred from the given string) */
 	FTextToken CreateTextTokenFromUnquotedString(FString InString)
 	{
@@ -616,7 +554,6 @@ void FTextFilterExpressionEvaluator::ConstructExpressionParser()
 	TokenDefinitions.DefineToken(&ConsumeOperator<FAnd>);
 	TokenDefinitions.DefineToken(&ConsumeOperator<FNot>);
 	TokenDefinitions.DefineToken(&ConsumeNumber);
-	TokenDefinitions.DefineToken(&ConsumeQuotedText);
 	TokenDefinitions.DefineToken((ExpressionEvaluatorMode == ETextFilterExpressionEvaluatorMode::Complex) ? &ConsumeComplexText : &ConsumeBasicText);
 
 	Grammar.DefineGrouping<FSubExpressionStart, FSubExpressionEnd>();
