VS2015 fix operator= needs a default implementation

diff --git a/Engine/Source/Runtime/Core/Private/Misc/ExpressionParser.cpp b/Engine/Source/Runtime/Core/Private/Misc/ExpressionParser.cpp
--- a/Engine/Source/Runtime/Core/Private/Misc/ExpressionParser.cpp
+++ b/Engine/Source/Runtime/Core/Private/Misc/ExpressionParser.cpp
@@ -903,6 +903,8 @@ PRAGMA_DISABLE_OPTIMIZATION
 		FHugeType(int32 InId) : FMoveableType(InId) {}
 		FHugeType(FHugeType&& In) : FMoveableType(MoveTemp(In)) {}
 		FHugeType(const FHugeType& In) : FMoveableType(In) {}
+
+		FHugeType& operator=(FHugeType&& In) = default;
 		
 		uint8 Padding[1024];
 	};
