Fix UE-16920 which was caused by not properly handling transitions between windowed-fullscreen <-> fullscreen, this requires popping out of one fullscreen mode to re-enter the other in order to get the proper events that the CGL/NSOpenGL APIs need.

diff --git a/Engine/Source/Runtime/Core/Private/Mac/CocoaWindow.cpp b/Engine/Source/Runtime/Core/Private/Mac/CocoaWindow.cpp
--- a/Engine/Source/Runtime/Core/Private/Mac/CocoaWindow.cpp
+++ b/Engine/Source/Runtime/Core/Private/Mac/CocoaWindow.cpp
@@ -85,6 +85,9 @@ NSString* NSPerformDragOperation = @"NSPerformDragOperation";
 - (void)setWindowMode:(EWindowMode::Type)NewWindowMode
 {
 	WindowMode = NewWindowMode;
+	
+	NSView* OpenGLView = [self openGLView];
+	[[NSNotificationCenter defaultCenter] postNotificationName:NSViewGlobalFrameDidChangeNotification object:OpenGLView];
 }
 
 - (EWindowMode::Type)windowMode
diff --git a/Engine/Source/Runtime/Core/Private/Mac/MacWindow.cpp b/Engine/Source/Runtime/Core/Private/Mac/MacWindow.cpp
--- a/Engine/Source/Runtime/Core/Private/Mac/MacWindow.cpp
+++ b/Engine/Source/Runtime/Core/Private/Mac/MacWindow.cpp
@@ -374,7 +374,12 @@ void FMacWindow::SetWindowMode( EWindowMode::Type NewWindowMode )
 	bool bMakeFullscreen = NewWindowMode != EWindowMode::Windowed;
 	bool bIsFullscreen = GetWindowMode() != EWindowMode::Windowed;
 
-	if( bIsFullscreen != bMakeFullscreen )
+	if(bIsFullscreen == bMakeFullscreen && NewWindowMode != GetWindowMode())
+	{
+		SetWindowMode(EWindowMode::Windowed);
+	}
+	
+	if( bIsFullscreen != bMakeFullscreen || NewWindowMode != GetWindowMode() )
 	{
 		bool WindowIsFullScreen = !bMakeFullscreen;
 		
