Framework: Improved the comments on FRotator Pitch/Yaw/Roll and removed INTEL_ORDER_ROTATOR

diff --git a/Engine/Source/Runtime/Core/Public/Math/Rotator.h b/Engine/Source/Runtime/Core/Public/Math/Rotator.h
--- a/Engine/Source/Runtime/Core/Public/Math/Rotator.h
+++ b/Engine/Source/Runtime/Core/Public/Math/Rotator.h
@@ -11,14 +11,13 @@
 struct FRotator
 {
 public:
-
-	/** Looking up and down (0=Straight Ahead, +Up, -Down). */
+	/** Rotation around the right axis (around Y axis), Looking up and down (0=Straight Ahead, +Up, -Down) */
 	float Pitch; 
 
-	/** Rotating around (running in circles), 0=East, +North, -South. */
+	/** Rotation around the up axis (around Z axis), Running in circles 0=East, +North, -South. */
 	float Yaw; 
 
-	/** Rotation about axis of screen, 0=Straight, +Clockwise, -CCW. */
+	/** Rotation around the forward axis (around X axis), Tilting your head, 0=Straight, +Clockwise, -CCW. */
 	float Roll;
 
 public:
@@ -31,7 +30,7 @@ public:
 	/**
 	 * Default constructor (no initialization).
 	 */
-	FORCEINLINE FRotator() { };
+	FORCEINLINE FRotator() { }
 
 	/**
 	 * Constructor
@@ -73,7 +72,7 @@ public:
 	 * @param R The other rotator.
 	 * @return The result of adding a rotator to this.
 	 */
-	FRotator operator+( const FRotator &R ) const;
+	FRotator operator+( const FRotator& R ) const;
 
 	/**
 	 * Get the result of subtracting a rotator from this.
@@ -81,7 +80,7 @@ public:
 	 * @param R The other rotator.
 	 * @return The result of subtracting a rotator from this.
 	 */
-	FRotator operator-( const FRotator &R ) const;
+	FRotator operator-( const FRotator& R ) const;
 
 	/**
 	 * Get the result of scaling this rotator.
@@ -117,7 +116,7 @@ public:
 	 * @param R The other rotator.
 	 * @return true if two rotators are identical, otherwise false.
 	 */
-	bool operator==( const FRotator &R ) const;
+	bool operator==( const FRotator& R ) const;
 
 	/**
 	 * Checks whether two rotators are different.
@@ -125,7 +124,7 @@ public:
 	 * @param V The other rotator.
 	 * @return true if two rotators are different, otherwise false.
 	 */
-	bool operator!=( const FRotator &V ) const;
+	bool operator!=( const FRotator& V ) const;
 
 	// Assignment operators.
 
@@ -135,7 +134,7 @@ public:
 	 * @param R The other rotator.
 	 * @return Copy of rotator after addition.
 	 */
-	FRotator operator+=( const FRotator &R );
+	FRotator operator+=( const FRotator& R );
 
 	/**
 	 * Subtracts another rotator from this.
@@ -143,7 +142,7 @@ public:
 	 * @param R The other rotator.
 	 * @return Copy of rotator after subtraction.
 	 */
-	FRotator operator-=( const FRotator &R );
+	FRotator operator-=( const FRotator& R );
 
 public:
 
@@ -171,7 +170,7 @@ public:
 	 * @param Tolerance Error Tolerance.
 	 * @return true if two rotators are equal, within specified tolerance, otherwise false.
 	 */
-	bool Equals( const FRotator &R, float Tolerance=KINDA_SMALL_NUMBER ) const;
+	bool Equals( const FRotator& R, float Tolerance=KINDA_SMALL_NUMBER ) const;
 
 	/**
 	 * Adds to each component of the rotator.
@@ -194,7 +193,7 @@ public:
 	 * @param RotGrid A Rotator specifying how to snap each component.
 	 * @return Snapped version of rotation.
 	 */
-	FRotator GridSnap( const FRotator &RotGrid ) const;
+	FRotator GridSnap( const FRotator& RotGrid ) const;
 
 	/**
 	 * Convert a rotation into a vector facing in its direction.
@@ -403,20 +402,11 @@ public:
  * @param R rotator to be scaled.
  * @return Scaled rotator.
  */
-FORCEINLINE FRotator operator*( float Scale, const FRotator &R )
+FORCEINLINE FRotator operator*( float Scale, const FRotator& R )
 {
 	return R.operator*( Scale );
 }
 
-#if PLATFORM_LITTLE_ENDIAN
-	#define INTEL_ORDER_ROTATOR(x) (x)
-#else
-	static FORCEINLINE FRotator INTEL_ORDER_ROTATOR(FRotator r)
-	{
-		return FRotator(INTEL_ORDER32(r.Pitch), INTEL_ORDER32(r.Yaw), INTEL_ORDER32(r.Roll));
-	}
-#endif
-
 
 FORCEINLINE FRotator::FRotator( float InF ) 
 	:	Pitch(InF), Yaw(InF), Roll(InF) 
@@ -433,13 +423,13 @@ FORCEINLINE FRotator::FRotator(EForceInit)
 {}
 
 
-FORCEINLINE FRotator FRotator::operator+( const FRotator &R ) const
+FORCEINLINE FRotator FRotator::operator+( const FRotator& R ) const
 {
 	return FRotator( Pitch+R.Pitch, Yaw+R.Yaw, Roll+R.Roll );
 }
 
 
-FORCEINLINE FRotator FRotator::operator-( const FRotator &R ) const
+FORCEINLINE FRotator FRotator::operator-( const FRotator& R ) const
 {
 	return FRotator( Pitch-R.Pitch, Yaw-R.Yaw, Roll-R.Roll );
 }
@@ -464,26 +454,26 @@ FORCEINLINE FRotator FRotator::operator-() const
 }
 
 
-FORCEINLINE bool FRotator::operator==( const FRotator &R ) const
+FORCEINLINE bool FRotator::operator==( const FRotator& R ) const
 {
 	return Pitch==R.Pitch && Yaw==R.Yaw && Roll==R.Roll;
 }
 
 
-FORCEINLINE bool FRotator::operator!=( const FRotator &V ) const
+FORCEINLINE bool FRotator::operator!=( const FRotator& V ) const
 {
 	return Pitch!=V.Pitch || Yaw!=V.Yaw || Roll!=V.Roll;
 }
 
 
-FORCEINLINE FRotator FRotator::operator+=( const FRotator &R )
+FORCEINLINE FRotator FRotator::operator+=( const FRotator& R )
 {
 	Pitch += R.Pitch; Yaw += R.Yaw; Roll += R.Roll;
 	return *this;
 }
 
 
-FORCEINLINE FRotator FRotator::operator-=( const FRotator &R )
+FORCEINLINE FRotator FRotator::operator-=( const FRotator& R )
 {
 	Pitch -= R.Pitch; Yaw -= R.Yaw; Roll -= R.Roll;
 	return *this;
@@ -512,7 +502,7 @@ FORCEINLINE bool FRotator::IsZero() const
 }
 
 
-FORCEINLINE bool FRotator::Equals(const FRotator &R, float Tolerance) const
+FORCEINLINE bool FRotator::Equals(const FRotator& R, float Tolerance) const
 {
 #if PLATFORM_ENABLE_VECTORINTRINSICS
 	const VectorRegister RegA = VectorLoadFloat3_W0(this);
@@ -537,7 +527,7 @@ FORCEINLINE FRotator FRotator::Add( float DeltaPitch, float DeltaYaw, float Delt
 }
 
 
-FORCEINLINE FRotator FRotator::GridSnap( const FRotator &RotGrid ) const
+FORCEINLINE FRotator FRotator::GridSnap( const FRotator& RotGrid ) const
 {
 	return FRotator
 		(
