Skip VC++ version checks in WindowsPlatformCompilerSetup.h when running Unreal Code Analyzer, as it runs clang on Windows.

diff --git a/Engine/Source/Runtime/Core/Public/Windows/WindowsPlatformCompilerSetup.h b/Engine/Source/Runtime/Core/Public/Windows/WindowsPlatformCompilerSetup.h
--- a/Engine/Source/Runtime/Core/Public/Windows/WindowsPlatformCompilerSetup.h
+++ b/Engine/Source/Runtime/Core/Public/Windows/WindowsPlatformCompilerSetup.h
@@ -6,18 +6,20 @@
 #error PLATFORM_WINDOWS not defined
 #endif
 
+/** Unreal Code Analyzer runs clang on Windows. */
+#ifndef UNREAL_CODE_ANALYZER
 /**
- * Future-proofing the min version check so we keep bumping it whenever we upgrade.
- */ 
-#if _MSC_VER > 1900
+* Future-proofing the min version check so we keep bumping it whenever we upgrade.
+*/
+#if _MSC_VER > 1900 
 	#pragma message("Detected compiler newer than Visual Studio 2013, please update min version checking in WindowsPlatformCompilerSetup.h")
 #endif
 
-
 /**
  * We require at least Visual Studio 2013 to compile
  */
 static_assert(_MSC_VER >= 1800, "Visual Studio 2013 or later is required to compile on Windows platform");
+#endif // UNREAL_CODE_ANALYZER
 
 #pragma warning (error:      4001 4002 4003 4004      4006 4007 4008 4009 4010 4011 4012 4013 4014 4015 4016 4017 4018 4019 4020 4021 4022 4023 4024 4025 4026 4027 4028 4029 4030 4031 4032 4033 4034 4035 4036 4037 4038 4039 4040 4041 4042 4043 4044 4045 4046 4047 4048 4049 4050 4051 4052 4053 4054 4055 4056 4057 4058 4059 4060           4063 4064 4065 4066 4067 4068 4069 4070 4071 4072 4073 4074 4075 4076 4077 4078 4079 4080 4081 4082 4083 4084 4085 4086 4087 4088 4089 4090 4091 4092 4093 4094 4095 4096 4097 4098 4099)
 #pragma warning (error: 4100 4101 4102 4103 4104 4105 4106 4107 4108 4109 4110 4111 4112 4113 4114 4115 4116 4117 4118 4119 4120 4121 4122 4123 4124 4125 4126 4127 4128 4129 4130 4131 4132 4133 4134 4135 4136 4137 4138 4139 4140 4141 4142 4143 4144 4145 4146 4147 4148 4149 4150 4151 4152 4153 4154 4155 4156 4157 4158 4159 4160 4161 4162 4163 4164      4166 4167 4168 4169 4170 4171 4172 4173 4174 4175 4176 4177 4178 4179 4180 4181 4182 4183 4184 4185 4186 4187 4188 4189 4190 4191 4192 4193 4194 4195 4196 4197 4198 4199)
