Modify the when we ensure that the Mac modifier key state is kept up-to-date to avoid updating it with garbage modifier states from NSEvent types that don't guarantee to set them. Fixes JIRA UE-14960.

diff --git a/Engine/Source/Runtime/Core/Private/Mac/MacApplication.cpp b/Engine/Source/Runtime/Core/Private/Mac/MacApplication.cpp
--- a/Engine/Source/Runtime/Core/Private/Mac/MacApplication.cpp
+++ b/Engine/Source/Runtime/Core/Private/Mac/MacApplication.cpp
@@ -499,23 +499,8 @@ void FMacApplication::ProcessEvent(const FDeferredMacEvent& Event)
 	TSharedPtr<FMacWindow> EventWindow = FindWindowByNSWindow(Event.Window);
 	if (Event.Type)
 	{
-		if (CurrentModifierFlags != Event.ModifierFlags)
-		{
-			NSUInteger ModifierFlags = Event.ModifierFlags;
-
-			HandleModifierChange(ModifierFlags, (1<<4), 7, MMK_RightCommand);
-			HandleModifierChange(ModifierFlags, (1<<3), 6, MMK_LeftCommand);
-			HandleModifierChange(ModifierFlags, (1<<1), 0, MMK_LeftShift);
-			HandleModifierChange(ModifierFlags, (1<<16), 8, MMK_CapsLock);
-			HandleModifierChange(ModifierFlags, (1<<5), 4, MMK_LeftAlt);
-			HandleModifierChange(ModifierFlags, (1<<0), 2, MMK_LeftControl);
-			HandleModifierChange(ModifierFlags, (1<<2), 1, MMK_RightShift);
-			HandleModifierChange(ModifierFlags, (1<<6), 5, MMK_RightAlt);
-			HandleModifierChange(ModifierFlags, (1<<13), 3, MMK_RightControl);
-
-			CurrentModifierFlags = ModifierFlags;
-		}
-
+		bool bModifiersValid = false;
+		
 		switch (Event.Type)
 		{
 			case NSMouseMoved:
@@ -523,22 +508,26 @@ void FMacApplication::ProcessEvent(const FDeferredMacEvent& Event)
 			case NSRightMouseDragged:
 			case NSOtherMouseDragged:
 				ProcessMouseMovedEvent(Event, EventWindow);
+				bModifiersValid = true;
 				break;
 
 			case NSLeftMouseDown:
 			case NSRightMouseDown:
 			case NSOtherMouseDown:
 				ProcessMouseDownEvent(Event, EventWindow);
+				bModifiersValid = true;
 				break;
 
 			case NSLeftMouseUp:
 			case NSRightMouseUp:
 			case NSOtherMouseUp:
 				ProcessMouseUpEvent(Event, EventWindow);
+				bModifiersValid = true;
 				break;
 
 			case NSScrollWheel:
 				ProcessScrollWheelEvent(Event, EventWindow);
+				bModifiersValid = true;
 				break;
 
 			case NSEventTypeMagnify:
@@ -547,16 +536,41 @@ void FMacApplication::ProcessEvent(const FDeferredMacEvent& Event)
 			case NSEventTypeBeginGesture:
 			case NSEventTypeEndGesture:
 				ProcessGestureEvent(Event);
+				bModifiersValid = true;
 				break;
 
 			case NSKeyDown:
 				ProcessKeyDownEvent(Event, EventWindow);
+				bModifiersValid = true;
 				break;
 
 			case NSKeyUp:
 				ProcessKeyUpEvent(Event);
+				bModifiersValid = true;
+				break;
+				
+			case NSMouseEntered:
+			case NSMouseExited:
+				bModifiersValid = true;
 				break;
 		}
+		
+		if (bModifiersValid && CurrentModifierFlags != Event.ModifierFlags)
+		{
+			NSUInteger ModifierFlags = Event.ModifierFlags;
+			
+			HandleModifierChange(ModifierFlags, (1<<4), 7, MMK_RightCommand);
+			HandleModifierChange(ModifierFlags, (1<<3), 6, MMK_LeftCommand);
+			HandleModifierChange(ModifierFlags, (1<<1), 0, MMK_LeftShift);
+			HandleModifierChange(ModifierFlags, (1<<16), 8, MMK_CapsLock);
+			HandleModifierChange(ModifierFlags, (1<<5), 4, MMK_LeftAlt);
+			HandleModifierChange(ModifierFlags, (1<<0), 2, MMK_LeftControl);
+			HandleModifierChange(ModifierFlags, (1<<2), 1, MMK_RightShift);
+			HandleModifierChange(ModifierFlags, (1<<6), 5, MMK_RightAlt);
+			HandleModifierChange(ModifierFlags, (1<<13), 3, MMK_RightControl);
+			
+			CurrentModifierFlags = ModifierFlags;
+		}
 	}
 	else if (EventWindow.IsValid())
 	{
