Move PS4 drawcall stats out to a verbose group.

diff --git a/Engine/Source/Runtime/Core/Public/Stats/Stats2.h b/Engine/Source/Runtime/Core/Public/Stats/Stats2.h
--- a/Engine/Source/Runtime/Core/Public/Stats/Stats2.h
+++ b/Engine/Source/Runtime/Core/Public/Stats/Stats2.h
@@ -2005,6 +2005,7 @@ DECLARE_STATS_GROUP(TEXT("FPS Chart"),STATGROUP_FPSChart, STATCAT_Advanced);
 DECLARE_STATS_GROUP(TEXT("GPU Particles"),STATGROUP_GPUParticles, STATCAT_Advanced);
 DECLARE_STATS_GROUP(TEXT("Game"),STATGROUP_Game, STATCAT_Advanced);
 DECLARE_STATS_GROUP(TEXT("Gnm"),STATGROUP_PS4RHI, STATCAT_Advanced);
+DECLARE_STATS_GROUP_VERBOSE(TEXT("GnmVerbose"), STATGROUP_PS4RHIVERBOSE, STATCAT_Advanced);
 DECLARE_STATS_GROUP(TEXT("Init Views"),STATGROUP_InitViews, STATCAT_Advanced);
 DECLARE_STATS_GROUP(TEXT("Landscape"),STATGROUP_Landscape, STATCAT_Advanced);
 DECLARE_STATS_GROUP(TEXT("Light Rendering"),STATGROUP_LightRendering, STATCAT_Advanced);
