ensureOnce is now ensure.
ensure is now ensureAlways.
Redundant ensureMsg removed and replaced with ensureMsgf.
'debug ensureAlways' command added.

diff --git a/Engine/Plugins/2D/Paper2D/Source/PaperTiledImporter/Private/PaperTiledImporterFactory.cpp b/Engine/Plugins/2D/Paper2D/Source/PaperTiledImporter/Private/PaperTiledImporterFactory.cpp
--- a/Engine/Plugins/2D/Paper2D/Source/PaperTiledImporter/Private/PaperTiledImporterFactory.cpp
+++ b/Engine/Plugins/2D/Paper2D/Source/PaperTiledImporter/Private/PaperTiledImporterFactory.cpp
@@ -348,7 +348,7 @@ EReimportResult::Type UPaperTiledImporterFactory::Reimport(UObject* Obj)
 	if (UPaperTileMap* TileMap = Cast<UPaperTileMap>(Obj))
 	{
 		//@TODO: Not implemented yet
-		ensureMsg(false, TEXT("Tile map reimport is not implemented yet"));
+		ensureMsgf(false, TEXT("Tile map reimport is not implemented yet"));
 	}
 	return EReimportResult::Failed;
 }
@@ -1130,7 +1130,7 @@ void FTiledObject::AddToSpriteGeometryCollection(const FVector2D& Offset, const
 			UE_LOG(LogPaperTiledImporter, Warning, TEXT("Ignoring Tiled Object of type Polyline"));
 			break;
 		default:
-			ensureMsg(false, TEXT("Unknown enumerant in ETiledObjectType"));
+			ensureMsgf(false, TEXT("Unknown enumerant in ETiledObjectType"));
 			break;
 		}
 
diff --git a/Engine/Source/Developer/BlueprintCompilerCppBackend/Private/BlueprintCompilerCppBackend.cpp b/Engine/Source/Developer/BlueprintCompilerCppBackend/Private/BlueprintCompilerCppBackend.cpp
--- a/Engine/Source/Developer/BlueprintCompilerCppBackend/Private/BlueprintCompilerCppBackend.cpp
+++ b/Engine/Source/Developer/BlueprintCompilerCppBackend/Private/BlueprintCompilerCppBackend.cpp
@@ -147,7 +147,7 @@ FString FBlueprintCompilerCppBackend::TermToText(const FBPTerminal* Term, const
 			}
 			else
 			{
-				ensureMsg(false, TEXT("It is not possible to express this interface property as a literal value!"));
+				ensureMsgf(false, TEXT("It is not possible to express this interface property as a literal value!"));
 				return Term->Name;
 			}
 		}
@@ -155,7 +155,7 @@ FString FBlueprintCompilerCppBackend::TermToText(const FBPTerminal* Term, const
 			// else if (CoerceProperty->IsA(UMulticastDelegateProperty::StaticClass()))
 			// Cannot assign a literal to a multicast delegate; it should be added instead of assigned
 		{
-			ensureMsg(false, TEXT("It is not possible to express this type as a literal value!"));
+			ensureMsgf(false, TEXT("It is not possible to express this type as a literal value!"));
 			return Term->Name;
 		}
 	}
diff --git a/Engine/Source/Developer/DirectoryWatcher/Private/FileCache.cpp b/Engine/Source/Developer/DirectoryWatcher/Private/FileCache.cpp
--- a/Engine/Source/Developer/DirectoryWatcher/Private/FileCache.cpp
+++ b/Engine/Source/Developer/DirectoryWatcher/Private/FileCache.cpp
@@ -221,7 +221,7 @@ TOptional<FDirectoryState> FAsyncDirectoryReader::GetLiveState()
 {
 	TOptional<FDirectoryState> OldState;
 
-	if (ensureMsg(IsComplete(), TEXT("Invalid property access from thread before task completion")))
+	if (ensureMsgf(IsComplete(), TEXT("Invalid property access from thread before task completion")))
 	{
 		Swap(OldState, LiveState);
 	}
@@ -234,7 +234,7 @@ TOptional<FDirectoryState> FAsyncDirectoryReader::GetCachedState()
 {
 	TOptional<FDirectoryState> OldState;
 
-	if (ensureMsg(IsComplete(), TEXT("Invalid property access from thread before task completion")))
+	if (ensureMsgf(IsComplete(), TEXT("Invalid property access from thread before task completion")))
 	{
 		Swap(OldState, CachedState);
 	}
diff --git a/Engine/Source/Developer/DirectoryWatcher/Public/FileCache.h b/Engine/Source/Developer/DirectoryWatcher/Public/FileCache.h
--- a/Engine/Source/Developer/DirectoryWatcher/Public/FileCache.h
+++ b/Engine/Source/Developer/DirectoryWatcher/Public/FileCache.h
@@ -348,7 +348,7 @@ struct FFileCacheConfig
 	/** Set up this cache to generate MD5 hashes for its constituent files */
 	FFileCacheConfig& RequireFileHashes(bool bInRequireFileHashes)
 	{
-		if (ensureMsg(bInRequireFileHashes || !bDetectMoves, TEXT("Unable to disable file hashing when move detection is enabled")))
+		if (ensureMsgf(bInRequireFileHashes || !bDetectMoves, TEXT("Unable to disable file hashing when move detection is enabled")))
 		{
 			bRequireFileHashes = bInRequireFileHashes;
 		}
diff --git a/Engine/Source/Editor/BlueprintGraph/Private/BlueprintActionDatabaseRegistrar.cpp b/Engine/Source/Editor/BlueprintGraph/Private/BlueprintActionDatabaseRegistrar.cpp
--- a/Engine/Source/Editor/BlueprintGraph/Private/BlueprintActionDatabaseRegistrar.cpp
+++ b/Engine/Source/Editor/BlueprintGraph/Private/BlueprintActionDatabaseRegistrar.cpp
@@ -349,7 +349,7 @@ int32 FBlueprintActionDatabaseRegistrar::RegisterClassFactoryActions(const UClas
 //------------------------------------------------------------------------------
 bool FBlueprintActionDatabaseRegistrar::AddActionToDatabase(UObject const* ActionKey, UBlueprintNodeSpawner* NodeSpawner)
 {
-	ensureMsg(NodeSpawner->NodeClass == GeneratingClass, TEXT("We expect a nodes to add only spawners for its own type... Maybe a sub-class is adding nodes it shouldn't?"));
+	ensureMsgf(NodeSpawner->NodeClass == GeneratingClass, TEXT("We expect a nodes to add only spawners for its own type... Maybe a sub-class is adding nodes it shouldn't?"));
 	if (IsOpenForRegistration(ActionKey))
 	{
 		ActionKey = BlueprintActionDatabaseRegistrarImpl::ResolveActionKey(ActionKey);
diff --git a/Engine/Source/Editor/BlueprintGraph/Private/BlueprintActionFilter.cpp b/Engine/Source/Editor/BlueprintGraph/Private/BlueprintActionFilter.cpp
--- a/Engine/Source/Editor/BlueprintGraph/Private/BlueprintActionFilter.cpp
+++ b/Engine/Source/Editor/BlueprintGraph/Private/BlueprintActionFilter.cpp
@@ -1385,7 +1385,7 @@ static bool BlueprintActionFilterImpl::IsMissmatchedPropertyType(FBlueprintActio
 				}
 				else
 				{
-					ensureMsg(false, TEXT("Unhandled property/node pair, we've probably made some bad assuptions."));
+					ensureMsgf(false, TEXT("Unhandled property/node pair, we've probably made some bad assuptions."));
 				}
 			}
 		}
diff --git a/Engine/Source/Editor/GraphEditor/Private/SGraphEditorImpl.cpp b/Engine/Source/Editor/GraphEditor/Private/SGraphEditorImpl.cpp
--- a/Engine/Source/Editor/GraphEditor/Private/SGraphEditorImpl.cpp
+++ b/Engine/Source/Editor/GraphEditor/Private/SGraphEditorImpl.cpp
@@ -138,7 +138,7 @@ void SGraphEditorImpl::SelectAllNodes()
 		UEdGraphNode* Node = EdGraphObj->Nodes[NodeIndex];
 		if (Node)
 		{
-			ensureMsg(Node->IsValidLowLevel(), TEXT("Node is invalid"));
+			ensureMsgf(Node->IsValidLowLevel(), TEXT("Node is invalid"));
 			NewSet.Add(Node);
 		}
 	}
diff --git a/Engine/Source/Editor/GraphEditor/Private/SNodePanel.cpp b/Engine/Source/Editor/GraphEditor/Private/SNodePanel.cpp
--- a/Engine/Source/Editor/GraphEditor/Private/SNodePanel.cpp
+++ b/Engine/Source/Editor/GraphEditor/Private/SNodePanel.cpp
@@ -132,7 +132,7 @@ void FGraphSelectionManager::SetSelectionSet(FGraphPanelSelectionSet& NewSet)
 
 void FGraphSelectionManager::SetNodeSelection(SelectedItemType Node, bool bSelect)
 {
-	ensureMsg(Node != nullptr, TEXT("Node is invalid"));
+	ensureMsgf(Node != nullptr, TEXT("Node is invalid"));
 	if (bSelect)
 	{
 		SelectedNodes.Add(Node);
diff --git a/Engine/Source/Editor/Kismet/Private/BlueprintEditor.cpp b/Engine/Source/Editor/Kismet/Private/BlueprintEditor.cpp
--- a/Engine/Source/Editor/Kismet/Private/BlueprintEditor.cpp
+++ b/Engine/Source/Editor/Kismet/Private/BlueprintEditor.cpp
@@ -6608,7 +6608,7 @@ void FBlueprintEditor::NewDocument_OnClicked(ECreatedDocumentType GraphType)
 	}
 	else
 	{
-		ensureMsg(false, TEXT("GraphType is invalid") );
+		ensureMsgf(false, TEXT("GraphType is invalid") );
 	}
 
 	// Now open the new graph
diff --git a/Engine/Source/Editor/Kismet/Private/SKismetInspector.cpp b/Engine/Source/Editor/Kismet/Private/SKismetInspector.cpp
--- a/Engine/Source/Editor/Kismet/Private/SKismetInspector.cpp
+++ b/Engine/Source/Editor/Kismet/Private/SKismetInspector.cpp
@@ -531,7 +531,7 @@ void SKismetInspector::UpdateFromObjects(const TArray<UObject*>& PropertyObjects
 		{
 			if (!Object->IsValidLowLevel())
 			{
-				ensureMsg(false, TEXT("Object in KismetInspector is invalid, see TTP 281915"));
+				ensureMsgf(false, TEXT("Object in KismetInspector is invalid, see TTP 281915"));
 				continue;
 			}
 
diff --git a/Engine/Source/Editor/Kismet/Private/SSCSEditor.cpp b/Engine/Source/Editor/Kismet/Private/SSCSEditor.cpp
--- a/Engine/Source/Editor/Kismet/Private/SSCSEditor.cpp
+++ b/Engine/Source/Editor/Kismet/Private/SSCSEditor.cpp
@@ -1036,7 +1036,7 @@ FSCSEditorTreeNodeComponent::FSCSEditorTreeNodeComponent(UActorComponent* InComp
 	AActor* Owner = InComponentTemplate->GetOwner();
 	if (Owner != nullptr)
 	{
-		ensureMsg(Owner->HasAllFlags(RF_ClassDefaultObject), TEXT("Use a different node class for instanced components"));
+		ensureMsgf(Owner->HasAllFlags(RF_ClassDefaultObject), TEXT("Use a different node class for instanced components"));
 	}
 }
 
diff --git a/Engine/Source/Editor/KismetCompiler/Private/KismetCompiler.cpp b/Engine/Source/Editor/KismetCompiler/Private/KismetCompiler.cpp
--- a/Engine/Source/Editor/KismetCompiler/Private/KismetCompiler.cpp
+++ b/Engine/Source/Editor/KismetCompiler/Private/KismetCompiler.cpp
@@ -1819,7 +1819,7 @@ void FKismetCompilerContext::FinishCompilingClass(UClass* Class)
 	AActor* ActorCDO = Cast<AActor>(Class->GetDefaultObject());
 	if (ActorCDO)
 	{
-		ensureMsg(!ActorCDO->bExchangedRoles, TEXT("Your CDO has had ExchangeNetRoles called on it (likely via RerunConstructionScripts) which should never have happened. This will cause issues replicating this actor over the network due to mutated transient data!"));
+		ensureMsgf(!ActorCDO->bExchangedRoles, TEXT("Your CDO has had ExchangeNetRoles called on it (likely via RerunConstructionScripts) which should never have happened. This will cause issues replicating this actor over the network due to mutated transient data!"));
 	}
 }
 
diff --git a/Engine/Source/Editor/KismetCompiler/Private/KismetCompilerVMBackend.cpp b/Engine/Source/Editor/KismetCompiler/Private/KismetCompilerVMBackend.cpp
--- a/Engine/Source/Editor/KismetCompiler/Private/KismetCompilerVMBackend.cpp
+++ b/Engine/Source/Editor/KismetCompiler/Private/KismetCompilerVMBackend.cpp
@@ -456,7 +456,7 @@ public:
 			}
 			//else if (UClassProperty* ClassProperty = Cast<UClassProperty>(CoerceProperty))
 			//{
-			//	ensureMsg(false, TEXT("Class property literals are not supported yet!"));
+			//	ensureMsgf(false, TEXT("Class property literals are not supported yet!"));
 			//}
 			else if (UStructProperty* StructProperty = Cast<UStructProperty>(CoerceProperty))
 			{
@@ -563,7 +563,7 @@ public:
 			{
 				if (Term->Name == TEXT(""))
 				{
-					ensureMsg(false, TEXT("Cannot use an empty literal expression for a delegate property"));
+					ensureMsgf(false, TEXT("Cannot use an empty literal expression for a delegate property"));
 				}
 				else
 				{
@@ -1220,7 +1220,7 @@ public:
 		}
 		else
 		{
-			ensureMsg(false, TEXT("FScriptBuilderBase::EmitGoto unknown type"));
+			ensureMsgf(false, TEXT("FScriptBuilderBase::EmitGoto unknown type"));
 		}
 	}
 
diff --git a/Engine/Source/Editor/UnrealEd/Private/Kismet2/BlueprintEditorUtils.cpp b/Engine/Source/Editor/UnrealEd/Private/Kismet2/BlueprintEditorUtils.cpp
--- a/Engine/Source/Editor/UnrealEd/Private/Kismet2/BlueprintEditorUtils.cpp
+++ b/Engine/Source/Editor/UnrealEd/Private/Kismet2/BlueprintEditorUtils.cpp
@@ -1229,7 +1229,7 @@ struct FEditoronlyBlueprintHelper
 								}
 								else
 								{
-									ensureMsg(false, TEXT("Can not describe why the blueprint should be fixed."));
+									ensureMsgf(false, TEXT("Can not describe why the blueprint should be fixed."));
 									UE_LOG(LogBlueprint, Warning, TEXT("%s Unknown reason. Pin '%s' on node '%s'"), *ReasonPrefix, *PinName, *PinNodeName);
 								}
 							}
diff --git a/Engine/Source/Editor/UnrealEd/Private/Kismet2/KismetDebugUtilities.cpp b/Engine/Source/Editor/UnrealEd/Private/Kismet2/KismetDebugUtilities.cpp
--- a/Engine/Source/Editor/UnrealEd/Private/Kismet2/KismetDebugUtilities.cpp
+++ b/Engine/Source/Editor/UnrealEd/Private/Kismet2/KismetDebugUtilities.cpp
@@ -465,7 +465,7 @@ void FKismetDebugUtilities::AttemptToBreakExecution(UBlueprint* BlueprintObj, co
 	{
 		bShouldInStackDebug = false;
 		//@TODO: Determine exactly what behavior we want for breakpoints hit when not in PIE/SIE
-		//ensureMsg(false, TEXT("Breakpoints placed in a function instead of the event graph are not supported yet"));
+		//ensureMsgf(false, TEXT("Breakpoints placed in a function instead of the event graph are not supported yet"));
 	}
 
 	// Now enter within-the-frame debugging mode
@@ -866,7 +866,7 @@ FKismetDebugUtilities::EWatchTextResult FKismetDebugUtilities::GetWatchText(FStr
 			static bool bErrorOnce = true;
 			if (bErrorOnce)
 			{
-				ensureMsg(false, TEXT("Error: Invalid (but non-null) property associated with pin; cannot get variable value"));
+				ensureMsgf(false, TEXT("Error: Invalid (but non-null) property associated with pin; cannot get variable value"));
 				bErrorOnce = false;
 			}
 			return EWTR_NoProperty;
diff --git a/Engine/Source/Editor/UnrealEd/Private/SkeletalMeshEdit.cpp b/Engine/Source/Editor/UnrealEd/Private/SkeletalMeshEdit.cpp
--- a/Engine/Source/Editor/UnrealEd/Private/SkeletalMeshEdit.cpp
+++ b/Engine/Source/Editor/UnrealEd/Private/SkeletalMeshEdit.cpp
@@ -957,7 +957,7 @@ bool UnFbx::FFbxImporter::ImportCurveToAnimSequence(class UAnimSequence * Target
 			else
 			{
 				// this should not happen, we already checked before adding
-				ensureMsg(0, TEXT("FBX Import: Critical error: no memory?"));
+				ensureMsgf(0, TEXT("FBX Import: Critical error: no memory?"));
 			}
 		}
 		else
diff --git a/Engine/Source/Runtime/Core/Private/Modules/ModuleManager.cpp b/Engine/Source/Runtime/Core/Private/Modules/ModuleManager.cpp
--- a/Engine/Source/Runtime/Core/Private/Modules/ModuleManager.cpp
+++ b/Engine/Source/Runtime/Core/Private/Modules/ModuleManager.cpp
@@ -200,7 +200,7 @@ void FModuleManager::AddModuleToModulesList(const FName InModuleName, TSharedRef
 void FModuleManager::AddModule(const FName InModuleName)
 {
 	// Do we already know about this module?  If not, we'll create information for this module now.
-	if (!((ensureMsg(InModuleName != NAME_None, TEXT("FModuleManager::AddModule() was called with an invalid module name (empty string or 'None'.)  This is not allowed.")) &&
+	if (!((ensureMsgf(InModuleName != NAME_None, TEXT("FModuleManager::AddModule() was called with an invalid module name (empty string or 'None'.)  This is not allowed.")) &&
 		!Modules.Contains(InModuleName))))
 	{
 		return;
@@ -1131,7 +1131,7 @@ bool FModuleManager::DoesLoadedModuleHaveUObjects( const FName ModuleName )
 TSharedRef<FModuleManager::FModuleInfo> FModuleManager::GetOrCreateModule(FName InModuleName)
 {
 	check(IsInGameThread());
-	ensureMsg(InModuleName != NAME_None, TEXT("FModuleManager::AddModule() was called with an invalid module name (empty string or 'None'.)  This is not allowed."));
+	ensureMsgf(InModuleName != NAME_None, TEXT("FModuleManager::AddModule() was called with an invalid module name (empty string or 'None'.)  This is not allowed."));
 
 	if (Modules.Contains(InModuleName))
 	{
diff --git a/Engine/Source/Runtime/Core/Private/Windows/WindowsPlatformProcess.cpp b/Engine/Source/Runtime/Core/Private/Windows/WindowsPlatformProcess.cpp
--- a/Engine/Source/Runtime/Core/Private/Windows/WindowsPlatformProcess.cpp
+++ b/Engine/Source/Runtime/Core/Private/Windows/WindowsPlatformProcess.cpp
@@ -157,7 +157,7 @@ void FWindowsPlatformProcess::PushDllDirectory(const TCHAR* Directory)
 void FWindowsPlatformProcess::PopDllDirectory(const TCHAR* Directory)
 {
 	// don't allow too many pops (indicates bad code that should be fixed, but won't kill anything, so using ensure)
-	ensureMsg(DllDirectoryStack.Num() > 0, TEXT("Tried to PopDllDirectory too many times"));
+	ensureMsgf(DllDirectoryStack.Num() > 0, TEXT("Tried to PopDllDirectory too many times"));
 	// verify we are popping the top
 	checkf(DllDirectoryStack.Top() == Directory, TEXT("There was a PushDllDirectory/PopDllDirectory mismatch (Popped %s, which didn't match %s)"), *DllDirectoryStack.Top(), Directory);
 	// pop it off
diff --git a/Engine/Source/Runtime/Core/Public/Math/BoxSphereBounds.h b/Engine/Source/Runtime/Core/Public/Math/BoxSphereBounds.h
--- a/Engine/Source/Runtime/Core/Public/Math/BoxSphereBounds.h
+++ b/Engine/Source/Runtime/Core/Public/Math/BoxSphereBounds.h
@@ -250,9 +250,9 @@ public:
 #if ENABLE_NAN_DIAGNOSTIC
 	FORCEINLINE void DiagnosticCheckNaN() const
 	{
-		ensureOnceMsgf(!Origin.ContainsNaN(), TEXT("Origin contains NaN: %s"), *Origin.ToString());
-		ensureOnceMsgf(!BoxExtent.ContainsNaN(), TEXT("BoxExtent contains NaN: %s"), *BoxExtent.ToString());
-		ensureOnceMsgf(!FMath::IsNaN(SphereRadius) && FMath::IsFinite(SphereRadius), TEXT("SphereRadius contains NaN: %f"), SphereRadius);
+		ensureMsgf(!Origin.ContainsNaN(), TEXT("Origin contains NaN: %s"), *Origin.ToString());
+		ensureMsgf(!BoxExtent.ContainsNaN(), TEXT("BoxExtent contains NaN: %s"), *BoxExtent.ToString());
+		ensureMsgf(!FMath::IsNaN(SphereRadius) && FMath::IsFinite(SphereRadius), TEXT("SphereRadius contains NaN: %f"), SphereRadius);
 	}
 #else
 	FORCEINLINE void DiagnosticCheckNaN() const {}
diff --git a/Engine/Source/Runtime/Core/Public/Math/Quat.h b/Engine/Source/Runtime/Core/Public/Math/Quat.h
--- a/Engine/Source/Runtime/Core/Public/Math/Quat.h
+++ b/Engine/Source/Runtime/Core/Public/Math/Quat.h
@@ -380,7 +380,7 @@ public:
 #if ENABLE_NAN_DIAGNOSTIC
 	FORCEINLINE void DiagnosticCheckNaN() const
 	{
-		ensureOnceMsgf(!ContainsNaN(), TEXT("FQuat contains NaN: %s"), *ToString());
+		ensureMsgf(!ContainsNaN(), TEXT("FQuat contains NaN: %s"), *ToString());
 	}
 #else
 	FORCEINLINE void DiagnosticCheckNaN() const {}
diff --git a/Engine/Source/Runtime/Core/Public/Math/Transform.h b/Engine/Source/Runtime/Core/Public/Math/Transform.h
--- a/Engine/Source/Runtime/Core/Public/Math/Transform.h
+++ b/Engine/Source/Runtime/Core/Public/Math/Transform.h
@@ -44,17 +44,17 @@ public:
 #if ENABLE_NAN_DIAGNOSTIC
 	FORCEINLINE void DiagnosticCheckNaN_Scale3D() const
 	{
-		ensureOnceMsgf(!Scale3D.ContainsNaN(), TEXT("FTransform Scale3D contains NaN: %s"), *Scale3D.ToString());
+		ensureMsgf(!Scale3D.ContainsNaN(), TEXT("FTransform Scale3D contains NaN: %s"), *Scale3D.ToString());
 	}
 
 	FORCEINLINE void DiagnosticCheckNaN_Translate() const
 	{
-		ensureOnceMsgf(!Translation.ContainsNaN(), TEXT("FTransform Translation contains NaN: %s"), *Translation.ToString());
+		ensureMsgf(!Translation.ContainsNaN(), TEXT("FTransform Translation contains NaN: %s"), *Translation.ToString());
 	}
 
 	FORCEINLINE void DiagnosticCheckNaN_Rotate() const
 	{
-		ensureOnceMsgf(!Rotation.ContainsNaN(), TEXT("FTransform Rotation contains NaN: %s"), *Rotation.ToString());
+		ensureMsgf(!Rotation.ContainsNaN(), TEXT("FTransform Rotation contains NaN: %s"), *Rotation.ToString());
 	}
 
 	FORCEINLINE void DiagnosticCheckNaN_All() const
diff --git a/Engine/Source/Runtime/Core/Public/Math/TransformVectorized.h b/Engine/Source/Runtime/Core/Public/Math/TransformVectorized.h
--- a/Engine/Source/Runtime/Core/Public/Math/TransformVectorized.h
+++ b/Engine/Source/Runtime/Core/Public/Math/TransformVectorized.h
@@ -43,17 +43,17 @@ public:
 #if ENABLE_NAN_DIAGNOSTIC
 	FORCEINLINE void DiagnosticCheckNaN_Scale3D() const
 	{
-		ensureOnceMsgf(!VectorContainsNaNOrInfinite(Scale3D), TEXT("FTransform Vectorized Scale3D contains NaN"));
+		ensureMsgf(!VectorContainsNaNOrInfinite(Scale3D), TEXT("FTransform Vectorized Scale3D contains NaN"));
 	}
 
 	FORCEINLINE void DiagnosticCheckNaN_Translate() const
 	{
-		ensureOnceMsgf(!VectorContainsNaNOrInfinite(Translation), TEXT("FTransform Vectorized Translation contains NaN"));
+		ensureMsgf(!VectorContainsNaNOrInfinite(Translation), TEXT("FTransform Vectorized Translation contains NaN"));
 	}
 
 	FORCEINLINE void DiagnosticCheckNaN_Rotate() const
 	{
-		ensureOnceMsgf(!VectorContainsNaNOrInfinite(Rotation), TEXT("FTransform Vectorized Rotation contains NaN"));
+		ensureMsgf(!VectorContainsNaNOrInfinite(Rotation), TEXT("FTransform Vectorized Rotation contains NaN"));
 	}
 
 	FORCEINLINE void DiagnosticCheckNaN_All() const
diff --git a/Engine/Source/Runtime/Core/Public/Math/UnrealMatrix.h b/Engine/Source/Runtime/Core/Public/Math/UnrealMatrix.h
--- a/Engine/Source/Runtime/Core/Public/Math/UnrealMatrix.h
+++ b/Engine/Source/Runtime/Core/Public/Math/UnrealMatrix.h
@@ -263,7 +263,7 @@ inline FMatrix FMatrix::InverseFast() const
 		GetScaledAxis( EAxis::Z ).IsNearlyZero(SMALL_NUMBER) ) 
 	{
 		UE_LOG(LogUnrealMath, Error, TEXT("FMatrix::InverseFast(), trying to invert a NIL matrix, this results in NaNs! Use Inverse() instead."));
-		ensureMsg(false, TEXT("FMatrix::InverseFast(), trying to invert a NIL matrix, this results in NaNs! Use Inverse() instead."));
+		ensureMsgf(false, TEXT("FMatrix::InverseFast(), trying to invert a NIL matrix, this results in NaNs! Use Inverse() instead."));
 	}
 #endif
 	FMatrix Result;
diff --git a/Engine/Source/Runtime/Core/Public/Math/Vector.h b/Engine/Source/Runtime/Core/Public/Math/Vector.h
--- a/Engine/Source/Runtime/Core/Public/Math/Vector.h
+++ b/Engine/Source/Runtime/Core/Public/Math/Vector.h
@@ -38,7 +38,7 @@ public:
 #if ENABLE_NAN_DIAGNOSTIC
 	FORCEINLINE void DiagnosticCheckNaN() const
 	{
-		ensureOnceMsgf(!ContainsNaN(), TEXT("FVector contains NaN: %s"), *ToString());
+		ensureMsgf(!ContainsNaN(), TEXT("FVector contains NaN: %s"), *ToString());
 	}
 #else
 	FORCEINLINE void DiagnosticCheckNaN() const {}
diff --git a/Engine/Source/Runtime/Core/Public/Math/Vector2D.h b/Engine/Source/Runtime/Core/Public/Math/Vector2D.h
--- a/Engine/Source/Runtime/Core/Public/Math/Vector2D.h
+++ b/Engine/Source/Runtime/Core/Public/Math/Vector2D.h
@@ -471,7 +471,7 @@ public:
 #if ENABLE_NAN_DIAGNOSTIC
 	FORCEINLINE void DiagnosticCheckNaN() const
 	{
-		ensureOnceMsgf(!ContainsNaN(), TEXT("FVector contains NaN: %s"), *ToString());
+		ensureMsgf(!ContainsNaN(), TEXT("FVector contains NaN: %s"), *ToString());
 	}
 #else
 	FORCEINLINE void DiagnosticCheckNaN() const {}
diff --git a/Engine/Source/Runtime/Core/Public/Math/Vector4.h b/Engine/Source/Runtime/Core/Public/Math/Vector4.h
--- a/Engine/Source/Runtime/Core/Public/Math/Vector4.h
+++ b/Engine/Source/Runtime/Core/Public/Math/Vector4.h
@@ -348,7 +348,7 @@ public:
 #if ENABLE_NAN_DIAGNOSTIC
 	FORCEINLINE void DiagnosticCheckNaN() const
 	{
-		ensureOnceMsgf(!ContainsNaN(), TEXT("FVector contains NaN: %s"), *ToString());
+		ensureMsgf(!ContainsNaN(), TEXT("FVector contains NaN: %s"), *ToString());
 	}
 #else
 	FORCEINLINE void DiagnosticCheckNaN() const { }
diff --git a/Engine/Source/Runtime/Core/Public/Misc/AssertionMacros.h b/Engine/Source/Runtime/Core/Public/Misc/AssertionMacros.h
--- a/Engine/Source/Runtime/Core/Public/Misc/AssertionMacros.h
+++ b/Engine/Source/Runtime/Core/Public/Misc/AssertionMacros.h
@@ -115,7 +115,7 @@
 
 	namespace UE4Asserts_Private
 	{
-		// This is used by ensureOnce to generate a bool per instance
+		// This is used by ensure to generate a bool per instance
 		// by passing a lambda which will uniquely instantiate the template.
 		template <typename Type>
 		bool TrueOnFirstCallOnly(const Type&)
@@ -136,19 +136,17 @@
 		}
 	}
 
-	#define ensure(         InExpression                ) ((InExpression) != 0 || FDebug::OptionallyLogFormattedEnsureMessageReturningFalse(true,                                          #InExpression, __FILE__, __LINE__, TEXT("")               ) || FPlatformMisc::DebugBreakAndPromptForRemoteReturningFalse(true))
-	#define ensureMsg(      InExpression, InMsg         ) ((InExpression) != 0 || FDebug::OptionallyLogFormattedEnsureMessageReturningFalse(true,                                          #InExpression, __FILE__, __LINE__, InMsg                  ) || FPlatformMisc::DebugBreakAndPromptForRemoteReturningFalse(true))
-	#define ensureMsgf(     InExpression, InFormat, ... ) ((InExpression) != 0 || FDebug::OptionallyLogFormattedEnsureMessageReturningFalse(true,                                          #InExpression, __FILE__, __LINE__, InFormat, ##__VA_ARGS__) || FPlatformMisc::DebugBreakAndPromptForRemoteReturningFalse(true))
-	#define ensureOnce(     InExpression                ) ((InExpression) != 0 || FDebug::OptionallyLogFormattedEnsureMessageReturningFalse(UE4Asserts_Private::TrueOnFirstCallOnly([]{}), #InExpression, __FILE__, __LINE__, TEXT("")               ) || UE4Asserts_Private::OptionallyDebugBreakAndPromptForRemoteReturningFalse(UE4Asserts_Private::TrueOnFirstCallOnly([]{}), true))
-	#define ensureOnceMsgf( InExpression, InFormat, ... ) ((InExpression) != 0 || FDebug::OptionallyLogFormattedEnsureMessageReturningFalse(UE4Asserts_Private::TrueOnFirstCallOnly([]{}), #InExpression, __FILE__, __LINE__, InFormat, ##__VA_ARGS__) || UE4Asserts_Private::OptionallyDebugBreakAndPromptForRemoteReturningFalse(UE4Asserts_Private::TrueOnFirstCallOnly([]{}), true))
+	#define ensure(           InExpression                ) ((InExpression) != 0 || FDebug::OptionallyLogFormattedEnsureMessageReturningFalse(UE4Asserts_Private::TrueOnFirstCallOnly([]{}), #InExpression, __FILE__, __LINE__, TEXT("")               ) || UE4Asserts_Private::OptionallyDebugBreakAndPromptForRemoteReturningFalse(UE4Asserts_Private::TrueOnFirstCallOnly([]{}), true))
+	#define ensureMsgf(       InExpression, InFormat, ... ) ((InExpression) != 0 || FDebug::OptionallyLogFormattedEnsureMessageReturningFalse(UE4Asserts_Private::TrueOnFirstCallOnly([]{}), #InExpression, __FILE__, __LINE__, InFormat, ##__VA_ARGS__) || UE4Asserts_Private::OptionallyDebugBreakAndPromptForRemoteReturningFalse(UE4Asserts_Private::TrueOnFirstCallOnly([]{}), true))
+	#define ensureAlways(     InExpression                ) ((InExpression) != 0 || FDebug::OptionallyLogFormattedEnsureMessageReturningFalse(true,                                          #InExpression, __FILE__, __LINE__, TEXT("")               ) || FPlatformMisc::DebugBreakAndPromptForRemoteReturningFalse(true))
+	#define ensureAlwaysMsgf( InExpression, InFormat, ... ) ((InExpression) != 0 || FDebug::OptionallyLogFormattedEnsureMessageReturningFalse(true,                                          #InExpression, __FILE__, __LINE__, InFormat, ##__VA_ARGS__) || FPlatformMisc::DebugBreakAndPromptForRemoteReturningFalse(true))
 
 #else	// DO_CHECK
 
-	#define ensure(         InExpression                ) ((InExpression) != 0)
-	#define ensureMsg(      InExpression, InMsg         ) ((InExpression) != 0)
-	#define ensureMsgf(     InExpression, InFormat, ... ) ((InExpression) != 0)
-	#define ensureOnce(     InExpression                ) ((InExpression) != 0)
-	#define ensureOnceMsgf( InExpression, InFormat, ... ) ((InExpression) != 0)
+	#define ensure(           InExpression                ) ((InExpression) != 0)
+	#define ensureMsgf(       InExpression, InFormat, ... ) ((InExpression) != 0)
+	#define ensureAlways(     InExpression                ) ((InExpression) != 0)
+	#define ensureAlwaysMsgf( InExpression, InFormat, ... ) ((InExpression) != 0)
 
 #endif	// DO_CHECK
 
diff --git a/Engine/Source/Runtime/Core/Public/Misc/Build.h b/Engine/Source/Runtime/Core/Public/Misc/Build.h
--- a/Engine/Source/Runtime/Core/Public/Misc/Build.h
+++ b/Engine/Source/Runtime/Core/Public/Misc/Build.h
@@ -170,7 +170,7 @@
 	Basic options that by default depend on the build configuration and platform
 
 	DO_GUARD_SLOW									If true, then checkSlow, checkfSlow and verifySlow are compiled into the executable.
-	DO_CHECK										If true, then checkCode, checkf, verify, check, checkNoEntry, checkNoReentry, checkNoRecursion, verifyf, checkf, ensure, ensureMsg, and ensureMsgf are compiled into the executables
+	DO_CHECK										If true, then checkCode, checkf, verify, check, checkNoEntry, checkNoReentry, checkNoRecursion, verifyf, checkf, ensure, ensureAlways, ensureMsgf and ensureAlwaysMsgf are compiled into the executables
 	STATS											If true, then the stats system is compiled into the executable.
 	ALLOW_DEBUG_FILES								If true, then debug files like screen shots and profiles can be saved from the executable.
 	NO_LOGGING										If true, then no logs or text output will be produced
diff --git a/Engine/Source/Runtime/Core/Public/Misc/FeedbackContext.h b/Engine/Source/Runtime/Core/Public/Misc/FeedbackContext.h
--- a/Engine/Source/Runtime/Core/Public/Misc/FeedbackContext.h
+++ b/Engine/Source/Runtime/Core/Public/Misc/FeedbackContext.h
@@ -210,7 +210,7 @@ public:
 			checkSlow(Stack.Num() != 0 && Stack.Last() == this);
 
 			auto* Task = Stack.Last();
-			if (ensureMsg(Task == this, TEXT("Out-of-order scoped task construction/destruction")))
+			if (ensureMsgf(Task == this, TEXT("Out-of-order scoped task construction/destruction")))
 			{
 				Stack.Pop(false);
 			}
diff --git a/Engine/Source/Runtime/Core/Public/Windows/WindowsPlatformCodeAnalysis.h b/Engine/Source/Runtime/Core/Public/Windows/WindowsPlatformCodeAnalysis.h
--- a/Engine/Source/Runtime/Core/Public/Windows/WindowsPlatformCodeAnalysis.h
+++ b/Engine/Source/Runtime/Core/Public/Windows/WindowsPlatformCodeAnalysis.h
@@ -105,7 +105,7 @@
 	//-V:UE_BUILD_SHIPPING:501
 	//-V:WITH_EDITOR:501
 	//-V:PLATFORM_:517,547
-	//-V:ensureMsg:562
+	//-V:ensureMsgf:562
 	//-V:WindowsMinorVersion:547 
 	//-V:AddUninitialized(sizeof(void*)/:514
 	//-V:TestTrue:678
diff --git a/Engine/Source/Runtime/CoreUObject/Private/UObject/UObjectBaseUtility.cpp b/Engine/Source/Runtime/CoreUObject/Private/UObject/UObjectBaseUtility.cpp
--- a/Engine/Source/Runtime/CoreUObject/Private/UObject/UObjectBaseUtility.cpp
+++ b/Engine/Source/Runtime/CoreUObject/Private/UObject/UObjectBaseUtility.cpp
@@ -283,7 +283,7 @@ bool UObjectBaseUtility::IsA( const UClass* SomeBase ) const
 	#endif
 
 	#if (UCLASS_FAST_ISA_IMPL & 1) && (UCLASS_FAST_ISA_IMPL & 2)
-		ensureOnceMsgf(bOldResult == bNewResult, TEXT("New cast code failed"));
+		ensureMsgf(bOldResult == bNewResult, TEXT("New cast code failed"));
 	#endif
 
 	#if UCLASS_FAST_ISA_IMPL & 1
diff --git a/Engine/Source/Runtime/CoreUObject/Private/UObject/UObjectGlobals.cpp b/Engine/Source/Runtime/CoreUObject/Private/UObject/UObjectGlobals.cpp
--- a/Engine/Source/Runtime/CoreUObject/Private/UObject/UObjectGlobals.cpp
+++ b/Engine/Source/Runtime/CoreUObject/Private/UObject/UObjectGlobals.cpp
@@ -1784,7 +1784,7 @@ bool StaticAllocateObjectErrorTests( UClass* InClass, UObject* InOuter, FName In
 			const FString ErrorMsg = FString::Printf(TEXT("Class which was marked abstract was trying to be loaded.  It will be nulled out on save. %s %s"), *InName.ToString(), *InClass->GetName());
 			// if we are trying instantiate an abstract class in the editor we'll warn the user that it will be nulled out on save
 			UE_LOG(LogUObjectGlobals, Warning, TEXT("%s"), *ErrorMsg);
-			ensureMsg(false, *ErrorMsg);
+			ensureMsgf(false, *ErrorMsg);
 		}
 		else
 		{
diff --git a/Engine/Source/Runtime/Engine/Private/ActiveSound.cpp b/Engine/Source/Runtime/Engine/Private/ActiveSound.cpp
--- a/Engine/Source/Runtime/Engine/Private/ActiveSound.cpp
+++ b/Engine/Source/Runtime/Engine/Private/ActiveSound.cpp
@@ -58,7 +58,7 @@ FActiveSound::FActiveSound()
 
 FActiveSound::~FActiveSound()
 {
-	ensureMsg(WaveInstances.Num() == 0, TEXT("Destroyed an active sound that had active wave instances."));
+	ensureMsgf(WaveInstances.Num() == 0, TEXT("Destroyed an active sound that had active wave instances."));
 }
 
 FArchive& operator<<( FArchive& Ar, FActiveSound* ActiveSound )
diff --git a/Engine/Source/Runtime/Engine/Private/Animation/AnimSequenceBase.cpp b/Engine/Source/Runtime/Engine/Private/Animation/AnimSequenceBase.cpp
--- a/Engine/Source/Runtime/Engine/Private/Animation/AnimSequenceBase.cpp
+++ b/Engine/Source/Runtime/Engine/Private/Animation/AnimSequenceBase.cpp
@@ -354,7 +354,7 @@ void UAnimSequenceBase::UpdateAnimNotifyTrackCache()
 		{
 			// this notifyindex isn't valid, delete
 			// this should not happen, but if it doesn, find best place to add
-			ensureMsg(0, TEXT("AnimNotifyTrack: Wrong indices found"));
+			ensureMsgf(0, TEXT("AnimNotifyTrack: Wrong indices found"));
 			AnimNotifyTracks[0].Notifies.Add(&Notifies[NotifyIndex]);
 		}
 	}
diff --git a/Engine/Source/Runtime/Engine/Private/Components/CharacterMovementComponent.cpp b/Engine/Source/Runtime/Engine/Private/Components/CharacterMovementComponent.cpp
--- a/Engine/Source/Runtime/Engine/Private/Components/CharacterMovementComponent.cpp
+++ b/Engine/Source/Runtime/Engine/Private/Components/CharacterMovementComponent.cpp
@@ -761,7 +761,7 @@ void UCharacterMovementComponent::OnMovementModeChanged(EMovementMode PreviousMo
 	}
 
 	CharacterOwner->OnMovementModeChanged(PreviousMovementMode, PreviousCustomMode);
-	ensureOnce(GroundMovementMode == MOVE_Walking || GroundMovementMode == MOVE_NavWalking);
+	ensure(GroundMovementMode == MOVE_Walking || GroundMovementMode == MOVE_NavWalking);
 };
 
 
@@ -776,7 +776,7 @@ uint8 UCharacterMovementComponent::PackNetworkMovementMode() const
 {
 	if (MovementMode != MOVE_Custom)
 	{
-		ensureOnce(GroundMovementMode == MOVE_Walking || GroundMovementMode == MOVE_NavWalking);
+		ensure(GroundMovementMode == MOVE_Walking || GroundMovementMode == MOVE_NavWalking);
 		const uint8 GroundModeBit = (GroundMovementMode == MOVE_Walking ? 0 : 1);
 		return uint8(MovementMode.GetValue()) | (GroundModeBit << PackedMovementModeConstants::GroundShift);
 	}
@@ -811,7 +811,7 @@ void UCharacterMovementComponent::ApplyNetworkMovementMode(const uint8 ReceivedM
 	TEnumAsByte<EMovementMode> NetGroundMode(MOVE_None);
 	uint8 NetCustomMode(0);
 	UnpackNetworkMovementMode(ReceivedMode, NetMovementMode, NetCustomMode, NetGroundMode);
-	ensureOnce(NetGroundMode == MOVE_Walking || NetGroundMode == MOVE_NavWalking);
+	ensure(NetGroundMode == MOVE_Walking || NetGroundMode == MOVE_NavWalking);
 
 	GroundMovementMode = NetGroundMode;
 	SetMovementMode(NetMovementMode, NetCustomMode);
diff --git a/Engine/Source/Runtime/Engine/Private/Components/SceneComponent.cpp b/Engine/Source/Runtime/Engine/Private/Components/SceneComponent.cpp
--- a/Engine/Source/Runtime/Engine/Private/Components/SceneComponent.cpp
+++ b/Engine/Source/Runtime/Engine/Private/Components/SceneComponent.cpp
@@ -344,7 +344,7 @@ void USceneComponent::UpdateComponentToWorldWithParent(USceneComponent* Parent,
 	FTransform NewTransform = CalcNewComponentToWorld(RelativeTransform, Parent);
 
 #if !(UE_BUILD_SHIPPING || UE_BUILD_TEST)
-	ensureOnce(NewTransform.IsValid());
+	ensure(NewTransform.IsValid());
 #endif
 
 	// If transform has changed..
diff --git a/Engine/Source/Runtime/Engine/Private/DataChannel.cpp b/Engine/Source/Runtime/Engine/Private/DataChannel.cpp
--- a/Engine/Source/Runtime/Engine/Private/DataChannel.cpp
+++ b/Engine/Source/Runtime/Engine/Private/DataChannel.cpp
@@ -2052,7 +2052,7 @@ bool UActorChannel::ReplicateActor()
 	{
 		FString Error(FString::Printf(TEXT("Attempt to replicate '%s' while already replicating that Actor!"), *Actor->GetName()));
 		UE_LOG(LogNet, Log, TEXT("%s"), *Error);
-		ensureMsg(false,*Error);
+		ensureMsgf(false,*Error);
 		return false;
 	}
 
diff --git a/Engine/Source/Runtime/Engine/Private/EdGraph/EdGraphPin.cpp b/Engine/Source/Runtime/Engine/Private/EdGraph/EdGraphPin.cpp
--- a/Engine/Source/Runtime/Engine/Private/EdGraph/EdGraphPin.cpp
+++ b/Engine/Source/Runtime/Engine/Private/EdGraph/EdGraphPin.cpp
@@ -127,8 +127,8 @@ void UEdGraphPin::MakeLinkTo(UEdGraphPin* ToPin)
 			UEdGraphNode* MyNode = GetOwningNode();
 
 			// Check that the other pin does not link to us
-			ensureMsg(!ToPin->LinkedTo.Contains(this), *GetLinkInfoString( LOCTEXT("MakeLinkTo", "MakeLinkTo").ToString(), LOCTEXT("IsLinked", "is linked with pin").ToString(), ToPin));			    
-			ensureMsg(MyNode->GetOuter() == ToPin->GetOwningNode()->GetOuter(), *GetLinkInfoString( LOCTEXT("MakeLinkTo", "MakeLinkTo").ToString(), LOCTEXT("OuterMismatch", "has a different outer than pin").ToString(), ToPin)); // Ensure both pins belong to the same graph
+			ensureMsgf(!ToPin->LinkedTo.Contains(this), *GetLinkInfoString( LOCTEXT("MakeLinkTo", "MakeLinkTo").ToString(), LOCTEXT("IsLinked", "is linked with pin").ToString(), ToPin));			    
+			ensureMsgf(MyNode->GetOuter() == ToPin->GetOwningNode()->GetOuter(), *GetLinkInfoString( LOCTEXT("MakeLinkTo", "MakeLinkTo").ToString(), LOCTEXT("OuterMismatch", "has a different outer than pin").ToString(), ToPin)); // Ensure both pins belong to the same graph
 
 			// Add to both lists
 			LinkedTo.Add(ToPin);
@@ -154,13 +154,13 @@ void UEdGraphPin::BreakLinkTo(UEdGraphPin* ToPin)
 			LinkedTo.Remove(ToPin);
 
 			// Check that the other pin links to us
-			ensureMsg(ToPin->LinkedTo.Contains(this), *GetLinkInfoString( LOCTEXT("BreakLinkTo", "BreakLinkTo").ToString(), LOCTEXT("NotLinked", "not reciprocally linked with pin").ToString(), ToPin) );
+			ensureMsgf(ToPin->LinkedTo.Contains(this), *GetLinkInfoString( LOCTEXT("BreakLinkTo", "BreakLinkTo").ToString(), LOCTEXT("NotLinked", "not reciprocally linked with pin").ToString(), ToPin) );
 			ToPin->LinkedTo.Remove(this);
 		}
 		else
 		{
 			// Check that the other pin does not link to us
-			ensureMsg(!ToPin->LinkedTo.Contains(this), *GetLinkInfoString( LOCTEXT("MakeLinkTo", "MakeLinkTo").ToString(), LOCTEXT("IsLinked", "is linked with pin").ToString(), ToPin));
+			ensureMsgf(!ToPin->LinkedTo.Contains(this), *GetLinkInfoString( LOCTEXT("MakeLinkTo", "MakeLinkTo").ToString(), LOCTEXT("IsLinked", "is linked with pin").ToString(), ToPin));
 		}
 	}
 }
diff --git a/Engine/Source/Runtime/Engine/Private/EdGraph/EdGraphSchema.cpp b/Engine/Source/Runtime/Engine/Private/EdGraph/EdGraphSchema.cpp
--- a/Engine/Source/Runtime/Engine/Private/EdGraph/EdGraphSchema.cpp
+++ b/Engine/Source/Runtime/Engine/Private/EdGraph/EdGraphSchema.cpp
@@ -449,7 +449,7 @@ void UEdGraphSchema::BreakSinglePinLink(UEdGraphPin* SourcePin, UEdGraphPin* Tar
 FPinConnectionResponse UEdGraphSchema::MovePinLinks(UEdGraphPin& MoveFromPin, UEdGraphPin& MoveToPin, bool bIsIntermediateMove) const
 {
 #if WITH_EDITOR
-	ensureMsg(bIsIntermediateMove || !MoveToPin.GetOwningNode()->GetGraph()->HasAnyFlags(RF_Transient),
+	ensureMsgf(bIsIntermediateMove || !MoveToPin.GetOwningNode()->GetGraph()->HasAnyFlags(RF_Transient),
 		TEXT("When moving to an Intermediate pin, use FKismetCompilerContext::MovePinLinksToIntermediate() instead of UEdGraphSchema::MovePinLinks()"));
 #endif // #if WITH_EDITOR
 
@@ -482,7 +482,7 @@ FPinConnectionResponse UEdGraphSchema::MovePinLinks(UEdGraphPin& MoveFromPin, UE
 FPinConnectionResponse UEdGraphSchema::CopyPinLinks(UEdGraphPin& CopyFromPin, UEdGraphPin& CopyToPin, bool bIsIntermediateCopy) const
 {
 #if WITH_EDITOR
-	ensureMsg(bIsIntermediateCopy || !CopyToPin.GetOwningNode()->GetGraph()->HasAnyFlags(RF_Transient),
+	ensureMsgf(bIsIntermediateCopy || !CopyToPin.GetOwningNode()->GetGraph()->HasAnyFlags(RF_Transient),
 		TEXT("When copying to an Intermediate pin, use FKismetCompilerContext::CopyPinLinksToIntermediate() instead of UEdGraphSchema::CopyPinLinks()"));
 #endif // #if WITH_EDITOR
 
diff --git a/Engine/Source/Runtime/Engine/Private/NetworkDriver.cpp b/Engine/Source/Runtime/Engine/Private/NetworkDriver.cpp
--- a/Engine/Source/Runtime/Engine/Private/NetworkDriver.cpp
+++ b/Engine/Source/Runtime/Engine/Private/NetworkDriver.cpp
@@ -779,7 +779,7 @@ void UNetDriver::InternalProcessRemoteFunction
 		{
 			FString Error(FString::Printf(TEXT("Attempt to replicate function '%s' on Actor '%s' while it is in the middle of variable replication!"), *Function->GetName(), *Actor->GetName()));
 			UE_LOG(LogScript, Error, TEXT("%s"), *Error);
-			ensureMsg(false, *Error);
+			ensureMsgf(false, *Error);
 			return;
 		}
 
diff --git a/Engine/Source/Runtime/Engine/Private/Particles/ParticleTrail2EmitterInstance.cpp b/Engine/Source/Runtime/Engine/Private/Particles/ParticleTrail2EmitterInstance.cpp
--- a/Engine/Source/Runtime/Engine/Private/Particles/ParticleTrail2EmitterInstance.cpp
+++ b/Engine/Source/Runtime/Engine/Private/Particles/ParticleTrail2EmitterInstance.cpp
@@ -2270,8 +2270,8 @@ bool FParticleRibbonEmitterInstance::ResolveSourcePoint(int32 InTrailIdx,
 					{
 						const FVector WorldOrigin = SourceEmitter->SimulationToWorld.GetOrigin();
 						UParticleSystemComponent* Comp = SourceEmitter->Component;
-						if (!ensureOnceMsgf(!SourceParticle->Location.ContainsNaN(), TEXT("NaN in SourceParticle Location. Template: %s, Component: %s"), Comp ? *GetNameSafe(Comp->Template) : TEXT("UNKNOWN"), *GetPathNameSafe(Comp)) ||
-							!ensureOnceMsgf(!WorldOrigin.ContainsNaN(), TEXT("NaN in WorldOrigin. Template: %s, Component: %s"), Comp ? *GetNameSafe(Comp->Template) : TEXT("UNKNOWN"), *GetPathNameSafe(Comp))
+						if (!ensureMsgf(!SourceParticle->Location.ContainsNaN(), TEXT("NaN in SourceParticle Location. Template: %s, Component: %s"), Comp ? *GetNameSafe(Comp->Template) : TEXT("UNKNOWN"), *GetPathNameSafe(Comp)) ||
+							!ensureMsgf(!WorldOrigin.ContainsNaN(), TEXT("NaN in WorldOrigin. Template: %s, Component: %s"), Comp ? *GetNameSafe(Comp->Template) : TEXT("UNKNOWN"), *GetPathNameSafe(Comp))
 							)
 						{
 							// Contains NaN!
diff --git a/Engine/Source/Runtime/Engine/Private/UnrealEngine.cpp b/Engine/Source/Runtime/Engine/Private/UnrealEngine.cpp
--- a/Engine/Source/Runtime/Engine/Private/UnrealEngine.cpp
+++ b/Engine/Source/Runtime/Engine/Private/UnrealEngine.cpp
@@ -4711,6 +4711,14 @@ bool UEngine::HandleDebugCommand( const TCHAR* Cmd, FOutputDevice& Ar )
 			return true;
 		}
 	}
+	else if( FParse::Command( &Cmd, TEXT("ENSUREALWAYS") ) )
+	{
+		UE_LOG(LogEngine, Warning, TEXT("Printed warning to log.") );
+		if( !ensureAlways( 0 ) )
+		{
+			return true;
+		}
+	}
 	else if( FParse::Command( &Cmd, TEXT( "FATAL" ) ) )
 	{
 		UE_LOG( LogEngine, Warning, TEXT( "Printed warning to log." ) );
diff --git a/Engine/Source/Runtime/Engine/Private/World.cpp b/Engine/Source/Runtime/Engine/Private/World.cpp
--- a/Engine/Source/Runtime/Engine/Private/World.cpp
+++ b/Engine/Source/Runtime/Engine/Private/World.cpp
@@ -1533,7 +1533,7 @@ void UWorld::TransferBlueprintDebugReferences(UWorld* NewWorld)
 							UE_LOG(LogWorld, Warning, TEXT("  OldObject path is %s"), *OldTargetObject->GetPathName());
 							UE_LOG(LogWorld, Warning, TEXT("  OldObject class path is %s"), *OldTargetObject->GetClass()->GetPathName());
 
-							ensureMsg(false, TEXT("Failed to find an appropriate object to debug back in the editor world"));
+							ensureMsgf(false, TEXT("Failed to find an appropriate object to debug back in the editor world"));
 						}
 
 						NewTargetObject = NULL;
diff --git a/Engine/Source/Runtime/Engine/Public/SceneManagement.h b/Engine/Source/Runtime/Engine/Public/SceneManagement.h
--- a/Engine/Source/Runtime/Engine/Public/SceneManagement.h
+++ b/Engine/Source/Runtime/Engine/Public/SceneManagement.h
@@ -1282,7 +1282,7 @@ public:
 		if (!bTriggered)
 		{
 			bTriggered = true;
-			ensureMsg(false, TEXT("FSimpleElementCollector::DrawMesh called"));
+			ensureMsgf(false, TEXT("FSimpleElementCollector::DrawMesh called"));
 		}
 
 		return false; 
@@ -1296,7 +1296,7 @@ public:
 		if (!bTriggered)
 		{
 			bTriggered = true;
-			ensureMsg(false, TEXT("FSimpleElementCollector::DrawMesh called"));
+			ensureMsgf(false, TEXT("FSimpleElementCollector::DrawMesh called"));
 		}
 
 		return 0;
diff --git a/Engine/Source/Runtime/Renderer/Private/RendererScene.cpp b/Engine/Source/Runtime/Renderer/Private/RendererScene.cpp
--- a/Engine/Source/Runtime/Renderer/Private/RendererScene.cpp
+++ b/Engine/Source/Runtime/Renderer/Private/RendererScene.cpp
@@ -1914,7 +1914,7 @@ void FScene::Release()
 	{
 		for (auto* ActorComponent : TObjectRange<UActorComponent>())
 		{
-			if ( !ensureMsg(!ActorComponent->IsRegistered() || ActorComponent->GetScene() != this, 
+			if ( !ensureMsgf(!ActorComponent->IsRegistered() || ActorComponent->GetScene() != this, 
 					*FString::Printf(TEXT("Component Name: %s World Name: %s Component Asset: %s"), 
 										*ActorComponent->GetFullName(), 
 										*GetWorld()->GetFullName(), 
diff --git a/Engine/Source/Runtime/Slate/Private/Framework/Application/MenuStack.cpp b/Engine/Source/Runtime/Slate/Private/Framework/Application/MenuStack.cpp
--- a/Engine/Source/Runtime/Slate/Private/Framework/Application/MenuStack.cpp
+++ b/Engine/Source/Runtime/Slate/Private/Framework/Application/MenuStack.cpp
@@ -146,7 +146,7 @@ TSharedRef<SWindow> FMenuStack::PushMenu( const TSharedRef<SWindow>& ParentWindo
 {
 	// Deprecated method that is no longer called by its deprecated counterpart in FSlateApplication so it will only create a dummy message warning the caller not to use this method
 	// Backwards compatibility of the API is taken care of by FSlateApplication.
-	ensureMsg(false, TEXT("PushMenu() returning an SWindow is deprecated. Use Push() that returns an IMenu instead."));
+	ensureMsgf(false, TEXT("PushMenu() returning an SWindow is deprecated. Use Push() that returns an IMenu instead."));
 
 	FSlateRect Anchor(SummonLocation, SummonLocation + SummonLocationSize);
 	FVector2D ExpectedSize(300, 200);
diff --git a/Engine/Source/Runtime/Slate/Private/Framework/Application/SlateApplication.cpp b/Engine/Source/Runtime/Slate/Private/Framework/Application/SlateApplication.cpp
--- a/Engine/Source/Runtime/Slate/Private/Framework/Application/SlateApplication.cpp
+++ b/Engine/Source/Runtime/Slate/Private/Framework/Application/SlateApplication.cpp
@@ -4308,7 +4308,7 @@ bool FSlateApplication::ProcessMouseButtonUpEvent( FPointerEvent& MouseEvent )
 	{
 		//FWidgetPath MouseCaptorPath = MouseCaptor.ToWidgetPath(MouseEvent.GetPointerIndex());
 		FWidgetPath MouseCaptorPath = MouseCaptor.ToWidgetPath( FWeakWidgetPath::EInterruptedPathHandling::Truncate, &MouseEvent );
-		if ( ensureMsg(MouseCaptorPath.Widgets.Num() > 0, TEXT("A window had a widget with mouse capture. That entire window has been dismissed before the mouse up could be processed.")) )
+		if ( ensureMsgf(MouseCaptorPath.Widgets.Num() > 0, TEXT("A window had a widget with mouse capture. That entire window has been dismissed before the mouse up could be processed.")) )
 		{
 #if PLATFORM_MAC
 			NSWindow* ActiveNativeWindow = [NSApp keyWindow];
diff --git a/Engine/Source/Runtime/Slate/Private/Framework/Docking/TabManager.cpp b/Engine/Source/Runtime/Slate/Private/Framework/Docking/TabManager.cpp
--- a/Engine/Source/Runtime/Slate/Private/Framework/Docking/TabManager.cpp
+++ b/Engine/Source/Runtime/Slate/Private/Framework/Docking/TabManager.cpp
@@ -41,7 +41,7 @@ TSharedRef<FTabManager::FLayoutNode> FTabManager::FLayout::NewFromString_Helper(
 			}
 			else
 			{
-				ensureMsg(false, TEXT("Invalid placement mode."));
+				ensureMsgf(false, TEXT("Invalid placement mode."));
 				return FTabManager::FArea::Placement_Automatic;
 			}
 		}
@@ -61,7 +61,7 @@ TSharedRef<FTabManager::FLayoutNode> FTabManager::FLayout::NewFromString_Helper(
 			}
 			else
 			{
-				ensureMsg(false, TEXT("Invalid orientation."));
+				ensureMsgf(false, TEXT("Invalid orientation."));
 				return Orient_Horizontal;
 			}
 			
@@ -82,7 +82,7 @@ TSharedRef<FTabManager::FLayoutNode> FTabManager::FLayout::NewFromString_Helper(
 			}
 			else
 			{
-				ensureMsg(false, TEXT("Invalid tab state."));
+				ensureMsgf(false, TEXT("Invalid tab state."));
 				return ETabState::OpenedTab;
 			}
 		}
@@ -181,7 +181,7 @@ TSharedRef<FTabManager::FLayoutNode> FTabManager::FLayout::NewFromString_Helper(
 	}
 	else
 	{
-		ensureMsg(false, TEXT("Unrecognized node type."));
+		ensureMsgf(false, TEXT("Unrecognized node type."));
 		return FTabManager::NewArea(FTabManager::FallbackWindowSize);
 	}	
 }
@@ -324,7 +324,7 @@ TSharedRef<FJsonObject> FTabManager::FLayout::PersistToString_Helper(const TShar
 	}
 	else
 	{
-		ensureMsg( false, TEXT("Unable to persist layout node of unknown type.") );
+		ensureMsgf( false, TEXT("Unable to persist layout node of unknown type.") );
 	}
 
 	return JsonObj;
diff --git a/Engine/Source/Runtime/Slate/Private/Widgets/Layout/SScrollBox.cpp b/Engine/Source/Runtime/Slate/Private/Widgets/Layout/SScrollBox.cpp
--- a/Engine/Source/Runtime/Slate/Private/Widgets/Layout/SScrollBox.cpp
+++ b/Engine/Source/Runtime/Slate/Private/Widgets/Layout/SScrollBox.cpp
@@ -389,7 +389,7 @@ bool SScrollBox::ScrollDescendantIntoView(const FGeometry& MyGeometry, const TSh
 	FindChildGeometries( MyGeometry, WidgetsToFind, Result );
 
 	FArrangedWidget* WidgetGeometry = Result.Find( WidgetToFind.ToSharedRef() );
-	if ( ensureMsg( WidgetGeometry, TEXT("Unable to scroll to descendant as it's not a child of the scrollbox") ) )
+	if ( ensureMsgf( WidgetGeometry, TEXT("Unable to scroll to descendant as it's not a child of the scrollbox") ) )
 	{
 		// @todo: This is a workaround because DesiredScrollOffset can exceed the ScrollMax when mouse dragging on the scroll bar and we need it clamped here or the offset is wrong
 		ScrollBy(MyGeometry, 0, EAllowOverscroll::No, false);
diff --git a/Engine/Source/Runtime/SlateCore/Private/Widgets/SNullWidget.cpp b/Engine/Source/Runtime/SlateCore/Private/Widgets/SNullWidget.cpp
--- a/Engine/Source/Runtime/SlateCore/Private/Widgets/SNullWidget.cpp
+++ b/Engine/Source/Runtime/SlateCore/Private/Widgets/SNullWidget.cpp
@@ -35,7 +35,7 @@ public:
 private:
 	virtual void SetVisibility( TAttribute<EVisibility> InVisibility ) override final
 	{
-		ensureMsg( false, TEXT("Attempting to SetVisibility() on SNullWidget. Mutating SNullWidget is not allowed.") );
+		ensureMsgf( false, TEXT("Attempting to SetVisibility() on SNullWidget. Mutating SNullWidget is not allowed.") );
 	}
 public:
 	
diff --git a/Engine/Source/Runtime/SlateCore/Private/Widgets/SWindow.cpp b/Engine/Source/Runtime/SlateCore/Private/Widgets/SWindow.cpp
--- a/Engine/Source/Runtime/SlateCore/Private/Widgets/SWindow.cpp
+++ b/Engine/Source/Runtime/SlateCore/Private/Widgets/SWindow.cpp
@@ -1002,7 +1002,7 @@ SOverlay::FOverlaySlot& SWindow::AddOverlaySlot( const int32 ZOrder )
 {
 	if(!WindowOverlay.IsValid())
 	{
-		ensureMsg( false, TEXT("This window does not support overlays. The added slot will not be visible!") );
+		ensureMsgf( false, TEXT("This window does not support overlays. The added slot will not be visible!") );
 		WindowOverlay = SNew(SOverlay).Visibility( EVisibility::HitTestInvisible );
 	}
 
diff --git a/Engine/Source/Runtime/SlateRHIRenderer/Private/SlateRHIRenderer.cpp b/Engine/Source/Runtime/SlateRHIRenderer/Private/SlateRHIRenderer.cpp
--- a/Engine/Source/Runtime/SlateRHIRenderer/Private/SlateRHIRenderer.cpp
+++ b/Engine/Source/Runtime/SlateRHIRenderer/Private/SlateRHIRenderer.cpp
@@ -699,7 +699,7 @@ FIntRect FSlateRHIRenderer::SetupVirtualScreenBuffer(const bool bPrimaryWorkArea
 	FSlateApplication::Get().GetDisplayMetrics(DisplayMetrics);
 
 #if !PLATFORM_WINDOWS && !PLATFORM_MAC && !PLATFORM_LINUX
-	ensureMsg(0, TEXT("This functionality is not valid for this platform"));	
+	ensureMsgf(0, TEXT("This functionality is not valid for this platform"));	
 	return FIntRect(FIntPoint(0, 0), FIntPoint(DisplayMetrics.PrimaryDisplayWidth, DisplayMetrics.PrimaryDisplayHeight));
 #endif	
 
@@ -768,7 +768,7 @@ FIntRect FSlateRHIRenderer::SetupVirtualScreenBuffer(const bool bPrimaryWorkArea
 void FSlateRHIRenderer::CopyWindowsToVirtualScreenBuffer(const TArray<FString>& KeypressBuffer)
 {
 #if !PLATFORM_WINDOWS && !PLATFORM_MAC && !PLATFORM_LINUX
-	ensureMsg(0, TEXT("This functionality is not valid for this platform"));
+	ensureMsgf(0, TEXT("This functionality is not valid for this platform"));
 	return;
 #endif
 
diff --git a/Engine/Source/Runtime/UMG/Private/Components/Widget.cpp b/Engine/Source/Runtime/UMG/Private/Components/Widget.cpp
--- a/Engine/Source/Runtime/UMG/Private/Components/Widget.cpp
+++ b/Engine/Source/Runtime/UMG/Private/Components/Widget.cpp
@@ -710,7 +710,7 @@ bool UWidget::IsChildOf(UWidget* PossibleParent)
 
 TSharedRef<SWidget> UWidget::RebuildWidget()
 {
-	ensureMsg(false, TEXT("You must implement RebuildWidget() in your child class"));
+	ensureMsgf(false, TEXT("You must implement RebuildWidget() in your child class"));
 	return SNew(SSpacer);
 }
 
diff --git a/Engine/Source/Runtime/Windows/D3D11RHI/Private/D3D11Commands.cpp b/Engine/Source/Runtime/Windows/D3D11RHI/Private/D3D11Commands.cpp
--- a/Engine/Source/Runtime/Windows/D3D11RHI/Private/D3D11Commands.cpp
+++ b/Engine/Source/Runtime/Windows/D3D11RHI/Private/D3D11Commands.cpp
@@ -845,7 +845,7 @@ void FD3D11DynamicRHI::RHISetRenderTargets(
 			FD3D11TextureBase* NewRenderTarget = GetD3D11TextureFromRHITexture(NewRenderTargetsRHI[RenderTargetIndex].Texture);
 			RenderTargetView = NewRenderTarget->GetRenderTargetView(RTMipIndex, RTSliceIndex);
 
-			ensureMsg(RenderTargetView, TEXT("Texture being set as render target has no RTV"));
+			ensureMsgf(RenderTargetView, TEXT("Texture being set as render target has no RTV"));
 #if CHECK_SRV_TRANSITIONS			
 			if (RenderTargetView)
 			{
