Stats - Refactored reading from a stats file (wip)

diff --git a/Engine/Source/Programs/UnrealFrontend/Private/Commands/StatsDumpMemoryCommand.cpp b/Engine/Source/Programs/UnrealFrontend/Private/Commands/StatsDumpMemoryCommand.cpp
--- a/Engine/Source/Programs/UnrealFrontend/Private/Commands/StatsDumpMemoryCommand.cpp
+++ b/Engine/Source/Programs/UnrealFrontend/Private/Commands/StatsDumpMemoryCommand.cpp
@@ -221,19 +221,21 @@ void FStatsMemoryDumpCommand::InternalRun()
 	FString SourceFilepath;
 	FParse::Value( FCommandLine::Get(), TEXT( "-INFILE=" ), SourceFilepath );
 
-	FStatsReadFile* NewReader = FStatsReadFile::CreateReaderForRawStats( *SourceFilepath, nullptr );
+	FStatsReadFile* NewReader = FStatsReadFile::CreateReaderForRawStats( *SourceFilepath, FStatsReadFile::FOnNewCombinedHistory::CreateRaw( this, &FStatsMemoryDumpCommand::ProcessMemoryOperations ) );
 	if (NewReader)
 	{
-		NewReader->ReadSynchronously();
 		PlatformName = NewReader->Header.PlatformName;
-		ProcessMemoryOperations( *NewReader );
+		NewReader->ReadAndProcessSynchronously();
 		delete NewReader;
+
+
 	}
 }
 
 
-void FStatsMemoryDumpCommand::ProcessMemoryOperations( const FStatsReadFile& File )
+void FStatsMemoryDumpCommand::ProcessMemoryOperations( FStatsReadFile* InFile )
 {
+	const FStatsReadFile& File = *InFile;
 	const TMap<int64, FStatPacketArray>& CombinedHistory = File.CombinedHistory;
 
 	// This is only example code, no fully implemented, may sometimes crash.
@@ -264,6 +266,9 @@ void FStatsMemoryDumpCommand::ProcessMemoryOperations( const FStatsReadFile& Fil
 	uint32 LastSequenceTagForNamedMarker = 0;
 	Snapshots.Add( TPairInitializer<uint32, FName>( LastSequenceTagForNamedMarker, TEXT( "BeginSnapshot" ) ) );
 
+	const int32 OnerPercent = FMath::Max( int32(InFile->FileInfo.TotalStatMessagesNum / 100), 65536 );
+
+	// #YRX_Stats: 2015-07-10 This still should be generic
 	for( int32 FrameIndex = 0; FrameIndex < Frames.Num(); ++FrameIndex )
 	{
 		const int64 TargetFrame = Frames[FrameIndex];
@@ -287,19 +292,26 @@ void FStatsMemoryDumpCommand::ProcessMemoryOperations( const FStatsReadFile& Fil
 			const FStatMessagesArray& Data = StatPacket.StatMessages;
 
 			const int32 NumStatMessages = Data.Num();
-			const int32 OnerPercent = FMath::Max( NumStatMessages / 100, 1024 );
 			for( int32 Index = 0; Index < NumStatMessages; Index++ )
 			{
 				CurrentStatMessageIndex++;
 
-				if( Index % OnerPercent )
+				if (CurrentStatMessageIndex % OnerPercent == 0)
 				{
 					const double CurrentSeconds = FPlatformTime::Seconds();
 					if( CurrentSeconds > PreviousSeconds + NumSecondsBetweenLogs )
 					{
-						const int32 CurrentPct = int32( 100.0*CurrentStatMessageIndex / NumStatMessages );
-						UE_LOG( LogStats, Log, TEXT( "Processing %3i%% (%i/%i) stat messages [Frame: %3i, Packet: %2i]" ), CurrentPct, Index, NumStatMessages, FrameIndex, PacketIndex );
+						const int32 PercentagePos = int32( 100.0*CurrentStatMessageIndex / InFile->FileInfo.TotalStatMessagesNum );
+						InFile->StageProgress.Set( PercentagePos );
+						UE_LOG( LogStats, Verbose, TEXT( "Processing %3i%% (%10i/%10i) stat messages [Frame: %3i, Packet: %2i]" ), PercentagePos, CurrentStatMessageIndex, InFile->FileInfo.TotalStatMessagesNum, FrameIndex, PacketIndex );
 						PreviousSeconds = CurrentSeconds;
+
+						// Abandon support, simply return for now.
+						if (InFile->bShouldStopProcessing == true)
+						{
+							InFile->SetProcessingStage( EStatsProcessingStage::RS_Stopped );
+							return;
+						}
 					}
 				}
 
@@ -326,12 +338,12 @@ void FStatsMemoryDumpCommand::ProcessMemoryOperations( const FStatsReadFile& Fil
 							NumMemoryOperations++;
 							// @see FStatsMallocProfilerProxy::TrackAlloc
 							// After AllocPtr message there is always alloc size message and the sequence tag.
-							Index++;
+							Index++; CurrentStatMessageIndex++;
 							const FStatMessage& AllocSizeMessage = Data[Index];
 							const int64 AllocSize = AllocSizeMessage.GetValue_int64();
 
 							// Read OperationSequenceTag.
-							Index++;
+							Index++; CurrentStatMessageIndex++;
 							const FStatMessage& SequenceTagMessage = Data[Index];
 							const uint32 SequenceTag = SequenceTagMessage.GetValue_int64();
 
@@ -351,6 +363,7 @@ void FStatsMemoryDumpCommand::ProcessMemoryOperations( const FStatsReadFile& Fil
 								StackState->bIsBrokenCallstack
 								) );
 							LastSequenceTagForNamedMarker = SequenceTag;
+
 						}
 						else if (MemOp == EMemoryOperation::Realloc)
 						{
@@ -358,17 +371,17 @@ void FStatsMemoryDumpCommand::ProcessMemoryOperations( const FStatsReadFile& Fil
 							const uint64 OldPtr = Ptr;
 
 							// Read NewPtr
-							Index++;
+							Index++; CurrentStatMessageIndex++;
 							const FStatMessage& AllocPtrMessage = Data[Index];
 							const uint64 NewPtr = AllocPtrMessage.GetValue_Ptr() & ~(uint64)EMemoryOperation::Mask;
 
 							// After AllocPtr message there is always alloc size message and the sequence tag.
-							Index++;
+							Index++; CurrentStatMessageIndex++;
 							const FStatMessage& ReallocSizeMessage = Data[Index];
 							const int64 ReallocSize = ReallocSizeMessage.GetValue_int64();
 
 							// Read OperationSequenceTag.
-							Index++;
+							Index++; CurrentStatMessageIndex++;
 							const FStatMessage& SequenceTagMessage = Data[Index];
 							const uint32 SequenceTag = SequenceTagMessage.GetValue_int64();
 
@@ -394,7 +407,7 @@ void FStatsMemoryDumpCommand::ProcessMemoryOperations( const FStatsReadFile& Fil
 						{
 							NumMemoryOperations++;
 							// Read OperationSequenceTag.
-							Index++;
+							Index++; CurrentStatMessageIndex++;
 							const FStatMessage& SequenceTagMessage = Data[Index];
 							const uint32 SequenceTag = SequenceTagMessage.GetValue_int64();
 
@@ -463,6 +476,8 @@ void FStatsMemoryDumpCommand::ProcessMemoryOperations( const FStatsReadFile& Fil
 		}
 	}
 
+	InFile->StageProgress.Set( 100 );
+
 	// End marker.
 	Snapshots.Add( TPairInitializer<uint32, FName>( TNumericLimits<uint32>::Max(), TEXT( "EndSnapshot" ) ) );
 
@@ -492,9 +507,18 @@ void FStatsMemoryDumpCommand::ProcessMemoryOperations( const FStatsReadFile& Fil
 	}
 	*/
 
+	InFile->SetProcessingStage( EStatsProcessingStage::RS_SortSequence );
+
 	// Sort all memory operation by the sequence tag, iterate through all operation and generate memory usage.
 	SequenceAllocationArray.Sort( FAllocationInfoSequenceTagLess() );
 
+	// Abandon support, simply return for now.
+	if (InFile->bShouldStopProcessing == true)
+	{
+		InFile->SetProcessingStage( EStatsProcessingStage::RS_Stopped );
+		return;
+	}
+
 	// Named markers/snapshots
 
 	// Alive allocations.
@@ -526,6 +550,13 @@ void FStatsMemoryDumpCommand::ProcessMemoryOperations( const FStatsReadFile& Fil
 				const int32 CurrentPct = int32( 100.0*(Index + 1) / NumSequenceAllocations );
 				UE_LOG( LogStats, Log, TEXT( "Processing allocations %3i%% (%10i/%10i)" ), CurrentPct, Index + 1, NumSequenceAllocations );
 				PreviousSeconds = CurrentSeconds;
+
+				// Abandon support, simply return for now.
+				if (InFile->bShouldStopProcessing == true)
+				{
+					InFile->SetProcessingStage( EStatsProcessingStage::RS_Stopped );
+					return;
+				}
 			}
 		}
 
@@ -613,7 +644,7 @@ void FStatsMemoryDumpCommand::ProcessMemoryOperations( const FStatsReadFile& Fil
 
 #if	UE_BUILD_DEBUG
 				const FString ReallocCallstack = FStatsCallstack::GetHumanReadable( Alloc.EncodedCallstack );
-				UE_LOG( LogStats, Warning, TEXT( "ReallocWithoutAlloc: %s %i %i/%i [%i]" ), *ReallocCallstack, Alloc.Size, Alloc.OldPtr, Alloc.Ptr, Alloc.SequenceTag );
+				//UE_LOG( LogStats, Warning, TEXT( "ReallocWithoutAlloc: %s %i %i/%i [%i]" ), *ReallocCallstack, Alloc.Size, Alloc.OldPtr, Alloc.Ptr, Alloc.SequenceTag );
 #endif // UE_BUILD_DEBUG
 
 				AllocationMap.Add( Alloc.Ptr, Alloc );
@@ -639,7 +670,7 @@ void FStatsMemoryDumpCommand::ProcessMemoryOperations( const FStatsReadFile& Fil
 
 #if	UE_BUILD_DEBUG
 				const FString FWACallstack = FStatsCallstack::GetHumanReadable( Alloc.EncodedCallstack );
-				UE_LOG( LogStats, Warning, TEXT( "FreeWithoutAllocCallstack: %s %i" ), *FWACallstack, Alloc.Ptr );
+				//UE_LOG( LogStats, Warning, TEXT( "FreeWithoutAllocCallstack: %s %i" ), *FWACallstack, Alloc.Ptr );
 #endif // UE_BUILD_DEBUG
 			}
 		}
@@ -682,6 +713,8 @@ void FStatsMemoryDumpCommand::ProcessMemoryOperations( const FStatsReadFile& Fil
 		}
 	}
 
+	InFile->SetProcessingStage( EStatsProcessingStage::RS_Finished );
+
 	// Frame-240 Frame-120 Frame-060
 	TMap<FString, FCombinedAllocationInfo> FrameBegin_Exit;
 	CompareSnapshots_FString( TEXT( "BeginSnapshot" ), TEXT( "EngineLoop.Exit" ), FrameBegin_Exit );
diff --git a/Engine/Source/Programs/UnrealFrontend/Private/Commands/StatsDumpMemoryCommand.h b/Engine/Source/Programs/UnrealFrontend/Private/Commands/StatsDumpMemoryCommand.h
--- a/Engine/Source/Programs/UnrealFrontend/Private/Commands/StatsDumpMemoryCommand.h
+++ b/Engine/Source/Programs/UnrealFrontend/Private/Commands/StatsDumpMemoryCommand.h
@@ -223,7 +223,7 @@ protected:
 	void InternalRun();
 
 	/** Basic memory profiling, only for debugging purpose. */
-	void ProcessMemoryOperations( const FStatsReadFile& File );
+	void ProcessMemoryOperations( FStatsReadFile* InFile );
 
 	/** Generates a basic memory usage report and prints it to the log. */
 	void GenerateMemoryUsageReport( const TMap<uint64, FAllocationInfo>& AllocationMap, const TSet<FName>& UObjectRawNames );
diff --git a/Engine/Source/Runtime/Core/Private/Stats/StatsFile.cpp b/Engine/Source/Runtime/Core/Private/Stats/StatsFile.cpp
--- a/Engine/Source/Runtime/Core/Private/Stats/StatsFile.cpp
+++ b/Engine/Source/Runtime/Core/Private/Stats/StatsFile.cpp
@@ -10,6 +10,7 @@
 
 #include "StatsData.h"
 #include "StatsFile.h"
+#include "ScopeExit.h"
 
 DECLARE_CYCLE_STAT( TEXT( "Stream File" ), STAT_StreamFile, STATGROUP_StatSystem );
 DECLARE_CYCLE_STAT( TEXT( "Wait For Write" ), STAT_StreamFileWaitForWrite, STATGROUP_StatSystem );
@@ -523,6 +524,28 @@ void FRawStatsWriteFile::WriteStatPacket( FArchive& Ar, FStatPacket& StatPacket
 	}
 }
 
+/*-----------------------------------------------------------------------------
+	FAsyncRawStatsReadFile
+-----------------------------------------------------------------------------*/
+
+FAsyncRawStatsFile::FAsyncRawStatsFile( FStatsReadFile* InOwner )
+	: Owner( InOwner )
+{}
+
+void FAsyncRawStatsFile::DoWork()
+{
+	Owner->ReadAndProcessSynchronously();
+}
+
+void FAsyncRawStatsFile::Abandon()
+{
+	Owner->RequestStop();
+}
+
+/*-----------------------------------------------------------------------------
+	FStatsReadFile
+-----------------------------------------------------------------------------*/
+
 const double FStatsReadFile::NumSecondsBetweenUpdates = 5.0;
 
 FStatsReadFile* FStatsReadFile::CreateReaderForRegularStats( const TCHAR* Filename, FOnNewStatPacket InNewStatPacket )
@@ -533,28 +556,54 @@ FStatsReadFile* FStatsReadFile::CreateReaderForRegularStats( const TCHAR* Filena
 	{
 		StatsReadFile->NewStatPacket = InNewStatPacket;
 	}
+	else
+	{
+		delete StatsReadFile;
+	}
 	return bValid ? StatsReadFile : nullptr;
 }
 
 
-FStatsReadFile* FStatsReadFile::CreateReaderForRawStats( const TCHAR* Filename, FOnNewRawStatPacket InNewRawStatPacket )
+FStatsReadFile* FStatsReadFile::CreateReaderForRawStats( const TCHAR* Filename, FOnNewCombinedHistory InNewCombinedHistory )
 {
 	FStatsReadFile* StatsReadFile = new FStatsReadFile( Filename, true );
 	const bool bValid = StatsReadFile->PrepareLoading();
 	if (bValid)
 	{
-		StatsReadFile->NewRawStatPacket = InNewRawStatPacket;
+		StatsReadFile->NewCombinedHistory = InNewCombinedHistory;
+	}
+	else
+	{
+		delete StatsReadFile;
 	}
 	return bValid ? StatsReadFile : nullptr;
 }
 
-void FStatsReadFile::ReadSynchronously()
+void FStatsReadFile::ReadAndProcessSynchronously()
 {
-	const double StartTime = FPlatformTime::Seconds();
+	ReadStats();
+	ProcessStats();
+
+	if (GetProcessingStage() == EStatsProcessingStage::RS_Stopped)
+	{
+		SetProcessingStage( EStatsProcessingStage::RS_Invalid );
+	}
+}
 
+void FStatsReadFile::ReadAndProcessAsynchronously()
+{
+	if (bRawStatsFile)
+	{
+		AsyncWork = new FAsyncTask<FAsyncRawStatsFile>( this );
+		AsyncWork->StartBackgroundTask();
+	}
+}
+
+void FStatsReadFile::ReadStats()
+{
 	if (bRawStatsFile)
 	{
-		const int64 FileSize = Reader->TotalSize();
+		const double StartTime = FPlatformTime::Seconds();
 
 		// Buffer used to store the compressed and decompressed data.
 		TArray<uint8> SrcArray;
@@ -567,9 +616,9 @@ void FStatsReadFile::ReadSynchronously()
 
 		// Update stage progress once per NumSecondsBetweenUpdates(5) seconds to avoid spamming.
 		double PreviousSeconds = FPlatformTime::Seconds();
-		ReadingStage.Set( int32( ERawStatsReadingStage::RS_ReadingAndCombiningPackets ) );
+		SetProcessingStage( EStatsProcessingStage::RS_ReadAndCombinePackets );
 
-		while (Reader->Tell() < FileSize)
+		while (Reader->Tell() < Reader->TotalSize())
 		{
 			// Read the compressed data.
 			FCompressedStatsData UncompressedData( SrcArray, DestArray );
@@ -594,55 +643,70 @@ void FStatsReadFile::ReadSynchronously()
 				return Item->ThreadId == StatPacket->ThreadId;
 			} );
 
-			const int64 PacketSize = StatPacket->StatMessages.GetAllocatedSize();
-			FileInfo.TotalStatMessagesNum += StatPacket->StatMessages.Num();
-
+			
 			if (CombinedPacket)
 			{
-				FileInfo.TotalPacketsSize -= (*CombinedPacket)->StatMessages.GetAllocatedSize();
 				(*CombinedPacket)->StatMessages += StatPacket->StatMessages;
-				FileInfo.TotalPacketsSize += (*CombinedPacket)->StatMessages.GetAllocatedSize();
-
 				delete StatPacket;
 			}
 			else
 			{
 				Frame.Packets.Add( StatPacket );
-				FileInfo.TotalPacketsSize += PacketSize;
-				FileInfo.MaximumPacketSize = FMath::Max( FileInfo.MaximumPacketSize, PacketSize );
+				FileInfo.MaximumPacketSize = FMath::Max<int64>( FileInfo.MaximumPacketSize, StatPacket->StatMessages.GetAllocatedSize() );
 			}
 
-			const double CurrentSeconds = FPlatformTime::Seconds(); 			
-			if (CurrentSeconds > LastUpdateTime + NumSecondsBetweenUpdates)
+			UpdateReadStageProgress();
+
+			if (bShouldStopProcessing == true)
 			{
-				const int32 PercentagePos = int32( 100.0*Reader->Tell() / FileSize );
-				StageProgress.Set( PercentagePos );
-				UE_LOG( LogStats, Verbose, TEXT( "%3i%% %10llu (%6.1f MB) read messages, last read frame %4i" ), PercentagePos, FileInfo.TotalStatMessagesNum, FileInfo.TotalPacketsSize / 1024.0f / 1024.0f, StatPacketFrameNum );
-				LastUpdateTime = CurrentSeconds;
+				SetProcessingStage( EStatsProcessingStage::RS_Stopped );
+				break;
 			}
-			
+
 			FileInfo.TotalPacketsNum++;
 		}
 
-		const double TotalTime = FPlatformTime::Seconds() - StartTime;
-		UE_LOG( LogStats, Log, TEXT( "Read took %.2f sec(s)" ), TotalTime );
-
-		DumpCombinedHistoryStats();	
-
-		//system( "pause" );
+		if (GetProcessingStage() != EStatsProcessingStage::RS_Stopped)
+		{
+			const double TotalTime = FPlatformTime::Seconds() - StartTime;
+			UE_LOG( LogStats, Log, TEXT( "Reading took %.2f sec(s)" ), TotalTime );
 
-		ReadingStage.Set( int32( ERawStatsReadingStage::RS_DecodingStatMessages ) );
+			UpdateCombinedHistoryStats();
+		}
+		else
+		{
+			UE_LOG( LogStats, Warning, TEXT( "Reading stopped, abandoning" ) );
+			// Clear all data.
+			CombinedHistory.Empty();
+		}
 	}
 }
 
-void FStatsReadFile::ReadAsynchronously()
+void FStatsReadFile::ProcessStats()
 {
+	if (bRawStatsFile)
+	{
+		if (GetProcessingStage() != EStatsProcessingStage::RS_Stopped)
+		{
+			SetProcessingStage( EStatsProcessingStage::RS_ProcessStatMessages );
+			NewCombinedHistory.ExecuteIfBound( this );
+
+			if (GetProcessingStage() == EStatsProcessingStage::RS_Stopped)
+			{
+				UE_LOG( LogStats, Warning, TEXT( "Processing stopped, abandoning" ) );
+			}
 
+			// Clear all data.
+			CombinedHistory.Empty();
+		}
+	}
 }
 
+
 FStatsReadFile::FStatsReadFile( const TCHAR* InFilename, bool bInRawStatsFile )
 	: Header( Stream.Header )
 	, Reader( nullptr )
+	, AsyncWork( nullptr )
 	, LastUpdateTime( 0.0 )
 	, Filename( InFilename )
 	, bRawStatsFile( bInRawStatsFile )
@@ -652,58 +716,82 @@ FStatsReadFile::FStatsReadFile( const TCHAR* InFilename, bool bInRawStatsFile )
 
 FStatsReadFile::~FStatsReadFile()
 {
+	RequestStop();
+
+	FPlatformProcess::ConditionalSleep( [&]()
+	{
+		return IsBusy();
+	}, 1.0f );
+
+	if (AsyncWork)
+	{
+		check( AsyncWork->IsDone() );
 
+		delete AsyncWork;
+		AsyncWork = nullptr;
+	}
+
+	delete Reader;
+	Reader = nullptr;
 }
 
 bool FStatsReadFile::PrepareLoading()
 {
 	const double StartTime = FPlatformTime::Seconds();
 
-	const int64 Size = IFileManager::Get().FileSize( *Filename );
-	if (Size < 4)
-	{
-		UE_LOG( LogStats, Error, TEXT( "Could not open: %s" ), *Filename );
-		return false;
-	}
+	SetProcessingStage( EStatsProcessingStage::RS_Started );
 
-	Reader = IFileManager::Get().CreateFileReader( *Filename );
-	if (!Reader)
 	{
-		UE_LOG( LogStats, Error, TEXT( "Could not open: %s" ), *Filename );
-		return false;
-	}
+		ON_SCOPE_EXIT
+		{
+			SetProcessingStage( EStatsProcessingStage::RS_Invalid );
+		};
 
-	if (!Stream.ReadHeader( *Reader ))
-	{
-		UE_LOG( LogStats, Error, TEXT( "Could not read, header is invalid: %s" ), *Filename );
-		return false;
-	}
+		const int64 Size = IFileManager::Get().FileSize( *Filename );
+		if (Size < 4)
+		{
+			UE_LOG( LogStats, Error, TEXT( "Could not open: %s" ), *Filename );
+			return false;
+		}
+
+		Reader = IFileManager::Get().CreateFileReader( *Filename );
+		if (!Reader)
+		{
+			UE_LOG( LogStats, Error, TEXT( "Could not open: %s" ), *Filename );
+			return false;
+		}
 
-	//UE_LOG( LogStats, Warning, TEXT( "Reading a raw stats file for memory profiling: %s" ), *Filename );
+		if (!Stream.ReadHeader( *Reader ))
+		{
+			UE_LOG( LogStats, Error, TEXT( "Could not read, header is invalid: %s" ), *Filename );
+			return false;
+		}
 
-	const bool bIsFinalized = Stream.Header.IsFinalized();
-	if (!bIsFinalized)
-	{
-		UE_LOG( LogStats, Error, TEXT( "Could not read, file is not finalized: %s" ), *Filename );
-		return false;
-	}
+		//UE_LOG( LogStats, Warning, TEXT( "Reading a raw stats file for memory profiling: %s" ), *Filename );
 
-	if( Stream.Header.Version < EStatMagicWithHeader::VERSION_6 )
-	{
-		UE_LOG( LogStats, Error, TEXT( "Could not read, invalid version: %s, expected %u, was %u" ), *Filename, (uint32)EStatMagicWithHeader::VERSION_6, Stream.Header.Version );
-		return false;
-	}
+		const bool bIsFinalized = Stream.Header.IsFinalized();
+		if (!bIsFinalized)
+		{
+			UE_LOG( LogStats, Error, TEXT( "Could not read, file is not finalized: %s" ), *Filename );
+			return false;
+		}
 
-	const bool bHasCompressedData = Stream.Header.HasCompressedData();
-	if( !bHasCompressedData )
-	{
-		UE_LOG( LogStats, Error, TEXT( "Could not read, required compressed data: %s" ), *Filename );
-		return false;
+		if (Stream.Header.Version < EStatMagicWithHeader::VERSION_6)
+		{
+			UE_LOG( LogStats, Error, TEXT( "Could not read, invalid version: %s, expected %u, was %u" ), *Filename, (uint32)EStatMagicWithHeader::VERSION_6, Stream.Header.Version );
+			return false;
+		}
+
+		const bool bHasCompressedData = Stream.Header.HasCompressedData();
+		if (!bHasCompressedData)
+		{
+			UE_LOG( LogStats, Error, TEXT( "Could not read, required compressed data: %s" ), *Filename );
+			return false;
+		}
 	}
 
 	State.MarkAsLoaded();
 
-	
 	// Read metadata.
 	TArray<FStatMessage> MetadataMessages;
 	Stream.ReadFNamesAndMetadataMessages( *Reader, MetadataMessages );
@@ -728,7 +816,6 @@ bool FStatsReadFile::PrepareLoading()
 	// Move file pointer to the first frame or first stat packet.
 	const int64 FrameOffset0 = Stream.FramesInfo[0].FrameFileOffset;
 	Reader->Seek( FrameOffset0 );
-	const int64 FileSize = Reader->TotalSize();
 
 	const double TotalTime = FPlatformTime::Seconds() - StartTime;
 	UE_LOG( LogStats, Log, TEXT( "Prepare loading took %.2f sec(s)" ), TotalTime );
@@ -736,7 +823,19 @@ bool FStatsReadFile::PrepareLoading()
 	return true;
 }
 
-void FStatsReadFile::DumpCombinedHistoryStats()
+void FStatsReadFile::UpdateReadStageProgress()
+{
+	const double CurrentSeconds = FPlatformTime::Seconds();
+	if (CurrentSeconds > LastUpdateTime + NumSecondsBetweenUpdates)
+	{
+		const int32 PercentagePos = int32( 100.0*Reader->Tell() / Reader->TotalSize() );
+		StageProgress.Set( PercentagePos );
+		UE_LOG( LogStats, Verbose, TEXT( "%3i%%, current num frames %4i" ), PercentagePos, CombinedHistory.Num() );
+		LastUpdateTime = CurrentSeconds;
+	}
+}
+
+void FStatsReadFile::UpdateCombinedHistoryStats()
 {
 	// Dump frame stats
 	int64 LastFrameNum = 0;
@@ -759,12 +858,16 @@ void FStatsReadFile::DumpCombinedHistoryStats()
 				FrameStatMessages );
 
 		LastFrameNum = FMath::Max( LastFrameNum, FrameNum );
+
+		FileInfo.TotalStatMessagesNum += FrameStatMessages;
+		FileInfo.TotalPacketsSize += FramePacketsSize;
 	}
 
-	UE_LOG( LogStats, Warning, TEXT( "TotalPacketSize: %6.1f MB, Max: %2f MB, TotalPacketsNum: %lli, Frames: %i" ),
+	UE_LOG( LogStats, Warning, TEXT( "Total PacketSize: %6.1f MB, Max: %2f MB, PacketsNum: %lli, StatMessagesNum: %lli, Frames: %i" ),
 			FileInfo.TotalPacketsSize / 1024.0f / 1024.0f,
 			FileInfo.MaximumPacketSize / 1024.0f / 1024.0f,
 			FileInfo.TotalPacketsNum,
+			FileInfo.TotalStatMessagesNum,
 			CombinedHistory.Num() );
 
 	TArray<int64> Frames;
diff --git a/Engine/Source/Runtime/Core/Public/Stats/StatsFile.h b/Engine/Source/Runtime/Core/Public/Stats/StatsFile.h
--- a/Engine/Source/Runtime/Core/Public/Stats/StatsFile.h
+++ b/Engine/Source/Runtime/Core/Public/Stats/StatsFile.h
@@ -767,14 +767,52 @@ struct FRawStatsFileInfo
 	int64 TotalPacketsNum;
 };
 
-/** Enumerates raw stats reading stages. */
-enum ERawStatsReadingStage
+/** Enumerates stats processing stages. */
+enum EStatsProcessingStage
 {
-	RS_ReadingAndCombiningPackets = 0,
-	RS_DecodingStatMessages,
+	RS_Started = 0,
+	RS_ReadAndCombinePackets,
+	RS_ProcessStatMessages,
 	/** Memory profiler only. */
-	RS_SortingSequence,
+	RS_SortSequence,
+	RS_Finished,
+	RS_Stopped,
 
+	/** Last stage, at this moment all data is read and processed or process has been stopped, we can now remove the instance. */
+	RS_Invalid,
+
+};
+
+struct FStatsReadFile;
+
+/**
+ *	Helper class used to read and process raw stats file on the async task.
+ */
+class FAsyncRawStatsFile
+{
+	/** Pointer to FStatsReadFile to call for async work. */
+	FStatsReadFile* Owner;
+
+public:
+	/** Initialization constructor. */
+	FAsyncRawStatsFile( FStatsReadFile* InOwner );
+
+	/** Call DoWork on the parent */
+	void DoWork();
+
+	TStatId GetStatId() const
+	{
+		return TStatId();
+	}
+
+	/** This task is abandonable. */
+	bool CanAbandon()
+	{
+		return true;
+	}
+
+	/** Abandon this task. */
+	void Abandon();
 };
 
 /** Struct used to read from ue4stats/ue4statsraw files, initializes all metadata and starts a process of reading the file asynchronously. */
@@ -786,40 +824,67 @@ protected:
 	/** Number of seconds between updating the current stage. */
 	static const double NumSecondsBetweenUpdates;
 
+public:
 	/** Delegate that FStatsReadFile calls on the instance whenever we have a new frame. */
 	DECLARE_DELEGATE_OneParam( FOnNewStatPacket, TArray<FStatMessage>* );
 
 	/** Delegate that FStatsReadFile calls on the instance whenever we have a new raw stats packet. */
-	DECLARE_DELEGATE_OneParam( FOnNewRawStatPacket, FStatPacket* );
+	//DECLARE_DELEGATE_OneParam( FOnNewRawStatPacket, FStatPacket* );
+
+	/** Delegate that FStatsReadFile calls on the instance whenever we have a whole combined history read. */
+	DECLARE_DELEGATE_OneParam( FOnNewCombinedHistory, FStatsReadFile* );
 
 public:
 	/** Creates a new reader for regular stats file. Will be nullptr for invalid files. */
 	static FStatsReadFile* CreateReaderForRegularStats( const TCHAR* Filename, FOnNewStatPacket InNewStatPacket );
 
 	/** Creates a new reader for raw stats file. Will be nullptr for invalid files. */
-	static FStatsReadFile* CreateReaderForRawStats( const TCHAR* Filename, FOnNewRawStatPacket InNewRawStatPacket );
+	static FStatsReadFile* CreateReaderForRawStats( const TCHAR* Filename, FOnNewCombinedHistory InNewCombinedHistory );
+
+	/** Reads and processes the file on the current thread. This is a blocking operation. */
+	void ReadAndProcessSynchronously();
+
+	/** Reads and processes the file using the async tasks on the pool thread. The read data is sent to the game thread using the task graph. This is a non-blocking operation. */
+	void ReadAndProcessAsynchronously();
+
+	void ReadStats();
+
+	void ProcessStats();
 
-	/** Read the file on the current thread. This is a blocking operation. */
-	void ReadSynchronously();
+	/** Sets a new processing stage for this file. */
+	void SetProcessingStage( EStatsProcessingStage NewStage )
+	{
+		ProcessingStage.Set( int32( NewStage ) );
+	}
 
-	/** Read the file using the async tasks on the pool thread. The read data is sent to the game thread using the task graph. This is a non-blocking operation. */
-	void ReadAsynchronously();
+	/**
+	 * @return current processing stage.
+	 */
+	const EStatsProcessingStage GetProcessingStage() const
+	{
+		return EStatsProcessingStage( ProcessingStage.GetValue() );
+	}
 
 	/**	 
 	 * @return true, if any asynchronous operation is being processed
 	 */
 	bool IsBusy()
 	{
-		// #YRX_STATS: 2015-07-09 
-		return false;
+		return AsyncWork != nullptr ? !AsyncWork->IsDone() : false;
+	}
+
+	/** Requests stopping the processing of the data. */
+	void RequestStop()
+	{
+		bShouldStopProcessing = true;
 	}
 	
 	/**
 	 * @return current stage of reading.
 	 */
-	const ERawStatsReadingStage GetReadingStage() const
+	const EStatsProcessingStage GetReadingStage() const
 	{
-		return (ERawStatsReadingStage)ReadingStage.GetValue();
+		return (EStatsProcessingStage)ProcessingStage.GetValue();
 	}
 
 protected:
@@ -837,8 +902,11 @@ protected:
 	 */
 	bool PrepareLoading();
 
+	/** Updates stage progress periodically, does debug logging if enabled. */
+	void UpdateReadStageProgress();
+
 	/** Dumps combined history stats. Only for raw stats. */
-	void DumpCombinedHistoryStats();
+	void UpdateCombinedHistoryStats();
 
 //protected:
 public:
@@ -854,11 +922,20 @@ public:
 	/** File reader. */
 	FArchive* Reader;
 
+	/** Async task. */
+	FAsyncTask<FAsyncRawStatsFile>* AsyncWork;
+
 	/** Delegate we fire every time we read one frame of data. The data must be freed after processing. */
-	mutable FOnNewStatPacket NewStatPacket;
+	FOnNewStatPacket NewStatPacket;
 
 	/** Delegate we fire every time we read one raw stats packet. The data must be freed after processing. */
-	mutable FOnNewRawStatPacket NewRawStatPacket;
+	//FOnNewRawStatPacket NewRawStatPacket;
+	
+	/**
+	 *	Delegate we fire after reading all raw stats packets and combining them into one history. 
+	 *	The combined history will be freed after processing.
+	 */
+	FOnNewCombinedHistory NewCombinedHistory;
 
 	/** Basic information about the stats file. */
 	FRawStatsFileInfo FileInfo;
@@ -869,12 +946,15 @@ public:
 	/** All raw names that contains a path to an UObject. */
 	TSet<FName> UObjectRawNames;
 
-	/** Current stage of reading. Thread safe. */
-	FThreadSafeCounter ReadingStage;
+	/** Current stage of processing. */
+	FThreadSafeCounter ProcessingStage;
 
 	/** Percentage progress of the current stage. */
 	FThreadSafeCounter StageProgress;
 
+	/** If true, we should break the processing loop. */
+	FThreadSafeBool bShouldStopProcessing;
+
 	/** Time of the last stage update. */
 	double LastUpdateTime;
 
