Deprecation message for ensureMsg.

diff --git a/Engine/Source/Runtime/Core/Public/Misc/AssertionMacros.h b/Engine/Source/Runtime/Core/Public/Misc/AssertionMacros.h
--- a/Engine/Source/Runtime/Core/Public/Misc/AssertionMacros.h
+++ b/Engine/Source/Runtime/Core/Public/Misc/AssertionMacros.h
@@ -150,6 +150,17 @@
 
 #endif	// DO_CHECK
 
+namespace UE4Asserts_Private
+{
+	DEPRECATED(4.9, "ensureMsg is deprecated, please use ensureMsgf instead.")
+	FORCEINLINE bool DeprecatedEnsure(bool bValue)
+	{
+		return bValue;
+	}
+}
+
+#define ensureMsg(InExpression, InMsg) UE4Asserts_Private::DeprecatedEnsure(ensureMsgf(InExpression, InMsg))
+
 // These are only for use as part of the GET_MEMBER_NAME_CHECKED macro
 template <typename T>
 bool JunkFunc(const T&);
