Linux: update cursor position cache on mouse move event (UE-18424).

- PR #1369 contributed by [REDACTED_PERSON].
- The code previously queried global cursor position which could differ from one sent with the event if framerate was low.

diff --git a/Engine/Source/Runtime/Core/Private/Linux/LinuxApplication.cpp b/Engine/Source/Runtime/Core/Private/Linux/LinuxApplication.cpp
--- a/Engine/Source/Runtime/Core/Private/Linux/LinuxApplication.cpp
+++ b/Engine/Source/Runtime/Core/Private/Linux/LinuxApplication.cpp
@@ -303,6 +303,10 @@ void FLinuxApplication::ProcessDeferredMessage( SDL_Event Event )
 			}
 			else
 			{
+				int xOffset, yOffset;
+				SDL_GetWindowPosition( NativeWindow, &xOffset, &yOffset );
+				LinuxCursor->SetCachedPosition( motionEvent.x + xOffset, motionEvent.y + yOffset );
+
 				FVector2D CurrentPosition = LinuxCursor->GetPosition();
 				if( LinuxCursor->UpdateCursorClipping( CurrentPosition ) )
 				{
@@ -312,8 +316,6 @@ void FLinuxApplication::ProcessDeferredMessage( SDL_Event Event )
 				{
 					if ( CurrentEventWindow->IsRegularWindow() )
 					{
-						int xOffset, yOffset;
-						SDL_GetWindowPosition( NativeWindow, &xOffset, &yOffset );
 						MessageHandler->GetWindowZoneForPoint( CurrentEventWindow.ToSharedRef(), CurrentPosition.X - xOffset, CurrentPosition.Y - yOffset );
 						MessageHandler->OnCursorSet();
 					}
diff --git a/Engine/Source/Runtime/Core/Private/Linux/LinuxCursor.cpp b/Engine/Source/Runtime/Core/Private/Linux/LinuxCursor.cpp
--- a/Engine/Source/Runtime/Core/Private/Linux/LinuxCursor.cpp
+++ b/Engine/Source/Runtime/Core/Private/Linux/LinuxCursor.cpp
@@ -192,6 +192,13 @@ void FLinuxCursor::InvalidateCaches()
 	bPositionCacheIsValid = false;
 }
 
+void FLinuxCursor::SetCachedPosition( const int32 X, const int32 Y )
+{
+	CachedGlobalXPosition = X;
+	CachedGlobalYPosition = Y;
+	bPositionCacheIsValid = true;
+}
+
 void FLinuxCursor::SetPosition( const int32 X, const int32 Y )
 {
 	int WndX, WndY;
@@ -201,7 +208,7 @@ void FLinuxCursor::SetPosition( const int32 X, const int32 Y )
 	SDL_GetWindowPosition( WndFocus, &WndX, &WndY );	//	get top left
 	SDL_WarpMouseInWindow( NULL, X - WndX, Y - WndY );
 
-	InvalidateCaches();
+	SetCachedPosition(X, Y);
 }
 
 void FLinuxCursor::AddOffset(const int32 DX, const int32 DY)
diff --git a/Engine/Source/Runtime/Core/Public/Linux/LinuxCursor.h b/Engine/Source/Runtime/Core/Public/Linux/LinuxCursor.h
--- a/Engine/Source/Runtime/Core/Public/Linux/LinuxCursor.h
+++ b/Engine/Source/Runtime/Core/Public/Linux/LinuxCursor.h
@@ -59,6 +59,9 @@ public:
 	 */
 	void InvalidateCaches();
 
+	/** Set the internal cached position, setting the cache to valid */
+	void SetCachedPosition( const int32 X, const int32 Y );
+
 private:
 
 	EMouseCursor::Type CurrentType;
