static_asserts in TSharedPtr when using with a type which inherits TSharedFromThis of the wrong thread safety type.
Fix for FNullGameMoviePlayer.

diff --git a/Engine/Source/Runtime/Core/Public/Templates/SharedPointer.h b/Engine/Source/Runtime/Core/Public/Templates/SharedPointer.h
--- a/Engine/Source/Runtime/Core/Public/Templates/SharedPointer.h
+++ b/Engine/Source/Runtime/Core/Public/Templates/SharedPointer.h
@@ -467,6 +467,16 @@ class TSharedPtr
 	// TSharedPtrs with UObjects are illegal.
 	static_assert(!TPointerIsConvertibleFromTo<ObjectType, const UObjectBase>::Value, "You cannot use TSharedPtr or TWeakPtr with UObjects. Consider a UPROPERTY() pointer or TWeakObjectPtr.");
 
+	enum
+	{
+		ObjectTypeHasSameModeSharedFromThis     = TPointerIsConvertibleFromTo<ObjectType, TSharedFromThis<ObjectType, Mode>>::Value,
+		ObjectTypeHasOppositeModeSharedFromThis = TPointerIsConvertibleFromTo<ObjectType, TSharedFromThis<ObjectType, (Mode == ESPMode::NotThreadSafe) ? ESPMode::ThreadSafe : ESPMode::NotThreadSafe>>::Value
+	};
+
+	// TSharedPtr of one mode to a type which has a TSharedFromThis only of another mode is illegal.
+	// A type which does not inherit TSharedFromThis at all is ok.
+	static_assert(ObjectTypeHasSameModeSharedFromThis || !ObjectTypeHasOppositeModeSharedFromThis, "You cannot use a TSharedPtr of one mode with a type which inherits TSharedFromThis of another mode.");
+
 public:
 
 	/**
diff --git a/Engine/Source/Runtime/MoviePlayer/Private/NullMoviePlayer.h b/Engine/Source/Runtime/MoviePlayer/Private/NullMoviePlayer.h
--- a/Engine/Source/Runtime/MoviePlayer/Private/NullMoviePlayer.h
+++ b/Engine/Source/Runtime/MoviePlayer/Private/NullMoviePlayer.h
@@ -5,7 +5,7 @@
 
 /** An placeholder implementation of the movie player for the editor */
 class FNullGameMoviePlayer : public IGameMoviePlayer,
-	public TSharedFromThis<FNullGameMoviePlayer, ESPMode::ThreadSafe>
+	public TSharedFromThis<FNullGameMoviePlayer>
 {
 public:
 	static TSharedPtr<FNullGameMoviePlayer> Get()
