Linux: use SDL_WarpMouseGlobal() instead of manual workaround (UE-18105).

- PR #1345 contributed by [REDACTED_PERSON].
- Also gets rid of misused SDL_SetWindowGrab().

diff --git a/Engine/Source/Runtime/Core/Private/Linux/LinuxCursor.cpp b/Engine/Source/Runtime/Core/Private/Linux/LinuxCursor.cpp
--- a/Engine/Source/Runtime/Core/Private/Linux/LinuxCursor.cpp
+++ b/Engine/Source/Runtime/Core/Private/Linux/LinuxCursor.cpp
@@ -7,8 +7,6 @@
 
 FLinuxCursor::FLinuxCursor()
 	: 	bHidden(false)
-	,	AccumulatedOffsetX(0)
-	,	AccumulatedOffsetY(0)
 	,	CachedGlobalXPosition(0)
 	,	CachedGlobalYPosition(0)
 	,	bPositionCacheIsValid(false)
@@ -180,9 +178,6 @@ FVector2D FLinuxCursor::GetPosition() const
 
 	int CursorX = CachedGlobalXPosition;
 	int CursorY = CachedGlobalYPosition;
-	
-	CursorX += AccumulatedOffsetX;
-	CursorY += AccumulatedOffsetY;
 
 	return FVector2D( CursorX, CursorY );
 }
@@ -201,27 +196,11 @@ void FLinuxCursor::SetCachedPosition( const int32 X, const int32 Y )
 
 void FLinuxCursor::SetPosition( const int32 X, const int32 Y )
 {
-	int WndX, WndY;
-
-	SDL_HWindow WndFocus = SDL_GetMouseFocus();
-
-	SDL_GetWindowPosition( WndFocus, &WndX, &WndY );	//	get top left
-	SDL_WarpMouseInWindow( NULL, X - WndX, Y - WndY );
+	SDL_WarpMouseGlobal( X, Y);
 
 	SetCachedPosition(X, Y);
 }
 
-void FLinuxCursor::AddOffset(const int32 DX, const int32 DY)
-{
-	AccumulatedOffsetX += DX; 
-	AccumulatedOffsetY += DY;
-}
-
-void FLinuxCursor::ResetOffset()
-{
-	AccumulatedOffsetX = AccumulatedOffsetY = 0;
-}
-
 void FLinuxCursor::SetType( const EMouseCursor::Type InNewCursor )
 {	
 	checkf( InNewCursor < EMouseCursor::TotalCursorCount, TEXT("Invalid cursor(%d) supplied"), InNewCursor );
@@ -270,11 +249,9 @@ void FLinuxCursor::Lock( const RECT* const Bounds )
 	if ( Bounds == NULL )
 	{
 		CursorClipRect = FIntRect();
-		SDL_SetWindowGrab( NULL, SDL_FALSE );
 	}
 	else
 	{
-		SDL_SetWindowGrab( NULL, SDL_TRUE );
 		CursorClipRect.Min.X = FMath::TruncToInt(Bounds->left);
 		CursorClipRect.Min.Y = FMath::TruncToInt(Bounds->top);
 		CursorClipRect.Max.X = FMath::TruncToInt(Bounds->right) - 1;
diff --git a/Engine/Source/Runtime/Core/Public/Linux/LinuxCursor.h b/Engine/Source/Runtime/Core/Public/Linux/LinuxCursor.h
--- a/Engine/Source/Runtime/Core/Public/Linux/LinuxCursor.h
+++ b/Engine/Source/Runtime/Core/Public/Linux/LinuxCursor.h
@@ -74,12 +74,6 @@ private:
 
 	uint32 CursorEvent;
 
-	/** Accumulated X offset (from relative mouse movement events) used for reconstructing global cursor position */
-	int32 AccumulatedOffsetX;
-
-	/** Accumulated Y offset (from relative mouse movement events) used for reconstructing global cursor position */
-	int32 AccumulatedOffsetY;
-
 	/** Cached global X position */
 	mutable int32 CachedGlobalXPosition;
 
