#ue4 - Move some math and collision stats to a verbose stats group.

- Quat<->Rotator conversions now in MathVerbose.
- Collsion Pre/Post filter now in CollisionVerbose.

diff --git a/Engine/Source/Runtime/Core/Private/Math/UnrealMath.cpp b/Engine/Source/Runtime/Core/Private/Math/UnrealMath.cpp
--- a/Engine/Source/Runtime/Core/Private/Math/UnrealMath.cpp
+++ b/Engine/Source/Runtime/Core/Private/Math/UnrealMath.cpp
@@ -12,8 +12,8 @@ DEFINE_LOG_CATEGORY(LogUnrealMath);
 * Math stats
 */
 
-DECLARE_CYCLE_STAT( TEXT( "Convert Rotator to Quat" ), STAT_MathConvertRotatorToQuat, STATGROUP_Math );
-DECLARE_CYCLE_STAT( TEXT( "Convert Quat to Rotator" ), STAT_MathConvertQuatToRotator, STATGROUP_Math );
+DECLARE_CYCLE_STAT( TEXT( "Convert Rotator to Quat" ), STAT_MathConvertRotatorToQuat, STATGROUP_MathVerbose );
+DECLARE_CYCLE_STAT( TEXT( "Convert Quat to Rotator" ), STAT_MathConvertQuatToRotator, STATGROUP_MathVerbose );
 
 /*-----------------------------------------------------------------------------
 	Globals
diff --git a/Engine/Source/Runtime/Core/Public/Stats/Stats2.h b/Engine/Source/Runtime/Core/Public/Stats/Stats2.h
--- a/Engine/Source/Runtime/Core/Public/Stats/Stats2.h
+++ b/Engine/Source/Runtime/Core/Public/Stats/Stats2.h
@@ -1996,6 +1996,7 @@ DECLARE_STATS_GROUP(TEXT("Beam Particles"),STATGROUP_BeamParticles, STATCAT_Adva
 DECLARE_STATS_GROUP(TEXT("CPU Stalls"), STATGROUP_CPUStalls, STATCAT_Advanced);
 DECLARE_STATS_GROUP(TEXT("Canvas"),STATGROUP_Canvas, STATCAT_Advanced);
 DECLARE_STATS_GROUP(TEXT("Collision"),STATGROUP_Collision, STATCAT_Advanced);
+DECLARE_STATS_GROUP(TEXT("CollisionVerbose"),STATGROUP_CollisionVerbose, STATCAT_Advanced);
 DECLARE_STATS_GROUP(TEXT("Crash Tracker"),STATGROUP_CrashTracker, STATCAT_Advanced);
 DECLARE_STATS_GROUP(TEXT("D3D11RHI"),STATGROUP_D3D11RHI, STATCAT_Advanced);
 DECLARE_STATS_GROUP(TEXT("DDC"),STATGROUP_DDC, STATCAT_Advanced);
@@ -2011,7 +2012,7 @@ DECLARE_STATS_GROUP(TEXT("Landscape"),STATGROUP_Landscape, STATCAT_Advanced);
 DECLARE_STATS_GROUP(TEXT("Light Rendering"),STATGROUP_LightRendering, STATCAT_Advanced);
 DECLARE_STATS_GROUP(TEXT("Load Time"), STATGROUP_LoadTime, STATCAT_Advanced);
 DECLARE_STATS_GROUP_VERBOSE(TEXT("Load Time (Verbose)"), STATGROUP_LoadTimeVerbose, STATCAT_Advanced);
-DECLARE_STATS_GROUP(TEXT("Math"), STATGROUP_Math, STATCAT_Advanced);
+DECLARE_STATS_GROUP(TEXT("MathVerbose"), STATGROUP_MathVerbose, STATCAT_Advanced);
 DECLARE_STATS_GROUP(TEXT("Memory Allocator"),STATGROUP_MemoryAllocator, STATCAT_Advanced);
 DECLARE_STATS_GROUP(TEXT("Memory Platform"),STATGROUP_MemoryPlatform, STATCAT_Advanced);
 DECLARE_STATS_GROUP(TEXT("Memory StaticMesh"),STATGROUP_MemoryStaticMesh, STATCAT_Advanced);
diff --git a/Engine/Source/Runtime/Engine/Public/Collision.h b/Engine/Source/Runtime/Engine/Public/Collision.h
--- a/Engine/Source/Runtime/Engine/Public/Collision.h
+++ b/Engine/Source/Runtime/Engine/Public/Collision.h
@@ -19,8 +19,8 @@ DECLARE_CYCLE_STAT_EXTERN(TEXT("GeomOverlapAny"),STAT_Collision_GeomOverlapAny,S
 DECLARE_CYCLE_STAT_EXTERN(TEXT("GeomOverlapSingle"),STAT_Collision_GeomOverlapSingle,STATGROUP_Collision, );
 DECLARE_CYCLE_STAT_EXTERN(TEXT("GeomOverlapMultiple"),STAT_Collision_GeomOverlapMultiple,STATGROUP_Collision, );
 DECLARE_CYCLE_STAT_EXTERN(TEXT("GeomComputePenetration"), STAT_Collision_GeomComputePenetration, STATGROUP_Collision, );
-DECLARE_CYCLE_STAT_EXTERN(TEXT("PreFilter"), STAT_Collision_PreFilter, STATGROUP_Collision, );
-DECLARE_CYCLE_STAT_EXTERN(TEXT("PostFilter"), STAT_Collision_PostFilter, STATGROUP_Collision, );
+DECLARE_CYCLE_STAT_EXTERN(TEXT("PreFilter"), STAT_Collision_PreFilter, STATGROUP_CollisionVerbose, );
+DECLARE_CYCLE_STAT_EXTERN(TEXT("PostFilter"), STAT_Collision_PostFilter, STATGROUP_CollisionVerbose, );
 
 /** Enable collision analyzer support */
 #if (1 && !(UE_BUILD_SHIPPING || UE_BUILD_TEST) && WITH_EDITOR && WITH_UNREAL_DEVELOPER_TOOLS && WITH_PHYSX)
