This move assignment operator has never been used and was incorrectly implemented

diff --git a/Engine/Source/Runtime/Core/Public/Containers/Array.h b/Engine/Source/Runtime/Core/Public/Containers/Array.h
--- a/Engine/Source/Runtime/Core/Public/Containers/Array.h
+++ b/Engine/Source/Runtime/Core/Public/Containers/Array.h
@@ -2862,7 +2862,7 @@ public:
 
 	            TIndirectArray() {}
 	FORCEINLINE TIndirectArray(TIndirectArray&& Other) : Array(MoveTemp(Other.Array)) {}
-	FORCEINLINE TIndirectArray& operator=(TIndirectArray&& Other) { Array = MoveTemp(Other.Temp); return *this; }
+	FORCEINLINE TIndirectArray& operator=(TIndirectArray&& Other) { Array = MoveTemp(Other.Array); return *this; }
 
 #endif
 
