#ue4 - MathVerbose and CollisionVerbose are now correctly declared as verbose.

diff --git a/Engine/Source/Runtime/Core/Public/Stats/Stats2.h b/Engine/Source/Runtime/Core/Public/Stats/Stats2.h
--- a/Engine/Source/Runtime/Core/Public/Stats/Stats2.h
+++ b/Engine/Source/Runtime/Core/Public/Stats/Stats2.h
@@ -1996,7 +1996,7 @@ DECLARE_STATS_GROUP(TEXT("Beam Particles"),STATGROUP_BeamParticles, STATCAT_Adva
 DECLARE_STATS_GROUP(TEXT("CPU Stalls"), STATGROUP_CPUStalls, STATCAT_Advanced);
 DECLARE_STATS_GROUP(TEXT("Canvas"),STATGROUP_Canvas, STATCAT_Advanced);
 DECLARE_STATS_GROUP(TEXT("Collision"),STATGROUP_Collision, STATCAT_Advanced);
-DECLARE_STATS_GROUP(TEXT("CollisionVerbose"),STATGROUP_CollisionVerbose, STATCAT_Advanced);
+DECLARE_STATS_GROUP_VERBOSE(TEXT("CollisionVerbose"),STATGROUP_CollisionVerbose, STATCAT_Advanced);
 DECLARE_STATS_GROUP(TEXT("Crash Tracker"),STATGROUP_CrashTracker, STATCAT_Advanced);
 DECLARE_STATS_GROUP(TEXT("D3D11RHI"),STATGROUP_D3D11RHI, STATCAT_Advanced);
 DECLARE_STATS_GROUP(TEXT("DDC"),STATGROUP_DDC, STATCAT_Advanced);
@@ -2012,7 +2012,7 @@ DECLARE_STATS_GROUP(TEXT("Landscape"),STATGROUP_Landscape, STATCAT_Advanced);
 DECLARE_STATS_GROUP(TEXT("Light Rendering"),STATGROUP_LightRendering, STATCAT_Advanced);
 DECLARE_STATS_GROUP(TEXT("Load Time"), STATGROUP_LoadTime, STATCAT_Advanced);
 DECLARE_STATS_GROUP_VERBOSE(TEXT("Load Time (Verbose)"), STATGROUP_LoadTimeVerbose, STATCAT_Advanced);
-DECLARE_STATS_GROUP(TEXT("MathVerbose"), STATGROUP_MathVerbose, STATCAT_Advanced);
+DECLARE_STATS_GROUP_VERBOSE(TEXT("MathVerbose"), STATGROUP_MathVerbose, STATCAT_Advanced);
 DECLARE_STATS_GROUP(TEXT("Memory Allocator"),STATGROUP_MemoryAllocator, STATCAT_Advanced);
 DECLARE_STATS_GROUP(TEXT("Memory Platform"),STATGROUP_MemoryPlatform, STATCAT_Advanced);
 DECLARE_STATS_GROUP(TEXT("Memory StaticMesh"),STATGROUP_MemoryStaticMesh, STATCAT_Advanced);
