Missing file in my last checkin.

diff --git a/Engine/Source/Runtime/Core/Public/Misc/CookingStatsInterface.h b/Engine/Source/Runtime/Core/Public/Misc/CookingStatsInterface.h
new file mode 100644
--- /dev/null
+++ b/Engine/Source/Runtime/Core/Public/Misc/CookingStatsInterface.h
@@ -0,0 +1,29 @@
+// Copyright 1998-2015 Epic Games, Inc. All Rights Reserved.
+
+#pragma once
+
+#include "ModuleManager.h"
+#include "Core.h"
+
+/**
+* CookingStats module interface
+*/
+class ICookingStatsInterface : public IModuleInterface
+{
+public:
+	/**
+	 * Tries to gets a pointer to the active AssetRegistryInterface implementation. 
+	 */
+	static inline ICookingStatsInterface * GetPtr()
+	{
+		static FName CookingStats("CookingStats");
+		auto Module = FModuleManager::GetModulePtr<ICookingStatsInterface>(CookingStats);
+		if (Module == nullptr)
+		{
+			Module = &FModuleManager::LoadModuleChecked<ICookingStatsInterface>(CookingStats);
+		}
+		return reinterpret_cast<ICookingStatsInterface*>(Module);
+	}
+
+};
+
