Enable default function definitions for VS2015

diff --git a/Engine/Source/Runtime/Core/Public/Windows/WIndowsPlatform.h b/Engine/Source/Runtime/Core/Public/Windows/WIndowsPlatform.h
--- a/Engine/Source/Runtime/Core/Public/Windows/WIndowsPlatform.h
+++ b/Engine/Source/Runtime/Core/Public/Windows/WIndowsPlatform.h
@@ -58,7 +58,11 @@ typedef FWindowsPlatformTypes FPlatformTypes;
 #define PLATFORM_USES_MICROSOFT_LIBC_FUNCTIONS				1
 #define PLATFORM_SUPPORTS_TBB								1
 #define PLATFORM_SUPPORTS_NAMED_PIPES						1
-#define PLATFORM_COMPILER_HAS_DEFAULTED_FUNCTIONS			0
+#if _MSC_VER >= 1900
+	#define PLATFORM_COMPILER_HAS_DEFAULTED_FUNCTIONS		1
+#else
+	#define PLATFORM_COMPILER_HAS_DEFAULTED_FUNCTIONS		0
+#endif
 #if _MSC_VER >= 1800 || __clang__
 	#define PLATFORM_COMPILER_HAS_VARIADIC_TEMPLATES		1
 	#define PLATFORM_COMPILER_HAS_EXPLICIT_OPERATORS		1
