Fixed undesired conditional logic that skipped some logic necessary for live culture swapping when in certain configurations such as shipping. Enabled use of culture switch on commandline in all configurations.

diff --git a/Engine/Source/Runtime/Core/Private/Internationalization/TextLocalizationManager.cpp b/Engine/Source/Runtime/Core/Private/Internationalization/TextLocalizationManager.cpp
--- a/Engine/Source/Runtime/Core/Private/Internationalization/TextLocalizationManager.cpp
+++ b/Engine/Source/Runtime/Core/Private/Internationalization/TextLocalizationManager.cpp
@@ -46,84 +46,81 @@ void EndInitTextLocalization()
 				GConfig->SetString(TEXT("Internationalization"), TEXT("Culture"), *RequestedCultureName, GEngineIni);
 			}
 		}
+		// Use culture override specified on commandline.
+		else if (FParse::Value(FCommandLine::Get(), TEXT("CULTURE="), RequestedCultureName))
+		{
+			//UE_LOG(LogInit, Log, TEXT("Overriding culture %s w/ command-line option".), *CultureName);
+		}
 		else
-#if !UE_BUILD_SHIPPING
-			// Use culture override specified on commandline.
-			if (FParse::Value(FCommandLine::Get(), TEXT("CULTURE="), RequestedCultureName))
+#if WITH_EDITOR
+			// See if we've been provided a culture override in the editor
+			if(GConfig->GetString( TEXT("Internationalization"), TEXT("Culture"), RequestedCultureName, GEditorSettingsIni ))
 			{
-				//UE_LOG(LogInit, Log, TEXT("Overriding culture %s w/ command-line option".), *CultureName);
+				//UE_LOG(LogInit, Log, TEXT("Overriding culture %s w/ editor configuration."), *CultureName);
 			}
 			else
-#endif // !UE_BUILD_SHIPPING
-#if WITH_EDITOR
-				// See if we've been provided a culture override in the editor
-				if(GConfig->GetString( TEXT("Internationalization"), TEXT("Culture"), RequestedCultureName, GEditorSettingsIni ))
+#endif // WITH_EDITOR
+				// Use culture specified in engine configuration.
+				if(GConfig->GetString( TEXT("Internationalization"), TEXT("Culture"), RequestedCultureName, GEngineIni ))
 				{
-					//UE_LOG(LogInit, Log, TEXT("Overriding culture %s w/ editor configuration."), *CultureName);
+					//UE_LOG(LogInit, Log, TEXT("Overriding culture %s w/ engine configuration."), *CultureName);
 				}
 				else
-#endif // WITH_EDITOR
-					// Use culture specified in engine configuration.
-					if(GConfig->GetString( TEXT("Internationalization"), TEXT("Culture"), RequestedCultureName, GEngineIni ))
+				{
+					RequestedCultureName = I18N.GetDefaultCulture()->GetName();
+				}
+
+				FString TargetCultureName = RequestedCultureName;
+				{
+					TArray<FString> LocalizationPaths;
+					if(ShouldLoadEditor)
 					{
-						//UE_LOG(LogInit, Log, TEXT("Overriding culture %s w/ engine configuration."), *CultureName);
+						LocalizationPaths += FPaths::GetEditorLocalizationPaths();
 					}
-					else
+					if(ShouldLoadGame)
 					{
-						RequestedCultureName = I18N.GetDefaultCulture()->GetName();
+						LocalizationPaths += FPaths::GetGameLocalizationPaths();
 					}
+					LocalizationPaths += FPaths::GetEngineLocalizationPaths();
 
-					FString TargetCultureName = RequestedCultureName;
-					{
-						TArray<FString> LocalizationPaths;
-						if(ShouldLoadEditor)
-						{
-							LocalizationPaths += FPaths::GetEditorLocalizationPaths();
-						}
-						if(ShouldLoadGame)
-						{
-							LocalizationPaths += FPaths::GetGameLocalizationPaths();
-						}
-						LocalizationPaths += FPaths::GetEngineLocalizationPaths();
+					// Validate the locale has data or fallback to one that does.
+					TArray< FCultureRef > AvailableCultures;
+					I18N.GetCulturesWithAvailableLocalization(LocalizationPaths, AvailableCultures, false);
 
-						// Validate the locale has data or fallback to one that does.
-						TArray< FCultureRef > AvailableCultures;
-						I18N.GetCulturesWithAvailableLocalization(LocalizationPaths, AvailableCultures, false);
-
-						FString ValidCultureName;
-						ValidCultureName.Empty();
-						FCulturePtr TargetCulture = I18N.GetCulture(TargetCultureName);
-						if (TargetCulture.IsValid())
+					FString ValidCultureName;
+					ValidCultureName.Empty();
+					FCulturePtr TargetCulture = I18N.GetCulture(TargetCultureName);
+					if (TargetCulture.IsValid())
+					{
+						TArray<FString> PrioritizedParentCultureNames = TargetCulture->GetPrioritizedParentCultureNames();
+						for (const FString& CultureName : PrioritizedParentCultureNames)
 						{
-							TArray<FString> PrioritizedParentCultureNames = TargetCulture->GetPrioritizedParentCultureNames();
-							for (const FString& CultureName : PrioritizedParentCultureNames)
+							FCulturePtr ValidCulture = I18N.GetCulture(CultureName);
+							if (ValidCulture.IsValid() && AvailableCultures.Contains(ValidCulture.ToSharedRef()))
 							{
-								FCulturePtr ValidCulture = I18N.GetCulture(CultureName);
-								if (ValidCulture.IsValid() && AvailableCultures.Contains(ValidCulture.ToSharedRef()))
-								{
-									ValidCultureName = CultureName;
-									break;
-								}
+								ValidCultureName = CultureName;
+								break;
 							}
 						}
+					}
 
-						if(!ValidCultureName.IsEmpty())
-						{
-							if(RequestedCultureName != ValidCultureName)
-							{
-								// Make the user aware that the localization data belongs to a parent culture.
-								UE_LOG(LogTextLocalizationManager, Log, TEXT("No specific translations for ('%s') exist, so ('%s') translations will be used."), *RequestedCultureName, *ValidCultureName);
-							}
-						}
-						else
+					if(!ValidCultureName.IsEmpty())
+					{
+						if(RequestedCultureName != ValidCultureName)
 						{
-							// Fallback to English.
-							UE_LOG(LogTextLocalizationManager, Log, TEXT("No translations for ('%s') exist, falling back to 'en' for localization and internationalization data."), *RequestedCultureName);
-							TargetCultureName = TEXT("en");
+							// Make the user aware that the localization data belongs to a parent culture.
+							UE_LOG(LogTextLocalizationManager, Log, TEXT("No specific translations for ('%s') exist, so ('%s') translations will be used."), *RequestedCultureName, *ValidCultureName);
 						}
 					}
+					else
+					{
+						// Fallback to English.
+						UE_LOG(LogTextLocalizationManager, Log, TEXT("No translations for ('%s') exist, falling back to 'en' for localization and internationalization data."), *RequestedCultureName);
+						TargetCultureName = TEXT("en");
+					}
+				}
 
-					I18N.SetCurrentCulture(TargetCultureName);
+				I18N.SetCurrentCulture(TargetCultureName);
 	}
 
 	FTextLocalizationManager::Get().LoadLocalizationResourcesForCulture(I18N.GetCurrentCulture()->GetName(), ShouldLoadEditor, ShouldLoadGame);
@@ -362,53 +359,48 @@ FTextDisplayStringRef FTextLocalizationManager::GetDisplayString(const FString&
 
 		const FTextDisplayStringRef UnlocalizedString = MakeShareable( new FString( SourceString ? **SourceString : TEXT("") ) );
 
-		// If live-culture-swap is enabled or the system is uninitialized - make entries so that they can be updated when system is initialized or a culture swap occurs.
-		CA_SUPPRESS(6236)
-			CA_SUPPRESS(6316)
-			if( !(bIsInitialized) || ENABLE_LOC_TESTING )
-			{
 #if ENABLE_LOC_TESTING
-				if( (bShouldLEETIFYAll || bShouldLEETIFYUnlocalizedString) && SourceString )
-				{
-					FInternationalization::Leetify(*UnlocalizedString);
-					if(*UnlocalizedString == *SourceString)
-					{
-						UE_LOG(LogTextLocalizationManager, Warning, TEXT("Leetify failed to alter a string (%s)."), **SourceString );
-					}
-				}
+		if( (bShouldLEETIFYAll || bShouldLEETIFYUnlocalizedString) && SourceString )
+		{
+			FInternationalization::Leetify(*UnlocalizedString);
+			if(*UnlocalizedString == *SourceString)
+			{
+				UE_LOG(LogTextLocalizationManager, Warning, TEXT("Leetify failed to alter a string (%s)."), **SourceString );
+			}
+		}
 #endif
 
-				if ( UnlocalizedString->IsEmpty() )
-				{
-					if ( !bIsInitialized )
-					{
-						*(UnlocalizedString) = AccessedStringBeforeLocLoadedErrorMsg;
-					}
-				}
+		if ( UnlocalizedString->IsEmpty() )
+		{
+			if ( !bIsInitialized )
+			{
+				*(UnlocalizedString) = AccessedStringBeforeLocLoadedErrorMsg;
+			}
+		}
 
-				FDisplayStringLookupTable::FDisplayStringEntry NewEntry(
+		// Make entries so that they can be updated when system is initialized or a culture swap occurs.
+		FDisplayStringLookupTable::FDisplayStringEntry NewEntry(
 #if ENABLE_LOC_TESTING
-					bShouldLEETIFYAll						/*bIsLocalized*/
+			bShouldLEETIFYAll						/*bIsLocalized*/
 #else
-					false												/*bIsLocalized*/
+			false												/*bIsLocalized*/
 #endif
-					, TEXT("")
-					, SourceString ? FCrc::StrCrc32(**SourceString) : 0	/*SourceStringHash*/
-					, UnlocalizedString									/*String*/
-					);
+			, TEXT("")
+			, SourceString ? FCrc::StrCrc32(**SourceString) : 0	/*SourceStringHash*/
+			, UnlocalizedString									/*String*/
+			);
 
-				if( !LiveKeyTable )
-				{
-					LiveKeyTable = &(DisplayStringLookupTable.NamespacesTable.Add( Namespace, FDisplayStringLookupTable::FKeysTable() ));
-				}
+		if( !LiveKeyTable )
+		{
+			LiveKeyTable = &(DisplayStringLookupTable.NamespacesTable.Add( Namespace, FDisplayStringLookupTable::FKeysTable() ));
+		}
 
-				LiveKeyTable->Add( Key, NewEntry );
+		LiveKeyTable->Add( Key, NewEntry );
 
-				NamespaceKeyLookupTable.Add(NewEntry.DisplayString, FNamespaceKeyEntry(Namespace, Key));
+		NamespaceKeyLookupTable.Add(NewEntry.DisplayString, FNamespaceKeyEntry(Namespace, Key));
 
-			}
 
-			return UnlocalizedString;
+		return UnlocalizedString;
 	}
 }
 
