Stats -  DebugLeakTest will close the stat file when requesting exit, fixed an issue with testing memory profiler on consoles

diff --git a/Engine/Source/Runtime/Core/Private/Stats/Stats2.cpp b/Engine/Source/Runtime/Core/Private/Stats/Stats2.cpp
--- a/Engine/Source/Runtime/Core/Private/Stats/Stats2.cpp
+++ b/Engine/Source/Runtime/Core/Private/Stats/Stats2.cpp
@@ -137,6 +137,12 @@ void DebugLeakTest()
 				(new FAutoDeleteAsyncTask<FAllocPool>())->StartBackgroundTask();
 			}
 		}
+
+		if (GIsRequestingExit)
+		{
+			// If we are writing stats data, stop it now.
+			DirectStatsCommand( TEXT( "stat stopfile" ), true );
+		}
 	}
 }
 
@@ -1283,9 +1289,6 @@ void FThreadStats::StopThread()
 			DirectStatsCommand( TEXT( "stat dumpsum -stop -ms=100" ), true );
 		}
 
-		// If we are writing stats data, stop it now.
-		DirectStatsCommand(TEXT("stat stopfile"), true);
-
 		FThreadStats::MasterDisableForever();
 
 		WaitForStats();
