Stats - Stage progress should be updated correctly

diff --git a/Engine/Source/Runtime/Core/Public/Stats/StatsFile.h b/Engine/Source/Runtime/Core/Public/Stats/StatsFile.h
--- a/Engine/Source/Runtime/Core/Public/Stats/StatsFile.h
+++ b/Engine/Source/Runtime/Core/Public/Stats/StatsFile.h
@@ -871,7 +871,6 @@ struct FCreateStatsReader
 };
 
 /** Struct used to read from ue4stats/ue4statsraw files, initializes all metadata and starts a process of reading the file asynchronously. */
-// #YRX_Stats: 2015-07-09 Remove CORE_API
 struct CORE_API FStatsReadFile
 {
 	friend class FAsyncRawStatsFile;
@@ -959,7 +958,11 @@ protected:
 	/** Sets a new processing stage for this file. */
 	void SetProcessingStage( EStatsProcessingStage NewStage )
 	{
-		ProcessingStage.Set( int32( NewStage ) );
+		if (GetProcessingStage() != NewStage)
+		{
+			ProcessingStage.Set( int32( NewStage ) );
+			StageProgress.Set( 0 );
+		}
 	}
 
 public:
