Stats - NumSecondsBetweenUpdates changed from 5 to 2 seconds

diff --git a/Engine/Source/Runtime/Core/Private/Stats/StatsFile.cpp b/Engine/Source/Runtime/Core/Private/Stats/StatsFile.cpp
--- a/Engine/Source/Runtime/Core/Private/Stats/StatsFile.cpp
+++ b/Engine/Source/Runtime/Core/Private/Stats/StatsFile.cpp
@@ -546,7 +546,7 @@ void FAsyncRawStatsFile::Abandon()
 	FStatsReadFile
 -----------------------------------------------------------------------------*/
 
-const double FStatsReadFile::NumSecondsBetweenUpdates = 5.0;
+const double FStatsReadFile::NumSecondsBetweenUpdates = 2.0;
 
 FStatsReadFile* FStatsReadFile::CreateReaderForRegularStats( const TCHAR* Filename )
 {
@@ -724,7 +724,7 @@ void FStatsReadFile::ReadStats()
 		// Background threads usually execute time consuming operations, so the frame number won't be valid.
 		// Needs to be combined by the thread and the time, not by the frame number.
 
-		// Update stage progress once per NumSecondsBetweenUpdates(5) seconds to avoid spamming.
+		// Update stage progress once per NumSecondsBetweenUpdates(2) seconds to avoid spamming.
 		SetProcessingStage( EStatsProcessingStage::SPS_ReadStats );
 
 		while (Reader->Tell() < Reader->TotalSize())
