Stats - Fixed bad processing stage set during prepare loading

diff --git a/Engine/Source/Runtime/Core/Private/Stats/StatsFile.cpp b/Engine/Source/Runtime/Core/Private/Stats/StatsFile.cpp
--- a/Engine/Source/Runtime/Core/Private/Stats/StatsFile.cpp
+++ b/Engine/Source/Runtime/Core/Private/Stats/StatsFile.cpp
@@ -607,15 +607,21 @@ bool FStatsReadFile::PrepareLoading()
 	SetProcessingStage( EStatsProcessingStage::SPS_Started );
 
 	{
+		bool bResult = true;
+
 		ON_SCOPE_EXIT
 		{
-			SetProcessingStage( EStatsProcessingStage::SPS_Invalid );
+			if (!bResult)
+			{
+				SetProcessingStage( EStatsProcessingStage::SPS_Invalid );
+			}
 		};
 
 		const int64 Size = IFileManager::Get().FileSize( *Filename );
 		if (Size < 4)
 		{
 			UE_LOG( LogStats, Error, TEXT( "Could not open: %s" ), *Filename );
+			bResult = false;
 			return false;
 		}
 
@@ -623,12 +629,14 @@ bool FStatsReadFile::PrepareLoading()
 		if (!Reader)
 		{
 			UE_LOG( LogStats, Error, TEXT( "Could not open: %s" ), *Filename );
+			bResult = false;
 			return false;
 		}
 
 		if (!Stream.ReadHeader( *Reader ))
 		{
 			UE_LOG( LogStats, Error, TEXT( "Could not read, header is invalid: %s" ), *Filename );
+			bResult = false;
 			return false;
 		}
 
@@ -638,12 +646,14 @@ bool FStatsReadFile::PrepareLoading()
 		if (!bIsFinalized)
 		{
 			UE_LOG( LogStats, Error, TEXT( "Could not read, file is not finalized: %s" ), *Filename );
+			bResult = false;
 			return false;
 		}
 
 		if (Stream.Header.Version < EStatMagicWithHeader::VERSION_6)
 		{
 			UE_LOG( LogStats, Error, TEXT( "Could not read, invalid version: %s, expected %u, was %u" ), *Filename, (uint32)EStatMagicWithHeader::VERSION_6, Stream.Header.Version );
+			bResult = false;
 			return false;
 		}
 
@@ -651,6 +661,7 @@ bool FStatsReadFile::PrepareLoading()
 		if (!bHasCompressedData)
 		{
 			UE_LOG( LogStats, Error, TEXT( "Could not read, required compressed data: %s" ), *Filename );
+			bResult = false;
 			return false;
 		}
 	}
