#ue4 - Add tests for FRotator::ZeroRotator equality with FQuat::Identity.

diff --git a/Engine/Source/Runtime/Core/Private/Tests/Math/UnrealMathTest.cpp b/Engine/Source/Runtime/Core/Private/Tests/Math/UnrealMathTest.cpp
--- a/Engine/Source/Runtime/Core/Private/Tests/Math/UnrealMathTest.cpp
+++ b/Engine/Source/Runtime/Core/Private/Tests/Math/UnrealMathTest.cpp
@@ -1069,6 +1069,21 @@ bool FVectorRegisterAbstractionTest::RunTest(const FString& Parameters)
 
 	// Quat<->Rotator conversions and equality
 	{
+		// Identity conversion
+		{
+			const FRotator R0 = FRotator::ZeroRotator;
+			const FRotator R1 = FRotator(FQuat::Identity);
+			LogRotatorTest(true, TEXT("FRotator::ZeroRotator ~= FQuat::Identity : Rotator"), R0, R1, R0.Equals(R1, 0.f));
+			LogRotatorTest(true, TEXT("FRotator::ZeroRotator == FQuat::Identity : Rotator"), R0, R1, R0 == R1);
+			LogRotatorTest(true, TEXT("FRotator::ZeroRotator not != FQuat::Identity : Rotator"), R0, R1, !(R0 != R1));
+
+			Q0 = FQuat::Identity;
+			Q1 = FQuat(FRotator::ZeroRotator);
+			LogQuaternionTest(TEXT("FRotator::ZeroRotator ~= FQuat::Identity : Quaternion"), Q0, Q1, Q0.Equals(Q1, 0.f));
+			LogQuaternionTest(TEXT("FRotator::ZeroRotator == FQuat::Identity : Quaternion"), Q0, Q1, Q0 == Q1);
+			LogQuaternionTest(TEXT("FRotator::ZeroRotator not != FQuat::Identity : Quaternion"), Q0, Q1, !(Q0 != Q1));
+		}
+
 		const float Nudge = KINDA_SMALL_NUMBER * 0.25f;
 		const FRotator RotArray[] = {
 			FRotator(0.f, 0.f, 0.f),
