FCulture::GetName now returns the culture's name using hyphens/dashes rather than underscores. Underscores were an implementation detail.

diff --git a/Engine/Source/Runtime/Core/Private/Internationalization/Culture.cpp b/Engine/Source/Runtime/Core/Private/Internationalization/Culture.cpp
--- a/Engine/Source/Runtime/Core/Private/Internationalization/Culture.cpp
+++ b/Engine/Source/Runtime/Core/Private/Internationalization/Culture.cpp
@@ -65,15 +65,15 @@ TArray<FString> FCulture::GetPrioritizedParentCultureNames()
 	TArray<FString> LocaleTagCombinations;
 	if (!ScriptCode.IsEmpty() && !RegionCode.IsEmpty())
 	{
-		LocaleTagCombinations.Add(LanguageCode + TEXT("_") + ScriptCode + TEXT("_") + RegionCode);
+		LocaleTagCombinations.Add(LanguageCode + TEXT("-") + ScriptCode + TEXT("-") + RegionCode);
 	}
 	if (!RegionCode.IsEmpty())
 	{
-		LocaleTagCombinations.Add(LanguageCode + TEXT("_") + RegionCode);
+		LocaleTagCombinations.Add(LanguageCode + TEXT("-") + RegionCode);
 	}
 	if (!ScriptCode.IsEmpty())
 	{
-		LocaleTagCombinations.Add(LanguageCode + TEXT("_") + ScriptCode);
+		LocaleTagCombinations.Add(LanguageCode + TEXT("-") + ScriptCode);
 	}
 	LocaleTagCombinations.Add(LanguageCode);
 
diff --git a/Engine/Source/Runtime/Core/Private/Internationalization/ICUCulture.cpp b/Engine/Source/Runtime/Core/Private/Internationalization/ICUCulture.cpp
--- a/Engine/Source/Runtime/Core/Private/Internationalization/ICUCulture.cpp
+++ b/Engine/Source/Runtime/Core/Private/Internationalization/ICUCulture.cpp
@@ -141,7 +141,9 @@ FString FCulture::FICUCultureImplementation::GetCanonicalName(const FString& Nam
 
 FString FCulture::FICUCultureImplementation::GetName() const
 {
-	return ICULocale.getName();
+	FString Result = ICULocale.getName();
+	Result.ReplaceInline(TEXT("_"), TEXT("-"));
+	return Result;
 }
 
 FString FCulture::FICUCultureImplementation::GetNativeName() const
