Stats - Fixed rawstats crash during shutdown

diff --git a/Engine/Source/Runtime/Core/Private/Stats/Stats2.cpp b/Engine/Source/Runtime/Core/Private/Stats/Stats2.cpp
--- a/Engine/Source/Runtime/Core/Private/Stats/Stats2.cpp
+++ b/Engine/Source/Runtime/Core/Private/Stats/Stats2.cpp
@@ -1289,6 +1289,9 @@ void FThreadStats::StopThread()
 			DirectStatsCommand( TEXT( "stat dumpsum -stop -ms=100" ), true );
 		}
 
+		// If we are writing stats data, stop it now.
+		DirectStatsCommand( TEXT( "stat stopfile" ), true );
+
 		FThreadStats::MasterDisableForever();
 
 		WaitForStats();
