Linux: added a platform survey (WIP).

diff --git a/Engine/Source/Runtime/Core/Private/Linux/LinuxPlatformSurvey.cpp b/Engine/Source/Runtime/Core/Private/Linux/LinuxPlatformSurvey.cpp
new file mode 100644
--- /dev/null
+++ b/Engine/Source/Runtime/Core/Private/Linux/LinuxPlatformSurvey.cpp
@@ -0,0 +1,55 @@
+// Copyright 1998-2015 Epic Games, Inc. All Rights Reserved.
+
+#include "CorePrivatePCH.h"
+
+#include "UnrealString.h"
+#include "SynthBenchmark.h"
+
+bool FLinuxPlatformSurvey::GetSurveyResults(FHardwareSurveyResults& OutResults, bool bWait)
+{
+	FMemory::Memset(&OutResults, 0, sizeof(FHardwareSurveyResults));
+	WriteFStringToResults(OutResults.Platform, TEXT("Linux"));
+
+	// CPU
+	OutResults.CPUCount = FPlatformMisc::NumberOfCores();  // TODO [RCL] 2015-07-15: parse /proc/cpuinfo for GHz/brand
+
+	// Memory
+	const FPlatformMemoryConstants& MemoryConstants = FPlatformMemory::GetConstants();
+	OutResults.MemoryMB = MemoryConstants.TotalPhysical / ( 1024ULL * 1024ULL );
+
+	// Misc
+	OutResults.bIsRemoteSession = FPlatformMisc::HasBeenStartedRemotely();
+	OutResults.bIsLaptopComputer = FPlatformMisc::IsRunningOnBattery();	// FIXME [RCL] 2015-07-15: incorrect. Laptops don't have to run on battery
+
+	// Synth benchmark
+	ISynthBenchmark::Get().Run(OutResults.SynthBenchmark);
+
+	OutResults.ErrorCount++;
+	WriteFStringToResults(OutResults.LastSurveyError, TEXT("Survey is incomplete"));
+	WriteFStringToResults(OutResults.LastSurveyErrorDetail, TEXT("CPU, OS details are missing"));
+
+	return true;
+}
+
+void FLinuxPlatformSurvey::GetOSName(FHardwareSurveyResults& OutResults)
+{
+	// TODO [RCL] 2015-07-15: check if /etc/os-release or /etc/redhat-release exist and parse it
+
+	/*
+	TArray<FString> OsReleaseLines;
+	if( FFileHelper::LoadANSITextFileToStrings(TEXT("/etc/os-release"), &IFileManager::Get(), OsReleaseLines))
+	{
+		//...
+	}
+	*/
+}
+
+void FLinuxPlatformSurvey::WriteFStringToResults(TCHAR* OutBuffer, const FString& InString)
+{
+	FMemory::Memset( OutBuffer, 0, sizeof(TCHAR) * FHardwareSurveyResults::MaxStringLength );
+	TCHAR* Cursor = OutBuffer;
+	for (int32 i = 0; i < FMath::Min(InString.Len(), FHardwareSurveyResults::MaxStringLength - 1); i++)
+	{
+		*Cursor++ = InString[i];
+	}
+}
diff --git a/Engine/Source/Runtime/Core/Public/Linux/LinuxPlatformSurvey.h b/Engine/Source/Runtime/Core/Public/Linux/LinuxPlatformSurvey.h
--- a/Engine/Source/Runtime/Core/Public/Linux/LinuxPlatformSurvey.h
+++ b/Engine/Source/Runtime/Core/Public/Linux/LinuxPlatformSurvey.h
@@ -13,7 +13,20 @@
 **/
 struct FLinuxPlatformSurvey : public FGenericPlatformSurvey
 {
-	// default implementation for now
+	/** Start, or check on, the hardware survey */
+	static bool GetSurveyResults(FHardwareSurveyResults& OutResults, bool bWait = false);
+
+private:
+
+	/**
+	 * Fills out the actual OS name ("distro" is an OS for all the practical purposes), as detailed as possible.
+	 */
+	static void GetOSName(FHardwareSurveyResults& Results);
+
+	/**
+	 * Safely write strings into the fixed length TCHAR buffers of a FHardwareSurveyResults member
+	 */
+	static void WriteFStringToResults(TCHAR* OutBuffer, const FString& InString);
 };
 
-typedef FLinuxPlatformSurvey FPlatformSurvey;
\ No newline at end of file
+typedef FLinuxPlatformSurvey FPlatformSurvey;
