Linux: fixed Unicode keyboard input (UE-13042)

diff --git a/Engine/Source/Runtime/Core/Private/Linux/LinuxApplication.cpp b/Engine/Source/Runtime/Core/Private/Linux/LinuxApplication.cpp
--- a/Engine/Source/Runtime/Core/Private/Linux/LinuxApplication.cpp
+++ b/Engine/Source/Runtime/Core/Private/Linux/LinuxApplication.cpp
@@ -252,11 +252,12 @@ void FLinuxApplication::ProcessDeferredMessage( SDL_Event Event )
 	case SDL_TEXTINPUT:
 		{
 			// Slate now gets all its text from here, I hope.
-			// Don't know if this will work with ingame text or not.
-			const bool bIsRepeated = Event.key.repeat != 0;
-			const TCHAR Character = *ANSI_TO_TCHAR(Event.text.text);
-
-			MessageHandler->OnKeyChar(Character, bIsRepeated);
+			const bool bIsRepeated = false;  //Event.key.repeat != 0;
+			const FString TextStr(UTF8_TO_TCHAR(Event.text.text));
+			for (auto TextIter = TextStr.CreateConstIterator(); TextIter; ++TextIter)
+			{
+				MessageHandler->OnKeyChar(*TextIter, bIsRepeated);
+			}
 		}
 		break;
 	case SDL_MOUSEMOTION:
