Remove Slate Stats out of core and add a verbose stat for high frequency Slate stats.

diff --git a/Engine/Source/Runtime/Core/Public/Stats/Stats2.h b/Engine/Source/Runtime/Core/Public/Stats/Stats2.h
--- a/Engine/Source/Runtime/Core/Public/Stats/Stats2.h
+++ b/Engine/Source/Runtime/Core/Public/Stats/Stats2.h
@@ -2041,8 +2041,6 @@ DECLARE_STATS_GROUP(TEXT("Server CPU"),STATGROUP_ServerCPU, STATCAT_Advanced);
 DECLARE_STATS_GROUP(TEXT("Shader Compiling"),STATGROUP_ShaderCompiling, STATCAT_Advanced);
 DECLARE_STATS_GROUP(TEXT("Shader Compression"),STATGROUP_Shaders, STATCAT_Advanced);
 DECLARE_STATS_GROUP(TEXT("Shadow Rendering"),STATGROUP_ShadowRendering, STATCAT_Advanced);
-DECLARE_STATS_GROUP(TEXT("Slate Memory"), STATGROUP_SlateMemory, STATCAT_Advanced);
-DECLARE_STATS_GROUP(TEXT("Slate"), STATGROUP_Slate, STATCAT_Advanced);
 DECLARE_STATS_GROUP(TEXT("Stat System"),STATGROUP_StatSystem, STATCAT_Advanced);
 DECLARE_STATS_GROUP(TEXT("Streaming Details"),STATGROUP_StreamingDetails, STATCAT_Advanced);
 DECLARE_STATS_GROUP(TEXT("Streaming"),STATGROUP_Streaming, STATCAT_Advanced);
diff --git a/Engine/Source/Runtime/SlateCore/Private/Rendering/DrawElements.cpp b/Engine/Source/Runtime/SlateCore/Private/Rendering/DrawElements.cpp
--- a/Engine/Source/Runtime/SlateCore/Private/Rendering/DrawElements.cpp
+++ b/Engine/Source/Runtime/SlateCore/Private/Rendering/DrawElements.cpp
@@ -2,7 +2,7 @@
 
 #include "SlateCorePrivatePCH.h"
 
-DECLARE_CYCLE_STAT(TEXT("FSlateDrawElement::Make Time"), STAT_SlateDrawElementMakeTime, STATGROUP_Slate);
+DECLARE_CYCLE_STAT(TEXT("FSlateDrawElement::Make Time"), STAT_SlateDrawElementMakeTime, STATGROUP_SlateVerbose);
 
 void FSlateDrawElement::Init(uint32 InLayer, const FPaintGeometry& PaintGeometry, const FSlateRect& InClippingRect, ESlateDrawEffect::Type InDrawEffects)
 {
diff --git a/Engine/Source/Runtime/SlateCore/Private/Rendering/ElementBatcher.cpp b/Engine/Source/Runtime/SlateCore/Private/Rendering/ElementBatcher.cpp
--- a/Engine/Source/Runtime/SlateCore/Private/Rendering/ElementBatcher.cpp
+++ b/Engine/Source/Runtime/SlateCore/Private/Rendering/ElementBatcher.cpp
@@ -6,8 +6,9 @@
 #include "Internationalization/Text.h"
 #include "SlateStats.h"
 
-//DECLARE_CYCLE_STAT(TEXT("Find Batch For Element Time"), STAT_SlateFindBatchForElement, STATGROUP_Slate);
-DECLARE_DWORD_COUNTER_STAT(TEXT("Num Elements (Prebatch)"), STAT_SlateNumPrebatchElements, STATGROUP_Slate);
+DECLARE_CYCLE_STAT(TEXT("Find Batch For Element Time"), STAT_SlateFindBatchForElement, STATGROUP_SlateVerbose);
+DECLARE_DWORD_COUNTER_STAT(TEXT("Num Elements (Prebatch)"), STAT_SlateNumPrebatchElements, STATGROUP_SlateVerbose);
+
 DECLARE_CYCLE_STAT(TEXT("Add Elements Time"), STAT_SlateAddElements, STATGROUP_Slate);
 
 SLATE_DECLARE_CYCLE_COUNTER(GSlateAddElements, "Add Elements");
diff --git a/Engine/Source/Runtime/SlateCore/Public/SlateCore.h b/Engine/Source/Runtime/SlateCore/Public/SlateCore.h
--- a/Engine/Source/Runtime/SlateCore/Public/SlateCore.h
+++ b/Engine/Source/Runtime/SlateCore/Public/SlateCore.h
@@ -17,6 +17,11 @@
 SLATECORE_API DECLARE_LOG_CATEGORY_EXTERN(LogSlate, Log, All);
 SLATECORE_API DECLARE_LOG_CATEGORY_EXTERN(LogSlateStyles, Log, All);
 
+DECLARE_STATS_GROUP(TEXT("Slate Memory"), STATGROUP_SlateMemory, STATCAT_Advanced);
+DECLARE_STATS_GROUP(TEXT("Slate"), STATGROUP_Slate, STATCAT_Advanced);
+DECLARE_STATS_GROUP_VERBOSE(TEXT("SlateVerbose"), STATGROUP_SlateVerbose, STATCAT_Advanced);
+
+
 // Compile all the RichText and MultiLine editable text?
 #define WITH_FANCY_TEXT 1
 
