Fix UE-18786 Toggling between windows on Mac brings up console command first time - forgot to add handling of NSFlagsChanged.

diff --git a/Engine/Source/Runtime/Core/Private/Mac/MacApplication.cpp b/Engine/Source/Runtime/Core/Private/Mac/MacApplication.cpp
--- a/Engine/Source/Runtime/Core/Private/Mac/MacApplication.cpp
+++ b/Engine/Source/Runtime/Core/Private/Mac/MacApplication.cpp
@@ -549,6 +549,10 @@ void FMacApplication::ProcessEvent(const FDeferredMacEvent& Event)
 				bModifiersValid = true;
 				break;
 				
+			case NSFlagsChanged:
+				bModifiersValid = true;
+				break;
+				
 			case NSMouseEntered:
 			case NSMouseExited:
 				bModifiersValid = true;
