Stats - Removed obsolete comment

diff --git a/Engine/Source/Runtime/Core/Private/Stats/StatsFile.cpp b/Engine/Source/Runtime/Core/Private/Stats/StatsFile.cpp
--- a/Engine/Source/Runtime/Core/Private/Stats/StatsFile.cpp
+++ b/Engine/Source/Runtime/Core/Private/Stats/StatsFile.cpp
@@ -871,9 +871,6 @@ void FStatsReadFile::ProcessStats()
 							}
 							else if (Op == EStatOperation::Memory)
 							{
-								// ProcessMemoryOperation
-
-								// Experimental code used only to test the implementation.
 								// First memory operation is Alloc or Free
 								const uint64 EncodedPtr = Message.GetValue_Ptr();
 								const EMemoryOperation MemOp = EMemoryOperation( EncodedPtr & (uint64)EMemoryOperation::Mask );
