Stats -  DebugLeakTest tweaks

diff --git a/Engine/Source/Runtime/Core/Private/Stats/Stats2.cpp b/Engine/Source/Runtime/Core/Private/Stats/Stats2.cpp
--- a/Engine/Source/Runtime/Core/Private/Stats/Stats2.cpp
+++ b/Engine/Source/Runtime/Core/Private/Stats/Stats2.cpp
@@ -94,7 +94,7 @@ void DebugLeakTest()
 		}
 
 
-		if (GFrameCounter <= 250)
+		if (GFrameCounter < 250)
 		{
 			// Background threads memory test.
 			struct FAllocTask
@@ -103,7 +103,7 @@ void DebugLeakTest()
 				{
 					DECLARE_SCOPE_CYCLE_COUNTER( TEXT( "FAllocTask::Alloc" ), Stat_FAllocTask_Alloc, STATGROUP_Quick );
 
-					int* IntAlloc = new int[1000];
+					int8* IntAlloc = new int8[112233];
 					int8* LeakTask = new int8[100000];
 					delete[] IntAlloc;
 				}
@@ -121,7 +121,7 @@ void DebugLeakTest()
 				{
 					DECLARE_SCOPE_CYCLE_COUNTER( TEXT( "FAllocPool::DoWork" ), Stat_FAllocPool_DoWork, STATGROUP_Quick );
 
-					int* IntAlloc = new int[1000];
+					int8* IntAlloc = new int8[223311];
 					int8* LeakTask = new int8[100000];
 					delete[] IntAlloc;
 				}
@@ -138,6 +138,15 @@ void DebugLeakTest()
 			}
 		}
 
+		for (int32 Index = 0; Index < 40; ++Index)
+		{
+			DECLARE_SCOPE_CYCLE_COUNTER( TEXT( "DebugLeakTest::Alloc" ), Stat_LeakTest_Alloc, STATGROUP_Quick );
+
+			int8* IntAlloc = new int8[331122];
+			int8* LeakTask = new int8[100000];
+			delete[] IntAlloc;
+		}
+
 		if (GIsRequestingExit)
 		{
 			// If we are writing stats data, stop it now.
