Linux: several keyboard input improvements.

- No longer assumes a US-QWERTY layout (uses SDL equivalent of Windows VK_ values, not raw keyboard scancodes)
- Fixes various keys not working for odd reasons (we lost bits downcasting from uint32 to uint16, causing key values to change or overlap other keys).

(Fixes UE-11582 and UE-14143 and probably others.)

Misc:
- Reverted to using WarpMouseInWindow due to reports that Global is not always reliable.
- Added a workaround for relative mouse movement being wrong in some cases.

diff --git a/Engine/Source/Runtime/Core/Private/Android/AndroidMisc.cpp b/Engine/Source/Runtime/Core/Private/Android/AndroidMisc.cpp
--- a/Engine/Source/Runtime/Core/Private/Android/AndroidMisc.cpp
+++ b/Engine/Source/Runtime/Core/Private/Android/AndroidMisc.cpp
@@ -463,7 +463,7 @@ FString FAndroidMisc::GetDefaultLocale()
 	return OSLanguage;
 }
 
-uint32 FAndroidMisc::GetCharKeyMap(uint16* KeyCodes, FString* KeyNames, uint32 MaxMappings)
+uint32 FAndroidMisc::GetCharKeyMap(uint32* KeyCodes, FString* KeyNames, uint32 MaxMappings)
 {
 #define ADDKEYMAP(KeyCode, KeyName)		if (NumMappings<MaxMappings) { KeyCodes[NumMappings]=KeyCode; if(KeyNames) { KeyNames[NumMappings]=KeyName; } ++NumMappings; };
 
@@ -589,7 +589,7 @@ uint32 FAndroidMisc::GetCharKeyMap(uint16* KeyCodes, FString* KeyNames, uint32 M
 
 }
 
-uint32 FAndroidMisc::GetKeyMap( uint16* KeyCodes, FString* KeyNames, uint32 MaxMappings )
+uint32 FAndroidMisc::GetKeyMap( uint32* KeyCodes, FString* KeyNames, uint32 MaxMappings )
 {
 #define ADDKEYMAP(KeyCode, KeyName)		if (NumMappings<MaxMappings) { KeyCodes[NumMappings]=KeyCode; if(KeyNames) { KeyNames[NumMappings]=KeyName; } ++NumMappings; };
 
diff --git a/Engine/Source/Runtime/Core/Private/GenericPlatform/GenericPlatformMisc.cpp b/Engine/Source/Runtime/Core/Private/GenericPlatform/GenericPlatformMisc.cpp
--- a/Engine/Source/Runtime/Core/Private/GenericPlatform/GenericPlatformMisc.cpp
+++ b/Engine/Source/Runtime/Core/Private/GenericPlatform/GenericPlatformMisc.cpp
@@ -680,7 +680,7 @@ const TCHAR* FGenericPlatformMisc::GameDir()
 	return *GameDir;
 }
 
-uint32 FGenericPlatformMisc::GetStandardPrintableKeyMap(uint16* KeyCodes, FString* KeyNames, uint32 MaxMappings, bool bMapUppercaseKeys, bool bMapLowercaseKeys)
+uint32 FGenericPlatformMisc::GetStandardPrintableKeyMap(uint32* KeyCodes, FString* KeyNames, uint32 MaxMappings, bool bMapUppercaseKeys, bool bMapLowercaseKeys)
 {
 	uint32 NumMappings = 0;
 
diff --git a/Engine/Source/Runtime/Core/Private/HTML5/HTML5PlatformMisc.cpp b/Engine/Source/Runtime/Core/Private/HTML5/HTML5PlatformMisc.cpp
--- a/Engine/Source/Runtime/Core/Private/HTML5/HTML5PlatformMisc.cpp
+++ b/Engine/Source/Runtime/Core/Private/HTML5/HTML5PlatformMisc.cpp
@@ -38,12 +38,12 @@ GenericApplication* FHTML5Misc::CreateApplication()
 	return FHTML5Application::CreateHTML5Application();
 }
 
-uint32 FHTML5Misc::GetCharKeyMap(uint16* KeyCodes, FString* KeyNames, uint32 MaxMappings)
+uint32 FHTML5Misc::GetCharKeyMap(uint32* KeyCodes, FString* KeyNames, uint32 MaxMappings)
 {
 	return FGenericPlatformMisc::GetStandardPrintableKeyMap(KeyCodes, KeyNames, MaxMappings, false, true);
 }
 
-uint32 FHTML5Misc::GetKeyMap(uint16* KeyCodes, FString* KeyNames, uint32 MaxMappings)
+uint32 FHTML5Misc::GetKeyMap(uint32* KeyCodes, FString* KeyNames, uint32 MaxMappings)
 {
 #define ADDKEYMAP(KeyCode, KeyName)		if (NumMappings<MaxMappings) { KeyCodes[NumMappings]=KeyCode; KeyNames[NumMappings]=KeyName; ++NumMappings; };
 
diff --git a/Engine/Source/Runtime/Core/Private/Linux/LinuxApplication.cpp b/Engine/Source/Runtime/Core/Private/Linux/LinuxApplication.cpp
--- a/Engine/Source/Runtime/Core/Private/Linux/LinuxApplication.cpp
+++ b/Engine/Source/Runtime/Core/Private/Linux/LinuxApplication.cpp
@@ -203,6 +203,19 @@ bool FLinuxApplication::GeneratesKeyCharMessage(const SDL_KeyboardEvent & KeyDow
 		(Sym != SDLK_DOWN && Sym != SDLK_LEFT && Sym != SDLK_RIGHT && Sym != SDLK_UP && Sym != SDLK_DELETE);
 }
 
+static inline uint32 CharCodeFromSDLKeySym(const SDL_Keycode KeySym)
+{
+	if ((KeySym & SDLK_SCANCODE_MASK) != 0)
+	{
+		return 0;
+	}
+	else if (KeySym == SDLK_DELETE)  // this doesn't use the scancode mask for some reason.
+	{
+		return 0;
+	}
+	return (uint32) KeySym;
+}
+
 void FLinuxApplication::ProcessDeferredMessage( SDL_Event Event )
 {
 	// This function can be reentered when entering a modal tick loop.
@@ -225,28 +238,33 @@ void FLinuxApplication::ProcessDeferredMessage( SDL_Event Event )
 	{
 	case SDL_KEYDOWN:
 		{
-			SDL_KeyboardEvent KeyEvent = Event.key;
-			const SDL_Keycode KeyCode = KeyEvent.keysym.scancode;
+			const SDL_KeyboardEvent &KeyEvent = Event.key;
+			const SDL_Keycode KeySym = KeyEvent.keysym.sym;
+			const uint32 CharCode = CharCodeFromSDLKeySym(KeySym);
 			const bool bIsRepeated = KeyEvent.repeat != 0;
 
 			// Text input is now handled in SDL_TEXTINPUT: see below
-			MessageHandler->OnKeyDown(KeyCode, KeyEvent.keysym.sym, bIsRepeated);
+			MessageHandler->OnKeyDown(KeySym, CharCode, bIsRepeated);
 
 			// Backspace input in only caught here.
-			if (KeyCode == SDL_SCANCODE_BACKSPACE || KeyCode == SDL_SCANCODE_RETURN)
+			if (KeySym == SDLK_BACKSPACE)
 			{
-				const TCHAR Character = SDL_GetKeyFromScancode(Event.key.keysym.scancode);
-				MessageHandler->OnKeyChar(Character, bIsRepeated);
+				MessageHandler->OnKeyChar('\b', bIsRepeated);
+			}
+			else if (KeySym == SDLK_RETURN)
+			{
+				MessageHandler->OnKeyChar('\r', bIsRepeated);
 			}
 		}
 		break;
 	case SDL_KEYUP:
 		{
-			SDL_KeyboardEvent keyEvent = Event.key;
-			const SDL_Keycode KeyCode = keyEvent.keysym.scancode;
-			const bool IsRepeat = keyEvent.repeat != 0;
+			const SDL_KeyboardEvent &KeyEvent = Event.key;
+			const SDL_Keycode KeySym = KeyEvent.keysym.sym;
+			const uint32 CharCode = CharCodeFromSDLKeySym(KeySym);
+			const bool IsRepeat = KeyEvent.repeat != 0;
 
-			MessageHandler->OnKeyUp( KeyCode, keyEvent.keysym.sym, IsRepeat );
+			MessageHandler->OnKeyUp( KeySym, CharCode, IsRepeat );
 		}
 		break;
 	case SDL_TEXTINPUT:
@@ -313,7 +331,20 @@ void FLinuxApplication::ProcessDeferredMessage( SDL_Event Event )
 
 			if(bUsingHighPrecisionMouseInput)
 			{
- 				MessageHandler->OnRawMouseMove(motionEvent.xrel, motionEvent.yrel);
+				// hack to work around jumps
+				const int kTooFarAway = 100;
+				const int kTooFarAwaySquare = kTooFarAway * kTooFarAway;
+				if (motionEvent.xrel * motionEvent.xrel + motionEvent.yrel * motionEvent.yrel > kTooFarAwaySquare)
+				{
+					UE_LOG(LogLinuxWindowEvent, Warning, TEXT("Suppressing too large relative mouse movement due to an apparent bug (%d, %d is larger than treshold %d)"),
+						motionEvent.xrel, motionEvent.yrel,
+						kTooFarAway
+						);
+				}
+				else
+				{
+					MessageHandler->OnRawMouseMove(motionEvent.xrel, motionEvent.yrel);
+				}
 			}
 			else
 			{
diff --git a/Engine/Source/Runtime/Core/Private/Linux/LinuxCursor.cpp b/Engine/Source/Runtime/Core/Private/Linux/LinuxCursor.cpp
--- a/Engine/Source/Runtime/Core/Private/Linux/LinuxCursor.cpp
+++ b/Engine/Source/Runtime/Core/Private/Linux/LinuxCursor.cpp
@@ -196,7 +196,19 @@ void FLinuxCursor::SetCachedPosition( const int32 X, const int32 Y )
 
 void FLinuxCursor::SetPosition( const int32 X, const int32 Y )
 {
-	SDL_WarpMouseGlobal( X, Y);
+	// according to reports on IRC SDL_WarpMouseGlobal() doesn't work on some WMs so use InWindow unless we don't have a window.
+	SDL_HWindow WndFocus = SDL_GetMouseFocus();
+
+	if (WndFocus)
+	{
+		int WndX, WndY;
+		SDL_GetWindowPosition(WndFocus, &WndX, &WndY);	//	get top left
+		SDL_WarpMouseInWindow(WndFocus, X - WndX, Y - WndY);
+	}
+	else
+	{
+		SDL_WarpMouseGlobal(X, Y);
+	}
 
 	SetCachedPosition(X, Y);
 }
diff --git a/Engine/Source/Runtime/Core/Private/Linux/LinuxPlatformMisc.cpp b/Engine/Source/Runtime/Core/Private/Linux/LinuxPlatformMisc.cpp
--- a/Engine/Source/Runtime/Core/Private/Linux/LinuxPlatformMisc.cpp
+++ b/Engine/Source/Runtime/Core/Private/Linux/LinuxPlatformMisc.cpp
@@ -259,7 +259,7 @@ void FLinuxPlatformMisc::PumpMessages( bool bFromMainLoop )
 	}
 }
 
-uint32 FLinuxPlatformMisc::GetCharKeyMap(uint16* KeyCodes, FString* KeyNames, uint32 MaxMappings)
+uint32 FLinuxPlatformMisc::GetCharKeyMap(uint32* KeyCodes, FString* KeyNames, uint32 MaxMappings)
 {
 	return FGenericPlatformMisc::GetStandardPrintableKeyMap(KeyCodes, KeyNames, MaxMappings, false, true);
 }
@@ -270,7 +270,7 @@ void FLinuxPlatformMisc::LowLevelOutputDebugString(const TCHAR *Message)
 	fprintf(stderr, "%ls", Message);	// there's no good way to implement that really
 }
 
-uint32 FLinuxPlatformMisc::GetKeyMap( uint16* KeyCodes, FString* KeyNames, uint32 MaxMappings )
+uint32 FLinuxPlatformMisc::GetKeyMap( uint32* KeyCodes, FString* KeyNames, uint32 MaxMappings )
 {
 #define ADDKEYMAP(KeyCode, KeyName)		if (NumMappings<MaxMappings) { KeyCodes[NumMappings]=KeyCode; KeyNames[NumMappings]=KeyName; ++NumMappings; };
 
@@ -278,48 +278,69 @@ uint32 FLinuxPlatformMisc::GetKeyMap( uint16* KeyCodes, FString* KeyNames, uint3
 
 	if (KeyCodes && KeyNames && (MaxMappings > 0))
 	{
-		ADDKEYMAP(SDL_SCANCODE_BACKSPACE, TEXT("BackSpace"));
-		ADDKEYMAP(SDL_SCANCODE_TAB, TEXT("Tab"));
-		ADDKEYMAP(SDL_SCANCODE_RETURN, TEXT("Enter"));
-		ADDKEYMAP(SDL_SCANCODE_RETURN2, TEXT("Enter"));
-		ADDKEYMAP(SDL_SCANCODE_KP_ENTER, TEXT("Enter"));
-		ADDKEYMAP(SDL_SCANCODE_PAUSE, TEXT("Pause"));
-
-		ADDKEYMAP(SDL_SCANCODE_ESCAPE, TEXT("Escape"));
-		ADDKEYMAP(SDL_SCANCODE_SPACE, TEXT("SpaceBar"));
-		ADDKEYMAP(SDL_SCANCODE_PAGEUP, TEXT("PageUp"));
-		ADDKEYMAP(SDL_SCANCODE_PAGEDOWN, TEXT("PageDown"));
-		ADDKEYMAP(SDL_SCANCODE_END, TEXT("End"));
-		ADDKEYMAP(SDL_SCANCODE_HOME, TEXT("Home"));
-
-		ADDKEYMAP(SDL_SCANCODE_LEFT, TEXT("Left"));
-		ADDKEYMAP(SDL_SCANCODE_UP, TEXT("Up"));
-		ADDKEYMAP(SDL_SCANCODE_RIGHT, TEXT("Right"));
-		ADDKEYMAP(SDL_SCANCODE_DOWN, TEXT("Down"));
-
-		ADDKEYMAP(SDL_SCANCODE_INSERT, TEXT("Insert"));
-		ADDKEYMAP(SDL_SCANCODE_DELETE, TEXT("Delete"));
-
-		ADDKEYMAP(SDL_SCANCODE_F1, TEXT("F1"));
-		ADDKEYMAP(SDL_SCANCODE_F2, TEXT("F2"));
-		ADDKEYMAP(SDL_SCANCODE_F3, TEXT("F3"));
-		ADDKEYMAP(SDL_SCANCODE_F4, TEXT("F4"));
-		ADDKEYMAP(SDL_SCANCODE_F5, TEXT("F5"));
-		ADDKEYMAP(SDL_SCANCODE_F6, TEXT("F6"));
-		ADDKEYMAP(SDL_SCANCODE_F7, TEXT("F7"));
-		ADDKEYMAP(SDL_SCANCODE_F8, TEXT("F8"));
-		ADDKEYMAP(SDL_SCANCODE_F9, TEXT("F9"));
-		ADDKEYMAP(SDL_SCANCODE_F10, TEXT("F10"));
-		ADDKEYMAP(SDL_SCANCODE_F11, TEXT("F11"));
-		ADDKEYMAP(SDL_SCANCODE_F12, TEXT("F12"));
-
-        ADDKEYMAP(SDL_SCANCODE_CAPSLOCK, TEXT("CapsLock"));
-        ADDKEYMAP(SDL_SCANCODE_LCTRL, TEXT("LeftControl"));
-        ADDKEYMAP(SDL_SCANCODE_LSHIFT, TEXT("LeftShift"));
-        ADDKEYMAP(SDL_SCANCODE_LALT, TEXT("LeftAlt"));
-        ADDKEYMAP(SDL_SCANCODE_RCTRL, TEXT("RightControl"));
-        ADDKEYMAP(SDL_SCANCODE_RSHIFT, TEXT("RightShift"));
-        ADDKEYMAP(SDL_SCANCODE_RALT, TEXT("RightAlt"));
+		ADDKEYMAP(SDLK_BACKSPACE, TEXT("BackSpace"));
+		ADDKEYMAP(SDLK_TAB, TEXT("Tab"));
+		ADDKEYMAP(SDLK_RETURN, TEXT("Enter"));
+		ADDKEYMAP(SDLK_RETURN2, TEXT("Enter"));
+		ADDKEYMAP(SDLK_KP_ENTER, TEXT("Enter"));
+		ADDKEYMAP(SDLK_PAUSE, TEXT("Pause"));
+
+		ADDKEYMAP(SDLK_ESCAPE, TEXT("Escape"));
+		ADDKEYMAP(SDLK_SPACE, TEXT("SpaceBar"));
+		ADDKEYMAP(SDLK_PAGEUP, TEXT("PageUp"));
+		ADDKEYMAP(SDLK_PAGEDOWN, TEXT("PageDown"));
+		ADDKEYMAP(SDLK_END, TEXT("End"));
+		ADDKEYMAP(SDLK_HOME, TEXT("Home"));
+
+		ADDKEYMAP(SDLK_LEFT, TEXT("Left"));
+		ADDKEYMAP(SDLK_UP, TEXT("Up"));
+		ADDKEYMAP(SDLK_RIGHT, TEXT("Right"));
+		ADDKEYMAP(SDLK_DOWN, TEXT("Down"));
+
+		ADDKEYMAP(SDLK_INSERT, TEXT("Insert"));
+		ADDKEYMAP(SDLK_DELETE, TEXT("Delete"));
+
+		ADDKEYMAP(SDLK_F1, TEXT("F1"));
+		ADDKEYMAP(SDLK_F2, TEXT("F2"));
+		ADDKEYMAP(SDLK_F3, TEXT("F3"));
+		ADDKEYMAP(SDLK_F4, TEXT("F4"));
+		ADDKEYMAP(SDLK_F5, TEXT("F5"));
+		ADDKEYMAP(SDLK_F6, TEXT("F6"));
+		ADDKEYMAP(SDLK_F7, TEXT("F7"));
+		ADDKEYMAP(SDLK_F8, TEXT("F8"));
+		ADDKEYMAP(SDLK_F9, TEXT("F9"));
+		ADDKEYMAP(SDLK_F10, TEXT("F10"));
+		ADDKEYMAP(SDLK_F11, TEXT("F11"));
+		ADDKEYMAP(SDLK_F12, TEXT("F12"));
+
+        ADDKEYMAP(SDLK_LCTRL, TEXT("LeftControl"));
+        ADDKEYMAP(SDLK_LSHIFT, TEXT("LeftShift"));
+        ADDKEYMAP(SDLK_LALT, TEXT("LeftAlt"));
+		ADDKEYMAP(SDLK_LGUI, TEXT("LeftCommand"));
+        ADDKEYMAP(SDLK_RCTRL, TEXT("RightControl"));
+        ADDKEYMAP(SDLK_RSHIFT, TEXT("RightShift"));
+        ADDKEYMAP(SDLK_RALT, TEXT("RightAlt"));
+		ADDKEYMAP(SDLK_RGUI, TEXT("RightCommand"));
+
+		ADDKEYMAP(SDLK_KP_0, TEXT("NumPadZero"));
+		ADDKEYMAP(SDLK_KP_1, TEXT("NumPadOne"));
+		ADDKEYMAP(SDLK_KP_2, TEXT("NumPadTwo"));
+		ADDKEYMAP(SDLK_KP_3, TEXT("NumPadThree"));
+		ADDKEYMAP(SDLK_KP_4, TEXT("NumPadFour"));
+		ADDKEYMAP(SDLK_KP_5, TEXT("NumPadFive"));
+		ADDKEYMAP(SDLK_KP_6, TEXT("NumPadSix"));
+		ADDKEYMAP(SDLK_KP_7, TEXT("NumPadSeven"));
+		ADDKEYMAP(SDLK_KP_8, TEXT("NumPadEight"));
+		ADDKEYMAP(SDLK_KP_9, TEXT("NumPadNine"));
+		ADDKEYMAP(SDLK_KP_MULTIPLY, TEXT("Multiply"));
+		ADDKEYMAP(SDLK_KP_PLUS, TEXT("Add"));
+		ADDKEYMAP(SDLK_KP_MINUS, TEXT("Subtract"));
+		ADDKEYMAP(SDLK_KP_DECIMAL, TEXT("Decimal"));
+		ADDKEYMAP(SDLK_KP_DIVIDE, TEXT("Divide"));
+
+        ADDKEYMAP(SDLK_CAPSLOCK, TEXT("CapsLock"));
+		ADDKEYMAP(SDLK_NUMLOCKCLEAR, TEXT("NumLock"));
+		ADDKEYMAP(SDLK_SCROLLLOCK, TEXT("ScrollLock"));
 	}
 
 	check(NumMappings < MaxMappings);
diff --git a/Engine/Source/Runtime/Core/Private/Mac/MacPlatformMisc.cpp b/Engine/Source/Runtime/Core/Private/Mac/MacPlatformMisc.cpp
--- a/Engine/Source/Runtime/Core/Private/Mac/MacPlatformMisc.cpp
+++ b/Engine/Source/Runtime/Core/Private/Mac/MacPlatformMisc.cpp
@@ -600,12 +600,12 @@ void FMacPlatformMisc::PumpMessages( bool bFromMainLoop )
 	}
 }
 
-uint32 FMacPlatformMisc::GetCharKeyMap(uint16* KeyCodes, FString* KeyNames, uint32 MaxMappings)
+uint32 FMacPlatformMisc::GetCharKeyMap(uint32* KeyCodes, FString* KeyNames, uint32 MaxMappings)
 {
 	return FGenericPlatformMisc::GetStandardPrintableKeyMap(KeyCodes, KeyNames, MaxMappings, false, true);
 }
 
-uint32 FMacPlatformMisc::GetKeyMap( uint16* KeyCodes, FString* KeyNames, uint32 MaxMappings )
+uint32 FMacPlatformMisc::GetKeyMap( uint32* KeyCodes, FString* KeyNames, uint32 MaxMappings )
 {
 #define ADDKEYMAP(KeyCode, KeyName)		if (NumMappings<MaxMappings) { KeyCodes[NumMappings]=KeyCode; KeyNames[NumMappings]=KeyName; ++NumMappings; };
 
diff --git a/Engine/Source/Runtime/Core/Private/WinRT/WinRTMisc.cpp b/Engine/Source/Runtime/Core/Private/WinRT/WinRTMisc.cpp
--- a/Engine/Source/Runtime/Core/Private/WinRT/WinRTMisc.cpp
+++ b/Engine/Source/Runtime/Core/Private/WinRT/WinRTMisc.cpp
@@ -127,12 +127,12 @@ void FWinRTMisc::PumpMessages(bool bFromMainLoop)
 	FApp::SetVolumeMultiplier(HasFocus ? 1.0f : FApp::GetUnfocusedVolumeMultiplier());
 }
 
-uint32 FWinRTMisc::GetCharKeyMap(uint16* KeyCodes, FString* KeyNames, uint32 MaxMappings)
+uint32 FWinRTMisc::GetCharKeyMap(uint32* KeyCodes, FString* KeyNames, uint32 MaxMappings)
 {
 	return FGenericPlatformMisc::GetStandardPrintableKeyMap(KeyCodes, KeyNames, MaxMappings, true, false);
 }
 
-uint32 FWinRTMisc::GetKeyMap(uint16* KeyCodes, FString* KeyNames, uint32 MaxMappings)
+uint32 FWinRTMisc::GetKeyMap(uint32* KeyCodes, FString* KeyNames, uint32 MaxMappings)
 {
 #define ADDKEYMAP(KeyCode, KeyName)		if (NumMappings<MaxMappings) { KeyCodes[NumMappings]=KeyCode; KeyNames[NumMappings]=KeyName; ++NumMappings; };
 
diff --git a/Engine/Source/Runtime/Core/Private/Windows/WindowsPlatformMisc.cpp b/Engine/Source/Runtime/Core/Private/Windows/WindowsPlatformMisc.cpp
--- a/Engine/Source/Runtime/Core/Private/Windows/WindowsPlatformMisc.cpp
+++ b/Engine/Source/Runtime/Core/Private/Windows/WindowsPlatformMisc.cpp
@@ -823,12 +823,12 @@ void FWindowsPlatformMisc::PumpMessages(bool bFromMainLoop)
 	FApp::SetVolumeMultiplier( HasFocus ? 1.0f : FApp::GetUnfocusedVolumeMultiplier() );
 }
 
-uint32 FWindowsPlatformMisc::GetCharKeyMap(uint16* KeyCodes, FString* KeyNames, uint32 MaxMappings)
+uint32 FWindowsPlatformMisc::GetCharKeyMap(uint32* KeyCodes, FString* KeyNames, uint32 MaxMappings)
 {
 	return FGenericPlatformMisc::GetStandardPrintableKeyMap(KeyCodes, KeyNames, MaxMappings, true, false);
 }
 
-uint32 FWindowsPlatformMisc::GetKeyMap( uint16* KeyCodes, FString* KeyNames, uint32 MaxMappings )
+uint32 FWindowsPlatformMisc::GetKeyMap( uint32* KeyCodes, FString* KeyNames, uint32 MaxMappings )
 {
 #define ADDKEYMAP(KeyCode, KeyName)		if (NumMappings<MaxMappings) { KeyCodes[NumMappings]=KeyCode; KeyNames[NumMappings]=KeyName; ++NumMappings; };
 
@@ -907,8 +907,8 @@ uint32 FWindowsPlatformMisc::GetKeyMap( uint16* KeyCodes, FString* KeyNames, uin
 		ADDKEYMAP( VK_LWIN, TEXT("LeftCommand") );
 		ADDKEYMAP( VK_RWIN, TEXT("RightCommand") );
 
-		TMap<uint16, uint16> ScanToVKMap;
-#define MAP_OEM_VK_TO_SCAN(KeyCode) { uint16 CharCode = MapVirtualKey(KeyCode,2); if (CharCode != 0) { ScanToVKMap.Add(CharCode,KeyCode); } }
+		TMap<uint32, uint32> ScanToVKMap;
+#define MAP_OEM_VK_TO_SCAN(KeyCode) { const uint32 CharCode = MapVirtualKey(KeyCode,2); if (CharCode != 0) { ScanToVKMap.Add(CharCode,KeyCode); } }
 		MAP_OEM_VK_TO_SCAN(VK_OEM_1);
 		MAP_OEM_VK_TO_SCAN(VK_OEM_2);
 		MAP_OEM_VK_TO_SCAN(VK_OEM_3);
@@ -925,7 +925,7 @@ uint32 FWindowsPlatformMisc::GetKeyMap( uint16* KeyCodes, FString* KeyNames, uin
 #undef  MAP_OEM_VK_TO_SCAN
 
 		static const uint32 MAX_KEY_MAPPINGS(256);
-		uint16 CharCodes[MAX_KEY_MAPPINGS];
+		uint32 CharCodes[MAX_KEY_MAPPINGS];
 		FString CharKeyNames[MAX_KEY_MAPPINGS];
 		const int32 CharMappings = GetCharKeyMap(CharCodes, CharKeyNames, MAX_KEY_MAPPINGS);
 
diff --git a/Engine/Source/Runtime/Core/Public/Android/AndroidMisc.h b/Engine/Source/Runtime/Core/Public/Android/AndroidMisc.h
--- a/Engine/Source/Runtime/Core/Public/Android/AndroidMisc.h
+++ b/Engine/Source/Runtime/Core/Public/Android/AndroidMisc.h
@@ -40,8 +40,8 @@ struct CORE_API FAndroidMisc : public FGenericPlatformMisc
 	// NOTE: THIS FUNCTION IS DEFINED IN ANDROIDOPENGL.CPP
 	static void GetValidTargetPlatforms(class TArray<class FString>& TargetPlatformNames);
 	static bool GetUseVirtualJoysticks();
-	static uint32 GetCharKeyMap(uint16* KeyCodes, FString* KeyNames, uint32 MaxMappings);
-	static uint32 GetKeyMap( uint16* KeyCodes, FString* KeyNames, uint32 MaxMappings );
+	static uint32 GetCharKeyMap(uint32* KeyCodes, FString* KeyNames, uint32 MaxMappings);
+	static uint32 GetKeyMap( uint32* KeyCodes, FString* KeyNames, uint32 MaxMappings );
 	static const TCHAR* GetDefaultDeviceProfileName() { return TEXT("Android_Default"); }
 	static bool GetVolumeButtonsHandledBySystem();
 	static void SetVolumeButtonsHandledBySystem(bool enabled);
diff --git a/Engine/Source/Runtime/Core/Public/GenericPlatform/GenericPlatformMisc.h b/Engine/Source/Runtime/Core/Public/GenericPlatform/GenericPlatformMisc.h
--- a/Engine/Source/Runtime/Core/Public/GenericPlatform/GenericPlatformMisc.h
+++ b/Engine/Source/Runtime/Core/Public/GenericPlatform/GenericPlatformMisc.h
@@ -356,12 +356,12 @@ struct CORE_API FGenericPlatformMisc
 	{
 	}
 
-	FORCEINLINE static uint32 GetKeyMap( uint16* KeyCodes, FString* KeyNames, uint32 MaxMappings )
+	FORCEINLINE static uint32 GetKeyMap( uint32* KeyCodes, FString* KeyNames, uint32 MaxMappings )
 	{
 		return 0;
 	}
 
-	FORCEINLINE static uint32 GetCharKeyMap(uint16* KeyCodes, FString* KeyNames, uint32 MaxMappings)
+	FORCEINLINE static uint32 GetCharKeyMap(uint32* KeyCodes, FString* KeyNames, uint32 MaxMappings)
 	{
 		return 0;
 	}
@@ -382,7 +382,7 @@ protected:
 	* @param bMapUppercaseKeys If true, will map A, B, C, etc to EKeys::A, EKeys::B, EKeys::C
 	* @param bMapLowercaseKeys If true, will map a, b, c, etc to EKeys::A, EKeys::B, EKeys::C
 	*/
-	static uint32 GetStandardPrintableKeyMap(uint16* KeyCodes, FString* KeyNames, uint32 MaxMappings, bool bMapUppercaseKeys, bool bMapLowercaseKeys);
+	static uint32 GetStandardPrintableKeyMap(uint32* KeyCodes, FString* KeyNames, uint32 MaxMappings, bool bMapUppercaseKeys, bool bMapLowercaseKeys);
 
 public:
 
diff --git a/Engine/Source/Runtime/Core/Public/HTML5/HTML5PlatformMisc.h b/Engine/Source/Runtime/Core/Public/HTML5/HTML5PlatformMisc.h
--- a/Engine/Source/Runtime/Core/Public/HTML5/HTML5PlatformMisc.h
+++ b/Engine/Source/Runtime/Core/Public/HTML5/HTML5PlatformMisc.h
@@ -20,8 +20,8 @@ struct CORE_API FHTML5Misc : public FGenericPlatformMisc
 	static void PlatformInit();
 	static void PlatformPostInit(bool ShowSplashScreen = false);
 	static class GenericApplication* CreateApplication();
-	static uint32 GetKeyMap(uint16* KeyCodes, FString* KeyNames, uint32 MaxMappings);
-	static uint32 GetCharKeyMap(uint16* KeyCodes, FString* KeyNames, uint32 MaxMappings);
+	static uint32 GetKeyMap(uint32* KeyCodes, FString* KeyNames, uint32 MaxMappings);
+	static uint32 GetCharKeyMap(uint32* KeyCodes, FString* KeyNames, uint32 MaxMappings);
 	static FString GetDefaultLocale();
 	static void SetCrashHandler(void (* CrashHandler)(const FGenericCrashContext& Context));
 	static EAppReturnType::Type MessageBoxExt( EAppMsgType::Type MsgType, const TCHAR* Text, const TCHAR* Caption );
diff --git a/Engine/Source/Runtime/Core/Public/Linux/LinuxPlatformMisc.h b/Engine/Source/Runtime/Core/Public/Linux/LinuxPlatformMisc.h
--- a/Engine/Source/Runtime/Core/Public/Linux/LinuxPlatformMisc.h
+++ b/Engine/Source/Runtime/Core/Public/Linux/LinuxPlatformMisc.h
@@ -61,8 +61,8 @@ struct CORE_API FLinuxPlatformMisc : public FGenericPlatformMisc
 	}
 
 	static void PumpMessages(bool bFromMainLoop);
-	static uint32 GetKeyMap( uint16* KeyCodes, FString* KeyNames, uint32 MaxMappings );
-	static uint32 GetCharKeyMap(uint16* KeyCodes, FString* KeyNames, uint32 MaxMappings);
+	static uint32 GetKeyMap( uint32* KeyCodes, FString* KeyNames, uint32 MaxMappings );
+	static uint32 GetCharKeyMap(uint32* KeyCodes, FString* KeyNames, uint32 MaxMappings);
 	static void LowLevelOutputDebugString(const TCHAR *Message);
 	static bool ControlScreensaver(EScreenSaverAction Action);
 
diff --git a/Engine/Source/Runtime/Core/Public/Mac/MacPlatformMisc.h b/Engine/Source/Runtime/Core/Public/Mac/MacPlatformMisc.h
--- a/Engine/Source/Runtime/Core/Public/Mac/MacPlatformMisc.h
+++ b/Engine/Source/Runtime/Core/Public/Mac/MacPlatformMisc.h
@@ -83,8 +83,8 @@ struct CORE_API FMacPlatformMisc : public FGenericPlatformMisc
 	}
 
 	static void PumpMessages(bool bFromMainLoop);
-	static uint32 GetKeyMap( uint16* KeyCodes, FString* KeyNames, uint32 MaxMappings );
-	static uint32 GetCharKeyMap(uint16* KeyCodes, FString* KeyNames, uint32 MaxMappings);
+	static uint32 GetKeyMap( uint32* KeyCodes, FString* KeyNames, uint32 MaxMappings );
+	static uint32 GetCharKeyMap(uint32* KeyCodes, FString* KeyNames, uint32 MaxMappings);
 	static void RequestExit(bool Force);
 	static void RequestMinimize();
 	static const TCHAR* GetSystemErrorMessage(TCHAR* OutBuffer, int32 BufferCount, int32 Error);
diff --git a/Engine/Source/Runtime/Core/Public/WinRT/WinRTMisc.h b/Engine/Source/Runtime/Core/Public/WinRT/WinRTMisc.h
--- a/Engine/Source/Runtime/Core/Public/WinRT/WinRTMisc.h
+++ b/Engine/Source/Runtime/Core/Public/WinRT/WinRTMisc.h
@@ -38,8 +38,8 @@ struct CORE_API FWinRTMisc : public FGenericPlatformMisc
 #endif
 
 	static void PumpMessages(bool bFromMainLoop);
-	static uint32 GetKeyMap(uint16* KeyCodes, FString* KeyNames, uint32 MaxMappings);
-	static uint32 GetCharKeyMap(uint16* KeyCodes, FString* KeyNames, uint32 MaxMappings);
+	static uint32 GetKeyMap(uint32* KeyCodes, FString* KeyNames, uint32 MaxMappings);
+	static uint32 GetCharKeyMap(uint32* KeyCodes, FString* KeyNames, uint32 MaxMappings);
 	static void LowLevelOutputDebugString(const TCHAR *Message);
 	static void RequestExit(bool Force);
 	static const TCHAR* GetSystemErrorMessage(TCHAR* OutBuffer, int32 BufferCount, int32 Error);
diff --git a/Engine/Source/Runtime/Core/Public/Windows/WindowsPlatformMisc.h b/Engine/Source/Runtime/Core/Public/Windows/WindowsPlatformMisc.h
--- a/Engine/Source/Runtime/Core/Public/Windows/WindowsPlatformMisc.h
+++ b/Engine/Source/Runtime/Core/Public/Windows/WindowsPlatformMisc.h
@@ -80,8 +80,8 @@ struct CORE_API FWindowsPlatformMisc
 	}
 
 	static void PumpMessages(bool bFromMainLoop);
-	static uint32 GetKeyMap( uint16* KeyCodes, FString* KeyNames, uint32 MaxMappings );
-	static uint32 GetCharKeyMap(uint16* KeyCodes, FString* KeyNames, uint32 MaxMappings);
+	static uint32 GetKeyMap( uint32* KeyCodes, FString* KeyNames, uint32 MaxMappings );
+	static uint32 GetCharKeyMap(uint32* KeyCodes, FString* KeyNames, uint32 MaxMappings);
 	static void SetUTF8Output();
 	static void LocalPrint(const TCHAR *Message);
 	static void RequestExit(bool Force);
diff --git a/Engine/Source/Runtime/InputCore/Classes/InputCoreTypes.h b/Engine/Source/Runtime/InputCore/Classes/InputCoreTypes.h
--- a/Engine/Source/Runtime/InputCore/Classes/InputCoreTypes.h
+++ b/Engine/Source/Runtime/InputCore/Classes/InputCoreTypes.h
@@ -504,13 +504,13 @@ struct INPUTCORE_API FInputKeyManager
 public:
 	static FInputKeyManager& Get();
 
-	void GetCodesFromKey(const FKey Key, const uint16*& KeyCode, const uint16*& CharCode) const;
+	void GetCodesFromKey(const FKey Key, const uint32*& KeyCode, const uint32*& CharCode) const;
 
 	/**
 	 * Retrieves the key mapped to the specified character code.
 	 * @param KeyCode	The key code to get the name for.
 	 */
-	FKey GetKeyFromCodes( const uint16 KeyCode, const uint16 CharCode ) const;
+	FKey GetKeyFromCodes( const uint32 KeyCode, const uint32 CharCode ) const;
 	void InitKeyMappings();
 private:
 	FInputKeyManager()
@@ -519,8 +519,8 @@ private:
 	}
 
 	static TSharedPtr< FInputKeyManager > Instance;
-	TMap<uint16, FKey> KeyMapVirtualToEnum;
-	TMap<uint16, FKey> KeyMapCharToEnum;
+	TMap<uint32, FKey> KeyMapVirtualToEnum;
+	TMap<uint32, FKey> KeyMapCharToEnum;
 };
 
 UCLASS(abstract)
diff --git a/Engine/Source/Runtime/InputCore/Private/InputCoreTypes.cpp b/Engine/Source/Runtime/InputCore/Private/InputCoreTypes.cpp
--- a/Engine/Source/Runtime/InputCore/Private/InputCoreTypes.cpp
+++ b/Engine/Source/Runtime/InputCore/Private/InputCoreTypes.cpp
@@ -1058,7 +1058,7 @@ FInputKeyManager& FInputKeyManager::Get()
 void FInputKeyManager::InitKeyMappings()
 {
 	static const uint32 MAX_KEY_MAPPINGS(256);
-	uint16 KeyCodes[MAX_KEY_MAPPINGS], CharCodes[MAX_KEY_MAPPINGS];
+	uint32 KeyCodes[MAX_KEY_MAPPINGS], CharCodes[MAX_KEY_MAPPINGS];
 	FString KeyNames[MAX_KEY_MAPPINGS], CharKeyNames[MAX_KEY_MAPPINGS];
 
 	uint32 const CharKeyMapSize(FPlatformMisc::GetCharKeyMap(CharCodes, CharKeyNames, MAX_KEY_MAPPINGS));
@@ -1092,7 +1092,7 @@ void FInputKeyManager::InitKeyMappings()
  *
  * @param	KeyCode	the key code to get the name for
  */
-FKey FInputKeyManager::GetKeyFromCodes( const uint16 KeyCode, const uint16 CharCode ) const
+FKey FInputKeyManager::GetKeyFromCodes( const uint32 KeyCode, const uint32 CharCode ) const
 {
 	const FKey* KeyPtr(KeyMapVirtualToEnum.Find(KeyCode));
 	if (KeyPtr == NULL)
@@ -1102,7 +1102,7 @@ FKey FInputKeyManager::GetKeyFromCodes( const uint16 KeyCode, const uint16 CharC
 	return KeyPtr ? *KeyPtr : EKeys::Invalid;
 }
 
-void FInputKeyManager::GetCodesFromKey(const FKey Key, const uint16*& KeyCode, const uint16*& CharCode) const
+void FInputKeyManager::GetCodesFromKey(const FKey Key, const uint32*& KeyCode, const uint32*& CharCode) const
 {
 	KeyCode = KeyMapCharToEnum.FindKey(Key);
 	CharCode = KeyMapVirtualToEnum.FindKey(Key);
