Added new stats for building audio / texture data.

diff --git a/Engine/Source/Developer/DerivedDataCache/Private/DDCStatsHelper.cpp b/Engine/Source/Developer/DerivedDataCache/Private/DDCStatsHelper.cpp
deleted file mode 100644
diff --git a/Engine/Source/Developer/DerivedDataCache/Private/DDCStatsHelper.h b/Engine/Source/Developer/DerivedDataCache/Private/DDCStatsHelper.h
deleted file mode 100644
diff --git a/Engine/Source/Developer/DerivedDataCache/Private/FileSystemDerivedDataBackend.cpp b/Engine/Source/Developer/DerivedDataCache/Private/FileSystemDerivedDataBackend.cpp
--- a/Engine/Source/Developer/DerivedDataCache/Private/FileSystemDerivedDataBackend.cpp
+++ b/Engine/Source/Developer/DerivedDataCache/Private/FileSystemDerivedDataBackend.cpp
@@ -218,6 +218,12 @@ public:
 	 */
 	virtual void PutCachedData(const TCHAR* CacheKey, TArray<uint8>& Data, bool bPutEvenIfExists, FCacheStatRecord* Stats) override
 	{
+		static FName NAME_PutCachedData(TEXT("PutCachedData"));
+		FDDCScopeStatHelper Stat(CacheKey, NAME_PutCachedData);
+		static FName NAME_FileDDCPath(TEXT("FileDDCPath"));
+		Stat.AddTag(NAME_FileDDCPath, CachePath);
+
+
 		check(!bFailed);
 		if (!bReadOnly)
 		{
diff --git a/Engine/Source/Developer/DerivedDataCache/Private/HierarchicalDerivedDataBackend.h b/Engine/Source/Developer/DerivedDataCache/Private/HierarchicalDerivedDataBackend.h
--- a/Engine/Source/Developer/DerivedDataCache/Private/HierarchicalDerivedDataBackend.h
+++ b/Engine/Source/Developer/DerivedDataCache/Private/HierarchicalDerivedDataBackend.h
@@ -156,6 +156,15 @@ public:
 	 */
 	virtual void PutCachedData(const TCHAR* CacheKey, TArray<uint8>& InData, bool bPutEvenIfExists, FCacheStatRecord* Stats) override
 	{
+		const static FName NAME_PutCachedData = FName(TEXT("PutCachedData"));
+		const static FName NAME_HierarchicalDDC = FName(TEXT("HierarchicalDDC"));
+		const static FName NAME_DataSize = FName(TEXT("DataSize"));
+
+		FDDCScopeStatHelper Stat(CacheKey, NAME_PutCachedData);
+		Stat.AddTag(NAME_HierarchicalDDC, FString());
+		Stat.AddTag(NAME_DataSize, FString::Printf(TEXT("%d bytes"), InData.Num()));
+
+
 		if (!bIsWritable)
 		{
 			return; // no point in continuing down the chain
diff --git a/Engine/Source/Editor/CookingStats/CookingStats.Build.cs b/Engine/Source/Editor/CookingStats/CookingStats.Build.cs
--- a/Engine/Source/Editor/CookingStats/CookingStats.Build.cs
+++ b/Engine/Source/Editor/CookingStats/CookingStats.Build.cs
@@ -8,15 +8,8 @@ public class CookingStats : ModuleRules
 	{
 		PrivateDependencyModuleNames.AddRange(
 			new string[] {
-				"Core",
-				"CoreUObject"
+				"Core"
 			}
 			);
-
-		if (UEBuildConfiguration.bBuildEditor == true)
-		{
-			PrivateIncludePathModuleNames.AddRange(new string[] { "DirectoryWatcher" });
-			DynamicallyLoadedModuleNames.AddRange(new string[] { "DirectoryWatcher" });
-		}
 	}
 }
diff --git a/Engine/Source/Editor/CookingStats/Private/CookingStatsPCH.h b/Engine/Source/Editor/CookingStats/Private/CookingStatsPCH.h
--- a/Engine/Source/Editor/CookingStats/Private/CookingStatsPCH.h
+++ b/Engine/Source/Editor/CookingStats/Private/CookingStatsPCH.h
@@ -2,7 +2,6 @@
 #define _COOKING_STATS_PRIVATE_
 
 #include "Core.h"
-#include "CoreUObject.h"
 
 DECLARE_LOG_CATEGORY_EXTERN(LogCookingStats, Log, All);
 
diff --git a/Engine/Source/Runtime/Core/Core.Build.cs b/Engine/Source/Runtime/Core/Core.Build.cs
--- a/Engine/Source/Runtime/Core/Core.Build.cs
+++ b/Engine/Source/Runtime/Core/Core.Build.cs
@@ -51,6 +51,7 @@ public class Core : ModuleRules
 			new string[] {
 				"TargetPlatform",
 				"DerivedDataCache",
+                "CookingStats",
 				"InputDevice",
                 "Analytics"
 			}
diff --git a/Engine/Source/Runtime/Core/Private/ProfilingDebugging/DDCStatsHelper.cpp b/Engine/Source/Runtime/Core/Private/ProfilingDebugging/DDCStatsHelper.cpp
new file mode 100644
--- /dev/null
+++ b/Engine/Source/Runtime/Core/Private/ProfilingDebugging/DDCStatsHelper.cpp
@@ -0,0 +1,137 @@
+#include "CorePrivatePCH.h"
+
+#include "DDCStatsHelper.h"
+#include "CookingStatsModule.h"
+#include "ModuleManager.h"
+
+namespace DDCStats
+{
+
+	ICookingStats* GetCookingStats()
+	{
+		static ICookingStats* CookingStats = nullptr;
+		static bool bInitialized = false;
+		if (bInitialized == false)
+		{
+			FCookingStatsModule* CookingStatsModule = FModuleManager::LoadModulePtr<FCookingStatsModule>(TEXT("CookingStats"));
+			if (CookingStatsModule)
+			{
+				CookingStats = &CookingStatsModule->Get();
+			}
+			bInitialized = true;
+		}
+		return CookingStats;
+	}
+
+	static uint32 GetValidTLSSlot()
+	{
+		uint32 TLSSlot = FPlatformTLS::AllocTlsSlot();
+		check(FPlatformTLS::IsValidTlsSlot(TLSSlot));
+		FPlatformTLS::SetTlsValue(TLSSlot, nullptr);
+		return TLSSlot;
+	}
+	
+
+	struct FDDCStatsTLSStore
+	{
+
+		FDDCStatsTLSStore() : TransactionGuid(NAME_None), CurrentIndex(0), RootScope(NULL) { }
+
+
+		FName GenerateNewTransactionName()
+		{
+			++CurrentIndex;
+			if (CurrentIndex <= 1)
+			{
+				FString TransactionGuidString = TEXT("DDCTransactionId");
+				TransactionGuidString += FGuid::NewGuid().ToString();
+				TransactionGuid = FName(*TransactionGuidString);
+				CurrentIndex = 1;
+			}
+			check(TransactionGuid != NAME_None);
+
+			return FName(TransactionGuid, CurrentIndex);
+		}
+
+		FName TransactionGuid;
+		int32 CurrentIndex;
+		FDDCScopeStatHelper *RootScope;
+	};
+
+	uint32 CookStatsFDDCStatsTLSStore = GetValidTLSSlot();
+
+}; // namespace DDCStats
+
+
+FDDCScopeStatHelper::FDDCScopeStatHelper(const TCHAR* CacheKey, const FName& FunctionName) : bHasParent(false)
+{
+
+	if (DDCStats::GetCookingStats() == nullptr)
+		return;
+
+	DDCStats::FDDCStatsTLSStore* TLSStore = (DDCStats::FDDCStatsTLSStore*)FPlatformTLS::GetTlsValue(DDCStats::CookStatsFDDCStatsTLSStore);
+	if (TLSStore == nullptr)
+	{
+		TLSStore = new DDCStats::FDDCStatsTLSStore();
+		FPlatformTLS::SetTlsValue(DDCStats::CookStatsFDDCStatsTLSStore, TLSStore);
+	}
+
+	TransactionGuid = TLSStore->GenerateNewTransactionName();
+
+
+	FString Temp;
+	AddTag(FunctionName, Temp);
+
+	if (TLSStore->RootScope == nullptr)
+	{
+		TLSStore->RootScope = this;
+		bHasParent = false;
+
+		static const FName NAME_CacheKey = FName(TEXT("CacheKey"));
+		AddTag(NAME_CacheKey, FString(CacheKey));
+	}
+	else
+	{
+		static const FName NAME_Parent = FName(TEXT("Parent"));
+		AddTag(NAME_Parent, TLSStore->RootScope->TransactionGuid.ToString());
+
+		bHasParent = true;
+	}
+
+
+	static const FName NAME_StartTime(TEXT("StartTime"));
+	AddTag(NAME_StartTime, FDateTime::Now().ToString());
+
+	StartTime = FPlatformTime::Seconds();
+
+	// if there is a root scope then we want to put out stuff under a child key of the parent and link to the parent scope
+}
+
+
+FDDCScopeStatHelper::~FDDCScopeStatHelper()
+{
+	if (DDCStats::GetCookingStats() == nullptr)
+		return;
+
+	double EndTime = FPlatformTime::Seconds();
+	float DurationMS = (EndTime - StartTime) * 1000.0f;
+
+	static const FName NAME_Duration = FName(TEXT("Duration"));
+	AddTag(NAME_Duration, FString::Printf(TEXT("%fms"), DurationMS));
+
+	DDCStats::FDDCStatsTLSStore* TLSStore = (DDCStats::FDDCStatsTLSStore*)FPlatformTLS::GetTlsValue(DDCStats::CookStatsFDDCStatsTLSStore);
+	check(TLSStore);
+
+	if (TLSStore->RootScope == this)
+	{
+		TLSStore->RootScope = nullptr;
+	}
+}
+
+void FDDCScopeStatHelper::AddTag(const FName& Tag, const FString& Value)
+{
+	if (DDCStats::GetCookingStats())
+	{
+		DDCStats::GetCookingStats()->AddTagValue(TransactionGuid, Tag, Value);
+	}
+}
\ No newline at end of file
diff --git a/Engine/Source/Runtime/Core/Public/ProfilingDebugging/DDCStatsHelper.h b/Engine/Source/Runtime/Core/Public/ProfilingDebugging/DDCStatsHelper.h
new file mode 100644
--- /dev/null
+++ b/Engine/Source/Runtime/Core/Public/ProfilingDebugging/DDCStatsHelper.h
@@ -0,0 +1,19 @@
+#pragma once
+
+
+class FDDCScopeStatHelper
+{
+private:
+	FName TransactionGuid;
+	double StartTime;
+	bool bHasParent;
+public:
+	CORE_API FDDCScopeStatHelper(const TCHAR* CacheKey, const FName& FunctionName);
+
+	CORE_API ~FDDCScopeStatHelper();
+
+	CORE_API void AddTag(const FName& Tag, const FString& Value);
+
+	CORE_API bool HasParent() const { return bHasParent; }
+};
+
diff --git a/Engine/Source/Runtime/Engine/Private/AudioDerivedData.cpp b/Engine/Source/Runtime/Engine/Private/AudioDerivedData.cpp
--- a/Engine/Source/Runtime/Engine/Private/AudioDerivedData.cpp
+++ b/Engine/Source/Runtime/Engine/Private/AudioDerivedData.cpp
@@ -11,6 +11,14 @@ DEFINE_LOG_CATEGORY_STATIC(LogAudioDerivedData, Log, All);
 	Derived data key generation.
 ------------------------------------------------------------------------------*/
 
+
+// enable audio building stats these appear in the Saved\Stats\stats.csv file
+#define BUILD_AUDIO_STATS 1
+
+#if BUILD_AUDIO_STATS
+#include "DDCStatsHelper.h"
+#endif
+
 #if WITH_EDITORONLY_DATA
 
 // If you want to bump this version, generate a new guid using
@@ -227,6 +235,14 @@ class FStreamedAudioCacheDerivedDataWorker : public FNonAbandonableTask
 	/** Build the streamed audio. This function is safe to call from any thread. */
 	void BuildStreamedAudio()
 	{
+#if BUILD_AUDIO_STATS
+		FString DDCKey;
+		GetStreamedAudioDerivedDataKeyFromSuffix(KeySuffix, DDCKey);
+		static const FName NAME_BuildStreamedAudio(TEXT("BuildStreamedAudio"));
+		FDDCScopeStatHelper ScopeTimer(*DDCKey, NAME_BuildStreamedAudio);
+#endif
+
+
 		GetStreamedAudioDerivedDataKeySuffix(SoundWave, AudioFormatName, KeySuffix);
 
 		DerivedData->Chunks.Empty();
diff --git a/Engine/Source/Runtime/Engine/Private/TextureDerivedData.cpp b/Engine/Source/Runtime/Engine/Private/TextureDerivedData.cpp
--- a/Engine/Source/Runtime/Engine/Private/TextureDerivedData.cpp
+++ b/Engine/Source/Runtime/Engine/Private/TextureDerivedData.cpp
@@ -37,6 +37,13 @@ enum
 #define TEXTURE_DERIVEDDATA_VER		TEXT("814DCC3DC72143F49509781513CB9855")
 
 
+// output texture building stats to Saved\Stats\Stats.csv
+#define BUILD_TEXTURE_STATS 1
+
+#if BUILD_TEXTURE_STATS
+#include "DDCStatsHelper.h"
+#endif
+
 /*------------------------------------------------------------------------------
 	Timing of derived data operations.
 ------------------------------------------------------------------------------*/
@@ -752,6 +759,14 @@ class FTextureCacheDerivedDataWorker : public FNonAbandonableTask
 	/** Build the texture. This function is safe to call from any thread. */
 	void BuildTexture()
 	{
+#if BUILD_TEXTURE_STATS
+		FString DerivedDataKey;
+		GetTextureDerivedDataKeyFromSuffix(KeySuffix, DerivedDataKey);
+		static const FName NAME_BuildTexture(TEXT("BuildTexture"));
+		FDDCScopeStatHelper ScopeStat(*DerivedDataKey, NAME_BuildTexture);
+#endif
+
+
 		if (SourceMips.Num())
 		{
 			TextureDerivedDataTimings::FScopedMeasurement Timer(TextureDerivedDataTimings::BuildTextureTime);
