Use the MODULE_API_VERSION define directly when checking that DLLs are the right version, so we can set the compatible engine changelist without needing to rebuild every module.

diff --git a/Engine/Source/Runtime/Core/Private/Modules/ModuleManager.cpp b/Engine/Source/Runtime/Core/Private/Modules/ModuleManager.cpp
--- a/Engine/Source/Runtime/Core/Private/Modules/ModuleManager.cpp
+++ b/Engine/Source/Runtime/Core/Private/Modules/ModuleManager.cpp
@@ -6,6 +6,7 @@
 #include "EngineBuildSettings.h"
 #include "UProjectInfo.h"
 #include "ScopeExit.h"
+#include "ModuleVersion.h"
 
 DEFINE_LOG_CATEGORY_STATIC(LogModuleManager, Log, All);
 
@@ -1051,7 +1052,8 @@ const TCHAR *FModuleManager::GetUBTConfiguration()
 bool FModuleManager::CheckModuleCompatibility(const TCHAR* Filename)
 {
 	int32 ModuleApiVersion = FPlatformProcess::GetDllApiVersion(Filename);
-	int32 CompiledInApiVersion = GCompatibleWithEngineVersion.GetChangelist();
+	int32 CompiledInApiVersion = MODULE_API_VERSION;
+
 	if (ModuleApiVersion != CompiledInApiVersion)
 	{
 		UE_LOG(LogModuleManager, Warning, TEXT("Found module file %s (API version %d), but it was incompatible with the current engine API version (%d). This is likely a stale module that must be recompiled."), Filename, ModuleApiVersion, CompiledInApiVersion);
