Fix alignment issues in binned alocator.

diff --git a/Engine/Source/Runtime/Core/Public/HAL/MallocBinned.h b/Engine/Source/Runtime/Core/Public/HAL/MallocBinned.h
--- a/Engine/Source/Runtime/Core/Public/HAL/MallocBinned.h
+++ b/Engine/Source/Runtime/Core/Public/HAL/MallocBinned.h
@@ -548,7 +548,7 @@ private:
 		return Pool;
 	}
 
-	FORCEINLINE FFreeMem* AllocateBlockFromPool(FPoolTable* Table, FPoolInfo* Pool)
+	FORCEINLINE FFreeMem* AllocateBlockFromPool(FPoolTable* Table, FPoolInfo* Pool, uint32 Alignment)
 	{
 		// Pick first available block and unlink it.
 		Pool->Taken++;
@@ -568,7 +568,7 @@ private:
 			}
 		}
 		STAT(UsedPeak = FMath::Max(UsedPeak, UsedCurrent += Table->BlockSize));
-		return Free;
+		return Align(Free, Alignment);
 	}
 
 	/**
@@ -607,6 +607,15 @@ private:
 				Pool->Link( Table->FirstPool );
 			}
 
+			void* BaseAddress = (void*)BasePtr;
+			uint32 BlockSize = Table->BlockSize;
+			PTRINT OffsetFromBase = (PTRINT)Ptr - (PTRINT)BaseAddress;
+			check(OffsetFromBase >= 0);
+			uint32 AlignOffset = OffsetFromBase % BlockSize;
+
+			// Patch pointer to include previously applied alignment.
+			Ptr = (void*)((PTRINT)Ptr - (PTRINT)AlignOffset);
+
 			// Free a pooled allocation.
 			FFreeMem* Free		= (FFreeMem*)Ptr;
 			Free->NumFreeBlocks	= 1;
@@ -638,7 +647,7 @@ private:
 			STAT(UsedCurrent -= Pool->GetBytes());
 			STAT(OsCurrent -= OsBytes);
 			STAT(WasteCurrent -= OsBytes - Pool->GetBytes());
-			OSFree(Ptr, OsBytes);
+			OSFree((void*)BasePtr, OsBytes);
 		}
 
 		MEM_TIME(MemTime += FPlatformTime::Seconds());
@@ -939,9 +948,10 @@ public:
 		{
 			Alignment = DEFAULT_BINNED_ALLOCATOR_ALIGNMENT;
 		}
-		check(Alignment <= PageSize);
+
 		Alignment = FMath::Max<uint32>(Alignment, DEFAULT_BINNED_ALLOCATOR_ALIGNMENT);
-		Size = FMath::Max<SIZE_T>(Alignment, Align(Size, Alignment));
+		SIZE_T SpareBytesCount = FMath::Min<SIZE_T>(DEFAULT_BINNED_ALLOCATOR_ALIGNMENT, Size);
+		Size = FMath::Max<SIZE_T>(PoolTable[0].BlockSize, Size + (Alignment - SpareBytesCount));
 		MEM_TIME(MemTime -= FPlatformTime::Seconds());
 		STAT(CurrentAllocs++);
 		STAT(TotalAllocs++);
@@ -963,11 +973,10 @@ public:
 				Pool = AllocatePoolMemory(Table, BINNED_ALLOC_POOL_SIZE/*PageSize*/, Size);
 			}
 
-			Free = AllocateBlockFromPool(Table, Pool);
+			Free = AllocateBlockFromPool(Table, Pool, Alignment);
 		}
 		else if ( ((Size >= BinnedSizeLimit && Size <= PagePoolTable[0].BlockSize) ||
-				  (Size > PageSize && Size <= PagePoolTable[1].BlockSize))
-				  && Alignment == DEFAULT_BINNED_ALLOCATOR_ALIGNMENT )
+			(Size > PageSize && Size <= PagePoolTable[1].BlockSize)))
 		{
 			// Bucket in a pool of 3*PageSize or 6*PageSize
 			uint32 BinType = Size < PageSize ? 0 : 1;
@@ -986,7 +995,7 @@ public:
 				Pool = AllocatePoolMemory(Table, PageCount*PageSize, BinnedSizeLimit+BinType);
 			}
 
-			Free = AllocateBlockFromPool(Table, Pool);
+			Free = AllocateBlockFromPool(Table, Pool, Alignment);
 		}
 		else
 		{
@@ -998,7 +1007,8 @@ public:
 			{
 				OutOfMemory(AlignedSize);
 			}
-			checkSlow(!((SIZE_T)Free & (PageSize-1)));
+
+			void* AlignedFree = Align(Free, Alignment);
 
 			// Create indirect.
 			FPoolInfo* Pool;
@@ -1007,8 +1017,20 @@ public:
 				FScopeLock PoolInfoLock(&AccessGuard);
 #endif
 				Pool = GetPoolInfo((UPTRINT)Free);
-			}
 
+				if ((UPTRINT)Free != ((UPTRINT)AlignedFree & ~(PageSize - 1)))
+				{
+					// Mark the FPoolInfo for AlignedFree to jump back to the FPoolInfo for ptr.
+					for (UPTRINT i = (UPTRINT)PageSize, Offset = 0; i < AlignedSize; i += PageSize, ++Offset)
+					{
+						FPoolInfo* TrailingPool = GetPoolInfo(((UPTRINT)Free) + i);
+						check(TrailingPool);
+						//Set trailing pools to point back to first pool
+						TrailingPool->SetAllocationSizes(0, 0, Offset, BinnedOSTableIndex);
+			}
+			}
+			}
+			Free = (FFreeMem*)AlignedFree;
 			Pool->SetAllocationSizes(Size, AlignedSize, BinnedOSTableIndex, BinnedOSTableIndex);
 			STAT(OsPeak = FMath::Max(OsPeak, OsCurrent += AlignedSize));
 			STAT(UsedPeak = FMath::Max(UsedPeak, UsedCurrent += Size));
@@ -1029,11 +1051,13 @@ public:
 		{
 			Alignment = DEFAULT_BINNED_ALLOCATOR_ALIGNMENT;
 		}
-		check(Alignment <= PageSize);
+
 		Alignment = FMath::Max<uint32>(Alignment, DEFAULT_BINNED_ALLOCATOR_ALIGNMENT);
+		const uint32 NewSizeUnmodified = NewSize;
+		SIZE_T SpareBytesCount = FMath::Min<SIZE_T>(DEFAULT_BINNED_ALLOCATOR_ALIGNMENT, NewSize);
 		if (NewSize)
 		{
-			NewSize = FMath::Max<SIZE_T>(Alignment, Align(NewSize, Alignment));
+			NewSize = FMath::Max<SIZE_T>(PoolTable[0].BlockSize, NewSize + (Alignment - SpareBytesCount));
 		}
 		MEM_TIME(MemTime -= FPlatformTime::Seconds());
 		UPTRINT BasePtr;
@@ -1046,22 +1070,26 @@ public:
 			{
 				// Allocated from pool, so grow or shrink if necessary.
 				check(Pool->TableIndex > 0); // it isn't possible to allocate a size of 0, Malloc will increase the size to DEFAULT_BINNED_ALLOCATOR_ALIGNMENT
-				if( NewSize>MemSizeToPoolTable[Pool->TableIndex]->BlockSize || NewSize<=MemSizeToPoolTable[Pool->TableIndex-1]->BlockSize )
+				if (NewSizeUnmodified > MemSizeToPoolTable[Pool->TableIndex]->BlockSize || NewSizeUnmodified <= MemSizeToPoolTable[Pool->TableIndex - 1]->BlockSize)
 				{
-					NewPtr = Malloc( NewSize, Alignment );
-					FMemory::Memcpy( NewPtr, Ptr, FMath::Min<SIZE_T>( NewSize, MemSizeToPoolTable[Pool->TableIndex]->BlockSize ) );
+					NewPtr = Malloc(NewSizeUnmodified, Alignment);
+					FMemory::Memcpy(NewPtr, Ptr, FMath::Min<SIZE_T>(NewSizeUnmodified, MemSizeToPoolTable[Pool->TableIndex]->BlockSize - (Alignment - SpareBytesCount)));
 					Free( Ptr );
 				}
+				else if (((UPTRINT)Ptr & (UPTRINT)(Alignment - 1)) != 0)
+				{
+					void* NewPtr = Align(Ptr, Alignment);
+					FMemory::Memmove(Ptr, NewPtr, NewSize);
+			}
 			}
 			else
 			{
 				// Allocated from OS.
-				checkSlow(!((UPTRINT)Ptr & (PageSize-1)));
 				if( NewSize > Pool->GetOsBytes(PageSize, BinnedOSTableIndex) || NewSize * 3 < Pool->GetOsBytes(PageSize, BinnedOSTableIndex) * 2 )
 				{
 					// Grow or shrink.
-					NewPtr = Malloc( NewSize, Alignment );
-					FMemory::Memcpy( NewPtr, Ptr, FMath::Min<SIZE_T>(NewSize, Pool->GetBytes()) );
+					NewPtr = Malloc(NewSizeUnmodified, Alignment);
+					FMemory::Memcpy(NewPtr, Ptr, FMath::Min<SIZE_T>(NewSizeUnmodified, Pool->GetBytes()));
 					Free( Ptr );
 				}
 				else
@@ -1070,13 +1098,13 @@ public:
 					STAT(UsedCurrent += NewSize - Pool->GetBytes());
 					STAT(UsedPeak = FMath::Max(UsedPeak, UsedCurrent));
 					STAT(WasteCurrent += Pool->GetBytes() - NewSize);
-					Pool->SetAllocationSizes(NewSize, Pool->GetOsBytes(PageSize, BinnedOSTableIndex), BinnedOSTableIndex, BinnedOSTableIndex);
+					Pool->SetAllocationSizes(NewSizeUnmodified, Pool->GetOsBytes(PageSize, BinnedOSTableIndex), BinnedOSTableIndex, BinnedOSTableIndex);
 				}
 			}
 		}
 		else if( Ptr == nullptr )
 		{
-			NewPtr = Malloc( NewSize, Alignment );
+			NewPtr = Malloc(NewSizeUnmodified, Alignment);
 		}
 		else
 		{
