Fix for a string comparison issue of properties when writing a single value to configs.

Spotted whilst testing a change to PhysicsSettings, where bSubstepping and bSubsteppingAsync were both changed due to a check for StartsWith(PropertyName)

diff --git a/Engine/Source/Runtime/Core/Private/Misc/ConfigCacheIni.cpp b/Engine/Source/Runtime/Core/Private/Misc/ConfigCacheIni.cpp
--- a/Engine/Source/Runtime/Core/Private/Misc/ConfigCacheIni.cpp
+++ b/Engine/Source/Runtime/Core/Private/Misc/ConfigCacheIni.cpp
@@ -3211,7 +3211,7 @@ bool FConfigFile::UpdateSinglePropertyInSection(const TCHAR* DiskFilename, const
 
 				NewFile += TheLine + LINE_TERMINATOR;
 			}
-			else if (bIsCurrentlyWritingSectionProperty && TheLine.StartsWith(PropertyName))
+			else if (bIsCurrentlyWritingSectionProperty && TheLine.StartsWith(FString::Printf(TEXT("%s="), PropertyName)))
 			{
 				bPropertyWasUpdated = true;
 				// If we have found the property in the file, let's update it with the new value.
