Fix CIS and a realloc bug.

diff --git a/Engine/Source/Runtime/Core/Public/HAL/MallocBinned.h b/Engine/Source/Runtime/Core/Public/HAL/MallocBinned.h
--- a/Engine/Source/Runtime/Core/Public/HAL/MallocBinned.h
+++ b/Engine/Source/Runtime/Core/Public/HAL/MallocBinned.h
@@ -1078,9 +1078,9 @@ public:
 				}
 				else if (((UPTRINT)Ptr & (UPTRINT)(Alignment - 1)) != 0)
 				{
-					void* NewPtr = Align(Ptr, Alignment);
-					FMemory::Memmove(Ptr, NewPtr, NewSize);
-			}
+					NewPtr = Align(Ptr, Alignment);
+					FMemory::Memmove(NewPtr, Ptr, NewSize);
+				}
 			}
 			else
 			{
