r.DBuffer is no scalability setting, changed that

diff --git a/Engine/Source/Runtime/Core/Private/HAL/ConsoleManager.cpp b/Engine/Source/Runtime/Core/Private/HAL/ConsoleManager.cpp
--- a/Engine/Source/Runtime/Core/Private/HAL/ConsoleManager.cpp
+++ b/Engine/Source/Runtime/Core/Private/HAL/ConsoleManager.cpp
@@ -2047,7 +2047,7 @@ static TAutoConsoleVariable<int32> CVarDBuffer(
 	TEXT("At the moment only can be ensures by full enablng this pass: r.EarlyZPassMovable=1 r.EarlyZPass=2\n")
 	TEXT(" 0: off\n")
 	TEXT(" 1: on (needs early pass rendering on all decal receivers and base pass lookups into the DBuffer, costs GPU memory, allows GBuffer compression)"),
-	ECVF_Scalability | ECVF_RenderThreadSafe);
+	ECVF_RenderThreadSafe | ECVF_ReadOnly);
 
 static TAutoConsoleVariable<float> CVarSkeletalMeshLODRadiusScale(
 	TEXT("r.SkeletalMeshLODRadiusScale"),
