#ue4 - Made FVector::IsUniform() just call AllComponentsEqual(), as the old implementation was similar but allowed 2*Tolerance difference between X and Z.

diff --git a/Engine/Source/Runtime/Core/Public/Math/Vector.h b/Engine/Source/Runtime/Core/Public/Math/Vector.h
--- a/Engine/Source/Runtime/Core/Public/Math/Vector.h
+++ b/Engine/Source/Runtime/Core/Public/Math/Vector.h
@@ -1647,7 +1647,7 @@ FORCEINLINE FVector FVector::Reciprocal() const
 
 FORCEINLINE bool FVector::IsUniform(float Tolerance) const
 {
-	return (FMath::Abs(X-Y) <= Tolerance) && (FMath::Abs(Y-Z) <= Tolerance);
+	return AllComponentsEqual(Tolerance);
 }
 
 FORCEINLINE FVector FVector::MirrorByVector( const FVector& MirrorNormal ) const
