#ue4 - Add FQuat methods GetForwardVector(), GetRightVector(), GetUpVector(). Improve comments for GetAxis[X|Y|Z] and GetRotationAxis().

UE-14015

diff --git a/Engine/Source/Runtime/Core/Public/Math/Quat.h b/Engine/Source/Runtime/Core/Public/Math/Quat.h
--- a/Engine/Source/Runtime/Core/Public/Math/Quat.h
+++ b/Engine/Source/Runtime/Core/Public/Math/Quat.h
@@ -339,19 +339,32 @@ public:
 	 */
 	void EnforceShortestArcWith( const FQuat& OtherQuat );
 	
-	/** Get X Rotation Axis. */
+	/** Get the forward direction (X axis) after it has been rotated by this Quaternion. */
 	FORCEINLINE FVector GetAxisX() const;
 
-	/** Get Y Rotation Axis. */
+	/** Get the right direction (Y axis) after it has been rotated by this Quaternion. */
 	FORCEINLINE FVector GetAxisY() const;
 
-	/** Get Z Rotation Axis. */
+	/** Get the up direction (Z axis) after it has been rotated by this Quaternion. */
 	FORCEINLINE FVector GetAxisZ() const;
 
-	/** @return rotator representation of this quaternion */
+	/** Get the forward direction (X axis) after it has been rotated by this Quaternion. */
+	FORCEINLINE FVector GetForwardVector() const;
+
+	/** Get the right direction (Y axis) after it has been rotated by this Quaternion. */
+	FORCEINLINE FVector GetRightVector() const;
+
+	/** Get the up direction (Z axis) after it has been rotated by this Quaternion. */
+	FORCEINLINE FVector GetUpVector() const;
+
+	/** Get the FRotator representation of this Quaternion. */
 	CORE_API FRotator Rotator() const;
 
-	/** @return Vector of the axis of the quaternion */
+	/**
+	 * Get the axis of rotation of the Quaternion.
+	 * This is the axis around which rotation occurs to transform the canonical coordinate system to the target orientation.
+	 * For the identity Quaternion which has no such rotation, FVector(1,0,0) is returned.
+	 */
 	FORCEINLINE FVector GetRotationAxis() const;
 
 	/**
@@ -946,6 +959,22 @@ FORCEINLINE FVector FQuat::GetAxisZ() const
 }
 
 
+FORCEINLINE FVector FQuat::GetForwardVector() const
+{
+	return GetAxisX();
+}
+
+FORCEINLINE FVector FQuat::GetRightVector() const
+{
+	return GetAxisY();
+}
+
+FORCEINLINE FVector FQuat::GetUpVector() const
+{
+	return GetAxisZ();
+}
+
+
 FORCEINLINE float FQuat::Error( const FQuat& Q1, const FQuat& Q2 )
 {
 	const float cosom = FMath::Abs(Q1.X * Q2.X + Q1.Y * Q2.Y + Q1.Z * Q2.Z + Q1.W * Q2.W);
