UE-18791/Deadlock when enabling memory profiling

diff --git a/Engine/Source/Runtime/Core/Public/HAL/MallocBinned.h b/Engine/Source/Runtime/Core/Public/HAL/MallocBinned.h
--- a/Engine/Source/Runtime/Core/Public/HAL/MallocBinned.h
+++ b/Engine/Source/Runtime/Core/Public/HAL/MallocBinned.h
@@ -1184,7 +1184,7 @@ public:
 	}
 
 
-	/** Called once per frame, gathers and sets all memory allocator statistics into the corresponding stats. */
+	/** Called once per frame, gathers and sets all memory allocator statistics into the corresponding stats. MUST BE THREAD SAFE. */
 	virtual void UpdateStats() override
 	{
 		FMalloc::UpdateStats();
diff --git a/Engine/Source/Runtime/Core/Public/HAL/MallocThreadSafeProxy.h b/Engine/Source/Runtime/Core/Public/HAL/MallocThreadSafeProxy.h
--- a/Engine/Source/Runtime/Core/Public/HAL/MallocThreadSafeProxy.h
+++ b/Engine/Source/Runtime/Core/Public/HAL/MallocThreadSafeProxy.h
@@ -35,8 +35,8 @@ public:
 	 */
 	virtual void* Malloc( SIZE_T Size, uint32 Alignment ) override
 	{
+		IncrementTotalMallocCalls();
 		FScopeLock ScopeLock( &SynchronizationObject );
-		STAT(TotalMallocCalls++);
 		return UsedMalloc->Malloc( Size, Alignment );
 	}
 
@@ -45,8 +45,8 @@ public:
 	 */
 	virtual void* Realloc( void* Ptr, SIZE_T NewSize, uint32 Alignment ) override
 	{
+		IncrementTotalReallocCalls();
 		FScopeLock ScopeLock( &SynchronizationObject );
-		STAT(TotalReallocCalls++);
 		return UsedMalloc->Realloc( Ptr, NewSize, Alignment );
 	}
 
@@ -57,19 +57,12 @@ public:
 	{
 		if( Ptr )
 		{
+			IncrementTotalFreeCalls();
 			FScopeLock ScopeLock( &SynchronizationObject );
-			STAT(TotalFreeCalls++);
 			UsedMalloc->Free( Ptr );
 		}
 	}
 
-	/** Called once per frame, gathers and sets all memory allocator statistics into the corresponding stats. */
-	virtual void UpdateStats() override
-	{
-		FScopeLock ScopeLock( &SynchronizationObject );
-		UsedMalloc->UpdateStats();
-	}
-
 	/** Writes allocator stats from the last update into the specified destination. */
 	virtual void GetAllocatorStats( FGenericMemoryStats& out_Stats ) override
 	{
diff --git a/Engine/Source/Runtime/Core/Public/HAL/MemoryBase.h b/Engine/Source/Runtime/Core/Public/HAL/MemoryBase.h
--- a/Engine/Source/Runtime/Core/Public/HAL/MemoryBase.h
+++ b/Engine/Source/Runtime/Core/Public/HAL/MemoryBase.h
@@ -104,7 +104,7 @@ public:
 		return false; 
 	}
 
-	/** Called once per frame, gathers and sets all memory allocator statistics into the corresponding stats. */
+	/** Called once per frame, gathers and sets all memory allocator statistics into the corresponding stats. MUST BE THREAD SAFE. */
 	virtual void UpdateStats();
 
 	/** Writes allocator stats from the last update into the specified destination. */
diff --git a/Engine/Source/Runtime/Core/Public/ProfilingDebugging/MallocProfiler.h b/Engine/Source/Runtime/Core/Public/ProfilingDebugging/MallocProfiler.h
--- a/Engine/Source/Runtime/Core/Public/ProfilingDebugging/MallocProfiler.h
+++ b/Engine/Source/Runtime/Core/Public/ProfilingDebugging/MallocProfiler.h
@@ -451,10 +451,9 @@ public:
 		return true; 
 	}
 
-	/** Called once per frame, gathers and sets all memory allocator statistics into the corresponding stats. */
+	/** Called once per frame, gathers and sets all memory allocator statistics into the corresponding stats. MUST BE THREAD SAFE. */
 	virtual void UpdateStats() override
 	{
-		FScopeLock Lock( &CriticalSection );
 		UsedMalloc->UpdateStats();
 	}
 
