UE4 - Fix D3D12 dialog

diff --git a/Engine/Source/Runtime/Core/Private/Windows/WindowsPlatformMisc.cpp b/Engine/Source/Runtime/Core/Private/Windows/WindowsPlatformMisc.cpp
--- a/Engine/Source/Runtime/Core/Private/Windows/WindowsPlatformMisc.cpp
+++ b/Engine/Source/Runtime/Core/Private/Windows/WindowsPlatformMisc.cpp
@@ -1789,7 +1789,12 @@ void FWindowsPlatformMisc::LoadPreInitModules()
 	// D3D11 is not supported on WinXP, so in this case we use the OpenGL RHI
 	if(FWindowsPlatformMisc::VerifyWindowsVersion(6, 0))
 	{
-		FModuleManager::Get().LoadModule(TEXT("D3D12RHI"));
+		//#todo-rco: Only try on Win10
+		const bool bForceD3D12 = FParse::Param(FCommandLine::Get(), TEXT("d3d12")) || FParse::Param(FCommandLine::Get(), TEXT("dx12"));
+		if (bForceD3D12)
+		{
+			FModuleManager::Get().LoadModule(TEXT("D3D12RHI"));
+		}
 		FModuleManager::Get().LoadModule(TEXT("D3D11RHI"));
 	}
 	FModuleManager::Get().LoadModule(TEXT("OpenGLDrv"));
