Crash reporter improvements (OR-4990/Automatically copy full dumps to Network Share)

diff --git a/Engine/Source/Developer/CrashDebugHelper/Private/Windows/WindowsPlatformStackWalkExt.cpp b/Engine/Source/Developer/CrashDebugHelper/Private/Windows/WindowsPlatformStackWalkExt.cpp
--- a/Engine/Source/Developer/CrashDebugHelper/Private/Windows/WindowsPlatformStackWalkExt.cpp
+++ b/Engine/Source/Developer/CrashDebugHelper/Private/Windows/WindowsPlatformStackWalkExt.cpp
@@ -1,7 +1,7 @@
 // Copyright 1998-2015 Epic Games, Inc. All Rights Reserved.
 
 #include "CrashDebugHelperPrivatePCH.h"
-#include "GenericPlatform/GenericPlatformContext.h"
+#include "GenericPlatform/GenericPlatformCrashContext.h"
 #include "WindowsPlatformStackWalkExt.h"
 #include "CrashDebugPDBCache.h"
 
diff --git a/Engine/Source/Programs/CrashReporter/CrashReportClient/Private/CrashDescription.cpp b/Engine/Source/Programs/CrashReporter/CrashReportClient/Private/CrashDescription.cpp
--- a/Engine/Source/Programs/CrashReporter/CrashReportClient/Private/CrashDescription.cpp
+++ b/Engine/Source/Programs/CrashReporter/CrashReportClient/Private/CrashDescription.cpp
@@ -1,7 +1,7 @@
 // Copyright 1998-2015 Epic Games, Inc. All Rights Reserved.
 
 #include "CrashReportClientApp.h"
-#include "GenericPlatformContext.h"
+#include "GenericPlatformCrashContext.h"
 
 #include "EngineVersion.h"
 #include "XmlFile.h"
diff --git a/Engine/Source/Programs/CrashReporter/CrashReportClient/Private/CrashDescription.h b/Engine/Source/Programs/CrashReporter/CrashReportClient/Private/CrashDescription.h
--- a/Engine/Source/Programs/CrashReporter/CrashReportClient/Private/CrashDescription.h
+++ b/Engine/Source/Programs/CrashReporter/CrashReportClient/Private/CrashDescription.h
@@ -1,7 +1,7 @@
 // Copyright 1998-2015 Epic Games, Inc. All Rights Reserved.
 
 #pragma once
-#include "GenericPlatformContext.h"
+#include "GenericPlatformCrashContext.h"
 
 /**
  *	Describes a unified crash, should be used by all platforms.
diff --git a/Engine/Source/Runtime/Core/Private/GenericPlatform/GenericPlatformContext.cpp b/Engine/Source/Runtime/Core/Private/GenericPlatform/GenericPlatformContext.cpp
deleted file mode 100644
diff --git a/Engine/Source/Runtime/Core/Private/GenericPlatform/GenericPlatformCrashContext.cpp b/Engine/Source/Runtime/Core/Private/GenericPlatform/GenericPlatformCrashContext.cpp
new file mode 100644
--- /dev/null
+++ b/Engine/Source/Runtime/Core/Private/GenericPlatform/GenericPlatformCrashContext.cpp
@@ -0,0 +1,292 @@
+// Copyright 1998-2015 Epic Games, Inc. All Rights Reserved.
+
+#include "CorePrivatePCH.h"
+#include "GenericPlatform/GenericPlatformCrashContext.h"
+#include "Misc./App.h"
+#include "EngineVersion.h"
+#include "EngineBuildSettings.h"
+
+const ANSICHAR* FGenericCrashContext::CrashContextRuntimeXMLNameA = "CrashContext.runtime-xml";
+const TCHAR* FGenericCrashContext::CrashContextRuntimeXMLNameW = TEXT( "CrashContext.runtime-xml" );
+bool FGenericCrashContext::bIsInitialized = false;
+FPlatformMemoryStats FGenericCrashContext::CrashMemoryStats = FPlatformMemoryStats();
+
+namespace NCachedCrashContextProperties
+{
+	static bool bIsInternalBuild;
+	static bool bIsPerforceBuild;
+	static bool bIsSourceDistribution;
+	static bool bIsUE4Release;
+	static FString ExecutableName;
+	static FString PlatformName;
+	static FString PlatformNameIni;
+	static FString BaseDir;
+	static FString RootDir;
+	static FString EpicAccountId;
+	static FString MachineIdStr;
+	static FString OsVersion;
+	static FString OsSubVersion;
+	static int32 NumberOfCores;
+	static int32 NumberOfCoresIncludingHyperthreads;
+	static FString CPUVendor;
+	static FString CPUBrand;
+	static FString PrimaryGPUBrand;
+	static FString UserName;
+	static FString DefaultLocale;
+	static int32 CrashDumpMode;
+	static int32 SecondsSinceStart;
+	static FString CrashGUID;
+}
+
+void FGenericCrashContext::Initialize()
+{
+	NCachedCrashContextProperties::bIsInternalBuild = FEngineBuildSettings::IsInternalBuild();
+	NCachedCrashContextProperties::bIsPerforceBuild = FEngineBuildSettings::IsPerforceBuild();
+	NCachedCrashContextProperties::bIsSourceDistribution = FEngineBuildSettings::IsSourceDistribution();
+	NCachedCrashContextProperties::bIsUE4Release = FApp::IsEngineInstalled();
+
+	NCachedCrashContextProperties::ExecutableName = FPlatformProcess::ExecutableName();
+	NCachedCrashContextProperties::PlatformName = FPlatformProperties::PlatformName();
+	NCachedCrashContextProperties::PlatformNameIni = FPlatformProperties::IniPlatformName();
+	NCachedCrashContextProperties::BaseDir = FPlatformProcess::BaseDir();
+	NCachedCrashContextProperties::RootDir = FPlatformMisc::RootDir();
+	NCachedCrashContextProperties::EpicAccountId = FPlatformMisc::GetEpicAccountId();
+	NCachedCrashContextProperties::MachineIdStr = FPlatformMisc::GetMachineId().ToString(EGuidFormats::Digits);
+	FPlatformMisc::GetOSVersions(NCachedCrashContextProperties::OsVersion, NCachedCrashContextProperties::OsSubVersion);
+	NCachedCrashContextProperties::NumberOfCores = FPlatformMisc::NumberOfCores();
+	NCachedCrashContextProperties::NumberOfCoresIncludingHyperthreads = FPlatformMisc::NumberOfCoresIncludingHyperthreads();
+
+	NCachedCrashContextProperties::CPUVendor = FPlatformMisc::GetCPUVendor();
+	NCachedCrashContextProperties::CPUBrand = FPlatformMisc::GetCPUBrand();
+	NCachedCrashContextProperties::PrimaryGPUBrand = FPlatformMisc::GetPrimaryGPUBrand();
+	NCachedCrashContextProperties::UserName = FPlatformProcess::UserName();
+	NCachedCrashContextProperties::DefaultLocale = FPlatformMisc::GetDefaultLocale();
+
+	// Using the -fullcrashdump parameter will cause full memory minidumps to be created for crashes
+	NCachedCrashContextProperties::CrashDumpMode = (int32)ECrashDumpMode::Default;
+	if (FCommandLine::IsInitialized())
+	{
+		const TCHAR* CmdLine = FCommandLine::Get();
+		if (FParse::Param( CmdLine, TEXT("fullcrashdump") ))
+		{
+			NCachedCrashContextProperties::CrashDumpMode = (int32)ECrashDumpMode::FullDump;
+		}
+	}
+
+	const FGuid Guid = FGuid::NewGuid();
+	NCachedCrashContextProperties::CrashGUID = FString::Printf(TEXT("UE4CC-%s-%s"), *NCachedCrashContextProperties::PlatformNameIni, *Guid.ToString(EGuidFormats::Digits));
+
+	// Initialize delegate for updating SecondsSinceStart, beacuse FPlatformTime::Seconds() is not POSIX safe.
+	const float PollingInterval = 1.0f;
+	FTicker::GetCoreTicker().AddTicker( FTickerDelegate::CreateLambda( []( float DeltaTime )
+	{
+		NCachedCrashContextProperties::SecondsSinceStart = int32(FPlatformTime::Seconds() - GStartTime);
+		return true;
+	} ), PollingInterval );
+
+	bIsInitialized = true;
+}
+
+FGenericCrashContext::FGenericCrashContext()
+{
+	CommonBuffer.Reserve( 32768 );
+}
+
+void FGenericCrashContext::SerializeContentToBuffer()
+{
+	// Must conform against:
+	// https://www.securecoding.cert.org/confluence/display/seccode/SIG30-C.+Call+only+asynchronous-safe+functions+within+signal+handlers
+	AddHeader();
+
+	BeginSection( TEXT( "RuntimeProperties" ) );
+	AddCrashProperty( TEXT( "Version" ), (int32)ECrashDescVersions::VER_2_AddedNewProperties );
+	AddCrashProperty( TEXT( "ProcessId" ), FPlatformProcess::GetCurrentProcessId() );
+	AddCrashProperty( TEXT( "IsInternalBuild" ), (int32)NCachedCrashContextProperties::bIsInternalBuild );
+	AddCrashProperty( TEXT( "IsPerforceBuild" ), (int32)NCachedCrashContextProperties::bIsPerforceBuild );
+	AddCrashProperty( TEXT( "IsSourceDistribution" ), (int32)NCachedCrashContextProperties::bIsSourceDistribution );
+
+	AddCrashProperty( TEXT( "SecondsSinceStart" ), (int32)NCachedCrashContextProperties::SecondsSinceStart );
+
+	// Add common crash properties.
+	AddCrashProperty( TEXT( "GameName" ), FApp::GetGameName() );
+	AddCrashProperty( TEXT( "ExecutableName" ), *NCachedCrashContextProperties::ExecutableName );
+	AddCrashProperty( TEXT( "BuildConfiguration" ), EBuildConfigurations::ToString( FApp::GetBuildConfiguration() ) );
+
+	AddCrashProperty( TEXT( "PlatformName" ), *NCachedCrashContextProperties::PlatformName );
+	AddCrashProperty( TEXT( "PlatformNameIni" ), *NCachedCrashContextProperties::PlatformNameIni );
+	AddCrashProperty( TEXT( "EngineMode" ), FPlatformMisc::GetEngineMode() );
+	AddCrashProperty( TEXT( "EngineVersion" ), *GEngineVersion.ToString() );
+	AddCrashProperty( TEXT( "CommandLine" ), FCommandLine::IsInitialized() ? FCommandLine::Get() : TEXT( "" ) );
+	AddCrashProperty( TEXT( "LanguageLCID" ), FInternationalization::Get().GetCurrentCulture()->GetLCID() );
+	AddCrashProperty( TEXT( "DefaultLocale" ), *NCachedCrashContextProperties::DefaultLocale );
+
+	AddCrashProperty( TEXT( "IsUE4Release" ), (int32)NCachedCrashContextProperties::bIsUE4Release );
+
+	// Remove periods from user names to match AutoReporter user names
+	// The name prefix is read by CrashRepository.AddNewCrash in the website code
+	const bool bSendUserName = NCachedCrashContextProperties::bIsInternalBuild;
+	AddCrashProperty( TEXT( "UserName" ), bSendUserName ? *NCachedCrashContextProperties::UserName.Replace( TEXT( "." ), TEXT( "" ) ) : TEXT( "" ) );
+
+	AddCrashProperty( TEXT( "BaseDir" ), *NCachedCrashContextProperties::BaseDir );
+	AddCrashProperty( TEXT( "RootDir" ), *NCachedCrashContextProperties::RootDir );
+	AddCrashProperty( TEXT( "MachineId" ), *NCachedCrashContextProperties::MachineIdStr );
+	AddCrashProperty( TEXT( "EpicAccountId" ), *NCachedCrashContextProperties::EpicAccountId );
+
+	AddCrashProperty( TEXT( "CallStack" ), TEXT( "" ) );
+	AddCrashProperty( TEXT( "SourceContext" ), TEXT( "" ) );
+	AddCrashProperty( TEXT( "UserDescription" ), TEXT( "" ) );
+	AddCrashProperty( TEXT( "ErrorMessage" ), (const TCHAR*)GErrorMessage ); // GErrorMessage may be broken.
+	AddCrashProperty( TEXT( "CrashDumpMode" ), NCachedCrashContextProperties::CrashDumpMode );
+
+	// Add misc stats.
+	AddCrashProperty( TEXT( "Misc.NumberOfCores" ), NCachedCrashContextProperties::NumberOfCores );
+	AddCrashProperty( TEXT( "Misc.NumberOfCoresIncludingHyperthreads" ), NCachedCrashContextProperties::NumberOfCoresIncludingHyperthreads );
+	AddCrashProperty( TEXT( "Misc.Is64bitOperatingSystem" ), (int32)FPlatformMisc::Is64bitOperatingSystem() );
+
+	AddCrashProperty( TEXT( "Misc.CPUVendor" ), *NCachedCrashContextProperties::CPUVendor );
+	AddCrashProperty( TEXT( "Misc.CPUBrand" ), *NCachedCrashContextProperties::CPUBrand );
+	AddCrashProperty( TEXT( "Misc.PrimaryGPUBrand" ), *NCachedCrashContextProperties::PrimaryGPUBrand );
+	AddCrashProperty( TEXT( "Misc.OSVersionMajor" ), *NCachedCrashContextProperties::OsVersion );
+	AddCrashProperty( TEXT( "Misc.OSVersionMinor" ), *NCachedCrashContextProperties::OsSubVersion );
+
+
+	// #YRX_Crash: 2015-07-21 Move to the crash report client.
+	/*{
+		uint64 AppDiskTotalNumberOfBytes = 0;
+		uint64 AppDiskNumberOfFreeBytes = 0;
+		FPlatformMisc::GetDiskTotalAndFreeSpace( FPlatformProcess::BaseDir(), AppDiskTotalNumberOfBytes, AppDiskNumberOfFreeBytes );
+		AddCrashProperty( TEXT( "Misc.AppDiskTotalNumberOfBytes" ), AppDiskTotalNumberOfBytes );
+		AddCrashProperty( TEXT( "Misc.AppDiskNumberOfFreeBytes" ), AppDiskNumberOfFreeBytes );
+	}*/
+
+	// FPlatformMemory::GetConstants is called in the GCreateMalloc, so we can assume it is always valid.
+	{
+		// Add memory stats.
+		const FPlatformMemoryConstants& MemConstants = FPlatformMemory::GetConstants();
+
+		AddCrashProperty( TEXT( "MemoryStats.TotalPhysical" ), (uint64)MemConstants.TotalPhysical );
+		AddCrashProperty( TEXT( "MemoryStats.TotalVirtual" ), (uint64)MemConstants.TotalVirtual );
+		AddCrashProperty( TEXT( "MemoryStats.PageSize" ), (uint64)MemConstants.PageSize );
+		AddCrashProperty( TEXT( "MemoryStats.TotalPhysicalGB" ), MemConstants.TotalPhysicalGB );
+	}
+
+	AddCrashProperty( TEXT( "MemoryStats.AvailablePhysical" ), (uint64)CrashMemoryStats.AvailablePhysical );
+	AddCrashProperty( TEXT( "MemoryStats.AvailableVirtual" ), (uint64)CrashMemoryStats.AvailableVirtual );
+	AddCrashProperty( TEXT( "MemoryStats.UsedPhysical" ), (uint64)CrashMemoryStats.UsedPhysical );
+	AddCrashProperty( TEXT( "MemoryStats.PeakUsedPhysical" ), (uint64)CrashMemoryStats.PeakUsedPhysical );
+	AddCrashProperty( TEXT( "MemoryStats.UsedVirtual" ), (uint64)CrashMemoryStats.UsedVirtual );
+	AddCrashProperty( TEXT( "MemoryStats.PeakUsedVirtual" ), (uint64)CrashMemoryStats.PeakUsedVirtual );
+	AddCrashProperty( TEXT( "MemoryStats.bIsOOM" ), (int32)FPlatformMemory::bIsOOM );
+	AddCrashProperty( TEXT( "MemoryStats.OOMAllocationSize"), (uint64)FPlatformMemory::OOMAllocationSize );
+	AddCrashProperty( TEXT( "MemoryStats.OOMAllocationAlignment"), (int32)FPlatformMemory::OOMAllocationAlignment );
+
+	//Architecture
+	//CrashedModuleName
+	//LoadedModules
+	EndSection( TEXT( "RuntimeProperties" ) );
+	
+	// Add platform specific properties.
+	BeginSection( TEXT("PlatformProperties") );
+	AddPlatformSpecificProperties();
+	EndSection( TEXT( "PlatformProperties" ) );
+
+	AddFooter();
+}
+
+const FString& FGenericCrashContext::GetUniqueCrashName()
+{
+	return NCachedCrashContextProperties::CrashGUID;
+}
+
+const bool FGenericCrashContext::IsFullCrashDump()
+{
+	return (NCachedCrashContextProperties::CrashDumpMode == (int32)ECrashDumpMode::FullDump);
+}
+
+void FGenericCrashContext::SerializeAsXML( const TCHAR* Filename )
+{
+	SerializeContentToBuffer();
+	// Use OS build-in functionality instead.
+	FFileHelper::SaveStringToFile( CommonBuffer, Filename, FFileHelper::EEncodingOptions::AutoDetect );
+}
+
+void FGenericCrashContext::AddCrashProperty( const TCHAR* PropertyName, const TCHAR* PropertyValue )
+{
+	CommonBuffer += TEXT( "<" );
+	CommonBuffer += PropertyName;
+	CommonBuffer += TEXT( ">" );
+
+
+	CommonBuffer += EscapeXMLString( PropertyValue );
+
+	CommonBuffer += TEXT( "</" );
+	CommonBuffer += PropertyName;
+	CommonBuffer += TEXT( ">" );
+	CommonBuffer += LINE_TERMINATOR;
+}
+
+void FGenericCrashContext::AddPlatformSpecificProperties()
+{
+	// Nothing really to do here. Can be overridden by the platform code.
+	// @see FWindowsPlatformCrashContext::AddPlatformSpecificProperties
+}
+
+void FGenericCrashContext::AddHeader()
+{
+	CommonBuffer += TEXT( "<?xml version=\"1.0\" encoding=\"UTF-8\"?>" ) LINE_TERMINATOR;
+	BeginSection( TEXT("FGenericCrashContext") );
+}
+
+void FGenericCrashContext::AddFooter()
+{
+	EndSection( TEXT( "FGenericCrashContext" ) );
+}
+
+void FGenericCrashContext::BeginSection( const TCHAR* SectionName )
+{
+	CommonBuffer += TEXT( "<" );
+	CommonBuffer += SectionName;
+	CommonBuffer += TEXT( ">" );
+	CommonBuffer += LINE_TERMINATOR;
+}
+
+void FGenericCrashContext::EndSection( const TCHAR* SectionName )
+{
+	CommonBuffer += TEXT( "</" );
+	CommonBuffer += SectionName;
+	CommonBuffer += TEXT( ">" );
+	CommonBuffer += LINE_TERMINATOR;
+}
+
+
+FString FGenericCrashContext::EscapeXMLString( const FString& Text )
+{
+	return Text
+		.Replace( TEXT( "&" ), TEXT( "&amp;" ) )
+		.Replace( TEXT( "\"" ), TEXT( "&quot;" ) )
+		.Replace( TEXT( "'" ), TEXT( "&apos;" ) )
+		.Replace( TEXT( "<" ), TEXT( "&lt;" ) )
+		.Replace( TEXT( ">" ), TEXT( "&gt;" ) );
+}
+
+FString FGenericCrashContext::UnescapeXMLString( const FString& Text )
+{
+	return Text
+		.Replace( TEXT( "&amp;" ), TEXT( "&" ) )
+		.Replace( TEXT( "&quot;" ), TEXT( "\"" ) )
+		.Replace( TEXT( "&apos;" ), TEXT( "'" ) )
+		.Replace( TEXT( "&lt;" ), TEXT( "<" ) )
+		.Replace( TEXT( "&gt;" ), TEXT( ">" ) );
+}
+
+FProgramCounterSymbolInfoEx::FProgramCounterSymbolInfoEx( FString InModuleName, FString InFunctionName, FString InFilename, uint32 InLineNumber, uint64 InSymbolDisplacement, uint64 InOffsetInModule, uint64 InProgramCounter ) :
+	ModuleName( InModuleName ),
+	FunctionName( InFunctionName ),
+	Filename( InFilename ),
+	LineNumber( InLineNumber ),
+	SymbolDisplacement( InSymbolDisplacement ),
+	OffsetInModule( InOffsetInModule ),
+	ProgramCounter( InProgramCounter )
+{
+
+}
\ No newline at end of file
diff --git a/Engine/Source/Runtime/Core/Private/GenericPlatform/GenericPlatformMisc.cpp b/Engine/Source/Runtime/Core/Private/GenericPlatform/GenericPlatformMisc.cpp
--- a/Engine/Source/Runtime/Core/Private/GenericPlatform/GenericPlatformMisc.cpp
+++ b/Engine/Source/Runtime/Core/Private/GenericPlatform/GenericPlatformMisc.cpp
@@ -10,7 +10,7 @@
 #include "SecureHash.h"
 #include "ExceptionHandling.h"
 #include "Containers/Map.h"
-#include "GenericPlatformContext.h"
+#include "GenericPlatformCrashContext.h"
 
 #include "UProjectInfo.h"
 
diff --git a/Engine/Source/Runtime/Core/Private/GenericPlatform/GenericPlatformStackWalk.cpp b/Engine/Source/Runtime/Core/Private/GenericPlatform/GenericPlatformStackWalk.cpp
--- a/Engine/Source/Runtime/Core/Private/GenericPlatform/GenericPlatformStackWalk.cpp
+++ b/Engine/Source/Runtime/Core/Private/GenericPlatform/GenericPlatformStackWalk.cpp
@@ -1,7 +1,7 @@
 // Copyright 1998-2015 Epic Games, Inc. All Rights Reserved.
 
 #include "CorePrivatePCH.h"
-#include "GenericPlatformContext.h"
+#include "GenericPlatformCrashContext.h"
 
 FProgramCounterSymbolInfo::FProgramCounterSymbolInfo() :
 	LineNumber( 0 ),
diff --git a/Engine/Source/Runtime/Core/Private/HTML5/HTML5PlatformMisc.cpp b/Engine/Source/Runtime/Core/Private/HTML5/HTML5PlatformMisc.cpp
--- a/Engine/Source/Runtime/Core/Private/HTML5/HTML5PlatformMisc.cpp
+++ b/Engine/Source/Runtime/Core/Private/HTML5/HTML5PlatformMisc.cpp
@@ -14,7 +14,7 @@
 #include "trace.h"
 
 #include "unicode/locid.h"
-#include "GenericPlatformContext.h"
+#include "GenericPlatformCrashContext.h"
 #include <SDL.h>
 #include <ctime>
 
diff --git a/Engine/Source/Runtime/Core/Private/Windows/WindowsPlatformCrashContext.cpp b/Engine/Source/Runtime/Core/Private/Windows/WindowsPlatformCrashContext.cpp
--- a/Engine/Source/Runtime/Core/Private/Windows/WindowsPlatformCrashContext.cpp
+++ b/Engine/Source/Runtime/Core/Private/Windows/WindowsPlatformCrashContext.cpp
@@ -295,7 +295,7 @@ int32 ReportCrashUsingCrashReportClient(EXCEPTION_POINTERS* ExceptionInfo, const
 			{
 				// No super safe due to dynamic memory allocations, but at least enables new functionality.
 				// Introduces a new runtime crash context. Will replace all Windows related crash reporting.
-				FPlatformCrashContext CrashContext;
+				FWindowsPlatformCrashContext CrashContext;
 
 				const FString CrashContextXMLPath = FPaths::Combine( *FPaths::GameLogDir(), *CrashContext.GetUniqueCrashName(), FPlatformCrashContext::CrashContextRuntimeXMLNameW );
 				CrashContext.SerializeAsXML( *CrashContextXMLPath );
diff --git a/Engine/Source/Runtime/Core/Public/Android/AndroidPlatformCrashContext.h b/Engine/Source/Runtime/Core/Public/Android/AndroidPlatformCrashContext.h
--- a/Engine/Source/Runtime/Core/Public/Android/AndroidPlatformCrashContext.h
+++ b/Engine/Source/Runtime/Core/Public/Android/AndroidPlatformCrashContext.h
@@ -2,7 +2,7 @@
 
 #pragma once
 
-#include "GenericPlatform/GenericPlatformContext.h"
+#include "GenericPlatform/GenericPlatformCrashContext.h"
 
 struct CORE_API FAndroidCrashContext : public FGenericCrashContext
 {
diff --git a/Engine/Source/Runtime/Core/Public/Apple/ApplePlatformCrashContext.h b/Engine/Source/Runtime/Core/Public/Apple/ApplePlatformCrashContext.h
--- a/Engine/Source/Runtime/Core/Public/Apple/ApplePlatformCrashContext.h
+++ b/Engine/Source/Runtime/Core/Public/Apple/ApplePlatformCrashContext.h
@@ -6,7 +6,7 @@
 
 #pragma once
 
-#include "GenericPlatform/GenericPlatformContext.h"
+#include "GenericPlatform/GenericPlatformCrashContext.h"
 
 /**
  * Declaration for common crash context implementation for Apple platforms.
diff --git a/Engine/Source/Runtime/Core/Public/GenericPlatform/GenericPlatformContext.h b/Engine/Source/Runtime/Core/Public/GenericPlatform/GenericPlatformContext.h
deleted file mode 100644
diff --git a/Engine/Source/Runtime/Core/Public/GenericPlatform/GenericPlatformCrashContext.h b/Engine/Source/Runtime/Core/Public/GenericPlatform/GenericPlatformCrashContext.h
new file mode 100644
--- /dev/null
+++ b/Engine/Source/Runtime/Core/Public/GenericPlatform/GenericPlatformCrashContext.h
@@ -0,0 +1,155 @@
+// Copyright 1998-2015 Epic Games, Inc. All Rights Reserved.
+
+#pragma once
+
+class FString;
+
+/** 
+ * Symbol information associated with a program counter. 
+ * FString version.
+ * To be used by external tools.
+ */
+struct CORE_API FProgramCounterSymbolInfoEx
+{
+	/** Module name. */
+	FString	ModuleName;
+
+	/** Function name. */
+	FString	FunctionName;
+
+	/** Filename. */
+	FString	Filename;
+
+	/** Line number in file. */
+	uint32	LineNumber;
+
+	/** Symbol displacement of address.	*/
+	uint64	SymbolDisplacement;
+
+	/** Program counter offset into module. */
+	uint64	OffsetInModule;
+
+	/** Program counter. */
+	uint64	ProgramCounter;
+
+	/** Default constructor. */
+	FProgramCounterSymbolInfoEx( FString InModuleName, FString InFunctionName, FString InFilename, uint32 InLineNumber, uint64 InSymbolDisplacement, uint64 InOffsetInModule, uint64 InProgramCounter );
+};
+
+
+/** Enumerates crash description versions. */
+enum class ECrashDescVersions : int32
+{
+	/** Introduces a new crash description format. */
+	VER_1_NewCrashFormat,
+
+	/** Added misc properties (CPU,GPU,OS,etc), memory related stats and platform specific properties as generic payload. */
+	VER_2_AddedNewProperties,
+};
+
+/** Enumerates crash dump modes. */
+enum class ECrashDumpMode : int32
+{
+	/** Default minidump settings. */
+	Default = 0,
+
+	/** Full memory crash minidump */
+	FullDump = 1,
+};
+
+/**
+ *	Contains a runtime crash's properties that are common for all platforms.
+ *	This may change in the future.
+ */
+struct CORE_API FGenericCrashContext
+{
+public:
+
+	/**
+	* We can't gather memory stats in crash handling function, so we gather them just before raising
+	* exception and use in crash reporting.
+	*/
+	static FPlatformMemoryStats CrashMemoryStats;
+	static const ANSICHAR* CrashContextRuntimeXMLNameA;
+	static const TCHAR* CrashContextRuntimeXMLNameW;
+
+	/** Initializes crash context related platform specific data that can be impossible to obtain after a crash. */
+	static void Initialize();
+
+	/**
+	 * @return true, if the generic crash context has been initialized.
+	 */
+	static bool IsInitalized()
+	{
+		return bIsInitialized;
+	}
+
+	/** Default constructor. */
+	FGenericCrashContext();
+
+	/** Serializes all data to the buffer. */
+	void SerializeContentToBuffer();
+
+	/**
+	 * @return the buffer containing serialized data.
+	 */
+	const FString& GetBuffer() const
+	{
+		return CommonBuffer;
+	}
+
+	/**
+	 * @return a globally unique crash name.
+	 */
+	const FString& GetUniqueCrashName();
+
+	/**
+	 * @return whether this crash is a full memory minidump
+	 */
+	const bool IsFullCrashDump();
+
+	/** Serializes crash's informations to the specified filename. Should be overridden for platforms where using FFileHelper is not safe, all POSIX platforms. */
+	virtual void SerializeAsXML( const TCHAR* Filename );
+
+	/** Writes a common property to the buffer. */
+	void AddCrashProperty( const TCHAR* PropertyName, const TCHAR* PropertyValue );
+
+	/** Writes a common property to the buffer. */
+	template <typename Type>
+	void AddCrashProperty( const TCHAR* PropertyName, const Type& Value )
+	{
+		AddCrashProperty( PropertyName, *TTypeToString<Type>::ToString( Value ) );
+	}
+
+private:
+	/** Serializes platform specific properties to the buffer. */
+	virtual void AddPlatformSpecificProperties();
+
+	/** Writes header information to the buffer. */
+	void AddHeader();
+
+	/** Writes footer to the buffer. */
+	void AddFooter();
+
+	void BeginSection( const TCHAR* SectionName );
+	void EndSection( const TCHAR* SectionName );
+
+	/** Escapes a specified XML string, naive implementation. */
+	static FString EscapeXMLString( const FString& Text );
+
+	/** Unescapes a specified XML string, naive implementation. */
+	static FString UnescapeXMLString( const FString& Text );
+
+	/**	Whether the Initialize() has been called */
+	static bool bIsInitialized;
+
+	/** The buffer used to store the crash's properties. */
+	FString CommonBuffer;
+
+	// FNoncopyable
+	FGenericCrashContext( const FGenericCrashContext& );
+	FGenericCrashContext& operator=(const FGenericCrashContext&);
+};
+
+struct CORE_API FGenericMemoryWarningContext
+{};
diff --git a/Engine/Source/Runtime/Core/Public/Linux/LinuxPlatformCrashContext.h b/Engine/Source/Runtime/Core/Public/Linux/LinuxPlatformCrashContext.h
--- a/Engine/Source/Runtime/Core/Public/Linux/LinuxPlatformCrashContext.h
+++ b/Engine/Source/Runtime/Core/Public/Linux/LinuxPlatformCrashContext.h
@@ -2,7 +2,7 @@
 
 #pragma once
 
-#include "GenericPlatform/GenericPlatformContext.h"
+#include "GenericPlatform/GenericPlatformCrashContext.h"
 
 // commandline parameter to suppress DWARF parsing (greatly speeds up callstack generation)
 #define CMDARG_SUPPRESS_DWARF_PARSING			"nodwarf"
diff --git a/Engine/Source/Runtime/Core/Public/Windows/WindowsPlatformCrashContext.h b/Engine/Source/Runtime/Core/Public/Windows/WindowsPlatformCrashContext.h
--- a/Engine/Source/Runtime/Core/Public/Windows/WindowsPlatformCrashContext.h
+++ b/Engine/Source/Runtime/Core/Public/Windows/WindowsPlatformCrashContext.h
@@ -2,7 +2,7 @@
 
 #pragma once
 
-#include "GenericPlatform/GenericPlatformContext.h"
+#include "GenericPlatform/GenericPlatformCrashContext.h"
 
 #include "AllowWindowsPlatformTypes.h"
 #include "DbgHelp.h"
