Fixed a typo

diff --git a/Engine/Source/Runtime/Core/Private/GenericPlatform/GenericPlatformCrashContext.cpp b/Engine/Source/Runtime/Core/Private/GenericPlatform/GenericPlatformCrashContext.cpp
--- a/Engine/Source/Runtime/Core/Private/GenericPlatform/GenericPlatformCrashContext.cpp
+++ b/Engine/Source/Runtime/Core/Private/GenericPlatform/GenericPlatformCrashContext.cpp
@@ -2,7 +2,7 @@
 
 #include "CorePrivatePCH.h"
 #include "GenericPlatform/GenericPlatformCrashContext.h"
-#include "Misc./App.h"
+#include "Misc/App.h"
 #include "EngineVersion.h"
 #include "EngineBuildSettings.h"
 
