Added 'intrusive' linked lists, which non-UObject structs/classes can inherit, to provide a linked list interface.

Also refactored TLinkedList methods - both this and TIntrusiveLinkedList inherit the same interface.

diff --git a/Engine/Source/Runtime/Core/Public/Containers/List.h b/Engine/Source/Runtime/Core/Public/Containers/List.h
--- a/Engine/Source/Runtime/Core/Public/Containers/List.h
+++ b/Engine/Source/Runtime/Core/Public/Containers/List.h
@@ -3,30 +3,30 @@
 #pragma once
 
 
-template <class ContainerType, class ElementType>
-class TLinkedListIterator
+template <class ContainerType>
+class TLinkedListIteratorBase
 {
 public:
-	explicit TLinkedListIterator(ContainerType* FirstLink)
+	explicit TLinkedListIteratorBase(ContainerType* FirstLink)
 		: CurrentLink(FirstLink)
 	{ }
 
 	/**
 	 * Advances the iterator to the next element.
 	 */
-	void Next()
+	FORCEINLINE void Next()
 	{
 		checkSlow(CurrentLink);
-		CurrentLink = CurrentLink->GetNextLink();
+		CurrentLink = (ContainerType*)CurrentLink->GetNextLink();
 	}
 
-	TLinkedListIterator& operator++()
+	FORCEINLINE TLinkedListIteratorBase& operator++()
 	{
 		Next();
 		return *this;
 	}
 
-	TLinkedListIterator operator++(int)
+	FORCEINLINE TLinkedListIteratorBase operator++(int)
 	{
 		auto Tmp = *this;
 		Next();
@@ -39,65 +39,94 @@ public:
 		return CurrentLink != nullptr;
 	}
 
+protected:
+
+	ContainerType* CurrentLink;
+
+	FORCEINLINE friend bool operator==(const TLinkedListIteratorBase& Lhs, const TLinkedListIteratorBase& Rhs) { return Lhs.CurrentLink == Rhs.CurrentLink; }
+	FORCEINLINE friend bool operator!=(const TLinkedListIteratorBase& Lhs, const TLinkedListIteratorBase& Rhs) { return Lhs.CurrentLink != Rhs.CurrentLink; }
+};
+
+template <class ContainerType, class ElementType>
+class TLinkedListIterator : public TLinkedListIteratorBase<ContainerType>
+{
+	typedef TLinkedListIteratorBase<ContainerType> Super;
+
+public:
+	explicit TLinkedListIterator(ContainerType* FirstLink)
+		: Super(FirstLink)
+	{
+	}
+
 	// Accessors.
-	ElementType& operator->() const
+	FORCEINLINE ElementType& operator->() const
 	{
-		checkSlow(CurrentLink);
-		return **CurrentLink;
+		checkSlow(this->CurrentLink);
+		return **(this->CurrentLink);
 	}
 
-	ElementType& operator*() const
+	FORCEINLINE ElementType& operator*() const
 	{
-		checkSlow(CurrentLink);
-		return **CurrentLink;
+		checkSlow(this->CurrentLink);
+		return **(this->CurrentLink);
 	}
+};
 
-private:
+template <class ContainerType, class ElementType>
+class TIntrusiveLinkedListIterator : public TLinkedListIteratorBase<ElementType>
+{
+	typedef TLinkedListIteratorBase<ElementType> Super;
 
-	ContainerType* CurrentLink;
+public:
+	explicit TIntrusiveLinkedListIterator(ElementType* FirstLink)
+		: Super(FirstLink)
+	{
+	}
+
+	// Accessors.
+	FORCEINLINE ElementType& operator->() const
+	{
+		checkSlow(this->CurrentLink);
+		return *(this->CurrentLink);
+	}
 
-	friend bool operator==(const TLinkedListIterator& Lhs, const TLinkedListIterator& Rhs) { return Lhs.CurrentLink == Rhs.CurrentLink; }
-	friend bool operator!=(const TLinkedListIterator& Lhs, const TLinkedListIterator& Rhs) { return Lhs.CurrentLink != Rhs.CurrentLink; }
+	FORCEINLINE ElementType& operator*() const
+	{
+		checkSlow(this->CurrentLink);
+		return *(this->CurrentLink);
+	}
 };
 
 
 /**
- * Encapsulates a link in a single linked list with constant access time.
+ * Base linked list class, used to implement methods shared by intrusive/non-intrusive linked lists
  */
-template <class ElementType>
-class TLinkedList
+template <class ContainerType, class ElementType, template<class, class> class IteratorType>
+class TLinkedListBase
 {
 public:
-
 	/**
 	 * Used to iterate over the elements of a linked list.
 	 */
-	typedef TLinkedListIterator<TLinkedList,       ElementType> TIterator;
-	typedef TLinkedListIterator<TLinkedList, const ElementType> TConstIterator;
+	typedef IteratorType<ContainerType,       ElementType> TIterator;
+	typedef IteratorType<ContainerType, const ElementType> TConstIterator;
 
-	/** Default constructor (empty list). */
-	TLinkedList()
-		: NextLink(nullptr)
-		, PrevLink(nullptr)
-	{ }
 
 	/**
-	 * Creates a new linked list with a single element.
-	 *
-	 * @param InElement The element to add to the list.
+	 * Default constructor (empty list)
 	 */
-	explicit TLinkedList(const ElementType& InElement)
-		: Element (InElement)
-		, NextLink(nullptr)
-		, PrevLink(nullptr)
-	{ }
+	TLinkedListBase()
+		: NextLink(NULL)
+		, PrevLink(NULL)
+	{
+	}
 
 	/**
 	 * Removes this element from the list in constant time.
 	 *
 	 * This function is safe to call even if the element is not linked.
 	 */
-	void Unlink()
+	FORCEINLINE void Unlink()
 	{
 		if( NextLink )
 		{
@@ -112,73 +141,224 @@ public:
 		PrevLink = nullptr;
 	}
 
+
 	/**
 	 * Adds this element to a list, before the given element.
 	 *
-	 * @param Before The link to insert this element before.
+	 * @param Before	The link to insert this element before.
 	 */
-	void Link(TLinkedList*& Before)
+	FORCEINLINE void LinkBefore(ContainerType* Before)
 	{
-		if(Before)
+		checkSlow(Before != NULL);
+
+		PrevLink = Before->PrevLink;
+		Before->PrevLink = &NextLink;
+
+		NextLink = Before;
+
+		if (PrevLink != NULL)
 		{
-			Before->PrevLink = &NextLink;
+			*PrevLink = (ContainerType*)this;
 		}
-		NextLink = Before;
-		PrevLink = &Before;
-		Before = this;
 	}
 
+	/**
+	 * Adds this element to the linked list, after the specified element
+	 *
+	 * @param After		The link to insert this element after.
+	 */
+	FORCEINLINE void LinkAfter(ContainerType* After)
+	{
+		checkSlow(After != NULL);
+
+		PrevLink = &After->NextLink;
+		NextLink = *PrevLink;
+		*PrevLink = (ContainerType*)this;
+
+		if (NextLink != NULL)
+		{
+			NextLink->PrevLink = &NextLink;
+		}
+	}
+
+	/**
+	 * Adds this element to the linked list, replacing the specified element.
+	 * This is equivalent to calling LinkBefore(Replace); Replace->Unlink();
+	 *
+	 * @param Replace	Pointer to the element to be replaced
+	 */
+	FORCEINLINE void LinkReplace(ContainerType* Replace)
+	{
+		checkSlow(Replace != NULL);
+
+		ContainerType**& ReplacePrev = Replace->PrevLink;
+		ContainerType*& ReplaceNext = Replace->NextLink;
+
+		PrevLink = ReplacePrev;
+		NextLink = ReplaceNext;
+
+		if (PrevLink != NULL)
+		{
+			*PrevLink = (ContainerType*)this;
+		}
+
+		if (NextLink != NULL)
+		{
+			NextLink->PrevLink = &NextLink;
+		}
+
+		ReplacePrev = NULL;
+		ReplaceNext = NULL;
+	}
+
+
+	/**
+	 * Adds this element as the head of the linked list, linking the input Head pointer to this element,
+	 * so that when the element is linked/unlinked, the Head linked list pointer will be correctly updated.
+	 *
+	 * If Head already has an element, this functions like LinkBefore.
+	 *
+	 * @param Head		Pointer to the head of the linked list - this pointer should be the main reference point for the linked list
+	 */
+	FORCEINLINE void LinkHead(ContainerType*& Head)
+	{
+		if (Head != NULL)
+		{
+			Head->PrevLink = &NextLink;
+		}
+
+		NextLink = Head;
+		PrevLink = &Head;
+		Head = (ContainerType*)this;
+	}
+
+
+	DEPRECATED(4.10, "This function is deprecated. Please use LinkHead.")
+	FORCEINLINE void Link(ContainerType*& Before)
+	{
+		LinkHead(Before);
+	}
+
+
 	/**
 	 * Returns whether element is currently linked.
 	 *
 	 * @return true if currently linked, false otherwise
 	 */
-	bool IsLinked()
+	FORCEINLINE bool IsLinked()
 	{
 		return PrevLink != nullptr;
 	}
 
-	TLinkedList* GetPrevLink() const
+	FORCEINLINE ContainerType** GetPrevLink() const
 	{
 		return PrevLink;
 	}
 
-	TLinkedList* GetNextLink() const
+	FORCEINLINE ContainerType* GetNextLink() const
+	{
+		return NextLink;
+	}
+
+	FORCEINLINE ContainerType* Next()
 	{
 		return NextLink;
 	}
 
+private:
+	/** The next link in the linked list */
+	ContainerType*  NextLink;
+
+	/** Pointer to 'NextLink', within the previous link in the linked list */
+	ContainerType** PrevLink;
+
+
+	FORCEINLINE friend TIterator      begin(      ContainerType& List) { return TIterator     (&List); }
+	FORCEINLINE friend TConstIterator begin(const ContainerType& List) { return TConstIterator(const_cast<ContainerType*>(&List)); }
+	FORCEINLINE friend TIterator      end  (      ContainerType& List) { return TIterator     (nullptr); }
+	FORCEINLINE friend TConstIterator end  (const ContainerType& List) { return TConstIterator(nullptr); }
+};
+
+/**
+ * Encapsulates a link in a single linked list with constant access time.
+ *
+ * This linked list is non-intrusive, i.e. it stores a copy of the element passed to it (typically a pointer)
+ */
+template <class ElementType>
+class TLinkedList : public TLinkedListBase<TLinkedList<ElementType>, ElementType, TLinkedListIterator>
+{
+	typedef TLinkedListBase<TLinkedList<ElementType>, ElementType, TLinkedListIterator>		Super;
+
+public:
+	/** Default constructor (empty list). */
+	TLinkedList()
+		: Super()
+	{
+	}
+
+	/**
+	 * Creates a new linked list with a single element.
+	 *
+	 * @param InElement The element to add to the list.
+	 */
+	explicit TLinkedList(const ElementType& InElement)
+		: Super()
+		, Element(InElement)
+	{
+	}
+
+
 	// Accessors.
-	ElementType* operator->()
+	FORCEINLINE ElementType* operator->()
 	{
 		return &Element;
 	}
-	const ElementType* operator->() const
+	FORCEINLINE const ElementType* operator->() const
 	{
 		return &Element;
 	}
-	ElementType& operator*()
+	FORCEINLINE ElementType& operator*()
 	{
 		return Element;
 	}
-	const ElementType& operator*() const
+	FORCEINLINE const ElementType& operator*() const
 	{
 		return Element;
 	}
-	TLinkedList* Next()
-	{
-		return NextLink;
-	}
+
 
 private:
 	ElementType   Element;
-	TLinkedList*  NextLink;
-	TLinkedList** PrevLink;
+};
 
-	friend TIterator      begin(      TLinkedList& List) { return TIterator     (&List); }
-	friend TConstIterator begin(const TLinkedList& List) { return TConstIterator(const_cast<TLinkedList*>(&List)); }
-	friend TIterator      end  (      TLinkedList& List) { return TIterator     (nullptr); }
-	friend TConstIterator end  (const TLinkedList& List) { return TConstIterator(nullptr); }
+/**
+ * Encapsulates a link in a single linked list with constant access time.
+ * Structs/classes must inherit this, to use it, e.g: struct FMyStruct : public TIntrusiveLinkedList<FMyStruct>
+ *
+ * This linked list is intrusive, i.e. the element is a subclass of this link, so that each link IS the element.
+ *
+ * Never reference TIntrusiveLinkedList outside of the above class/struct inheritance, only ever refer to the struct, e.g:
+ *	FMyStruct* MyLinkedList = NULL;
+ *
+ *	FMyStruct* StructLink = new FMyStruct();
+ *	StructLink->LinkHead(MyLinkedList);
+ *
+ *	for (FMyStruct::TIterator It(MyLinkedList); It; It.Next())
+ *	{
+ *		...
+ *	}
+ */
+template <class ElementType>
+class TIntrusiveLinkedList : public TLinkedListBase<ElementType, ElementType, TIntrusiveLinkedListIterator>
+{
+	typedef TLinkedListBase<ElementType, ElementType, TIntrusiveLinkedListIterator>		Super;
+
+public:
+	/** Default constructor (empty list). */
+	TIntrusiveLinkedList()
+		: Super()
+	{
+	}
 };
 
 
