#ue4 - Improve FQuat::FindBetween() and add some Vector->Orientation functions that include Quat results.

- Faster FQuat::FindBetween() implementation. Add FindBetweenNormals() and FindBetweenVectors() to indicate different input length requirements.
- Add FVector::ToOrientationQuat(). This preserves the up vector (result has no roll).
- Add FVector::ToOrientationRotator(), matching FVector::Rotation() implementation.
- Add unit tests.

UE-14015

diff --git a/Engine/Source/Runtime/Core/Private/Math/UnrealMath.cpp b/Engine/Source/Runtime/Core/Private/Math/UnrealMath.cpp
--- a/Engine/Source/Runtime/Core/Private/Math/UnrealMath.cpp
+++ b/Engine/Source/Runtime/Core/Private/Math/UnrealMath.cpp
@@ -163,15 +163,15 @@ void FRotator::SerializeCompressedShort( FArchive& Ar )
 	}
 }
 
-CORE_API FRotator FVector::Rotation() const
+FRotator FVector::ToOrientationRotator() const
 {
 	FRotator R;
 
 	// Find yaw.
-	R.Yaw = FMath::Atan2(Y,X) * 180.f / PI;
+	R.Yaw = FMath::Atan2(Y,X) * (180.f / PI);
 
 	// Find pitch.
-	R.Pitch = FMath::Atan2(Z,FMath::Sqrt(X*X+Y*Y)) * 180.f / PI;
+	R.Pitch = FMath::Atan2(Z,FMath::Sqrt(X*X+Y*Y)) * (180.f / PI);
 
 	// Find roll.
 	R.Roll = 0;
@@ -179,15 +179,15 @@ CORE_API FRotator FVector::Rotation() const
 	return R;
 }
 
-FRotator FVector4::Rotation() const
+FRotator FVector4::ToOrientationRotator() const
 {
 	FRotator R;
 
 	// Find yaw.
-	R.Yaw = FMath::Atan2(Y,X) * 180.f / PI;
+	R.Yaw = FMath::Atan2(Y,X) * (180.f / PI);
 
 	// Find pitch.
-	R.Pitch = FMath::Atan2(Z,FMath::Sqrt(X*X+Y*Y)) * 180.f / PI;
+	R.Pitch = FMath::Atan2(Z,FMath::Sqrt(X*X+Y*Y)) * (180.f / PI);
 
 	// Find roll.
 	R.Roll = 0;
@@ -195,6 +195,53 @@ FRotator FVector4::Rotation() const
 	return R;
 }
 
+
+FQuat FVector::ToOrientationQuat() const
+{
+	// Essentially an optimized Vector->Rotator->Quat made possible by knowing Roll == 0, and avoiding radians->degrees->radians.
+	// This is done to avoid adding any roll (which our API states as a constraint).
+	const float YawRad = FMath::Atan2(Y, X);
+	const float PitchRad = FMath::Atan2(Z, FMath::Sqrt(X*X + Y*Y));
+
+	const float DIVIDE_BY_2 = 0.5f;
+	float SP, SY;
+	float CP, CY;
+
+	FMath::SinCos(&SP, &CP, PitchRad * DIVIDE_BY_2);
+	FMath::SinCos(&SY, &CY, YawRad * DIVIDE_BY_2);
+
+	FQuat RotationQuat;
+	RotationQuat.X =  SP*SY;
+	RotationQuat.Y = -SP*CY;
+	RotationQuat.Z =  CP*SY;
+	RotationQuat.W =  CP*CY;
+	return RotationQuat;
+}
+
+
+FQuat FVector4::ToOrientationQuat() const
+{
+	// Essentially an optimized Vector->Rotator->Quat made possible by knowing Roll == 0, and avoiding radians->degrees->radians.
+	// This is done to avoid adding any roll (which our API states as a constraint).
+	const float YawRad = FMath::Atan2(Y, X);
+	const float PitchRad = FMath::Atan2(Z, FMath::Sqrt(X*X + Y*Y));
+
+	const float DIVIDE_BY_2 = 0.5f;
+	float SP, SY;
+	float CP, CY;
+
+	FMath::SinCos(&SP, &CP, PitchRad * DIVIDE_BY_2);
+	FMath::SinCos(&SY, &CY, YawRad * DIVIDE_BY_2);
+
+	FQuat RotationQuat;
+	RotationQuat.X =  SP*SY;
+	RotationQuat.Y = -SP*CY;
+	RotationQuat.Z =  CP*SY;
+	RotationQuat.W =  CP*CY;
+	return RotationQuat;
+}
+
+
 void FVector::FindBestAxisVectors( FVector& Axis1, FVector& Axis2 ) const
 {
 	const float NX = FMath::Abs(X);
@@ -696,52 +743,48 @@ bool FQuat::NetSerialize(FArchive& Ar, class UPackageMap*, bool& bOutSuccess)
 	return true;
 }
 
-FQuat FQuat::FindBetween(const FVector& vec1, const FVector& vec2)
+
+//
+// Based on:
+// http://lolengine.net/blog/2014/02/24/quaternion-from-two-vectors-final
+// http://www.euclideanspace.com/maths/algebra/vectors/angleBetween/index.htm
+//
+FORCEINLINE_DEBUGGABLE FQuat FindBetween_Helper(const FVector& A, const FVector& B, float NormAB)
 {
-	const FVector cross = vec1 ^ vec2;
-	const float crossMag = cross.Size();
+	float W = NormAB + FVector::DotProduct(A, B);
+	FQuat Result;
 
-	// See if vectors are parallel or anti-parallel
-	if(crossMag < KINDA_SMALL_NUMBER)
+	if (W >= 1e-6f * NormAB)
 	{
-		// If these vectors are parallel - just return identity quaternion (ie no rotation).
-		const float Dot = vec1 | vec2;
-		if(Dot > -KINDA_SMALL_NUMBER)
-		{
-			return FQuat::Identity; // no rotation
-		}
-		// Exactly opposite..
-		else
-		{
-			// ..rotation by 180 degrees around a vector orthogonal to vec1 & vec2
-			FVector Vec = vec1.SizeSquared() > vec2.SizeSquared() ? vec1 : vec2;
-			Vec.Normalize();
-
-			FVector AxisA, AxisB;
-			Vec.FindBestAxisVectors(AxisA, AxisB);
-
-			return FQuat(AxisA.X, AxisA.Y, AxisA.Z, 0.f); // (axis*sin(pi/2), cos(pi/2)) = (axis, 0)
-		}
+		//Axis = FVector::CrossProduct(A, B);
+		Result = FQuat(A.Y * B.Z - A.Z * B.Y,
+					   A.Z * B.X - A.X * B.Z,
+					   A.X * B.Y - A.Y * B.X,
+					   W);
 	}
-
-	// Not parallel, so use normal code
-	float angle = FMath::Asin(crossMag);
-
-	const float dot = vec1 | vec2;
-	if(dot < 0.0f)
+	else
 	{
-		angle = PI - angle;
+		// A and B point in opposite directions
+		W = 0.f;
+		Result = FMath::Abs(A.X) > FMath::Abs(A.Y)
+				? FQuat(-A.Z, 0.f, A.X, W)
+				: FQuat(0.f, -A.Z, A.Y, W);
 	}
 
-	float sinHalfAng, cosHalfAng;
-	FMath::SinCos(&sinHalfAng, &cosHalfAng, 0.5f * angle);
-	const FVector axis = cross / crossMag;
+	Result.Normalize();
+	return Result;
+}
 
-	return FQuat(
-		sinHalfAng * axis.X,
-		sinHalfAng * axis.Y,
-		sinHalfAng * axis.Z,
-		cosHalfAng );
+FQuat FQuat::FindBetweenNormals(const FVector& A, const FVector& B)
+{
+	const float NormAB = 1.f;
+	return FindBetween_Helper(A, B, NormAB);
+}
+
+FQuat FQuat::FindBetweenVectors(const FVector& A, const FVector& B)
+{
+	const float NormAB = FMath::Sqrt(A.SizeSquared() * B.SizeSquared());
+	return FindBetween_Helper(A, B, NormAB);
 }
 
 FQuat FQuat::Log() const
diff --git a/Engine/Source/Runtime/Core/Private/Tests/Math/UnrealMathTest.cpp b/Engine/Source/Runtime/Core/Private/Tests/Math/UnrealMathTest.cpp
--- a/Engine/Source/Runtime/Core/Private/Tests/Math/UnrealMathTest.cpp
+++ b/Engine/Source/Runtime/Core/Private/Tests/Math/UnrealMathTest.cpp
@@ -449,6 +449,54 @@ FORCENOINLINE FRotator TestQuaternionToRotator(const FQuat& Quat)
 }
 
 
+FORCENOINLINE FQuat FindBetween_Old(const FVector& vec1, const FVector& vec2)
+{
+	const FVector cross = vec1 ^ vec2;
+	const float crossMag = cross.Size();
+
+	// See if vectors are parallel or anti-parallel
+	if (crossMag < KINDA_SMALL_NUMBER)
+	{
+		// If these vectors are parallel - just return identity quaternion (ie no rotation).
+		const float Dot = vec1 | vec2;
+		if (Dot > -KINDA_SMALL_NUMBER)
+		{
+			return FQuat::Identity; // no rotation
+		}
+		// Exactly opposite..
+		else
+		{
+			// ..rotation by 180 degrees around a vector orthogonal to vec1 & vec2
+			FVector Vec = vec1.SizeSquared() > vec2.SizeSquared() ? vec1 : vec2;
+			Vec.Normalize();
+
+			FVector AxisA, AxisB;
+			Vec.FindBestAxisVectors(AxisA, AxisB);
+
+			return FQuat(AxisA.X, AxisA.Y, AxisA.Z, 0.f); // (axis*sin(pi/2), cos(pi/2)) = (axis, 0)
+		}
+	}
+
+	// Not parallel, so use normal code
+	float angle = FMath::Asin(crossMag);
+
+	const float dot = vec1 | vec2;
+	if (dot < 0.0f)
+	{
+		angle = PI - angle;
+	}
+
+	float sinHalfAng, cosHalfAng;
+	FMath::SinCos(&sinHalfAng, &cosHalfAng, 0.5f * angle);
+	const FVector axis = cross / crossMag;
+
+	return FQuat(
+		sinHalfAng * axis.X,
+		sinHalfAng * axis.Y,
+		sinHalfAng * axis.Z,
+		cosHalfAng);
+}
+
 
 // ROTATOR TESTS
 
@@ -1195,8 +1243,9 @@ bool FVectorRegisterAbstractionTest::RunTest(const FString& Parameters)
 			FVector::RightVector,
 			FVector::UpVector,
 			FVector(45.0f, -60.0f, 120.0f),
-			FVector(0.45f, -0.60f, 1.20f),
+			FVector(-45.0f, 60.0f, -120.0f),
 			FVector(0.57735026918962576451f, 0.57735026918962576451f, 0.57735026918962576451f),
+			-FVector::ForwardVector,
 		};
 
 		// ... and test within this tolerance.
@@ -1221,6 +1270,59 @@ bool FVectorRegisterAbstractionTest::RunTest(const FString& Parameters)
 				TEST_QUAT_ROTATE(QIndex, VIndex, Q, V, TestQuaternionMultiplyVector, Tolerance);
 			}
 		}
+
+
+		// FindBetween
+		{
+			for (FVector const& A: TestVectors)
+			{
+				for (FVector const &B : TestVectors)
+				{
+					const FVector ANorm = A.GetSafeNormal();
+					const FVector BNorm = B.GetSafeNormal();
+
+					const FQuat Old = FindBetween_Old(ANorm, BNorm);
+					const FQuat NewNormal = FQuat::FindBetweenNormals(ANorm, BNorm);
+					const FQuat NewVector = FQuat::FindBetweenVectors(A, B);
+
+					const FVector RotAOld = Old.RotateVector(ANorm);
+					const FVector RotANewNormal = NewNormal.RotateVector(ANorm);
+					const FVector RotANewVector = NewVector.RotateVector(ANorm);
+
+					if (A.IsZero() || B.IsZero())
+					{
+						LogTest(TEXT("FindBetween: Old == New (normal)"), TestQuatsEqual(Old, NewNormal, 1e-6f));
+						LogTest(TEXT("FindBetween: Old == New (vector)"), TestQuatsEqual(Old, NewVector, 1e-6f));
+					}
+					else
+					{
+						LogTest(TEXT("FindBetween: Old A->B"), RotAOld.Equals(BNorm));
+						LogTest(TEXT("FindBetween: New A->B (normal)"), RotANewNormal.Equals(BNorm));
+						LogTest(TEXT("FindBetween: New A->B (vector)"), RotANewVector.Equals(BNorm));
+					}
+				}
+			}
+		}
+
+
+		// FVector::ToOrientationRotator(), FVector::ToOrientationQuat()
+		{
+			for (FVector const& V: TestVectors)
+			{
+				const FVector VNormal = V.GetSafeNormal();
+
+				Q0 = FQuat::FindBetweenNormals(FVector::ForwardVector, VNormal);
+				Q1 = V.ToOrientationQuat();
+				const FRotator R0 = V.ToOrientationRotator();
+
+				const FVector Rotated0 = Q0.RotateVector(FVector::ForwardVector);
+				const FVector Rotated1 = Q1.RotateVector(FVector::ForwardVector);
+				const FVector Rotated2 = R0.RotateVector(FVector::ForwardVector);
+
+				LogTest(TEXT("V.ToOrientationQuat() rotate"), Rotated0.Equals(Rotated1));
+				LogTest(TEXT("V.ToOrientationRotator() rotate"), Rotated0.Equals(Rotated2));
+			}
+		}
 	}
 
 	// Quat multiplication
diff --git a/Engine/Source/Runtime/Core/Public/Math/Quat.h b/Engine/Source/Runtime/Core/Public/Math/Quat.h
--- a/Engine/Source/Runtime/Core/Public/Math/Quat.h
+++ b/Engine/Source/Runtime/Core/Public/Math/Quat.h
@@ -357,6 +357,9 @@ public:
 	/** Get the up direction (Z axis) after it has been rotated by this Quaternion. */
 	FORCEINLINE FVector GetUpVector() const;
 
+	/** Convert a rotation into a unit vector facing in its direction. Equivalent to GetForwardVector(). */
+	FORCEINLINE FVector Vector() const;
+
 	/** Get the FRotator representation of this Quaternion. */
 	CORE_API FRotator Rotator() const;
 
@@ -402,9 +405,22 @@ public:
 public:
 
 	/**
-	 * Generates the 'smallest' (geodesic) rotation between these two vectors.
+	 * Generates the 'smallest' (geodesic) rotation between two vectors of arbitrary length.
 	 */
-	static CORE_API FQuat FindBetween( const FVector& vec1, const FVector& vec2 );
+	static FORCEINLINE FQuat FindBetween( const FVector& Vector1, const FVector& Vector2 )
+	{
+		return FindBetweenVectors(Vector1, Vector2);
+	}
+
+	/**
+	 * Generates the 'smallest' (geodesic) rotation between two normals (assumed to be unit length).
+	 */
+	static CORE_API FQuat FindBetweenNormals( const FVector& Normal1, const FVector& Normal2 );
+
+	/**
+	 * Generates the 'smallest' (geodesic) rotation between two vectors of arbitrary length.
+	 */
+	static CORE_API FQuat FindBetweenVectors( const FVector& Vector1, const FVector& Vector2 );
 
 	/**
 	 * Error measure (angle) between two quaternions, ranged [0..1].
@@ -974,6 +990,11 @@ FORCEINLINE FVector FQuat::GetUpVector() const
 	return GetAxisZ();
 }
 
+FORCEINLINE FVector FQuat::Vector() const
+{
+	return GetAxisX();
+}
+
 
 FORCEINLINE float FQuat::Error( const FQuat& Q1, const FQuat& Q2 )
 {
diff --git a/Engine/Source/Runtime/Core/Public/Math/Rotator.h b/Engine/Source/Runtime/Core/Public/Math/Rotator.h
--- a/Engine/Source/Runtime/Core/Public/Math/Rotator.h
+++ b/Engine/Source/Runtime/Core/Public/Math/Rotator.h
@@ -196,9 +196,9 @@ public:
 	FRotator GridSnap( const FRotator& RotGrid ) const;
 
 	/**
-	 * Convert a rotation into a vector facing in its direction.
+	 * Convert a rotation into a unit vector facing in its direction.
 	 *
-	 * @return Rotation as a direction vector.
+	 * @return Rotation as a unit direction vector.
 	 */
 	CORE_API FVector Vector() const;
 
diff --git a/Engine/Source/Runtime/Core/Public/Math/UnrealMath.h b/Engine/Source/Runtime/Core/Public/Math/UnrealMath.h
--- a/Engine/Source/Runtime/Core/Public/Math/UnrealMath.h
+++ b/Engine/Source/Runtime/Core/Public/Math/UnrealMath.h
@@ -90,6 +90,16 @@ inline FVector FVector::PointPlaneProject(const FVector& Point, const FVector& A
 	return Point - Plane.PlaneDot(Point) * Plane;
 }
 
+FORCEINLINE_DEBUGGABLE FRotator FVector::Rotation() const
+{
+	return ToOrientationRotator();
+}
+
+FORCEINLINE_DEBUGGABLE FRotator FVector4::Rotation() const
+{
+	return ToOrientationRotator();
+}
+
 inline FPlane FPlane::TransformBy( const FMatrix& M ) const
 {
 	const FMatrix tmpTA = M.TransposeAdjoint();
diff --git a/Engine/Source/Runtime/Core/Public/Math/Vector.h b/Engine/Source/Runtime/Core/Public/Math/Vector.h
--- a/Engine/Source/Runtime/Core/Public/Math/Vector.h
+++ b/Engine/Source/Runtime/Core/Public/Math/Vector.h
@@ -622,13 +622,34 @@ public:
 	FORCEINLINE FVector ProjectOnToNormal(const FVector& Normal ) const;
 
 	/**
-	 * Return the FRotator corresponding to the direction that the vector
-	 * is pointing in.  Sets Yaw and Pitch to the proper numbers, and sets
-	 * roll to zero because the roll can't be determined from a vector.
+	 * Return the FRotator orientation corresponding to the direction in which the vector points.
+	 * Sets Yaw and Pitch to the proper numbers, and sets Roll to zero because the roll can't be determined from a vector.
 	 *
-	 * @return The FRotator from the vector's direction.
+	 * @return FRotator from the Vector's direction, without any roll.
+	 * @see ToOrientationQuat()
 	 */
-	CORE_API FRotator Rotation() const;
+	CORE_API FRotator ToOrientationRotator() const;
+
+	/**
+	 * Return the Quaternion orientation corresponding to the direction in which the vector points.
+	 * Similar to the FRotator version, returns a result without roll such that it preserves the up vector.
+	 *
+	 * @note If you don't care about preserving the up vector and just want the most direct rotation, you can use the faster
+	 * 'FQuat::FindBetweenVectors(FVector::ForwardVector, YourVector)' or 'FQuat::FindBetweenNormals(...)' if you know the vector is of unit length.
+	 *
+	 * @return Quaternion from the Vector's direction, without any roll.
+	 * @see ToOrientationRotator(), FQuat::FindBetweenVectors()
+	 */
+	CORE_API FQuat ToOrientationQuat() const;
+
+	/**
+	 * Return the FRotator orientation corresponding to the direction in which the vector points.
+	 * Sets Yaw and Pitch to the proper numbers, and sets Roll to zero because the roll can't be determined from a vector.
+	 * @note Identical to 'ToOrientationRotator()' and preserved for legacy reasons.
+	 * @return FRotator from the Vector's direction.
+	 * @see ToOrientationRotator(), ToOrientationQuat()
+	 */
+	FRotator Rotation() const;
 
 	/**
 	 * Find good arbitrary axis vectors to represent U and V axes of a plane,
diff --git a/Engine/Source/Runtime/Core/Public/Math/Vector4.h b/Engine/Source/Runtime/Core/Public/Math/Vector4.h
--- a/Engine/Source/Runtime/Core/Public/Math/Vector4.h
+++ b/Engine/Source/Runtime/Core/Public/Math/Vector4.h
@@ -298,13 +298,26 @@ public:
 	FORCEINLINE FVector4 UnsafeNormal3() const;
 
 	/**
-	 * Return the FRotator corresponding to the direction that the vector
-	 * is pointing in.  Sets Yaw and Pitch to the proper numbers, and sets
-	 * roll to zero because the roll can't be determined from a vector.
-	 *
-	 * @return The FRotator of the vector's direction.
+	 * Return the FRotator orientation corresponding to the direction in which the vector points.
+	 * Sets Yaw and Pitch to the proper numbers, and sets roll to zero because the roll can't be determined from a vector.
+	 * @return FRotator from the Vector's direction.
+	 */
+	CORE_API FRotator ToOrientationRotator() const;
+
+	/**
+	 * Return the Quaternion orientation corresponding to the direction in which the vector points.
+	 * @return Quaternion from the Vector's direction.
+	 */
+	CORE_API FQuat ToOrientationQuat() const;
+
+	/**
+	 * Return the FRotator orientation corresponding to the direction in which the vector points.
+	 * Sets Yaw and Pitch to the proper numbers, and sets roll to zero because the roll can't be determined from a vector.
+	 * Identical to 'ToOrientationRotator()'.
+	 * @return FRotator from the Vector's direction.
+	 * @see ToOrientationRotator()
 	 */
-	CORE_API FRotator Rotation() const;
+	FRotator Rotation() const;
 
 	/**
 	 * Set all of the vectors coordinates.
