Fixed additional places where locale names from ICU were leaving the abstraction layer without converting underscores to dashes.

diff --git a/Engine/Source/Runtime/Core/Private/Internationalization/ICUCulture.cpp b/Engine/Source/Runtime/Core/Private/Internationalization/ICUCulture.cpp
--- a/Engine/Source/Runtime/Core/Private/Internationalization/ICUCulture.cpp
+++ b/Engine/Source/Runtime/Core/Private/Internationalization/ICUCulture.cpp
@@ -136,7 +136,9 @@ FString FCulture::FICUCultureImplementation::GetCanonicalName(const FString& Nam
 
 	UErrorCode ICUStatus = U_ZERO_ERROR;
 	uloc_canonicalize(TCHAR_TO_ANSI( *Name ), CanonicalName, MaximumNameLength, &ICUStatus);
-	return CanonicalName;
+	FString CanonicalNameString = CanonicalName;
+	CanonicalNameString.ReplaceInline(TEXT("_"), TEXT("-"));
+	return CanonicalNameString;
 }
 
 FString FCulture::FICUCultureImplementation::GetName() const
diff --git a/Engine/Source/Runtime/Core/Private/Internationalization/ICUInternationalization.cpp b/Engine/Source/Runtime/Core/Private/Internationalization/ICUInternationalization.cpp
--- a/Engine/Source/Runtime/Core/Private/Internationalization/ICUInternationalization.cpp
+++ b/Engine/Source/Runtime/Core/Private/Internationalization/ICUInternationalization.cpp
@@ -260,7 +260,9 @@ void FICUInternationalization::GetCultureNames(TArray<FString>& CultureNames) co
 	CultureNames.Reset(LocaleCount);
 	for (int32 i = 0; i < LocaleCount; ++i)
 	{
-		CultureNames.Add(AvailableLocales[i].getName());
+		FString CultureName = AvailableLocales[i].getName();
+		CultureName.ReplaceInline(TEXT("_"), TEXT("-"));
+		CultureNames.Add(CultureName);
 	}
 }
 
