#ue4 - Clean up UnrealMathTest FVector3 comparison.

diff --git a/Engine/Source/Runtime/Core/Private/Tests/Math/UnrealMathTest.cpp b/Engine/Source/Runtime/Core/Private/Tests/Math/UnrealMathTest.cpp
--- a/Engine/Source/Runtime/Core/Private/Tests/Math/UnrealMathTest.cpp
+++ b/Engine/Source/Runtime/Core/Private/Tests/Math/UnrealMathTest.cpp
@@ -105,7 +105,7 @@ bool TestVectorsEqual3( VectorRegister Vec0, VectorRegister Vec1, float Toleranc
  *
  * @return true if equal(ish)
  */
-bool TestVectorsEqual3_NoVec( const FVector& Vec0, const FVector& Vec1, float Tolerance = 0.0f)
+bool TestFVector3Equal( const FVector& Vec0, const FVector& Vec1, float Tolerance = 0.0f)
 {
 	GScratch[0] = Vec0.X;
 	GScratch[1] = Vec0.Y;
@@ -1200,11 +1200,11 @@ bool FVectorRegisterAbstractionTest::RunTest(const FString& Parameters)
 		FVector FV0, FV1;
 		FV0 = Rotator0.Vector();
 		FV1 = FRotationMatrix( Rotator0 ).GetScaledAxis( EAxis::X );
-		LogTest( TEXT("Test0 Rotator::Vector()"), TestVectorsEqual3_NoVec(FV0, FV1, 1e-6f));
+		LogTest( TEXT("Test0 Rotator::Vector()"), TestFVector3Equal(FV0, FV1, 1e-6f));
 		
 		FV0 = FRotationMatrix( Rotator0 ).GetScaledAxis( EAxis::X );
 		FV1 = FQuatRotationMatrix( Q0 ).GetScaledAxis( EAxis::X );
-		LogTest( TEXT("Test0 FQuatRotationMatrix"), TestVectorsEqual3_NoVec(FV0, FV1, 1e-5f));
+		LogTest( TEXT("Test0 FQuatRotationMatrix"), TestFVector3Equal(FV0, FV1, 1e-5f));
 
 		Rotator0 = FRotator(45.0f,  60.0f, 120.0f);
 		Q0 = Rotator0.Quaternion();
@@ -1213,15 +1213,15 @@ bool FVectorRegisterAbstractionTest::RunTest(const FString& Parameters)
 
 		FV0 = Rotator0.Vector();
 		FV1 = FRotationMatrix( Rotator0 ).GetScaledAxis( EAxis::X );
-		LogTest( TEXT("Test1 Rotator::Vector()"), TestVectorsEqual3_NoVec(FV0, FV1, 1e-6f));
+		LogTest( TEXT("Test1 Rotator::Vector()"), TestFVector3Equal(FV0, FV1, 1e-6f));
 
 		FV0 = FRotationMatrix( Rotator0 ).GetScaledAxis( EAxis::X );
 		FV1 = FQuatRotationMatrix( Q0 ).GetScaledAxis( EAxis::X );
-		LogTest(TEXT("Test1 FQuatRotationMatrix"), TestVectorsEqual3_NoVec(FV0, FV1, 1e-5f));
+		LogTest(TEXT("Test1 FQuatRotationMatrix"), TestFVector3Equal(FV0, FV1, 1e-5f));
 
 		FV0 = FRotationMatrix(FRotator::ZeroRotator).GetScaledAxis(EAxis::X);
 		FV1 = FQuatRotationMatrix(FQuat::Identity).GetScaledAxis(EAxis::X);
-		LogTest(TEXT("Test2 FQuatRotationMatrix"), TestVectorsEqual3_NoVec(FV0, FV1, 1e-6f));
+		LogTest(TEXT("Test2 FQuatRotationMatrix"), TestFVector3Equal(FV0, FV1, 1e-6f));
 	}
 
 	// Quat Rotation tests
@@ -1255,7 +1255,7 @@ bool FVectorRegisterAbstractionTest::RunTest(const FString& Parameters)
 		#define TEST_QUAT_ROTATE(_QIndex, _VIndex, _Quat, _Vec, _Func, _Tolerance) \
 		{ \
 			const FString _TestName = FString::Printf(TEXT("Test Quat%d: Vec%d: %s"), _QIndex, _VIndex, TEXT(#_Func)); \
-			LogTest( *_TestName, TestVectorsEqual3_NoVec(_Quat.RotateVector(_Vec), _Func(_Quat, _Vec), _Tolerance) ); \
+			LogTest( *_TestName, TestFVector3Equal(_Quat.RotateVector(_Vec), _Func(_Quat, _Vec), _Tolerance) ); \
 		}
 
 		// Test loop
@@ -1296,9 +1296,9 @@ bool FVectorRegisterAbstractionTest::RunTest(const FString& Parameters)
 					}
 					else
 					{
-						LogTest(TEXT("FindBetween: Old A->B"), RotAOld.Equals(BNorm));
-						LogTest(TEXT("FindBetween: New A->B (normal)"), RotANewNormal.Equals(BNorm));
-						LogTest(TEXT("FindBetween: New A->B (vector)"), RotANewVector.Equals(BNorm));
+						LogTest(TEXT("FindBetween: Old A->B"), TestFVector3Equal(RotAOld, BNorm, KINDA_SMALL_NUMBER));
+						LogTest(TEXT("FindBetween: New A->B (normal)"), TestFVector3Equal(RotANewNormal, BNorm, KINDA_SMALL_NUMBER));
+						LogTest(TEXT("FindBetween: New A->B (vector)"), TestFVector3Equal(RotANewVector, BNorm, KINDA_SMALL_NUMBER));
 					}
 				}
 			}
@@ -1319,8 +1319,8 @@ bool FVectorRegisterAbstractionTest::RunTest(const FString& Parameters)
 				const FVector Rotated1 = Q1.RotateVector(FVector::ForwardVector);
 				const FVector Rotated2 = R0.RotateVector(FVector::ForwardVector);
 
-				LogTest(TEXT("V.ToOrientationQuat() rotate"), Rotated0.Equals(Rotated1));
-				LogTest(TEXT("V.ToOrientationRotator() rotate"), Rotated0.Equals(Rotated2));
+				LogTest(TEXT("V.ToOrientationQuat() rotate"), TestFVector3Equal(Rotated0, Rotated1, KINDA_SMALL_NUMBER));
+				LogTest(TEXT("V.ToOrientationRotator() rotate"), TestFVector3Equal(Rotated0, Rotated2, KINDA_SMALL_NUMBER));
 			}
 		}
 	}
