#ue4 - Fix param comment on FQuatRotationMatrix

diff --git a/Engine/Source/Runtime/Core/Public/Math/QuatRotationTranslationMatrix.h b/Engine/Source/Runtime/Core/Public/Math/QuatRotationTranslationMatrix.h
--- a/Engine/Source/Runtime/Core/Public/Math/QuatRotationTranslationMatrix.h
+++ b/Engine/Source/Runtime/Core/Public/Math/QuatRotationTranslationMatrix.h
@@ -33,7 +33,6 @@ public:
 	/** Constructor
 	*
 	* @param Q rotation
-	* @param Origin translation to apply
 	*/
 	FQuatRotationMatrix(const FQuat& Q)
 		: FQuatRotationTranslationMatrix(Q, FVector::ZeroVector)
