#ue4 - Perf: VInterpNormalRotationTo() now uses faster FQuat::RotateVector() rather than converting to matrix first.

diff --git a/Engine/Source/Runtime/Core/Private/Math/UnrealMath.cpp b/Engine/Source/Runtime/Core/Private/Math/UnrealMath.cpp
--- a/Engine/Source/Runtime/Core/Private/Math/UnrealMath.cpp
+++ b/Engine/Source/Runtime/Core/Private/Math/UnrealMath.cpp
@@ -2149,7 +2149,7 @@ CORE_API FVector FMath::VInterpNormalRotationTo(const FVector& Current, const FV
 	{
 		DeltaAngle = FMath::Clamp(DeltaAngle, -RotationStepRadians, RotationStepRadians);
 		DeltaQuat = FQuat(DeltaAxis, DeltaAngle);
-		return FQuatRotationTranslationMatrix( DeltaQuat, FVector::ZeroVector ).TransformVector(Current);
+		return DeltaQuat.RotateVector(Current);
 	}
 	return Target;
 }
