Fixed apps still responding controller input once a controller has been removed.  The system now sends a final event with default states once the controller is unplugged.

https://jira.ol.epicgames.net/browse/UE-18999

diff --git a/Engine/Source/Runtime/Core/Private/Windows/XInputInterface.cpp b/Engine/Source/Runtime/Core/Private/Windows/XInputInterface.cpp
--- a/Engine/Source/Runtime/Core/Private/Windows/XInputInterface.cpp
+++ b/Engine/Source/Runtime/Core/Private/Windows/XInputInterface.cpp
@@ -104,9 +104,13 @@ void XInputInterface::SendControllerEvents()
 			XINPUT_STATE XInputState;
 			FMemory::Memzero( &XInputState, sizeof(XINPUT_STATE) );
 
-			ControllerState.bIsConnected = ( XInputGetState( ControllerIndex, &XInputState ) == ERROR_SUCCESS ) ? true : false;
+			bool bWasConnected = ControllerState.bIsConnected;
 
-			if( ControllerState.bIsConnected )
+			ControllerState.bIsConnected = ( XInputGetState( ControllerIndex, &XInputState ) == ERROR_SUCCESS ) ? true : false;
+			
+			// If the controller is connected send events or if the controller was connected send a final event with default states so that 
+			// the game doesnt think that controller buttons are still held down
+			if( ControllerState.bIsConnected || bWasConnected )
 			{
 				bool CurrentStates[MAX_NUM_CONTROLLER_BUTTONS] = {0};
 		
