Fix crashes due to empty string

Merging 2632414

//depot/UE4-UT/...

to //depot/UE4/...

diff --git a/Engine/Source/Runtime/Core/Private/Serialization/AsyncIOSystemBase.cpp b/Engine/Source/Runtime/Core/Private/Serialization/AsyncIOSystemBase.cpp
--- a/Engine/Source/Runtime/Core/Private/Serialization/AsyncIOSystemBase.cpp
+++ b/Engine/Source/Runtime/Core/Private/Serialization/AsyncIOSystemBase.cpp
@@ -60,7 +60,7 @@ uint64 FAsyncIOSystemBase::QueueIORequest(
 	IORequest.RequestIndex				= RequestIndex++;
 	IORequest.FileSortKey				= INDEX_NONE;
 	IORequest.FileName					= FileName;
-	IORequest.FileNameHash				= FCrc::StrCrc32<TCHAR>(FileName.ToLower().GetCharArray().GetData());
+	IORequest.FileNameHash				= FCrc::StrCrc32<TCHAR>(*FileName.ToLower());
 	IORequest.Offset					= Offset;
 	IORequest.Size						= Size;
 	IORequest.UncompressedSize			= UncompressedSize;
@@ -103,7 +103,7 @@ uint64 FAsyncIOSystemBase::QueueDestroyHandleRequest(const FString& FileName)
 	FAsyncIORequest IORequest;
 	IORequest.RequestIndex				= RequestIndex++;
 	IORequest.FileName					= FileName;
-	IORequest.FileNameHash				= FCrc::StrCrc32<TCHAR>(FileName.ToLower().GetCharArray().GetData());
+	IORequest.FileNameHash				= FCrc::StrCrc32<TCHAR>(*FileName.ToLower());
 	IORequest.Priority					= AIOP_MIN;
 	IORequest.bIsDestroyHandleRequest	= true;
 
@@ -416,7 +416,7 @@ IFileHandle* FAsyncIOSystemBase::GetCachedFileHandle( const FString& FileName )
 		// Make sure it's valid before caching and using it.
 		if( FileHandle )
 		{
-			NameHashToHandleMap.Add(FCrc::StrCrc32<TCHAR>(FileName.ToLower().GetCharArray().GetData()), FileHandle);
+			NameHashToHandleMap.Add(FCrc::StrCrc32<TCHAR>(*FileName.ToLower()), FileHandle);
 		}
 	}
 
@@ -425,7 +425,7 @@ IFileHandle* FAsyncIOSystemBase::GetCachedFileHandle( const FString& FileName )
 
 IFileHandle* FAsyncIOSystemBase::FindCachedFileHandle( const FString& FileName )
 {
-	return FindCachedFileHandle(FCrc::StrCrc32<TCHAR>(FileName.ToLower().GetCharArray().GetData()));
+	return FindCachedFileHandle(FCrc::StrCrc32<TCHAR>(*FileName.ToLower()));
 }
 
 IFileHandle* FAsyncIOSystemBase::FindCachedFileHandle(const uint32 FileNameHash)
