Linux: ungrab all input before breaking into the debugger.

diff --git a/Engine/Source/Runtime/Core/Private/Linux/LinuxPlatformMisc.cpp b/Engine/Source/Runtime/Core/Private/Linux/LinuxPlatformMisc.cpp
--- a/Engine/Source/Runtime/Core/Private/Linux/LinuxPlatformMisc.cpp
+++ b/Engine/Source/Runtime/Core/Private/Linux/LinuxPlatformMisc.cpp
@@ -710,6 +710,21 @@ bool FLinuxPlatformMisc::IsDebuggerPresent()
 }
 #endif // !UE_BUILD_SHIPPING
 
+#if !UE_BUILD_SHIPPING
+void FLinuxPlatformMisc::UngrabAllInput()
+{
+	SDL_Window * GrabbedWindow = SDL_GetGrabbedWindow();
+	if (GrabbedWindow)
+	{
+		SDL_SetWindowGrab(GrabbedWindow, SDL_FALSE);
+		SDL_SetKeyboardGrab(GrabbedWindow, SDL_FALSE);
+	}
+
+	SDL_CaptureMouse(SDL_FALSE);
+}
+
+#endif // !UE_BUILD_SHIPPING
+
 bool FLinuxPlatformMisc::HasBeenStartedRemotely()
 {
 	static bool bHaveAnswer = false;
diff --git a/Engine/Source/Runtime/Core/Public/Linux/LinuxPlatformMisc.h b/Engine/Source/Runtime/Core/Public/Linux/LinuxPlatformMisc.h
--- a/Engine/Source/Runtime/Core/Public/Linux/LinuxPlatformMisc.h
+++ b/Engine/Source/Runtime/Core/Public/Linux/LinuxPlatformMisc.h
@@ -31,6 +31,7 @@ struct CORE_API FLinuxPlatformMisc : public FGenericPlatformMisc
 	{
 		if( IsDebuggerPresent() )
 		{
+			UngrabAllInput();
 			raise(SIGTRAP);
 		}
 	}
@@ -132,6 +133,13 @@ struct CORE_API FLinuxPlatformMisc : public FGenericPlatformMisc
 	 */
 	static bool PlatformInitMultimedia();
 
+#if !UE_BUILD_SHIPPING	// only in non-shipping because we break into the debugger in non-shipping builds only
+	/**
+	 * Ungrabs input (useful before breaking into debugging)
+	 */
+	static void UngrabAllInput();
+#endif // !UE_BUILD_SHIPPING
+
 	/**
 	 * Returns whether the program has been started remotely (e.g. over SSH)
 	 */
