Linux: fix crash when saving assets with localized names (UE-18253).

- Removed errorneous duplicate wchar_t->UTF-8 conversion.

diff --git a/Engine/Source/Runtime/Core/Private/Linux/LinuxPlatformFile.cpp b/Engine/Source/Runtime/Core/Private/Linux/LinuxPlatformFile.cpp
--- a/Engine/Source/Runtime/Core/Private/Linux/LinuxPlatformFile.cpp
+++ b/Engine/Source/Runtime/Core/Private/Linux/LinuxPlatformFile.cpp
@@ -770,7 +770,7 @@ FString FLinuxPlatformFile::GetFilenameOnDisk(const TCHAR* Filename)
 IFileHandle* FLinuxPlatformFile::OpenRead(const TCHAR* Filename, bool bAllowWrite)
 {
 	FString MappedToName;
-	int32 Handle = GCaseInsensMapper.OpenCaseInsensitiveRead(TCHAR_TO_UTF8(*NormalizeFilename(Filename)), MappedToName);
+	int32 Handle = GCaseInsensMapper.OpenCaseInsensitiveRead(NormalizeFilename(Filename), MappedToName);
 	if (Handle != -1)
 	{
 		return new FFileHandleLinux(Handle, *MappedToName, true);
