Added bool to Expose_TFormatSpecifier

diff --git a/Engine/Source/Runtime/Core/Public/Templates/UnrealTypeTraits.h b/Engine/Source/Runtime/Core/Public/Templates/UnrealTypeTraits.h
--- a/Engine/Source/Runtime/Core/Public/Templates/UnrealTypeTraits.h
+++ b/Engine/Source/Runtime/Core/Public/Templates/UnrealTypeTraits.h
@@ -142,6 +142,7 @@ struct TFormatSpecifier<type> \
 	} \
 };
 
+Expose_TFormatSpecifier(bool, "%i")
 Expose_TFormatSpecifier(uint8, "%u")
 Expose_TFormatSpecifier(uint16, "%u")
 Expose_TFormatSpecifier(uint32, "%u")
