Box2d is a struct, not a class.
FNativeClassHeaderGenerator and UStructProperty::GetCPPTypeForwardDeclaration use "struct".

diff --git a/Engine/Source/Runtime/Core/Public/Math/Box2D.h b/Engine/Source/Runtime/Core/Public/Math/Box2D.h
--- a/Engine/Source/Runtime/Core/Public/Math/Box2D.h
+++ b/Engine/Source/Runtime/Core/Public/Math/Box2D.h
@@ -6,7 +6,7 @@
 /**
  * Implements a rectangular 2D Box.
  */
-class FBox2D
+struct FBox2D
 {
 public:
 	/** Holds the box's minimum point. */
