[INTEGRATE] [Mac] Don't forward a left mouse button down event if inside the resize edge of the window, since the OSX window resizing will take over and we'll never get the corresponding mouse up event.
OPP-3803
Merging CL#2629242 using UE4-To-UE4-LauncherDev

diff --git a/Engine/Source/Runtime/Core/Private/Mac/MacApplication.cpp b/Engine/Source/Runtime/Core/Private/Mac/MacApplication.cpp
--- a/Engine/Source/Runtime/Core/Private/Mac/MacApplication.cpp
+++ b/Engine/Source/Runtime/Core/Private/Mac/MacApplication.cpp
@@ -667,8 +667,9 @@ void FMacApplication::ProcessMouseMovedEvent(const FDeferredMacEvent& Event, TSh
 {
 	if (EventWindow.IsValid() && EventWindow->IsRegularWindow())
 	{
-		bool IsMouseOverTitleBar = false;
-		const bool IsMovable = IsWindowMovable(EventWindow.ToSharedRef(), &IsMouseOverTitleBar);
+		const EWindowZone::Type Zone = GetCurrentWindowZone(EventWindow.ToSharedRef());
+		bool IsMouseOverTitleBar = (Zone == EWindowZone::TitleBar);
+		const bool IsMovable = IsMouseOverTitleBar || IsEdgeZone(Zone);
 		[EventWindow->GetWindowHandle() setMovable:IsMovable];
 		[EventWindow->GetWindowHandle() setMovableByWindowBackground:IsMouseOverTitleBar];
 	}
@@ -784,11 +785,10 @@ void FMacApplication::ProcessMouseDownEvent(const FDeferredMacEvent& Event, TSha
 
 	if (EventWindow.IsValid())
 	{
+		const EWindowZone::Type Zone = GetCurrentWindowZone(EventWindow.ToSharedRef());
 		if (Button == LastPressedMouseButton && (Event.ClickCount % 2) == 0)
 		{
-			bool IsMouseOverTitleBar = false;
-			const bool IsMovable = IsWindowMovable(EventWindow.ToSharedRef(), &IsMouseOverTitleBar);
-			if (IsMouseOverTitleBar)
+			if (Zone == EWindowZone::TitleBar)
 			{
 				const bool bShouldMinimize = [[NSUserDefaults standardUserDefaults] boolForKey:@"AppleMiniaturizeOnDoubleClick"];
 				FCocoaWindow* WindowHandle = EventWindow->GetWindowHandle();
@@ -806,7 +806,8 @@ void FMacApplication::ProcessMouseDownEvent(const FDeferredMacEvent& Event, TSha
 				MessageHandler->OnMouseDoubleClick(EventWindow, Button);
 			}
 		}
-		else
+		// Only forward left mouse button down events if it's not inside the resize edge zone of a normal resizable window.
+		else if (!EventWindow->IsRegularWindow() || Button != EMouseButtons::Left || !IsEdgeZone(Zone))
 		{
 			MessageHandler->OnMouseDown(EventWindow, Button);
 		}
@@ -1208,19 +1209,16 @@ NSScreen* FMacApplication::FindScreenByPoint(int32 X, int32 Y) const
 	return TargetScreen;
 }
 
-bool FMacApplication::IsWindowMovable(TSharedRef<FMacWindow> Window, bool* OutMovableByBackground) const
+EWindowZone::Type FMacApplication::GetCurrentWindowZone(const TSharedRef<FMacWindow>& Window) const
 {
-	if (OutMovableByBackground)
-	{
-		*OutMovableByBackground = false;
-	}
-
 	const FVector2D CursorPos = ((FMacCursor*)Cursor.Get())->GetPosition();
 	const int32 LocalMouseX = CursorPos.X - Window->PositionX;
 	const int32 LocalMouseY = CursorPos.Y - Window->PositionY;
+	return MessageHandler->GetWindowZoneForPoint(Window, LocalMouseX, LocalMouseY);
+}
 
-	const EWindowZone::Type Zone = MessageHandler->GetWindowZoneForPoint(Window, LocalMouseX, LocalMouseY);
-	const bool IsMouseOverTitleBar = Zone == EWindowZone::TitleBar;
+bool FMacApplication::IsEdgeZone(EWindowZone::Type Zone) const
+{
 	switch (Zone)
 	{
 		case EWindowZone::NotInWindow:
@@ -1234,11 +1232,6 @@ bool FMacApplication::IsWindowMovable(TSharedRef<FMacWindow> Window, bool* OutMo
 		case EWindowZone::BottomRightBorder:
 			return true;
 		case EWindowZone::TitleBar:
-			if (OutMovableByBackground)
-			{
-				*OutMovableByBackground = true;
-			}
-			return true;
 		case EWindowZone::ClientArea:
 		case EWindowZone::MinimizeButton:
 		case EWindowZone::MaximizeButton:
diff --git a/Engine/Source/Runtime/Core/Public/Mac/MacApplication.h b/Engine/Source/Runtime/Core/Public/Mac/MacApplication.h
--- a/Engine/Source/Runtime/Core/Public/Mac/MacApplication.h
+++ b/Engine/Source/Runtime/Core/Public/Mac/MacApplication.h
@@ -220,7 +220,8 @@ private:
 
 	FCocoaWindow* FindEventWindow(NSEvent* CocoaEvent) const;
 	NSScreen* FindScreenByPoint(int32 X, int32 Y) const;
-	bool IsWindowMovable(TSharedRef<FMacWindow> Window, bool* OutMovableByBackground) const;
+	EWindowZone::Type GetCurrentWindowZone(const TSharedRef<FMacWindow>& Window) const;
+	bool IsEdgeZone(EWindowZone::Type Zone) const;
 	bool IsPrintableKey(uint32 Character) const;
 	TCHAR ConvertChar(TCHAR Character) const;
 	TCHAR TranslateCharCode(TCHAR CharCode, uint32 KeyCode) const;
