Crash reporter improvements - More properties added, replaced a few magic strings with constants

diff --git a/Engine/Source/Runtime/Core/Private/GenericPlatform/GenericPlatformCrashContext.cpp b/Engine/Source/Runtime/Core/Private/GenericPlatform/GenericPlatformCrashContext.cpp
--- a/Engine/Source/Runtime/Core/Private/GenericPlatform/GenericPlatformCrashContext.cpp
+++ b/Engine/Source/Runtime/Core/Private/GenericPlatform/GenericPlatformCrashContext.cpp
@@ -6,8 +6,19 @@
 #include "EngineVersion.h"
 #include "EngineBuildSettings.h"
 
+/*-----------------------------------------------------------------------------
+	FGenericCrashContext
+-----------------------------------------------------------------------------*/
+
 const ANSICHAR* FGenericCrashContext::CrashContextRuntimeXMLNameA = "CrashContext.runtime-xml";
 const TCHAR* FGenericCrashContext::CrashContextRuntimeXMLNameW = TEXT( "CrashContext.runtime-xml" );
+
+const FString FGenericCrashContext::CrashContextExtension = TEXT( ".runtime-xml" );
+const FString FGenericCrashContext::RuntimePropertiesTag = TEXT( "RuntimeProperties" );
+const FString FGenericCrashContext::PlatformPropertiesTag = TEXT( "PlatformProperties" );
+const FString FGenericCrashContext::UE4MinidumpName = TEXT( "UE4Minidump.dmp" );
+const FString FGenericCrashContext::NewLineTag = TEXT( "&nl;" );
+
 bool FGenericCrashContext::bIsInitialized = false;
 FPlatformMemoryStats FGenericCrashContext::CrashMemoryStats = FPlatformMemoryStats();
 
@@ -17,6 +28,7 @@ namespace NCachedCrashContextProperties
 	static bool bIsPerforceBuild;
 	static bool bIsSourceDistribution;
 	static bool bIsUE4Release;
+	static FString GameName;
 	static FString ExecutableName;
 	static FString PlatformName;
 	static FString PlatformNameIni;
@@ -45,6 +57,7 @@ void FGenericCrashContext::Initialize()
 	NCachedCrashContextProperties::bIsSourceDistribution = FEngineBuildSettings::IsSourceDistribution();
 	NCachedCrashContextProperties::bIsUE4Release = FApp::IsEngineInstalled();
 
+	NCachedCrashContextProperties::GameName = FString::Printf( TEXT("UE4-%s"), FApp::GetGameName() );
 	NCachedCrashContextProperties::ExecutableName = FPlatformProcess::ExecutableName();
 	NCachedCrashContextProperties::PlatformName = FPlatformProperties::PlatformName();
 	NCachedCrashContextProperties::PlatformNameIni = FPlatformProperties::IniPlatformName();
@@ -98,17 +111,18 @@ void FGenericCrashContext::SerializeContentToBuffer()
 	// https://www.securecoding.cert.org/confluence/display/seccode/SIG30-C.+Call+only+asynchronous-safe+functions+within+signal+handlers
 	AddHeader();
 
-	BeginSection( TEXT( "RuntimeProperties" ) );
-	AddCrashProperty( TEXT( "Version" ), (int32)ECrashDescVersions::VER_2_AddedNewProperties );
+	BeginSection( *RuntimePropertiesTag );
+	AddCrashProperty( TEXT( "CrashVersion" ), ECrashDescVersions::VER_3_CrashContext );
+	AddCrashProperty( TEXT( "CrashGUID" ), *NCachedCrashContextProperties::CrashGUID );
 	AddCrashProperty( TEXT( "ProcessId" ), FPlatformProcess::GetCurrentProcessId() );
-	AddCrashProperty( TEXT( "IsInternalBuild" ), (int32)NCachedCrashContextProperties::bIsInternalBuild );
-	AddCrashProperty( TEXT( "IsPerforceBuild" ), (int32)NCachedCrashContextProperties::bIsPerforceBuild );
-	AddCrashProperty( TEXT( "IsSourceDistribution" ), (int32)NCachedCrashContextProperties::bIsSourceDistribution );
+	AddCrashProperty( TEXT( "IsInternalBuild" ), NCachedCrashContextProperties::bIsInternalBuild );
+	AddCrashProperty( TEXT( "IsPerforceBuild" ), NCachedCrashContextProperties::bIsPerforceBuild );
+	AddCrashProperty( TEXT( "IsSourceDistribution" ), NCachedCrashContextProperties::bIsSourceDistribution );
 
-	AddCrashProperty( TEXT( "SecondsSinceStart" ), (int32)NCachedCrashContextProperties::SecondsSinceStart );
+	AddCrashProperty( TEXT( "SecondsSinceStart" ), NCachedCrashContextProperties::SecondsSinceStart );
 
 	// Add common crash properties.
-	AddCrashProperty( TEXT( "GameName" ), FApp::GetGameName() );
+	AddCrashProperty( TEXT( "GameName" ), *NCachedCrashContextProperties::GameName );
 	AddCrashProperty( TEXT( "ExecutableName" ), *NCachedCrashContextProperties::ExecutableName );
 	AddCrashProperty( TEXT( "BuildConfiguration" ), EBuildConfigurations::ToString( FApp::GetBuildConfiguration() ) );
 
@@ -118,9 +132,9 @@ void FGenericCrashContext::SerializeContentToBuffer()
 	AddCrashProperty( TEXT( "EngineVersion" ), *GEngineVersion.ToString() );
 	AddCrashProperty( TEXT( "CommandLine" ), FCommandLine::IsInitialized() ? FCommandLine::Get() : TEXT( "" ) );
 	AddCrashProperty( TEXT( "LanguageLCID" ), FInternationalization::Get().GetCurrentCulture()->GetLCID() );
-	AddCrashProperty( TEXT( "DefaultLocale" ), *NCachedCrashContextProperties::DefaultLocale );
+	AddCrashProperty( TEXT( "AppDefaultLocale" ), *NCachedCrashContextProperties::DefaultLocale );
 
-	AddCrashProperty( TEXT( "IsUE4Release" ), (int32)NCachedCrashContextProperties::bIsUE4Release );
+	AddCrashProperty( TEXT( "IsUE4Release" ), NCachedCrashContextProperties::bIsUE4Release );
 
 	// Remove periods from user names to match AutoReporter user names
 	// The name prefix is read by CrashRepository.AddNewCrash in the website code
@@ -183,12 +197,12 @@ void FGenericCrashContext::SerializeContentToBuffer()
 	//Architecture
 	//CrashedModuleName
 	//LoadedModules
-	EndSection( TEXT( "RuntimeProperties" ) );
+	EndSection( *RuntimePropertiesTag );
 	
 	// Add platform specific properties.
-	BeginSection( TEXT("PlatformProperties") );
+	BeginSection( *PlatformPropertiesTag );
 	AddPlatformSpecificProperties();
-	EndSection( TEXT( "PlatformProperties" ) );
+	EndSection( *PlatformPropertiesTag );
 
 	AddFooter();
 }
@@ -266,7 +280,11 @@ FString FGenericCrashContext::EscapeXMLString( const FString& Text )
 		.Replace( TEXT( "\"" ), TEXT( "&quot;" ) )
 		.Replace( TEXT( "'" ), TEXT( "&apos;" ) )
 		.Replace( TEXT( "<" ), TEXT( "&lt;" ) )
-		.Replace( TEXT( ">" ), TEXT( "&gt;" ) );
+		.Replace( TEXT( ">" ), TEXT( "&gt;" ) )
+		// Replace newline for FXMLFile.
+		.Replace( TEXT( "\n" ), *NewLineTag )
+		// Ignore return carriage.
+		.Replace( TEXT( "\r" ), TEXT("") );
 }
 
 FString FGenericCrashContext::UnescapeXMLString( const FString& Text )
@@ -276,7 +294,8 @@ FString FGenericCrashContext::UnescapeXMLString( const FString& Text )
 		.Replace( TEXT( "&quot;" ), TEXT( "\"" ) )
 		.Replace( TEXT( "&apos;" ), TEXT( "'" ) )
 		.Replace( TEXT( "&lt;" ), TEXT( "<" ) )
-		.Replace( TEXT( "&gt;" ), TEXT( ">" ) );
+		.Replace( TEXT( "&gt;" ), TEXT( ">" ) )
+		.Replace( *NewLineTag, TEXT( "\n" ) );
 }
 
 FProgramCounterSymbolInfoEx::FProgramCounterSymbolInfoEx( FString InModuleName, FString InFunctionName, FString InFilename, uint32 InLineNumber, uint64 InSymbolDisplacement, uint64 InOffsetInModule, uint64 InProgramCounter ) :
diff --git a/Engine/Source/Runtime/Core/Private/Windows/WindowsPlatformCrashContext.cpp b/Engine/Source/Runtime/Core/Private/Windows/WindowsPlatformCrashContext.cpp
--- a/Engine/Source/Runtime/Core/Private/Windows/WindowsPlatformCrashContext.cpp
+++ b/Engine/Source/Runtime/Core/Private/Windows/WindowsPlatformCrashContext.cpp
@@ -301,7 +301,7 @@ int32 ReportCrashUsingCrashReportClient(EXCEPTION_POINTERS* ExceptionInfo, const
 				CrashContext.SerializeAsXML( *CrashContextXMLPath );
 				WerReportAddFile( ReportHandle, *CrashContextXMLPath, WerFileTypeOther, WER_FILE_ANONYMOUS_DATA );
 
-				const FString MinidumpFileName = FPaths::Combine( *FPaths::GameLogDir(), *CrashContext.GetUniqueCrashName(), TEXT("UE4Minidump.dmp") );
+				const FString MinidumpFileName = FPaths::Combine( *FPaths::GameLogDir(), *CrashContext.GetUniqueCrashName(), *FGenericCrashContext::UE4MinidumpName );
 				if( WriteMinidump( *MinidumpFileName, ExceptionInfo ) )
 				{
 					WerReportAddFile( ReportHandle, *MinidumpFileName, WerFileTypeMinidump, WER_FILE_ANONYMOUS_DATA );
diff --git a/Engine/Source/Runtime/Core/Public/GenericPlatform/GenericPlatformCrashContext.h b/Engine/Source/Runtime/Core/Public/GenericPlatform/GenericPlatformCrashContext.h
--- a/Engine/Source/Runtime/Core/Public/GenericPlatform/GenericPlatformCrashContext.h
+++ b/Engine/Source/Runtime/Core/Public/GenericPlatform/GenericPlatformCrashContext.h
@@ -45,6 +45,9 @@ enum class ECrashDescVersions : int32
 
 	/** Added misc properties (CPU,GPU,OS,etc), memory related stats and platform specific properties as generic payload. */
 	VER_2_AddedNewProperties,
+
+	/** Using crash context when available. */
+	VER_3_CrashContext = 3,
 };
 
 /** Enumerates crash dump modes. */
@@ -57,6 +60,23 @@ enum class ECrashDumpMode : int32
 	FullDump = 1,
 };
 
+/*-----------------------------------------------------------------------------
+	LexicalConversion
+-----------------------------------------------------------------------------*/
+
+namespace LexicalConversion
+{
+	inline FString ToString( const ECrashDescVersions& Value )
+	{
+		return TTypeToString<int32>::ToString( int32( Value ) );
+	}
+
+	inline FString ToString( const ECrashDumpMode& Value )
+	{
+		return TTypeToString<int32>::ToString( int32( Value ) );
+	}
+}
+
 /**
  *	Contains a runtime crash's properties that are common for all platforms.
  *	This may change in the future.
@@ -70,9 +90,16 @@ public:
 	* exception and use in crash reporting.
 	*/
 	static FPlatformMemoryStats CrashMemoryStats;
+	
 	static const ANSICHAR* CrashContextRuntimeXMLNameA;
 	static const TCHAR* CrashContextRuntimeXMLNameW;
 
+	static const FString CrashContextExtension;
+	static const FString RuntimePropertiesTag;
+	static const FString PlatformPropertiesTag;
+	static const FString UE4MinidumpName;
+	static const FString NewLineTag;
+
 	/** Initializes crash context related platform specific data that can be impossible to obtain after a crash. */
 	static void Initialize();
 
@@ -121,6 +148,12 @@ public:
 		AddCrashProperty( PropertyName, *TTypeToString<Type>::ToString( Value ) );
 	}
 
+	/** Escapes a specified XML string, naive implementation. */
+	static FString EscapeXMLString( const FString& Text );
+
+	/** Unescapes a specified XML string, naive implementation. */
+	static FString UnescapeXMLString( const FString& Text );
+
 private:
 	/** Serializes platform specific properties to the buffer. */
 	virtual void AddPlatformSpecificProperties();
@@ -134,12 +167,6 @@ private:
 	void BeginSection( const TCHAR* SectionName );
 	void EndSection( const TCHAR* SectionName );
 
-	/** Escapes a specified XML string, naive implementation. */
-	static FString EscapeXMLString( const FString& Text );
-
-	/** Unescapes a specified XML string, naive implementation. */
-	static FString UnescapeXMLString( const FString& Text );
-
 	/**	Whether the Initialize() has been called */
 	static bool bIsInitialized;
 
