Converted the UE4 Windows filesysytem implementation to use Windows API functions for getting file stat data

This yields better performance than using the POSIX wrapper functions, and will also allow us to use the stat data available while walking the directories (see IterateDirectory) without having issues with the two timestamps having different precisions.

Note: The new timestamps are more precise (they have 100-nanosecond precision), so this will result in anything that caches using only timestamps (such as the asset registry) performing a full re-scan. The auto-reimport is immune to this as it also tests the MD5 of the file.

diff --git a/Engine/Source/Runtime/Core/Private/Windows/WindowsPlatformFile.cpp b/Engine/Source/Runtime/Core/Private/Windows/WindowsPlatformFile.cpp
--- a/Engine/Source/Runtime/Core/Private/Windows/WindowsPlatformFile.cpp
+++ b/Engine/Source/Runtime/Core/Private/Windows/WindowsPlatformFile.cpp
@@ -4,8 +4,9 @@
 #include <sys/utime.h>
 
 
-// make an FTimeSpan object that represents the "epoch" for time_t (from a _stat struct)
-const FDateTime WindowsEpoch(1970, 1, 1);
+// make an FTimeSpan object that represents the "epoch" for FILETIME
+// FILETIME contains a 64-bit value representing the number of 100-nanosecond intervals since January 1, 1601 (UTC)
+const FDateTime WindowsEpoch(1601, 1, 1);
 
 #include "AllowWindowsPlatformTypes.h"
 	namespace FileConstants
@@ -14,6 +15,29 @@ const FDateTime WindowsEpoch(1970, 1, 1);
 	}
 #include "HideWindowsPlatformTypes.h"
 
+namespace
+{
+	FORCEINLINE FDateTime WindowsFileTimeToUEDateTime(const FILETIME& InFileTime)
+	{
+		ULARGE_INTEGER FileTime64;
+		FileTime64.LowPart = InFileTime.dwLowDateTime;
+		FileTime64.HighPart = InFileTime.dwHighDateTime;
+
+		return WindowsEpoch + FTimespan(static_cast<int64>(FileTime64.QuadPart));
+	}
+
+	FORCEINLINE FILETIME UEDateTimeToWindowsFileTime(const FDateTime& InDateTime)
+	{
+		ULARGE_INTEGER FileTime64;
+		FileTime64.QuadPart = (InDateTime - WindowsEpoch).GetTicks();
+
+		FILETIME FileTime;
+		FileTime.dwLowDateTime = FileTime64.LowPart;
+		FileTime.dwHighDateTime = FileTime64.HighPart;
+		return FileTime;
+	}
+}
+
 /**
  * This file reader uses overlapped i/o and double buffering to asynchronously read from files
  */
@@ -514,46 +538,42 @@ public:
 
 	virtual FDateTime GetTimeStamp(const TCHAR* Filename) override
 	{
-		// get file times
-		struct _stati64 FileInfo;
-		if(_wstati64(*NormalizeFilename(Filename), &FileInfo))
+		WIN32_FILE_ATTRIBUTE_DATA Info;
+		if (GetFileAttributesExW(*NormalizeFilename(Filename), GetFileExInfoStandard, &Info))
 		{
-			return FDateTime::MinValue();
+			return WindowsFileTimeToUEDateTime(Info.ftLastWriteTime);
 		}
 
-		// convert _stat time to FDateTime
-		FTimespan TimeSinceEpoch(0, 0, FileInfo.st_mtime);
-		return WindowsEpoch + TimeSinceEpoch;
+		return FDateTime::MinValue();
 	}
 
 	virtual void SetTimeStamp(const TCHAR* Filename, FDateTime DateTime) override
 	{
-		// get file times
-		struct _stati64 FileInfo;
-		if(_wstati64(*NormalizeFilename(Filename), &FileInfo))
+		HANDLE Handle = CreateFileW(*NormalizeFilename(Filename), FILE_WRITE_ATTRIBUTES, FILE_SHARE_READ, nullptr, OPEN_EXISTING, FILE_ATTRIBUTE_NORMAL | FILE_FLAG_OVERLAPPED, nullptr);
+		if (Handle != INVALID_HANDLE_VALUE)
+		{
+			const FILETIME ModificationFileTime = UEDateTimeToWindowsFileTime(DateTime);
+			if (!SetFileTime(Handle, nullptr, nullptr, &ModificationFileTime))
+			{
+				UE_LOG(LogTemp, Warning, TEXT("SetTimeStamp: Failed to SetFileTime on %s"), Filename);
+			}
+			CloseHandle(Handle);
+		}
+		else
 		{
-			return;
+			UE_LOG(LogTemp, Warning, TEXT("SetTimeStamp: Failed to open file %s"), Filename);
 		}
-
-		// change the modification time only
-		struct _utimbuf Times;
-		Times.actime = FileInfo.st_atime;
-		Times.modtime = (DateTime - WindowsEpoch).GetTotalSeconds();
-		_wutime(Filename, &Times);
 	}
 
 	virtual FDateTime GetAccessTimeStamp(const TCHAR* Filename) override
 	{
-		// get file times
-		struct _stati64 FileInfo;
-		if(_wstati64(*NormalizeFilename(Filename), &FileInfo))
+		WIN32_FILE_ATTRIBUTE_DATA Info;
+		if (GetFileAttributesExW(*NormalizeFilename(Filename), GetFileExInfoStandard, &Info))
 		{
-			return FDateTime::MinValue();
+			return WindowsFileTimeToUEDateTime(Info.ftLastAccessTime);
 		}
 
-		// convert _stat time to FDateTime
-		FTimespan TimeSinceEpoch(0, 0, FileInfo.st_atime);
-		return WindowsEpoch + TimeSinceEpoch;
+		return FDateTime::MinValue();
 	}
 
 	virtual FString GetFilenameOnDisk(const TCHAR* Filename) override
