Added functions to access low-level file stat data as a single file-system request

You can now use GetStatData to get all of the stat data for a file (its timestamps, size, read-only state, and type (file or directory)) in a single file-system request. Depending on the platform, this can be much more efficient than making multiple requests.

This change also adds alternate directory iteration functions (IterateDirectoryStat and IterateDirectoryStatRecursive) for when you need the stat data along with the files/directories on disk. For platforms that provide this stat data as part of the low-level directory iteration (such as Windows), this can yield much better performance than making two separate file-system requests. Platforms that do not provide this information as part of the directory iteration will just perform a stat request while iterating.

diff --git a/Engine/Source/Runtime/Core/Private/Android/AndroidFile.cpp b/Engine/Source/Runtime/Core/Private/Android/AndroidFile.cpp
--- a/Engine/Source/Runtime/Core/Private/Android/AndroidFile.cpp
+++ b/Engine/Source/Runtime/Core/Private/Android/AndroidFile.cpp
@@ -31,6 +31,29 @@ DEFINE_LOG_CATEGORY_STATIC(LogAndroidFile, Log, All);
 // make an FTimeSpan object that represents the "epoch" for time_t (from a stat struct)
 const FDateTime AndroidEpoch(1970, 1, 1);
 
+namespace
+{
+	FFileStatData AndroidStatToUEFileData(struct stat& FileInfo)
+	{
+		const bool bIsDirectory = S_ISDIR(FileInfo.st_mode);
+
+		int64 FileSize = -1;
+		if (!bIsDirectory)
+		{
+			FileSize = FileInfo.st_size;
+		}
+
+		return FFileStatData(
+			AndroidEpoch + FTimespan(0, 0, FileInfo.st_ctime), 
+			AndroidEpoch + FTimespan(0, 0, FileInfo.st_atime), 
+			AndroidEpoch + FTimespan(0, 0, FileInfo.st_mtime), 
+			FileSize,
+			bIsDirectory,
+			!!(FileInfo.st_mode & S_IWUSR)
+			);
+	}
+}
+
 #define USE_UTIME 0
 
 
@@ -1219,6 +1242,76 @@ public:
 		}
 	}
 
+	virtual FFileStatData GetStatData(const TCHAR* FilenameOrDirectory) override
+	{
+		return GetStatData(FilenameOrDirectory, false);
+	}
+
+	FFileStatData GetStatData(const TCHAR* FilenameOrDirectory, bool AllowLocal)
+	{
+#if LOG_ANDROID_FILE
+		FPlatformMisc::LowLevelOutputDebugStringf(TEXT("FAndroidPlatformFile::GetStatData('%s')"), FilenameOrDirectory);
+#endif
+		FString LocalPath;
+		FString AssetPath;
+		PathToAndroidPaths(LocalPath, AssetPath, FilenameOrDirectory, AllowLocal);
+
+		if (IsLocal(LocalPath))
+		{
+			struct stat FileInfo;
+			if (stat(TCHAR_TO_UTF8(*LocalPath), &FileInfo) != -1)
+			{
+				return AndroidStatToUEFileData(FileInfo);
+			}
+		}
+		else if (IsResource(AssetPath))
+		{
+			return FFileStatData(
+				FDateTime::MinValue(),					// CreationTime
+				FDateTime::MinValue(),					// AccessTime
+				FDateTime::MinValue(),					// ModificationTime
+				ZipResource.GetEntryLength(AssetPath),	// FileSize
+				false,									// bIsDirectory
+				true									// bIsReadOnly
+				);
+		}
+		else
+		{
+			AAsset * file = AAssetManager_open(AssetMgr, TCHAR_TO_UTF8(*AssetPath), AASSET_MODE_UNKNOWN);
+			if (nullptr != file)
+			{
+				bool isDirectory = false;
+				AAssetDir * subdir = AAssetManager_openDir(AssetMgr, TCHAR_TO_UTF8(*AssetPath));
+				if (nullptr != subdir)
+				{
+					isDirectory = true;
+					AAssetDir_close(subdir);
+				}
+
+				int64 FileSize = -1;
+				if (!isDirectory)
+				{
+					FileSize = AAsset_getLength(file);
+				}
+
+				FFileStatData StatData(
+					FDateTime::MinValue(),				// CreationTime
+					FDateTime::MinValue(),				// AccessTime
+					FDateTime::MinValue(),				// ModificationTime
+					FileSize,							// FileSize
+					isDirectory,						// bIsDirectory
+					true								// bIsReadOnly
+					);
+
+				AAsset_close(file);
+
+				return StatData;
+			}
+		}
+		
+		return FFileStatData();
+	}
+
 	virtual FString GetFilenameOnDisk(const TCHAR* Filename) override
 	{
 		return Filename;
@@ -1409,6 +1502,108 @@ public:
 	}
 
 	bool IterateDirectory(const TCHAR* Directory, FDirectoryVisitor& Visitor, bool AllowLocal, bool AllowAsset)
+	{
+		const FString DirectoryStr = Directory;
+
+		auto InternalVisitor = [&](const FString& InLocalPath, struct dirent* InEntry) -> bool
+		{
+			const FString DirPath = DirectoryStr / UTF8_TO_TCHAR(InEntry->d_name);
+			return Visitor.Visit(*DirPath, InEntry->d_type == DT_DIR);
+		};
+		
+		auto InternalResourceVisitor = [&](const FString& InResourceName) -> bool
+		{
+			return Visitor.Visit(*InResourceName, false);
+		};
+		
+		auto InternalAssetVisitor = [&](const char* InAssetPath) -> bool
+		{
+			bool isDirectory = false;
+			AAssetDir* subdir = AAssetManager_openDir(AssetMgr, InAssetPath);
+			if (nullptr != subdir)
+			{
+				isDirectory = true;
+				AAssetDir_close(subdir);
+			}
+
+			return Visitor.Visit(UTF8_TO_TCHAR(InAssetPath), isDirectory);
+		};
+
+		return IterateDirectoryCommon(Directory, InternalVisitor, InternalResourceVisitor, InternalAssetVisitor, AllowLocal, AllowAsset);
+	}
+
+	virtual bool IterateDirectoryStat(const TCHAR* Directory, FDirectoryStatVisitor& Visitor) override
+	{
+		return IterateDirectoryStat(Directory, Visitor, false, false);
+	}
+
+	bool IterateDirectoryStat(const TCHAR* Directory, FDirectoryStatVisitor& Visitor, bool AllowLocal, bool AllowAsset)
+	{
+		const FString DirectoryStr = Directory;
+
+		auto InternalVisitor = [&](const FString& InLocalPath, struct dirent* InEntry) -> bool
+		{
+			const FString DirPath = DirectoryStr / UTF8_TO_TCHAR(InEntry->d_name);
+
+			struct stat FileInfo;
+			if (stat(TCHAR_TO_UTF8(*(InLocalPath / UTF8_TO_TCHAR(InEntry->d_name))), &FileInfo) != -1)
+			{
+				return Visitor.Visit(*DirPath, AndroidStatToUEFileData(FileInfo));
+			}
+
+			return true;
+		};
+		
+		auto InternalResourceVisitor = [&](const FString& InResourceName) -> bool
+		{
+			return Visitor.Visit(
+				*InResourceName, 
+				FFileStatData(
+					FDateTime::MinValue(),						// CreationTime
+					FDateTime::MinValue(),						// AccessTime
+					FDateTime::MinValue(),						// ModificationTime
+					ZipResource.GetEntryLength(InResourceName),	// FileSize
+					false,										// bIsDirectory
+					true										// bIsReadOnly
+					)
+				);
+		};
+		
+		auto InternalAssetVisitor = [&](const char* InAssetPath) -> bool
+		{
+			bool isDirectory = false;
+			AAssetDir* subdir = AAssetManager_openDir(AssetMgr, InAssetPath);
+			if (nullptr != subdir)
+			{
+				isDirectory = true;
+				AAssetDir_close(subdir);
+			}
+
+			int64 FileSize = -1;
+			if (!isDirectory)
+			{
+				AAsset* file = AAssetManager_open(AssetMgr, InAssetPath, AASSET_MODE_UNKNOWN);
+				FileSize = AAsset_getLength(file);
+				AAssetDir_close(subdir);
+			}
+
+			return Visitor.Visit(
+				UTF8_TO_TCHAR(InAssetPath), 
+				FFileStatData(
+					FDateTime::MinValue(),	// CreationTime
+					FDateTime::MinValue(),	// AccessTime
+					FDateTime::MinValue(),	// ModificationTime
+					FileSize,				// FileSize
+					isDirectory,			// bIsDirectory
+					true					// bIsReadOnly
+					)
+				);
+		};
+
+		return IterateDirectoryCommon(Directory, InternalVisitor, InternalResourceVisitor, InternalAssetVisitor, AllowLocal, AllowAsset);
+	}
+
+	bool IterateDirectoryCommon(const TCHAR* Directory, const TFunctionRef<bool(const FString&, struct dirent*)>& Visitor, const TFunctionRef<bool(const FString&)>& ResourceVisitor, const TFunctionRef<bool(const char*)>& AssetVisitor, bool AllowLocal, bool AllowAsset)
 	{
 #if LOG_ANDROID_FILE
 		FPlatformMisc::LowLevelOutputDebugStringf(TEXT("FAndroidPlatformFile::IterateDirectory('%s')"), Directory);
@@ -1428,11 +1623,10 @@ public:
 				{
 					if (FCString::Strcmp(UTF8_TO_TCHAR(Entry->d_name), TEXT(".")) && FCString::Strcmp(UTF8_TO_TCHAR(Entry->d_name), TEXT("..")))
 					{
-						FString DirPath = FString(Directory) / UTF8_TO_TCHAR(Entry->d_name);
 #if LOG_ANDROID_FILE
-						FPlatformMisc::LowLevelOutputDebugStringf(TEXT("FAndroidPlatformFile::IterateDirectory('%s').. LOCAL Visit: '%s'"), Directory, *DirPath);
+						FPlatformMisc::LowLevelOutputDebugStringf(TEXT("FAndroidPlatformFile::IterateDirectory('%s').. LOCAL Visit: '%s'"), Directory, *(FString(Directory) / UTF8_TO_TCHAR(Entry->d_name)));
 #endif
-						Result = Visitor.Visit(*DirPath, Entry->d_type == DT_DIR);
+						Result = Visitor(LocalPath, Entry);
 					}
 				}
 				closedir(Handle);
@@ -1448,7 +1642,7 @@ public:
 #if LOG_ANDROID_FILE
 				FPlatformMisc::LowLevelOutputDebugStringf(TEXT("FAndroidPlatformFile::IterateDirectory('%s').. RESOURCE Visit: '%s'"), Directory, *ResourceDir.Current.Key());
 #endif
-				Result = Visitor.Visit(*ResourceDir.Current.Key(), false);
+				Result = ResourceVisitor(ResourceDir.Current.Key());
 			}
 		}
 		else if (IsResource(AssetPath + "/"))
@@ -1460,7 +1654,7 @@ public:
 #if LOG_ANDROID_FILE
 				FPlatformMisc::LowLevelOutputDebugStringf(TEXT("FAndroidPlatformFile::IterateDirectory('%s').. RESOURCE/ Visit: '%s'"), Directory, *ResourceDir.Current.Key());
 #endif
-				Result = Visitor.Visit(*ResourceDir.Current.Key(), false);
+				Result = ResourceVisitor(ResourceDir.Current.Key());
 			}
 		}
 		else if (AllowAsset)
@@ -1472,17 +1666,10 @@ public:
 				const char * fileName = nullptr;
 				while ((fileName = AAssetDir_getNextFileName(dir)) != nullptr && Result == true)
 				{
-					bool isDirectory = false;
-					AAssetDir * subdir = AAssetManager_openDir(AssetMgr, fileName);
-					if (nullptr != dir)
-					{
-						isDirectory = true;
-						AAssetDir_close(subdir);
-					}
 #if LOG_ANDROID_FILE
 					FPlatformMisc::LowLevelOutputDebugStringf(TEXT("FAndroidPlatformFile::IterateDirectory('%s').. ASSET Visit: '%s'"), Directory, UTF8_TO_TCHAR(fileName));
 #endif
-					Result = Visitor.Visit(UTF8_TO_TCHAR(fileName), isDirectory);
+					Result = AssetVisitor(fileName);
 				}
 				AAssetDir_close(dir);
 				return true;
diff --git a/Engine/Source/Runtime/Core/Private/Apple/ApplePlatformFile.cpp b/Engine/Source/Runtime/Core/Private/Apple/ApplePlatformFile.cpp
--- a/Engine/Source/Runtime/Core/Private/Apple/ApplePlatformFile.cpp
+++ b/Engine/Source/Runtime/Core/Private/Apple/ApplePlatformFile.cpp
@@ -10,6 +10,29 @@
 // make an FTimeSpan object that represents the "epoch" for time_t (from a stat struct)
 const FDateTime MacEpoch(1970, 1, 1);
 
+namespace
+{
+	FFileStatData MacStatToUEFileData(struct stat& FileInfo)
+	{
+		const bool bIsDirectory = S_ISDIR(FileInfo.st_mode);
+
+		int64 FileSize = -1;
+		if (!bIsDirectory)
+		{
+			FileSize = FileInfo.st_size;
+		}
+
+		return FFileStatData(
+			MacEpoch + FTimespan(0, 0, FileInfo.st_ctime), 
+			MacEpoch + FTimespan(0, 0, FileInfo.st_atime), 
+			MacEpoch + FTimespan(0, 0, FileInfo.st_mtime), 
+			FileSize,
+			bIsDirectory,
+			!!(FileInfo.st_mode & S_IWUSR)
+			);
+	}
+}
+
 /** 
  * Mac file handle implementation which limits number of open files per thread. This
  * is to prevent running out of system file handles (250). Should not be neccessary when
@@ -412,6 +435,7 @@ void FApplePlatformFile::SetTimeStamp(const TCHAR* Filename, const FDateTime Dat
 	Times.modtime = (DateTime - MacEpoch).GetTotalSeconds();
 	utime(TCHAR_TO_UTF8(*NormalizeFilename(Filename)), &Times);
 }
+
 FDateTime FApplePlatformFile::GetAccessTimeStamp(const TCHAR* Filename)
 {
 	// get file times
@@ -444,6 +468,7 @@ IFileHandle* FApplePlatformFile::OpenRead(const TCHAR* Filename, bool bAllowWrit
 	}
 	return NULL;
 }
+
 IFileHandle* FApplePlatformFile::OpenWrite(const TCHAR* Filename, bool bAppend, bool bAllowRead)
 {
 	int Flags = O_CREAT;
@@ -475,6 +500,7 @@ IFileHandle* FApplePlatformFile::OpenWrite(const TCHAR* Filename, bool bAppend,
 	}
 	return NULL;
 }
+
 bool FApplePlatformFile::DirectoryExists(const TCHAR* Directory)
 {
 	struct stat FileInfo;
@@ -484,6 +510,7 @@ bool FApplePlatformFile::DirectoryExists(const TCHAR* Directory)
 	}
 	return false;
 }
+
 bool FApplePlatformFile::CreateDirectory(const TCHAR* Directory)
 {
 	@autoreleasepool
@@ -494,45 +521,92 @@ bool FApplePlatformFile::CreateDirectory(const TCHAR* Directory)
 		return Result;
 	}
 }
+
 bool FApplePlatformFile::DeleteDirectory(const TCHAR* Directory)
 {
 	return rmdir(TCHAR_TO_UTF8(*NormalizeFilename(Directory))) == 0;
 }
+
+FFileStatData FApplePlatformFile::GetStatData(const TCHAR* FilenameOrDirectory)
+{
+	struct stat FileInfo;
+	if (Stat(FilenameOrDirectory, &FileInfo) != -1)
+	{
+		return MacStatToUEFileData(FileInfo);
+	}
+
+	return FFileStatData();
+}
+
 bool FApplePlatformFile::IterateDirectory(const TCHAR* Directory, FDirectoryVisitor& Visitor)
 {
 	@autoreleasepool
 	{
-		bool Result = false;
-		DIR* Handle = opendir(Directory[0] ? TCHAR_TO_UTF8(Directory) : ".");
-		if (Handle)
+		const FString DirectoryStr = Directory;
+		const FString NormalizedDirectoryStr = NormalizeFilename(Directory);
+
+		return IterateDirectoryCommon(Directory, [&](struct dirent* InEntry) -> bool
 		{
-			Result = true;
-			struct dirent *Entry;
-			while ((Entry = readdir(Handle)) != NULL)
+			// Normalize any unicode forms so we match correctly
+			const FString NormalizedFilename = UTF8_TO_TCHAR(([[[NSString stringWithUTF8String:InEntry->d_name] precomposedStringWithCanonicalMapping] cStringUsingEncoding:NSUTF8StringEncoding]));
+				
+			// Figure out whether it's a directory. Some protocols (like NFS) do not voluntarily return this as part of the directory entry, and need to be queried manually.
+			bool bIsDirectory = (InEntry->d_type == DT_DIR);
+			if (InEntry->d_type == DT_UNKNOWN)
 			{
-				if (FCStringAnsi::Strcmp(Entry->d_name, ".") && FCStringAnsi::Strcmp(Entry->d_name, "..") && FCStringAnsi::Strcmp(Entry->d_name, ".DS_Store"))
+				struct stat StatInfo;
+				if (stat(TCHAR_TO_UTF8(*(NormalizedDirectoryStr / NormalizedFilename)), &StatInfo) == 0)
 				{
-                    // Normalize any unicode forms so we match correctly
-                    FString NormalizedFilename = UTF8_TO_TCHAR(([[[NSString stringWithUTF8String:Entry->d_name] precomposedStringWithCanonicalMapping] cStringUsingEncoding:NSUTF8StringEncoding]));
-					
-                    // Figure out whether it's a directory. Some protocols (like NFS) do not voluntairily return this as part of the directory entry, and need to be queried manually.
-                    bool bIsDirectory = (Entry->d_type == DT_DIR);
-                    if(Entry->d_type == DT_UNKNOWN)
-                    {
-                        struct stat StatInfo;
-                        if(stat(TCHAR_TO_UTF8(*(FString(Directory) / Entry->d_name)), &StatInfo) == 0)
-                        {
-                            bIsDirectory = S_ISDIR(StatInfo.st_mode);
-                        }
-                    }
-					
-                    Result = Visitor.Visit(*(FString(Directory) / NormalizedFilename), bIsDirectory);
+					bIsDirectory = S_ISDIR(StatInfo.st_mode);
 				}
 			}
-			closedir(Handle);
+					
+			return Visitor.Visit(*(DirectoryStr / NormalizedFilename), bIsDirectory);
+		});
+	}
+}
+
+bool FApplePlatformFile::IterateDirectoryStat(const TCHAR* Directory, FDirectoryStatVisitor& Visitor)
+{
+	@autoreleasepool
+	{
+		const FString DirectoryStr = Directory;
+		const FString NormalizedDirectoryStr = NormalizeFilename(Directory);
+
+		return IterateDirectoryCommon(Directory, [&](struct dirent* InEntry) -> bool
+		{
+			// Normalize any unicode forms so we match correctly
+			const FString NormalizedFilename = UTF8_TO_TCHAR(([[[NSString stringWithUTF8String:InEntry->d_name] precomposedStringWithCanonicalMapping] cStringUsingEncoding:NSUTF8StringEncoding]));
+				
+			struct stat StatInfo;
+			if (stat(TCHAR_TO_UTF8(*(NormalizedDirectoryStr / NormalizedFilename)), &StatInfo) == 0)
+			{
+				return Visitor.Visit(*(DirectoryStr / NormalizedFilename), MacStatToUEFileData(StatInfo));
+			}
+
+			return true;
+		});
+	}
+}
+
+bool FApplePlatformFile::IterateDirectoryCommon(const TCHAR* Directory, const TFunctionRef<bool(struct dirent*)>& Visitor)
+{
+	bool Result = false;
+	DIR* Handle = opendir(Directory[0] ? TCHAR_TO_UTF8(Directory) : ".");
+	if (Handle)
+	{
+		Result = true;
+		struct dirent *Entry;
+		while ((Entry = readdir(Handle)) != NULL)
+		{
+			if (FCStringAnsi::Strcmp(Entry->d_name, ".") && FCStringAnsi::Strcmp(Entry->d_name, "..") && FCStringAnsi::Strcmp(Entry->d_name, ".DS_Store"))
+			{
+                Result = Visitor(Entry);
+			}
 		}
-		return Result;
+		closedir(Handle);
 	}
+	return Result;
 }
 
 bool FApplePlatformFile::CopyFile(const TCHAR* To, const TCHAR* From)
diff --git a/Engine/Source/Runtime/Core/Private/GenericPlatform/GenericPlatformFile.cpp b/Engine/Source/Runtime/Core/Private/GenericPlatform/GenericPlatformFile.cpp
--- a/Engine/Source/Runtime/Core/Private/GenericPlatform/GenericPlatformFile.cpp
+++ b/Engine/Source/Runtime/Core/Private/GenericPlatform/GenericPlatformFile.cpp
@@ -30,7 +30,7 @@ bool IPlatformFile::IterateDirectoryRecursively(const TCHAR* Directory, FDirecto
 			, Visitor(InVisitor)
 		{
 		}
-		virtual bool Visit(const TCHAR* FilenameOrDirectory, bool bIsDirectory)
+		virtual bool Visit(const TCHAR* FilenameOrDirectory, bool bIsDirectory) override
 		{
 			bool Result = Visitor.Visit(FilenameOrDirectory, bIsDirectory);
 			if (Result && bIsDirectory)
@@ -44,6 +44,32 @@ bool IPlatformFile::IterateDirectoryRecursively(const TCHAR* Directory, FDirecto
 	return IterateDirectory(Directory, Recurse);
 }
 
+bool IPlatformFile::IterateDirectoryStatRecursively(const TCHAR* Directory, FDirectoryStatVisitor& Visitor)
+{
+	class FStatRecurse : public FDirectoryStatVisitor
+	{
+	public:
+		IPlatformFile&			PlatformFile;
+		FDirectoryStatVisitor&	Visitor;
+		FStatRecurse(IPlatformFile&	InPlatformFile, FDirectoryStatVisitor& InVisitor)
+			: PlatformFile(InPlatformFile)
+			, Visitor(InVisitor)
+		{
+		}
+		virtual bool Visit(const TCHAR* FilenameOrDirectory, const FFileStatData& StatData) override
+		{
+			bool Result = Visitor.Visit(FilenameOrDirectory, StatData);
+			if (Result && StatData.bIsDirectory)
+			{
+				Result = PlatformFile.IterateDirectoryStat(FilenameOrDirectory, *this);
+			}
+			return Result;
+		}
+	};
+	FStatRecurse Recurse(*this, Visitor);
+	return IterateDirectoryStat(Directory, Recurse);
+}
+
 bool IPlatformFile::DeleteDirectoryRecursively(const TCHAR* Directory)
 {
 	class FRecurse : public FDirectoryVisitor
diff --git a/Engine/Source/Runtime/Core/Private/HAL/FileManagerGeneric.cpp b/Engine/Source/Runtime/Core/Private/HAL/FileManagerGeneric.cpp
--- a/Engine/Source/Runtime/Core/Private/HAL/FileManagerGeneric.cpp
+++ b/Engine/Source/Runtime/Core/Private/HAL/FileManagerGeneric.cpp
@@ -361,6 +361,11 @@ bool FFileManagerGeneric::DeleteDirectory( const TCHAR* Path, bool RequireExists
 	return GetLowLevel().DeleteDirectory( Path ) || ( !RequireExists && !GetLowLevel().DirectoryExists( Path ) );
 }
 
+FFileStatData FFileManagerGeneric::GetStatData(const TCHAR* FilenameOrDirectory)
+{
+	return GetLowLevel().GetStatData(FilenameOrDirectory);
+}
+
 void FFileManagerGeneric::FindFiles( TArray<FString>& Result, const TCHAR* InFilename, bool Files, bool Directories )
 {
 	class FFileMatch : public IPlatformFile::FDirectoryVisitor
@@ -442,6 +447,16 @@ bool FFileManagerGeneric::IterateDirectoryRecursively(const TCHAR* Directory, IP
 	return GetLowLevel().IterateDirectoryRecursively( Directory, Visitor );
 }
 
+bool FFileManagerGeneric::IterateDirectoryStat(const TCHAR* Directory, IPlatformFile::FDirectoryStatVisitor& Visitor)
+{
+	return GetLowLevel().IterateDirectoryStat( Directory, Visitor );
+}
+
+bool FFileManagerGeneric::IterateDirectoryStatRecursively(const TCHAR* Directory, IPlatformFile::FDirectoryStatVisitor& Visitor)
+{
+	return GetLowLevel().IterateDirectoryStatRecursively( Directory, Visitor );
+}
+
 double FFileManagerGeneric::GetFileAgeSeconds( const TCHAR* Filename )
 {
 	// make sure it exists
diff --git a/Engine/Source/Runtime/Core/Private/HTML5/HTML5PlatformFile.cpp b/Engine/Source/Runtime/Core/Private/HTML5/HTML5PlatformFile.cpp
--- a/Engine/Source/Runtime/Core/Private/HTML5/HTML5PlatformFile.cpp
+++ b/Engine/Source/Runtime/Core/Private/HTML5/HTML5PlatformFile.cpp
@@ -46,6 +46,29 @@
 // make an FTimeSpan object that represents the "epoch" for time_t (from a stat struct)
 const FDateTime HTML5Epoch(1970, 1, 1);
 
+namespace
+{
+	FFileStatData HTML5StatToUEFileData(struct stat& FileInfo)
+	{
+		const bool bIsDirectory = S_ISDIR(FileInfo.st_mode);
+
+		int64 FileSize = -1;
+		if (!bIsDirectory)
+		{
+			FileSize = FileInfo.st_size;
+		}
+
+		return FFileStatData(
+			HTML5Epoch + FTimespan(0, 0, FileInfo.st_ctime), 
+			HTML5Epoch + FTimespan(0, 0, FileInfo.st_atime), 
+			HTML5Epoch + FTimespan(0, 0, FileInfo.st_mtime), 
+			FileSize,
+			bIsDirectory,
+			!!(FileInfo.st_mode & S_IWUSR)
+			);
+	}
+}
+
 /** 
  * HTML5 file handle implementation
 **/
@@ -365,27 +388,80 @@ public:
 #endif
 	}
 
-	bool IterateDirectory(const TCHAR* Directory, FDirectoryVisitor& Visitor)
+	virtual FFileStatData GetStatData(const TCHAR* FilenameOrDirectory) override
 	{
-#if !PLATFORM_HTML5_WIN32 
-		bool Result = false;
-		// If Directory is an empty string, assume that we want to iterate Binaries/Mac (current dir), but because we're an app bundle, iterate bundle's Contents/Frameworks instead
-		DIR* Handle = opendir(TCHAR_TO_UTF8(*NormalizeFilename(Directory)));
-		if (Handle)
+		struct stat FileInfo;
+		if (stat(TCHAR_TO_UTF8(*NormalizeFilename(FilenameOrDirectory)), &FileInfo) != -1)
 		{
-			Result = true;
-			struct dirent *Entry;
-			while ((Entry = readdir(Handle)) != NULL)
+			return HTML5StatToUEFileData(FileInfo);
+		}
+
+		return FFileStatData();
+	}
+
+	virtual bool IterateDirectory(const TCHAR* Directory, FDirectoryVisitor& Visitor) override
+	{
+#if PLATFORM_HTML5_WIN32
+		const FString DirectoryStr = Directory;
+		return IterateDirectoryCommon(Directory, [&](struct _finddata_t& InFindData) -> bool
+		{
+			const bool bIsDirectory = ((InFindData.attrib & _A_SUBDIR) != 0);
+			return Visitor.Visit(*(DirectoryStr / UTF8_TO_TCHAR(InFindData.name)), bIsDirectory);
+		});
+#else
+		const FString DirectoryStr = Directory;
+		return IterateDirectoryCommon(Directory, [&](struct dirent* InEntry) -> bool
+		{
+			const bool bIsDirectory = InEntry->d_type == DT_DIR;
+			return Visitor.Visit(*(DirectoryStr / UTF8_TO_TCHAR(InEntry->d_name)), bIsDirectory);
+		});
+#endif
+	}
+
+	virtual bool IterateDirectoryStat(const TCHAR* Directory, FDirectoryStatVisitor& Visitor) override
+	{
+#if PLATFORM_HTML5_WIN32
+		const FString DirectoryStr = Directory;
+		const FString NormalizedDirectoryStr = NormalizeFilename(Directory);
+		return IterateDirectoryCommon(Directory, [&](struct _finddata_t& InFindData) -> bool
+		{
+			const bool bIsDirectory = ((InFindData.attrib & _A_SUBDIR) != 0);
+			
+			int64 FileSize = -1;
+			if (!bIsDirectory)
 			{
-				if (FCString::Strcmp(UTF8_TO_TCHAR(Entry->d_name), TEXT(".")) && FCString::Strcmp(UTF8_TO_TCHAR(Entry->d_name), TEXT("..")))
-				{
-					Result = Visitor.Visit(*(FString(Directory) / UTF8_TO_TCHAR(Entry->d_name)), Entry->d_type == DT_DIR);
-				}
+				FileSize = static_cast<int64>(InFindData.size);
 			}
-			closedir(Handle);
-		}
-		return Result;
-#else 
+
+			const FFileStatData StatData(
+				HTML5Epoch + FTimespan(0, 0, InFindData.time_create), 
+				HTML5Epoch + FTimespan(0, 0, InFindData.time_access), 
+				HTML5Epoch + FTimespan(0, 0, InFindData.time_write), 
+				FileSize,
+				bIsDirectory,
+				((InFindData.attrib & _A_RDONLY) != 0)
+				);
+			
+			return Visitor.Visit(*(DirectoryStr / UTF8_TO_TCHAR(InFindData.name)), StatData);
+		});
+#else
+		const FString DirectoryStr = Directory;
+		const FString NormalizedDirectoryStr = NormalizeFilename(Directory);
+		return IterateDirectoryCommon(Directory, [&](struct dirent* InEntry) -> bool
+		{
+			struct stat FileInfo;
+			if (stat(TCHAR_TO_UTF8(*(NormalizedDirectoryStr / UTF8_TO_TCHAR(InEntry->d_name))), &FileInfo) != -1)
+			{
+				return Visitor.Visit(*(DirectoryStr / UTF8_TO_TCHAR(InEntry->d_name)), HTML5StatToUEFileData(FileInfo));
+			}
+			return true;
+		});
+#endif
+	}
+
+#if PLATFORM_HTML5_WIN32
+	bool IterateDirectoryCommon(const TCHAR* Directory, const TFunctionRef<bool(struct _finddata_t&)>& Visitor)
+	{
 		bool Result = false;
 		// If Directory is an empty string, assume that we want to iterate Binaries/Mac (current dir), but because we're an app bundle, iterate bundle's Contents/Frameworks instead
 		FString FilePath = FPaths::Combine(*NormalizeFilename(Directory), TEXT("*"));
@@ -399,14 +475,35 @@ public:
 			{
 				if (FCString::Strcmp(UTF8_TO_TCHAR(c_file.name), TEXT(".")) && FCString::Strcmp(UTF8_TO_TCHAR(c_file.name), TEXT("..")))
 				{
-					Result = Visitor.Visit(*(FString(Directory) / UTF8_TO_TCHAR(c_file.name)), ((c_file.attrib & _A_SUBDIR) != 0));
+					Result = Visitor(c_file);
 				}
 			}
 			_findclose(hFile);
 		}
 		return Result;
-#endif 
 	}
+#else
+	bool IterateDirectoryCommon(const TCHAR* Directory, const TFunctionRef<bool(struct dirent*)>& Visitor)
+	{
+		bool Result = false;
+		// If Directory is an empty string, assume that we want to iterate Binaries/Mac (current dir), but because we're an app bundle, iterate bundle's Contents/Frameworks instead
+		DIR* Handle = opendir(TCHAR_TO_UTF8(*NormalizeFilename(Directory)));
+		if (Handle)
+		{
+			Result = true;
+			struct dirent *Entry;
+			while ((Entry = readdir(Handle)) != NULL)
+			{
+				if (FCString::Strcmp(UTF8_TO_TCHAR(Entry->d_name), TEXT(".")) && FCString::Strcmp(UTF8_TO_TCHAR(Entry->d_name), TEXT("..")))
+				{
+					Result = Visitor(Entry);
+				}
+			}
+			closedir(Handle);
+		}
+		return Result;
+	}
+#endif
 }; 
 
 IPlatformFile& IPlatformFile::GetPlatformPhysical()
diff --git a/Engine/Source/Runtime/Core/Private/IOS/IOSPlatformFile.cpp b/Engine/Source/Runtime/Core/Private/IOS/IOSPlatformFile.cpp
--- a/Engine/Source/Runtime/Core/Private/IOS/IOSPlatformFile.cpp
+++ b/Engine/Source/Runtime/Core/Private/IOS/IOSPlatformFile.cpp
@@ -11,6 +11,28 @@
 // make an FTimeSpan object that represents the "epoch" for time_t (from a stat struct)
 const FDateTime IOSEpoch(1970, 1, 1);
 
+namespace
+{
+	FFileStatData IOSStatToUEFileData(struct stat& FileInfo)
+	{
+		const bool bIsDirectory = S_ISDIR(FileInfo.st_mode);
+
+		int64 FileSize = -1;
+		if (!bIsDirectory)
+		{
+			FileSize = FileInfo.st_size;
+		}
+
+		return FFileStatData(
+			IOSEpoch + FTimespan(0, 0, FileInfo.st_ctime), 
+			IOSEpoch + FTimespan(0, 0, FileInfo.st_atime), 
+			IOSEpoch + FTimespan(0, 0, FileInfo.st_mtime), 
+			FileSize,
+			bIsDirectory,
+			!!(FileInfo.st_mode & S_IWUSR)
+			);
+	}
+}
 
 
 
@@ -518,6 +540,24 @@ FString FIOSPlatformFile::GetFilenameOnDisk(const TCHAR* Filename)
 	return Filename;
 }
 
+FFileStatData FIOSPlatformFile::GetStatData(const TCHAR* FilenameOrDirectory)
+{
+	struct stat FileInfo;
+	FString NormalizedFilename = NormalizeFilename(FilenameOrDirectory);
+
+	// check the read path
+	if(stat(TCHAR_TO_UTF8(*ConvertToIOSPath(NormalizedFilename, false)), &FileInfo) == -1)
+	{
+		// if not in the read path, check the write path
+		if(stat(TCHAR_TO_UTF8(*ConvertToIOSPath(NormalizedFilename, true)), &FileInfo) == -1)
+		{
+			return FFileStatData();
+		}
+	}
+
+	return IOSStatToUEFileData(FileInfo);
+}
+
 IFileHandle* FIOSPlatformFile::OpenRead(const TCHAR* Filename, bool bAllowWrite)
 {
 	FString NormalizedFilename = NormalizeFilename(Filename);
@@ -602,6 +642,47 @@ bool FIOSPlatformFile::DeleteDirectory(const TCHAR* Directory)
 }
 
 bool FIOSPlatformFile::IterateDirectory(const TCHAR* Directory, FDirectoryVisitor& Visitor)
+{
+	const FString DirectoryStr = Directory;
+	return IterateDirectoryCommon(Directory, [&](struct dirent* InEntry) -> bool
+	{
+		// Normalize any unicode forms so we match correctly
+		const FString NormalizedFilename = UTF8_TO_TCHAR(([[[NSString stringWithUTF8String:InEntry->d_name] precomposedStringWithCanonicalMapping] cStringUsingEncoding:NSUTF8StringEncoding]));
+		const FString FullPath = DirectoryStr / NormalizedFilename;
+
+		return Visitor.Visit(*FullPath, InEntry->d_type == DT_DIR);
+	});
+}
+
+bool FIOSPlatformFile::IterateDirectoryStat(const TCHAR* Directory, FDirectoryStatVisitor& Visitor)
+{
+	const FString DirectoryStr = Directory;
+	const FString NormalizedDirectoryStr = NormalizeFilename(Directory);
+
+	return IterateDirectoryCommon(Directory, [&](struct dirent* InEntry) -> bool
+	{
+		// Normalize any unicode forms so we match correctly
+		const FString NormalizedFilename = UTF8_TO_TCHAR(([[[NSString stringWithUTF8String:InEntry->d_name] precomposedStringWithCanonicalMapping] cStringUsingEncoding:NSUTF8StringEncoding]));
+		const FString FullPath = DirectoryStr / NormalizedFilename;
+		const FString FullNormalizedPath = NormalizedDirectoryStr / NormalizedFilename;
+
+		struct stat FileInfo;
+
+		// check the read path
+		if(stat(TCHAR_TO_UTF8(*ConvertToIOSPath(FullNormalizedPath, false)), &FileInfo) == -1)
+		{
+			// if not in the read path, check the write path
+			if(stat(TCHAR_TO_UTF8(*ConvertToIOSPath(FullNormalizedPath, true)), &FileInfo) == -1)
+			{
+				return true;
+			}
+		}
+
+		return Visitor.Visit(*FullPath, IOSStatToUEFileData(FileInfo));
+	});
+}
+
+bool FIOSPlatformFile::IterateDirectoryCommon(const TCHAR* Directory, const TFunctionRef<bool(struct dirent*)>& Visitor)
 {
 	bool Result = false;
 	const ANSICHAR* FrameworksPath;
@@ -633,9 +714,7 @@ bool FIOSPlatformFile::IterateDirectory(const TCHAR* Directory, FDirectoryVisito
 		{
 			if (FCStringAnsi::Strcmp(Entry->d_name, ".") && FCStringAnsi::Strcmp(Entry->d_name, ".."))
 			{
-				// Normalize any unicode forms so we match correctly
-				FString NormalizedFilename = UTF8_TO_TCHAR(([[[NSString stringWithUTF8String:Entry->d_name] precomposedStringWithCanonicalMapping] cStringUsingEncoding:NSUTF8StringEncoding]));
-				Result = Visitor.Visit(*(FString(Directory) / NormalizedFilename), Entry->d_type == DT_DIR);
+				Result = Visitor(Entry);
 			}
 		}
 		closedir(Handle);
diff --git a/Engine/Source/Runtime/Core/Private/IOS/IOSPlatformFile.h b/Engine/Source/Runtime/Core/Private/IOS/IOSPlatformFile.h
--- a/Engine/Source/Runtime/Core/Private/IOS/IOSPlatformFile.h
+++ b/Engine/Source/Runtime/Core/Private/IOS/IOSPlatformFile.h
@@ -30,15 +30,21 @@ public:
 	virtual FDateTime GetAccessTimeStamp(const TCHAR* Filename) override;
 	virtual FString GetFilenameOnDisk(const TCHAR* Filename) override;
 
+	virtual FFileStatData GetStatData(const TCHAR* FilenameOrDirectory) override;
+
 	virtual IFileHandle* OpenRead(const TCHAR* Filename, bool bAllowWrite = false) override;
 	virtual IFileHandle* OpenWrite(const TCHAR* Filename, bool bAppend = false, bool bAllowRead = false) override;
 
 	virtual bool DirectoryExists(const TCHAR* Directory) override;
 	virtual bool CreateDirectory(const TCHAR* Directory) override;
 	virtual bool DeleteDirectory(const TCHAR* Directory) override;
-	bool IterateDirectory(const TCHAR* Directory, FDirectoryVisitor& Visitor);
+
+	virtual bool IterateDirectory(const TCHAR* Directory, FDirectoryVisitor& Visitor) override;
+	virtual bool IterateDirectoryStat(const TCHAR* Directory, FDirectoryStatVisitor& Visitor) override;
 
 private:
+	bool IterateDirectoryCommon(const TCHAR* Directory, const TFunctionRef<bool(struct dirent*)>& Visitor);
+
 	FString ConvertToIOSPath(const FString& Filename, bool bForWrite);
 
 	FString ReadDirectory;
diff --git a/Engine/Source/Runtime/Core/Private/Linux/LinuxPlatformFile.cpp b/Engine/Source/Runtime/Core/Private/Linux/LinuxPlatformFile.cpp
--- a/Engine/Source/Runtime/Core/Private/Linux/LinuxPlatformFile.cpp
+++ b/Engine/Source/Runtime/Core/Private/Linux/LinuxPlatformFile.cpp
@@ -9,6 +9,29 @@ DEFINE_LOG_CATEGORY_STATIC(LogLinuxPlatformFile, Log, All);
 // make an FTimeSpan object that represents the "epoch" for time_t (from a stat struct)
 const FDateTime UnixEpoch(1970, 1, 1);
 
+namespace
+{
+	FFileStatData UnixStatToUEFileData(struct stat& FileInfo)
+	{
+		const bool bIsDirectory = S_ISDIR(FileInfo.st_mode);
+
+		int64 FileSize = -1;
+		if (!bIsDirectory)
+		{
+			FileSize = FileInfo.st_size;
+		}
+
+		return FFileStatData(
+			UnixEpoch + FTimespan(0, 0, FileInfo.st_ctime), 
+			UnixEpoch + FTimespan(0, 0, FileInfo.st_atime), 
+			UnixEpoch + FTimespan(0, 0, FileInfo.st_mtime), 
+			FileSize,
+			bIsDirectory,
+			!!(FileInfo.st_mode & S_IWUSR)
+			);
+	}
+}
+
 /** 
  * Linux file handle implementation which limits number of open files per thread. This
  * is to prevent running out of system file handles. Should not be neccessary when
@@ -885,7 +908,79 @@ bool FLinuxPlatformFile::DeleteDirectory(const TCHAR* Directory)
 	return rmdir(TCHAR_TO_UTF8(*CaseSensitiveFilename)) == 0;
 }
 
+FFileStatData FLinuxPlatformFile::GetStatData(const TCHAR* FilenameOrDirectory)
+{
+	FString CaseSensitiveFilename;
+	if (!GCaseInsensMapper.MapCaseInsensitiveFile(NormalizeFilename(FilenameOrDirectory), CaseSensitiveFilename))
+	{
+		// could not find the file
+		return FFileStatData();
+	}
+
+	struct stat FileInfo;
+	if(stat(TCHAR_TO_UTF8(*CaseSensitiveFilename), &FileInfo) == -1)
+	{
+		return FFileStatData();
+	}
+
+	return UnixStatToUEFileData(FileInfo);
+}
+
 bool FLinuxPlatformFile::IterateDirectory(const TCHAR* Directory, FDirectoryVisitor& Visitor)
+{
+	const FString DirectoryStr = Directory;
+	const FString NormalizedDirectoryStr = NormalizeFilename(Directory);
+
+	return IterateDirectoryCommon(Directory, [&](struct dirent* InEntry) -> bool
+	{
+		const FString UnicodeEntryName = UTF8_TO_TCHAR(InEntry->d_name);
+				
+		bool bIsDirectory = false;
+		if (InEntry->d_type != DT_UNKNOWN)
+		{
+			bIsDirectory = InEntry->d_type == DT_DIR;
+		}
+		else
+		{
+			// filesystem does not support d_type, fallback to stat
+			struct stat FileInfo;
+			const FString AbsoluteUnicodeName = NormalizedDirectoryStr / UnicodeEntryName;	
+			if (stat(TCHAR_TO_UTF8(*AbsoluteUnicodeName), &FileInfo) != -1)
+			{
+				bIsDirectory = ((FileInfo.st_mode & S_IFMT) == S_IFDIR);
+			}
+			else
+			{
+				int ErrNo = errno;
+				UE_LOG(LogLinuxPlatformFile, Warning, TEXT( "Cannot determine whether '%s' is a directory - d_type not supported and stat() failed with errno=%d (%s)"), *AbsoluteUnicodeName, ErrNo, UTF8_TO_TCHAR(strerror(ErrNo)));
+			}
+		}
+
+		return Visitor.Visit(*(DirectoryStr / UnicodeEntryName), bIsDirectory);
+	});
+}
+
+bool FLinuxPlatformFile::IterateDirectoryStat(const TCHAR* Directory, FDirectoryStatVisitor& Visitor)
+{
+	const FString DirectoryStr = Directory;
+	const FString NormalizedDirectoryStr = NormalizeFilename(Directory);
+
+	return IterateDirectoryCommon(Directory, [&](struct dirent* InEntry) -> bool
+	{
+		const FString UnicodeEntryName = UTF8_TO_TCHAR(InEntry->d_name);
+				
+		struct stat FileInfo;
+		const FString AbsoluteUnicodeName = NormalizedDirectoryStr / UnicodeEntryName;	
+		if (stat(TCHAR_TO_UTF8(*AbsoluteUnicodeName), &FileInfo) != -1)
+		{
+			return Visitor.Visit(*(DirectoryStr / UnicodeEntryName), UnixStatToUEFileData(FileInfo));
+		}
+
+		return true;
+	});
+}
+
+bool FLinuxPlatformFile::IterateDirectoryCommon(const TCHAR* Directory, const TFunctionRef<bool(struct dirent*)>& Visitor)
 {
 	bool Result = false;
 
@@ -894,33 +989,12 @@ bool FLinuxPlatformFile::IterateDirectory(const TCHAR* Directory, FDirectoryVisi
 	if (Handle)
 	{
 		Result = true;
-		struct dirent *Entry;
+		struct dirent* Entry;
 		while ((Entry = readdir(Handle)) != NULL)
 		{
 			if (FCString::Strcmp(UTF8_TO_TCHAR(Entry->d_name), TEXT(".")) && FCString::Strcmp(UTF8_TO_TCHAR(Entry->d_name), TEXT("..")))
 			{
-				bool bIsDirectory = false;
-				FString UnicodeEntryName = UTF8_TO_TCHAR(Entry->d_name);
-				if (Entry->d_type != DT_UNKNOWN)
-				{
-					bIsDirectory = Entry->d_type == DT_DIR;
-				}
-				else
-				{
-					// filesystem does not support d_type, fallback to stat
-					struct stat FileInfo;
-					FString AbsoluteUnicodeName = NormalizedDirectory / UnicodeEntryName;	
-					if (stat(TCHAR_TO_UTF8(*AbsoluteUnicodeName), &FileInfo) != -1)
-					{
-						bIsDirectory = ((FileInfo.st_mode & S_IFMT) == S_IFDIR);
-					}
-					else
-					{
-						int ErrNo = errno;
-						UE_LOG(LogLinuxPlatformFile, Warning, TEXT( "Cannot determine whether '%s' is a directory - d_type not supported and stat() failed with errno=%d (%s)"), *AbsoluteUnicodeName, ErrNo, UTF8_TO_TCHAR(strerror(ErrNo)));
-					}
-				}
-				Result = Visitor.Visit(*(FString(Directory) / UnicodeEntryName), bIsDirectory);
+				Result = Visitor(Entry);
 			}
 		}
 		closedir(Handle);
diff --git a/Engine/Source/Runtime/Core/Private/Windows/WindowsPlatformFile.cpp b/Engine/Source/Runtime/Core/Private/Windows/WindowsPlatformFile.cpp
--- a/Engine/Source/Runtime/Core/Private/Windows/WindowsPlatformFile.cpp
+++ b/Engine/Source/Runtime/Core/Private/Windows/WindowsPlatformFile.cpp
@@ -659,7 +659,73 @@ public:
 		RemoveDirectoryW(*NormalizeDirectory(Directory));
 		return !DirectoryExists(Directory);
 	}
-	bool IterateDirectory(const TCHAR* Directory, FDirectoryVisitor& Visitor)
+	virtual FFileStatData GetStatData(const TCHAR* FilenameOrDirectory) override
+	{
+		WIN32_FILE_ATTRIBUTE_DATA Info;
+		if (GetFileAttributesExW(*NormalizeFilename(FilenameOrDirectory), GetFileExInfoStandard, &Info))
+		{
+			const bool bIsDirectory = !!(Info.dwFileAttributes & FILE_ATTRIBUTE_DIRECTORY);
+
+			int64 FileSize = -1;
+			if (!bIsDirectory)
+			{
+				LARGE_INTEGER li;
+				li.HighPart = Info.nFileSizeHigh;
+				li.LowPart = Info.nFileSizeLow;
+				FileSize = static_cast<int64>(li.QuadPart);
+			}
+
+			return FFileStatData(
+				WindowsFileTimeToUEDateTime(Info.ftCreationTime),
+				WindowsFileTimeToUEDateTime(Info.ftLastAccessTime),
+				WindowsFileTimeToUEDateTime(Info.ftLastWriteTime),
+				FileSize, 
+				bIsDirectory,
+				!!(Info.dwFileAttributes & FILE_ATTRIBUTE_READONLY)
+				);
+		}
+
+		return FFileStatData();
+	}
+	virtual bool IterateDirectory(const TCHAR* Directory, FDirectoryVisitor& Visitor) override
+	{
+		const FString DirectoryStr = Directory;
+		return IterateDirectoryCommon(Directory, [&](const WIN32_FIND_DATAW& InData) -> bool
+		{
+			const bool bIsDirectory = !!(InData.dwFileAttributes & FILE_ATTRIBUTE_DIRECTORY);
+			return Visitor.Visit(*(DirectoryStr / InData.cFileName), bIsDirectory);
+		});
+	}
+	virtual bool IterateDirectoryStat(const TCHAR* Directory, FDirectoryStatVisitor& Visitor) override
+	{
+		const FString DirectoryStr = Directory;
+		return IterateDirectoryCommon(Directory, [&](const WIN32_FIND_DATAW& InData) -> bool
+		{
+			const bool bIsDirectory = !!(InData.dwFileAttributes & FILE_ATTRIBUTE_DIRECTORY);
+
+			int64 FileSize = -1;
+			if (!bIsDirectory)
+			{
+				LARGE_INTEGER li;
+				li.HighPart = InData.nFileSizeHigh;
+				li.LowPart = InData.nFileSizeLow;
+				FileSize = static_cast<int64>(li.QuadPart);
+			}
+
+			return Visitor.Visit(
+				*(DirectoryStr / InData.cFileName), 
+				FFileStatData(
+					WindowsFileTimeToUEDateTime(InData.ftCreationTime),
+					WindowsFileTimeToUEDateTime(InData.ftLastAccessTime),
+					WindowsFileTimeToUEDateTime(InData.ftLastWriteTime),
+					FileSize, 
+					bIsDirectory,
+					!!(InData.dwFileAttributes & FILE_ATTRIBUTE_READONLY)
+					)
+				);
+		});
+	}
+	bool IterateDirectoryCommon(const TCHAR* Directory, const TFunctionRef<bool(const WIN32_FIND_DATAW&)>& Visitor)
 	{
 		bool Result = false;
 		WIN32_FIND_DATAW Data;
@@ -671,7 +737,7 @@ public:
 			{
 				if (FCString::Strcmp(Data.cFileName, TEXT(".")) && FCString::Strcmp(Data.cFileName, TEXT("..")))
 				{
-					Result = Visitor.Visit(*(FString(Directory) / Data.cFileName), !!(Data.dwFileAttributes&FILE_ATTRIBUTE_DIRECTORY));
+					Result = Visitor(Data);
 				}
 			} while (Result && FindNextFileW(Handle, &Data));
 			FindClose(Handle);
diff --git a/Engine/Source/Runtime/Core/Public/Apple/ApplePlatformFile.h b/Engine/Source/Runtime/Core/Public/Apple/ApplePlatformFile.h
--- a/Engine/Source/Runtime/Core/Public/Apple/ApplePlatformFile.h
+++ b/Engine/Source/Runtime/Core/Public/Apple/ApplePlatformFile.h
@@ -36,8 +36,15 @@ public:
 	virtual bool DirectoryExists(const TCHAR* Directory) override;
 	virtual bool CreateDirectory(const TCHAR* Directory) override;
 	virtual bool DeleteDirectory(const TCHAR* Directory) override;
+
+	virtual FFileStatData GetStatData(const TCHAR* FilenameOrDirectory) override;
+
 	virtual bool IterateDirectory(const TCHAR* Directory, FDirectoryVisitor& Visitor) override;
+	virtual bool IterateDirectoryStat(const TCHAR* Directory, FDirectoryStatVisitor& Visitor) override;
+
 	virtual bool CopyFile(const TCHAR* To, const TCHAR* From) override;
 private:
+	bool IterateDirectoryCommon(const TCHAR* Directory, const TFunctionRef<bool(struct dirent*)>& Visitor);
+
 	int32 Stat(const TCHAR* Filename, struct stat* OutFileInfo);
 };
diff --git a/Engine/Source/Runtime/Core/Public/GenericPlatform/GenericPlatformFile.h b/Engine/Source/Runtime/Core/Public/GenericPlatform/GenericPlatformFile.h
--- a/Engine/Source/Runtime/Core/Public/GenericPlatform/GenericPlatformFile.h
+++ b/Engine/Source/Runtime/Core/Public/GenericPlatform/GenericPlatformFile.h
@@ -63,6 +63,56 @@ public:
 };
 
 
+/**
+ * Contains the information that's returned from stat'ing a file or directory 
+ */
+struct FFileStatData
+{
+	FFileStatData()
+		: CreationTime(FDateTime::MinValue())
+		, AccessTime(FDateTime::MinValue())
+		, ModificationTime(FDateTime::MinValue())
+		, FileSize(-1)
+		, bIsDirectory(false)
+		, bIsReadOnly(false)
+		, bIsValid(false)
+	{
+	}
+
+	FFileStatData(FDateTime InCreationTime, FDateTime InAccessTime,	FDateTime InModificationTime, const int64 InFileSize, const bool InIsDirectory, const bool InIsReadOnly)
+		: CreationTime(InCreationTime)
+		, AccessTime(InAccessTime)
+		, ModificationTime(InModificationTime)
+		, FileSize(InFileSize)
+		, bIsDirectory(InIsDirectory)
+		, bIsReadOnly(InIsReadOnly)
+		, bIsValid(true)
+	{
+	}
+
+	/** The time that the file or directory was originally created, or FDateTime::MinValue if the creation time is unknown */
+	FDateTime CreationTime;
+
+	/** The time that the file or directory was last accessed, or FDateTime::MinValue if the access time is unknown */
+	FDateTime AccessTime;
+
+	/** The time the the file or directory was last modified, or FDateTime::MinValue if the modification time is unknown */
+	FDateTime ModificationTime;
+
+	/** Size of the file (in bytes), or -1 if the file size is unknown */
+	int64 FileSize;
+
+	/** True if this data is for a directory, false if it's for a file */
+	bool bIsDirectory : 1;
+
+	/** True if this file is read-only */
+	bool bIsReadOnly : 1;
+
+	/** True if file or directory was found, false otherwise. Note that this value being true does not ensure that the other members are filled in with meaningful data, as not all file systems have access to all of this data */
+	bool bIsValid : 1;
+};
+
+
 /**
 * File I/O Interface
 **/
@@ -165,7 +215,10 @@ public:
 	/** Delete a directory and return true if the directory was deleted or otherwise does not exist. **/
 	virtual bool		DeleteDirectory(const TCHAR* Directory) = 0;
 
-	/** Base class for file and directory visitors. **/
+	/** Return the stat data for the given file or directory. Check the FFileStatData::bIsValid member before using the returned data */
+	virtual FFileStatData GetStatData(const TCHAR* FilenameOrDirectory) = 0;
+
+	/** Base class for file and directory visitors that take only the name. **/
 	class FDirectoryVisitor
 	{
 	public:
@@ -177,6 +230,20 @@ public:
 		**/
 		virtual bool Visit(const TCHAR* FilenameOrDirectory, bool bIsDirectory) = 0;
 	};
+
+	/** Base class for file and directory visitors that take all the stat data. **/
+	class FDirectoryStatVisitor
+	{
+	public:
+		/** 
+		 * Callback for a single file or a directory in a directory iteration.
+		 * @param FilenameOrDirectory		If bIsDirectory is true, this is a directory (with no trailing path delimiter), otherwise it is a file name.
+		 * @param StatData					The stat data for the file or directory.
+		 * @return							true if the iteration should continue.
+		**/
+		virtual bool Visit(const TCHAR* FilenameOrDirectory, const FFileStatData& StatData) = 0;
+	};
+
 	/** 
 	 * Call the Visit function of the visitor once for each file or directory in a single directory. This function does not explore subdirectories.
 	 * @param Directory		The directory to iterate the contents of.
@@ -185,6 +252,14 @@ public:
 	**/
 	virtual bool		IterateDirectory(const TCHAR* Directory, FDirectoryVisitor& Visitor) = 0;
 
+	/** 
+	 * Call the Visit function of the visitor once for each file or directory in a single directory. This function does not explore subdirectories.
+	 * @param Directory		The directory to iterate the contents of.
+	 * @param Visitor		Visitor to call for each element of the directory
+	 * @return				false if the directory did not exist or if the visitor returned false.
+	**/
+	virtual bool		IterateDirectoryStat(const TCHAR* Directory, FDirectoryStatVisitor& Visitor) = 0;
+
 
 	/////////////////////////////////////////////////////////////////////////////////////////////////////////
 	/////////// Utility Functions. These have a default implementation that uses the pure virtual operations.
@@ -199,6 +274,14 @@ public:
 	**/
 	virtual bool IterateDirectoryRecursively(const TCHAR* Directory, FDirectoryVisitor& Visitor);
 
+	/** 
+	 * Call the Visit function of the visitor once for each file or directory in a directory tree. This function explores subdirectories.
+	 * @param Directory		The directory to iterate the contents of, recursively.
+	 * @param Visitor		Visitor to call for each element of the directory and each element of all subdirectories.
+	 * @return				false if the directory did not exist or if the visitor returned false.
+	**/
+	virtual bool IterateDirectoryStatRecursively(const TCHAR* Directory, FDirectoryStatVisitor& Visitor);
+
 	/** 
 	 * Delete all files and subdirectories in a directory, then delete the directory itself
 	 * @param Directory		The directory to delete.
diff --git a/Engine/Source/Runtime/Core/Public/HAL/FileManager.h b/Engine/Source/Runtime/Core/Public/HAL/FileManager.h
--- a/Engine/Source/Runtime/Core/Public/HAL/FileManager.h
+++ b/Engine/Source/Runtime/Core/Public/HAL/FileManager.h
@@ -105,6 +105,9 @@ public:
 	/** Deletes a directory. */
 	virtual bool DeleteDirectory( const TCHAR* Path, bool RequireExists=0, bool Tree=0 )=0;
 
+	/** Return the stat data for the given file or directory. Check the FFileStatData::bIsValid member before using the returned data */
+	virtual FFileStatData GetStatData(const TCHAR* FilenameOrDirectory) = 0;
+
 	/** Finds file or directories. */
 	virtual void FindFiles( TArray<FString>& FileNames, const TCHAR* Filename, bool Files, bool Directories)=0;
 
@@ -139,6 +142,22 @@ public:
 	**/
 	virtual bool IterateDirectoryRecursively(const TCHAR* Directory, IPlatformFile::FDirectoryVisitor& Visitor) = 0;
 
+	/** 
+	 * Call the Visit function of the visitor once for each file or directory in a single directory. This function does not explore subdirectories.
+	 * @param Directory		The directory to iterate the contents of.
+	 * @param Visitor		Visitor to call for each element of the directory
+	 * @return				false if the directory did not exist or if the visitor returned false.
+	**/
+	virtual bool IterateDirectoryStat(const TCHAR* Directory, IPlatformFile::FDirectoryStatVisitor& Visitor) = 0;
+
+	/** 
+	 * Call the Visit function of the visitor once for each file or directory in a directory tree. This function explores subdirectories.
+	 * @param Directory		The directory to iterate the contents of, recursively.
+	 * @param Visitor		Visitor to call for each element of the directory and each element of all subdirectories.
+	 * @return				false if the directory did not exist or if the visitor returned false.
+	**/
+	virtual bool IterateDirectoryStatRecursively(const TCHAR* Directory, IPlatformFile::FDirectoryStatVisitor& Visitor) = 0;
+
 	/** Gets the age of a file measured in seconds. */
 	virtual double GetFileAgeSeconds( const TCHAR* Filename )=0;
 
diff --git a/Engine/Source/Runtime/Core/Public/HAL/FileManagerGeneric.h b/Engine/Source/Runtime/Core/Public/HAL/FileManagerGeneric.h
--- a/Engine/Source/Runtime/Core/Public/HAL/FileManagerGeneric.h
+++ b/Engine/Source/Runtime/Core/Public/HAL/FileManagerGeneric.h
@@ -74,6 +74,7 @@ public:
 	virtual bool	MakeDirectory( const TCHAR* Path, bool Tree=0 ) override;
 	virtual bool	DeleteDirectory( const TCHAR* Path, bool RequireExists=0, bool Tree=0 ) override;
 
+	virtual FFileStatData GetStatData(const TCHAR* FilenameOrDirectory) override;
 
 	/**
 	 * Finds all the files within the given directory, with optional file extension filter.
@@ -103,6 +104,22 @@ public:
 	**/
 	bool IterateDirectoryRecursively(const TCHAR* Directory, IPlatformFile::FDirectoryVisitor& Visitor) override;
 
+	/** 
+	 * Call the Visit function of the visitor once for each file or directory in a single directory. This function does not explore subdirectories.
+	 * @param Directory		The directory to iterate the contents of.
+	 * @param Visitor		Visitor to call for each element of the directory
+	 * @return				false if the directory did not exist or if the visitor returned false.
+	**/
+	bool IterateDirectoryStat(const TCHAR* Directory, IPlatformFile::FDirectoryStatVisitor& Visitor) override;
+
+	/** 
+	 * Call the Visit function of the visitor once for each file or directory in a directory tree. This function explores subdirectories.
+	 * @param Directory		The directory to iterate the contents of, recursively.
+	 * @param Visitor		Visitor to call for each element of the directory and each element of all subdirectories.
+	 * @return				false if the directory did not exist or if the visitor returned false.
+	**/
+	bool IterateDirectoryStatRecursively(const TCHAR* Directory, IPlatformFile::FDirectoryStatVisitor& Visitor) override;
+
 	/**
 	 * Converts passed in filename to use a relative path.
 	 *
diff --git a/Engine/Source/Runtime/Core/Public/HAL/IPlatformFileCachedWrapper.h b/Engine/Source/Runtime/Core/Public/HAL/IPlatformFileCachedWrapper.h
--- a/Engine/Source/Runtime/Core/Public/HAL/IPlatformFileCachedWrapper.h
+++ b/Engine/Source/Runtime/Core/Public/HAL/IPlatformFileCachedWrapper.h
@@ -310,6 +310,10 @@ public:
 	{
 		return LowerLevel->DeleteDirectory(Directory);
 	}
+	virtual FFileStatData GetStatData(const TCHAR* FilenameOrDirectory) override
+	{
+		return LowerLevel->GetStatData(FilenameOrDirectory);
+	}
 	virtual bool		IterateDirectory(const TCHAR* Directory, IPlatformFile::FDirectoryVisitor& Visitor) override
 	{
 		return LowerLevel->IterateDirectory(Directory, Visitor);
@@ -318,6 +322,14 @@ public:
 	{
 		return LowerLevel->IterateDirectoryRecursively(Directory, Visitor);
 	}
+	virtual bool		IterateDirectoryStat(const TCHAR* Directory, IPlatformFile::FDirectoryStatVisitor& Visitor) override
+	{
+		return LowerLevel->IterateDirectoryStat(Directory, Visitor);
+	}
+	virtual bool		IterateDirectoryStatRecursively(const TCHAR* Directory, IPlatformFile::FDirectoryStatVisitor& Visitor) override
+	{
+		return LowerLevel->IterateDirectoryStatRecursively(Directory, Visitor);
+	}
 	virtual bool		DeleteDirectoryRecursively(const TCHAR* Directory) override
 	{
 		return LowerLevel->DeleteDirectoryRecursively(Directory);
diff --git a/Engine/Source/Runtime/Core/Public/HAL/IPlatformFileLogWrapper.h b/Engine/Source/Runtime/Core/Public/HAL/IPlatformFileLogWrapper.h
--- a/Engine/Source/Runtime/Core/Public/HAL/IPlatformFileLogWrapper.h
+++ b/Engine/Source/Runtime/Core/Public/HAL/IPlatformFileLogWrapper.h
@@ -257,6 +257,15 @@ public:
 		FILE_LOG(LogPlatformFile, Log, TEXT("DeleteDirectory return %d [%fms]"), int32(Result), ThisTime);
 		return Result;
 	}
+	virtual FFileStatData GetStatData(const TCHAR* FilenameOrDirectory) override
+	{
+		FILE_LOG(LogPlatformFile, Log, TEXT("GetStatData %s"), FilenameOrDirectory);
+		double StartTime = FPlatformTime::Seconds();
+		FFileStatData Result = LowerLevel->GetStatData(FilenameOrDirectory);
+		float ThisTime = 1000.0f * float(FPlatformTime::Seconds() - StartTime);
+		FILE_LOG(LogPlatformFile, Log, TEXT("GetStatData return %d [%fms]"), int32(Result.bIsValid), ThisTime);
+		return Result;
+	}
 
 	class FLogVisitor : public IPlatformFile::FDirectoryVisitor
 	{
@@ -266,7 +275,7 @@ public:
 			: Visitor(InVisitor)
 		{
 		}
-		virtual bool Visit(const TCHAR* FilenameOrDirectory, bool bIsDirectory)
+		virtual bool Visit(const TCHAR* FilenameOrDirectory, bool bIsDirectory) override
 		{
 			FILE_LOG(LogPlatformFile, Verbose, TEXT("Visit %s %d"), FilenameOrDirectory, int32(bIsDirectory));
 			double StartTime = FPlatformTime::Seconds();
@@ -297,6 +306,47 @@ public:
 		FILE_LOG(LogPlatformFile, Log, TEXT("IterateDirectoryRecursively return %d [%fms]"), int32(Result), ThisTime);
 		return Result;
 	}
+
+	class FLogStatVisitor : public IPlatformFile::FDirectoryStatVisitor
+	{
+	public:
+		FDirectoryStatVisitor&	Visitor;
+		FLogStatVisitor(FDirectoryStatVisitor& InVisitor)
+			: Visitor(InVisitor)
+		{
+		}
+		virtual bool Visit(const TCHAR* FilenameOrDirectory, const FFileStatData& StatData) override
+		{
+			FILE_LOG(LogPlatformFile, Verbose, TEXT("Visit %s %d"), FilenameOrDirectory, int32(StatData.bIsDirectory));
+			double StartTime = FPlatformTime::Seconds();
+			bool Result = Visitor.Visit(FilenameOrDirectory, StatData);
+			float ThisTime = 1000.0f * float(FPlatformTime::Seconds() - StartTime);
+			FILE_LOG(LogPlatformFile, Verbose, TEXT("Visit return %d [%fms]"), int32(Result), ThisTime);
+			return Result;
+		}
+	};
+
+	virtual bool		IterateDirectoryStat(const TCHAR* Directory, IPlatformFile::FDirectoryStatVisitor& Visitor) override
+	{
+		FILE_LOG(LogPlatformFile, Log, TEXT("IterateDirectoryStat %s"), Directory);
+		double StartTime = FPlatformTime::Seconds();
+		FLogStatVisitor LogVisitor(Visitor);
+		bool Result = LowerLevel->IterateDirectoryStat(Directory, LogVisitor);
+		float ThisTime = 1000.0f * float(FPlatformTime::Seconds() - StartTime);
+		FILE_LOG(LogPlatformFile, Log, TEXT("IterateDirectoryStat return %d [%fms]"), int32(Result), ThisTime);
+		return Result;
+	}
+	virtual bool		IterateDirectoryStatRecursively(const TCHAR* Directory, IPlatformFile::FDirectoryStatVisitor& Visitor) override
+	{
+		FILE_LOG(LogPlatformFile, Log, TEXT("IterateDirectoryStatRecursively %s"), Directory);
+		double StartTime = FPlatformTime::Seconds();
+		FLogStatVisitor LogVisitor(Visitor);
+		bool Result = LowerLevel->IterateDirectoryStatRecursively(Directory, LogVisitor);
+		float ThisTime = 1000.0f * float(FPlatformTime::Seconds() - StartTime);
+		FILE_LOG(LogPlatformFile, Log, TEXT("IterateDirectoryStatRecursively return %d [%fms]"), int32(Result), ThisTime);
+		return Result;
+	}
+
 	virtual bool		DeleteDirectoryRecursively(const TCHAR* Directory) override
 	{
 		FILE_LOG(LogPlatformFile, Log, TEXT("DeleteDirectoryRecursively %s"), Directory);
diff --git a/Engine/Source/Runtime/Core/Public/HAL/IPlatformFileOpenLogWrapper.h b/Engine/Source/Runtime/Core/Public/HAL/IPlatformFileOpenLogWrapper.h
--- a/Engine/Source/Runtime/Core/Public/HAL/IPlatformFileOpenLogWrapper.h
+++ b/Engine/Source/Runtime/Core/Public/HAL/IPlatformFileOpenLogWrapper.h
@@ -168,6 +168,10 @@ public:
 	{
 		return LowerLevel->DeleteDirectory(Directory);
 	}
+	virtual FFileStatData GetStatData(const TCHAR* FilenameOrDirectory) override
+	{
+		return LowerLevel->GetStatData(FilenameOrDirectory);
+	}
 	virtual bool		IterateDirectory(const TCHAR* Directory, IPlatformFile::FDirectoryVisitor& Visitor) override
 	{
 		return LowerLevel->IterateDirectory( Directory, Visitor );
@@ -176,6 +180,14 @@ public:
 	{
 		return LowerLevel->IterateDirectoryRecursively( Directory, Visitor );
 	}
+	virtual bool		IterateDirectoryStat(const TCHAR* Directory, IPlatformFile::FDirectoryStatVisitor& Visitor) override
+	{
+		return LowerLevel->IterateDirectoryStat( Directory, Visitor );
+	}
+	virtual bool		IterateDirectoryStatRecursively(const TCHAR* Directory, IPlatformFile::FDirectoryStatVisitor& Visitor) override
+	{
+		return LowerLevel->IterateDirectoryStatRecursively( Directory, Visitor );
+	}
 	virtual bool		DeleteDirectoryRecursively(const TCHAR* Directory) override
 	{
 		return LowerLevel->DeleteDirectoryRecursively( Directory );
diff --git a/Engine/Source/Runtime/Core/Public/HAL/IPlatformFileProfilerWrapper.h b/Engine/Source/Runtime/Core/Public/HAL/IPlatformFileProfilerWrapper.h
--- a/Engine/Source/Runtime/Core/Public/HAL/IPlatformFileProfilerWrapper.h
+++ b/Engine/Source/Runtime/Core/Public/HAL/IPlatformFileProfilerWrapper.h
@@ -58,6 +58,8 @@ struct FProfiledFileStatsOp : public FProfiledFileStatsBase
 		Create,
 		Copy,
 		Iterate,
+		IterateStat,
+		GetStatData,
 
 		Count
 	};
@@ -406,6 +408,15 @@ public:
 		return Result;
 	}
 
+	virtual FFileStatData GetStatData(const TCHAR* FilenameOrDirectory) override
+	{
+		StatsType* FileStat = CreateStat( FilenameOrDirectory );
+		FProfiledFileStatsOp* OpStat = FileStat->CreateOpStat( FProfiledFileStatsOp::EOpType::GetStatData );
+		FFileStatData Result = LowerLevel->GetStatData(FilenameOrDirectory);
+		OpStat->Duration += FPlatformTime::Seconds() * 1000.0 - OpStat->LastOpTime;
+		return Result;
+	}
+
 	virtual bool		IterateDirectory(const TCHAR* Directory, IPlatformFile::FDirectoryVisitor& Visitor) override
 	{
 		StatsType* FileStat = CreateStat( Directory );
@@ -422,6 +433,24 @@ public:
 		OpStat->Duration += FPlatformTime::Seconds() * 1000.0 - OpStat->LastOpTime;
 		return Result;
 	}
+
+	virtual bool		IterateDirectoryStat(const TCHAR* Directory, IPlatformFile::FDirectoryStatVisitor& Visitor) override
+	{
+		StatsType* FileStat = CreateStat( Directory );
+		FProfiledFileStatsOp* OpStat = FileStat->CreateOpStat( FProfiledFileStatsOp::EOpType::IterateStat );
+		bool Result = LowerLevel->IterateDirectoryStat( Directory, Visitor );
+		OpStat->Duration += FPlatformTime::Seconds() * 1000.0 - OpStat->LastOpTime;
+		return Result;
+	}
+	virtual bool		IterateDirectoryStatRecursively(const TCHAR* Directory, IPlatformFile::FDirectoryStatVisitor& Visitor) override
+	{
+		StatsType* FileStat = CreateStat( Directory );
+		FProfiledFileStatsOp* OpStat = FileStat->CreateOpStat( FProfiledFileStatsOp::EOpType::IterateStat );
+		bool Result = LowerLevel->IterateDirectoryStatRecursively( Directory, Visitor );
+		OpStat->Duration += FPlatformTime::Seconds() * 1000.0 - OpStat->LastOpTime;
+		return Result;
+	}
+
 	virtual bool		DeleteDirectoryRecursively(const TCHAR* Directory) override
 	{
 		StatsType* FileStat = CreateStat( Directory );
@@ -608,6 +637,11 @@ public:
 	{
 		return LowerLevel->DeleteDirectory(Directory);
 	}
+	
+	virtual FFileStatData GetStatData(const TCHAR* FilenameOrDirectory) override
+	{
+		return LowerLevel->GetStatData(FilenameOrDirectory);
+	}
 
 	virtual bool		IterateDirectory(const TCHAR* Directory, IPlatformFile::FDirectoryVisitor& Visitor) override
 	{
@@ -617,6 +651,16 @@ public:
 	{
 		return LowerLevel->IterateDirectoryRecursively( Directory, Visitor );
 	}
+
+	virtual bool		IterateDirectoryStat(const TCHAR* Directory, IPlatformFile::FDirectoryStatVisitor& Visitor) override
+	{
+		return LowerLevel->IterateDirectoryStat( Directory, Visitor );
+	}
+	virtual bool		IterateDirectoryStatRecursively(const TCHAR* Directory, IPlatformFile::FDirectoryStatVisitor& Visitor) override
+	{
+		return LowerLevel->IterateDirectoryStatRecursively( Directory, Visitor );
+	}
+
 	virtual bool		DeleteDirectoryRecursively(const TCHAR* Directory) override
 	{
 		return LowerLevel->DeleteDirectoryRecursively( Directory );
diff --git a/Engine/Source/Runtime/Core/Public/Linux/LinuxPlatformFile.h b/Engine/Source/Runtime/Core/Public/Linux/LinuxPlatformFile.h
--- a/Engine/Source/Runtime/Core/Public/Linux/LinuxPlatformFile.h
+++ b/Engine/Source/Runtime/Core/Public/Linux/LinuxPlatformFile.h
@@ -36,6 +36,14 @@ public:
 	virtual bool DirectoryExists(const TCHAR* Directory) override;
 	virtual bool CreateDirectory(const TCHAR* Directory) override;
 	virtual bool DeleteDirectory(const TCHAR* Directory) override;
+
+	virtual FFileStatData GetStatData(const TCHAR* FilenameOrDirectory) override;
+
 	bool CreateDirectoriesFromPath(const TCHAR* Path);
-	bool IterateDirectory(const TCHAR* Directory, FDirectoryVisitor& Visitor) override;
+
+	virtual bool IterateDirectory(const TCHAR* Directory, FDirectoryVisitor& Visitor) override;
+	virtual bool IterateDirectoryStat(const TCHAR* Directory, FDirectoryStatVisitor& Visitor) override;
+
+protected:
+	bool IterateDirectoryCommon(const TCHAR* Directory, const TFunctionRef<bool(struct dirent*)>& Visitor);
 };
diff --git a/Engine/Source/Runtime/NetworkFile/Private/NetworkPlatformFile.cpp b/Engine/Source/Runtime/NetworkFile/Private/NetworkPlatformFile.cpp
--- a/Engine/Source/Runtime/NetworkFile/Private/NetworkPlatformFile.cpp
+++ b/Engine/Source/Runtime/NetworkFile/Private/NetworkPlatformFile.cpp
@@ -456,6 +456,14 @@ bool FNetworkPlatformFile::DeleteDirectory(const TCHAR* Directory)
 	return InnerPlatformFile->DeleteDirectory(Directory);
 }
 
+FFileStatData FNetworkPlatformFile::GetStatData(const TCHAR* FilenameOrDirectory)
+{
+//	FScopeLock ScopeLock(&SynchronizationObject);
+
+	// perform a local operation
+	return InnerPlatformFile->GetStatData(FilenameOrDirectory);
+}
+
 bool FNetworkPlatformFile::IterateDirectory(const TCHAR* InDirectory, IPlatformFile::FDirectoryVisitor& Visitor)
 {
 //	FScopeLock ScopeLock(&SynchronizationObject);
@@ -532,6 +540,102 @@ bool FNetworkPlatformFile::IterateDirectoryRecursively(const TCHAR* InDirectory,
 	return RetVal;
 }
 
+bool FNetworkPlatformFile::IterateDirectoryStat(const TCHAR* InDirectory, IPlatformFile::FDirectoryStatVisitor& Visitor)
+{
+//	FScopeLock ScopeLock(&SynchronizationObject);
+
+	// for .dll, etc searches that don't specify a path, we need to strip off the path
+	// before we send it to the visitor
+	bool bHadNoPath = InDirectory[0] == 0;
+
+	// local files go right to the source
+	FString RelativeDirectory = InDirectory;
+	MakeStandardNetworkFilename(RelativeDirectory);
+	if (IsInLocalDirectory(RelativeDirectory))
+	{
+		return InnerPlatformFile->IterateDirectoryStat(InDirectory, Visitor);
+	}
+
+	// we loop until this is false
+	bool RetVal = true;
+
+	FServerTOC::FDirectory* ServerDirectory = ServerFiles.FindDirectory(RelativeDirectory);
+	if (ServerDirectory != NULL)
+	{
+		// loop over the server files and look if they are in this exact directory
+		for (FServerTOC::FDirectory::TIterator It(*ServerDirectory); It && RetVal == true; ++It)
+		{
+			if (FPaths::GetPath(It.Key()) == RelativeDirectory)
+			{
+				// timestamps of 0 mean directories
+				bool bIsDirectory = It.Value() == 0;
+			
+				// todo: this data is just wrong for most things, but can we afford to get the files from the server to get the correct info? Could the server provide this instead?
+				const FFileStatData StatData(
+					FDateTime::MinValue(), 
+					FDateTime::MinValue(), 
+					(bIsDirectory) ? FDateTime::MinValue() : It.Value(),
+					-1,		// FileSize
+					bIsDirectory, 
+					true	// IsReadOnly
+					);
+
+				// visit (stripping off the path if needed)
+				RetVal = Visitor.Visit(bHadNoPath ? *FPaths::GetCleanFilename(It.Key()) : *It.Key(), StatData);
+			}
+		}
+	}
+
+	return RetVal;
+}
+
+bool FNetworkPlatformFile::IterateDirectoryStatRecursively(const TCHAR* InDirectory, IPlatformFile::FDirectoryStatVisitor& Visitor)
+{
+//	FScopeLock ScopeLock(&SynchronizationObject);
+
+	// local files go right to the source
+	FString RelativeDirectory = InDirectory;
+	MakeStandardNetworkFilename(RelativeDirectory);
+
+	if (IsInLocalDirectory(RelativeDirectory))
+	{
+		return InnerPlatformFile->IterateDirectoryStatRecursively(InDirectory, Visitor);
+	}
+
+	// we loop until this is false
+	bool RetVal = true;
+
+	for (TMap<FString, FServerTOC::FDirectory*>::TIterator DirIt(ServerFiles.Directories); DirIt && RetVal == true; ++DirIt)
+	{
+		if (DirIt.Key().StartsWith(RelativeDirectory))
+		{
+			FServerTOC::FDirectory& ServerDirectory = *DirIt.Value();
+		
+			// loop over the server files and look if they are in this exact directory
+			for (FServerTOC::FDirectory::TIterator It(ServerDirectory); It && RetVal == true; ++It)
+			{
+				// timestamps of 0 mean directories
+				bool bIsDirectory = It.Value() == 0;
+
+				// todo: this data is just wrong for most things, but can we afford to get the files from the server to get the correct info? Could the server provide this instead?
+				const FFileStatData StatData(
+					FDateTime::MinValue(), 
+					FDateTime::MinValue(), 
+					(bIsDirectory) ? FDateTime::MinValue() : It.Value(),
+					0,		// FileSize
+					bIsDirectory, 
+					true	// IsReadOnly
+					);
+
+				// visit!
+				RetVal = Visitor.Visit(*It.Key(), StatData);
+			}
+		}
+	}
+
+	return RetVal;
+}
+
 bool FNetworkPlatformFile::DeleteDirectoryRecursively(const TCHAR* Directory)
 {
 //	FScopeLock ScopeLock(&SynchronizationObject);
@@ -607,12 +711,12 @@ void FNetworkPlatformFile::GetFileInfo(const TCHAR* Filename, FFileInfo& Info)
 		EnsureFileIsLocal(RelativeFilename);
 	}
 
-	// @todo: This is pretty inefficient. If GetFileInfo was a first class function, this could be avoided
-	Info.FileExists = InnerPlatformFile->FileExists(Filename);
-	Info.ReadOnly = InnerPlatformFile->IsReadOnly(Filename);
-	Info.Size = InnerPlatformFile->FileSize(Filename);
-	Info.TimeStamp = InnerPlatformFile->GetTimeStamp(Filename);
-	Info.AccessTimeStamp = InnerPlatformFile->GetAccessTimeStamp(Filename);
+	const FFileStatData StatData = InnerPlatformFile->GetStatData(Filename);
+	Info.FileExists = StatData.bIsValid && !StatData.bIsDirectory;
+	Info.ReadOnly = StatData.bIsReadOnly;
+	Info.Size = StatData.FileSize;
+	Info.TimeStamp = StatData.ModificationTime;
+	Info.AccessTimeStamp = StatData.AccessTime;
 }
 
 void FNetworkPlatformFile::ConvertServerFilenameToClientFilename(FString& FilenameToConvert)
diff --git a/Engine/Source/Runtime/NetworkFile/Public/NetworkPlatformFile.h b/Engine/Source/Runtime/NetworkFile/Public/NetworkPlatformFile.h
--- a/Engine/Source/Runtime/NetworkFile/Public/NetworkPlatformFile.h
+++ b/Engine/Source/Runtime/NetworkFile/Public/NetworkPlatformFile.h
@@ -101,8 +101,14 @@ public:
 	virtual bool		CreateDirectory(const TCHAR* Directory) override;
 	virtual bool		DeleteDirectory(const TCHAR* Directory) override;
 
+	virtual FFileStatData GetStatData(const TCHAR* FilenameOrDirectory) override;
+
 	virtual bool		IterateDirectory(const TCHAR* Directory, IPlatformFile::FDirectoryVisitor& Visitor) override;
 	virtual bool		IterateDirectoryRecursively(const TCHAR* Directory, IPlatformFile::FDirectoryVisitor& Visitor) override;
+
+	virtual bool		IterateDirectoryStat(const TCHAR* Directory, IPlatformFile::FDirectoryStatVisitor& Visitor) override;
+	virtual bool		IterateDirectoryStatRecursively(const TCHAR* Directory, IPlatformFile::FDirectoryStatVisitor& Visitor) override;
+
 	virtual bool		DeleteDirectoryRecursively(const TCHAR* Directory) override;
 	virtual bool		CopyFile(const TCHAR* To, const TCHAR* From) override;
 
diff --git a/Engine/Source/Runtime/PakFile/Public/IPlatformFilePak.h b/Engine/Source/Runtime/PakFile/Public/IPlatformFilePak.h
--- a/Engine/Source/Runtime/PakFile/Public/IPlatformFilePak.h
+++ b/Engine/Source/Runtime/PakFile/Public/IPlatformFilePak.h
@@ -1160,6 +1160,40 @@ public:
 		return LowerLevel->DeleteDirectory(Directory);
 	}
 
+	virtual FFileStatData GetStatData(const TCHAR* FilenameOrDirectory) override
+	{
+		// Check pak files first.
+		FPakFile* PakFile = nullptr;
+		const FPakEntry* FileEntry = FindFileInPakFiles(FilenameOrDirectory, &PakFile);
+		if (FileEntry)
+		{
+			return FFileStatData(
+				PakFile->GetTimestamp(),
+				PakFile->GetTimestamp(),
+				PakFile->GetTimestamp(),
+				(FileEntry->CompressionMethod != COMPRESS_None) ? FileEntry->UncompressedSize : FileEntry->Size, 
+				false,	// IsDirectory
+				true	// IsReadOnly
+				);
+		}
+
+		// Then check pak directories
+		if (DirectoryExistsInPakFiles(FilenameOrDirectory))
+		{
+			return FFileStatData(
+				PakFile->GetTimestamp(),
+				PakFile->GetTimestamp(),
+				PakFile->GetTimestamp(),
+				-1,		// FileSize
+				true,	// IsDirectory
+				true	// IsReadOnly
+				);
+		}
+
+		// Fall back to lower level.
+		return LowerLevel->GetStatData(FilenameOrDirectory);
+	}
+
 	/**
 	 * Helper class to filter out files which have already been visited in one of the pak files.
 	 */
@@ -1264,6 +1298,130 @@ public:
 		return IPlatformFile::IterateDirectoryRecursively(Directory, PakVisitor);
 	}
 
+	/**
+	 * Helper class to filter out files which have already been visited in one of the pak files.
+	 */
+	class FPakStatVisitor : public IPlatformFile::FDirectoryStatVisitor
+	{
+	public:
+		/** Wrapped visitor. */
+		FDirectoryStatVisitor&	Visitor;
+		/** Visited pak files. */
+		TSet<FString>& VisitedPakFiles;
+		/** Cached list of pak files. */
+		TArray<FPakListEntry>& Paks;
+
+		/** Constructor. */
+		FPakStatVisitor(FDirectoryStatVisitor& InVisitor, TArray<FPakListEntry>& InPaks, TSet<FString>& InVisitedPakFiles)
+			: Visitor(InVisitor)
+			, VisitedPakFiles(InVisitedPakFiles)
+			, Paks(InPaks)
+		{}
+		virtual bool Visit(const TCHAR* FilenameOrDirectory, const FFileStatData& StatData)
+		{
+			if (StatData.bIsDirectory == false)
+			{
+				FString StandardFilename(FilenameOrDirectory);
+				FPaths::MakeStandardFilename(StandardFilename);
+
+				if (VisitedPakFiles.Contains(StandardFilename))
+				{
+					// Already visited, continue iterating.
+					return true;
+				}
+				else if (FPakPlatformFile::FindFileInPakFiles(Paks, FilenameOrDirectory, NULL))
+				{
+					VisitedPakFiles.Add(StandardFilename);
+				}
+			}			
+			return Visitor.Visit(FilenameOrDirectory, StatData);
+		}
+	};
+
+	virtual bool IterateDirectoryStat(const TCHAR* Directory, IPlatformFile::FDirectoryStatVisitor& Visitor) override
+	{
+		bool Result = true;
+		TSet<FString> FilesVisitedInPak;
+
+		TArray<FPakListEntry> Paks;
+		GetMountedPaks(Paks);
+
+		// Iterate pak files first
+		for (int32 PakIndex = 0; PakIndex < Paks.Num(); PakIndex++)
+		{
+			FPakFile& PakFile = *Paks[PakIndex].PakFile;
+			
+			const bool bIncludeFiles = true;
+			const bool bIncludeFolders = true;
+			TSet<FString> FilesVisitedInThisPak;
+			FString StandardDirectory = Directory;
+			FPaths::MakeStandardFilename(StandardDirectory);
+
+			PakFile.FindFilesAtPath(FilesVisitedInThisPak, *StandardDirectory, bIncludeFiles, bIncludeFolders);
+			for (TSet<FString>::TConstIterator SetIt(FilesVisitedInThisPak); SetIt && Result; ++SetIt)
+			{
+				const FString& Filename = *SetIt;
+				if (!FilesVisitedInPak.Contains(Filename))
+				{
+					bool bIsDir = Filename.Len() && Filename[Filename.Len() - 1] == '/';
+
+					int64 FileSize = -1;
+					if (!bIsDir)
+					{
+						const FPakEntry* FileEntry = FindFileInPakFiles(*Filename);
+						if (FileEntry)
+						{
+							FileSize = (FileEntry->CompressionMethod != COMPRESS_None) ? FileEntry->UncompressedSize : FileEntry->Size;
+						}
+					}
+
+					const FFileStatData StatData(
+						PakFile.GetTimestamp(),
+						PakFile.GetTimestamp(),
+						PakFile.GetTimestamp(),
+						FileSize, 
+						bIsDir,
+						true	// IsReadOnly
+						);
+
+					if (bIsDir)
+					{
+						Result = Visitor.Visit(*Filename.LeftChop(1), StatData) && Result;
+					}
+					else
+					{
+						Result = Visitor.Visit(*Filename, StatData) && Result;
+					}
+					FilesVisitedInPak.Add(Filename);
+				}
+			}
+		}
+		if (Result && LowerLevel->DirectoryExists(Directory))
+		{
+			if (FilesVisitedInPak.Num())
+			{
+				// Iterate inner filesystem using FPakVisitor
+				FPakStatVisitor PakVisitor(Visitor, Paks, FilesVisitedInPak);
+				Result = LowerLevel->IterateDirectoryStat(Directory, PakVisitor);
+			}
+			else
+			{
+				 // No point in using FPakVisitor as it will only slow things down.
+				Result = LowerLevel->IterateDirectoryStat(Directory, Visitor);
+			}
+		}
+		return Result;
+	}
+
+	virtual bool IterateDirectoryStatRecursively(const TCHAR* Directory, IPlatformFile::FDirectoryStatVisitor& Visitor) override
+	{
+		TSet<FString> FilesVisitedInPak;
+		TArray<FPakListEntry> Paks;
+		GetMountedPaks(Paks);
+		FPakStatVisitor PakVisitor(Visitor, Paks, FilesVisitedInPak);
+		return IPlatformFile::IterateDirectoryStatRecursively(Directory, PakVisitor);
+	}
+
 	virtual bool DeleteDirectoryRecursively(const TCHAR* Directory) override
 	{
 		// Can't delete directories existing in pak files. See DeleteDirectory(..) for more info.
diff --git a/Engine/Source/Runtime/SandboxFile/Public/IPlatformFileSandboxWrapper.h b/Engine/Source/Runtime/SandboxFile/Public/IPlatformFileSandboxWrapper.h
--- a/Engine/Source/Runtime/SandboxFile/Public/IPlatformFileSandboxWrapper.h
+++ b/Engine/Source/Runtime/SandboxFile/Public/IPlatformFileSandboxWrapper.h
@@ -415,6 +415,16 @@ public:
 		return LowerLevel->DeleteDirectory( *ConvertToSandboxPath( Directory ) );
 	}
 
+	virtual FFileStatData GetStatData(const TCHAR* FilenameOrDirectory) override
+	{
+		FFileStatData Result = LowerLevel->GetStatData( *ConvertToSandboxPath( FilenameOrDirectory ) );
+		if (!Result.bIsValid && (OkForInnerAccess(FilenameOrDirectory, false) && OkForInnerAccess(FilenameOrDirectory, true)))
+		{
+			Result = LowerLevel->GetStatData( FilenameOrDirectory );
+		}
+		return Result;
+	}
+
 	class FSandboxVisitor : public IPlatformFile::FDirectoryVisitor
 	{
 	public:
@@ -427,7 +437,7 @@ public:
 			, SandboxFile( InSandboxFile )
 		{
 		}
-		virtual bool Visit( const TCHAR* FilenameOrDirectory, bool bIsDirectory )
+		virtual bool Visit( const TCHAR* FilenameOrDirectory, bool bIsDirectory ) override
 		{
 			bool CanVisit = true;
 			FString LocalFilename( FilenameOrDirectory );
@@ -486,6 +496,77 @@ public:
 		return Result;
 	}
 
+	class FSandboxStatVisitor : public IPlatformFile::FDirectoryStatVisitor
+	{
+	public:
+		FDirectoryStatVisitor&	Visitor;
+		FSandboxPlatformFile& SandboxFile;
+		TSet<FString> VisitedSandboxFiles;
+
+		FSandboxStatVisitor( FDirectoryStatVisitor& InVisitor, FSandboxPlatformFile& InSandboxFile )
+			: Visitor( InVisitor )
+			, SandboxFile( InSandboxFile )
+		{
+		}
+		virtual bool Visit( const TCHAR* FilenameOrDirectory, const FFileStatData& StatData ) override
+		{
+			bool CanVisit = true;
+			FString LocalFilename( FilenameOrDirectory );
+			
+			if( FCString::Strnicmp( *LocalFilename, *SandboxFile.GetSandboxDirectory(), SandboxFile.GetSandboxDirectory().Len() ) == 0 )
+			{
+				// FilenameOrDirectory is already pointing to the sandbox directory so add it to the list of sanbox files.
+				// The filename is always stored with the abslute sandbox path.
+				VisitedSandboxFiles.Add( *LocalFilename );
+				// Now convert the sandbox path back to engine path because the sandbox folder should not be exposed
+				// to the engine and remain transparent.
+				LocalFilename = LocalFilename.Mid(SandboxFile.GetSandboxDirectory().Len());
+				if (LocalFilename.StartsWith(TEXT("Engine/")))
+				{
+					LocalFilename = SandboxFile.GetAbsoluteRootDirectory() / LocalFilename;
+				}
+				else
+				{
+					LocalFilename = SandboxFile.GetAbsolutePathToGameDirectory() / LocalFilename;
+				}
+			}
+			else
+			{
+				// Favourize Sandbox files over normal path files.
+				CanVisit = !VisitedSandboxFiles.Contains( SandboxFile.ConvertToSandboxPath(*LocalFilename ) )
+					&& SandboxFile.OkForInnerAccess(*LocalFilename, StatData.bIsDirectory);
+			}
+			if( CanVisit )
+			{
+				bool Result = Visitor.Visit( *LocalFilename, StatData );
+				return Result;
+			}
+			else
+			{
+				// Continue iterating.
+				return true;
+			}
+		}
+	};
+
+	virtual bool		IterateDirectoryStat(const TCHAR* Directory, IPlatformFile::FDirectoryStatVisitor& Visitor) override
+	{
+		FSandboxStatVisitor SandboxVisitor( Visitor, *this );
+		bool Result = false;
+		LowerLevel->IterateDirectoryStat( *ConvertToSandboxPath( Directory ), SandboxVisitor );
+		Result = LowerLevel->IterateDirectoryStat( Directory, SandboxVisitor );
+		return Result;
+	}
+
+	virtual bool		IterateDirectoryStatRecursively(const TCHAR* Directory, IPlatformFile::FDirectoryStatVisitor& Visitor) override
+	{
+		FSandboxStatVisitor SandboxVisitor( Visitor, *this );
+		bool Result = false;
+		LowerLevel->IterateDirectoryStatRecursively( *ConvertToSandboxPath( Directory ), SandboxVisitor );
+		Result = LowerLevel->IterateDirectoryStatRecursively( Directory, SandboxVisitor );
+		return Result;
+	}
+
 	virtual bool		DeleteDirectoryRecursively(const TCHAR* Directory) override
 	{
 		// Directories can be deleted only under the sandbox path
