bump up the read size of the pipe reads.
potential fix for UE-17454
#mac

diff --git a/Engine/Source/Runtime/Core/Private/Mac/MacPlatformProcess.cpp b/Engine/Source/Runtime/Core/Private/Mac/MacPlatformProcess.cpp
--- a/Engine/Source/Runtime/Core/Private/Mac/MacPlatformProcess.cpp
+++ b/Engine/Source/Runtime/Core/Private/Mac/MacPlatformProcess.cpp
@@ -933,7 +933,7 @@ FString FMacPlatformProcess::ReadPipe( void* ReadPipe )
 
 	FString Output;
 
-	const int32 READ_SIZE = 4096;
+	const int32 READ_SIZE = 8192;
 	ANSICHAR Buffer[READ_SIZE+1];
 	int32 BytesRead = 0;
 
