CIS fix

diff --git a/Engine/Source/Runtime/Core/Private/GenericPlatform/GenericPlatformCrashContext.cpp b/Engine/Source/Runtime/Core/Private/GenericPlatform/GenericPlatformCrashContext.cpp
--- a/Engine/Source/Runtime/Core/Private/GenericPlatform/GenericPlatformCrashContext.cpp
+++ b/Engine/Source/Runtime/Core/Private/GenericPlatform/GenericPlatformCrashContext.cpp
@@ -112,7 +112,7 @@ void FGenericCrashContext::SerializeContentToBuffer()
 	AddHeader();
 
 	BeginSection( *RuntimePropertiesTag );
-	AddCrashProperty( TEXT( "CrashVersion" ), ECrashDescVersions::VER_3_CrashContext );
+	AddCrashProperty( TEXT( "CrashVersion" ), (int32)ECrashDescVersions::VER_3_CrashContext );
 	AddCrashProperty( TEXT( "CrashGUID" ), *NCachedCrashContextProperties::CrashGUID );
 	AddCrashProperty( TEXT( "ProcessId" ), FPlatformProcess::GetCurrentProcessId() );
 	AddCrashProperty( TEXT( "IsInternalBuild" ), NCachedCrashContextProperties::bIsInternalBuild );
