Re-ordered updating of mac modifier key states to ensure that they are set before dispatching events to Slate Application.

This addresses UE-19036

diff --git a/Engine/Source/Runtime/Core/Private/Mac/MacApplication.cpp b/Engine/Source/Runtime/Core/Private/Mac/MacApplication.cpp
--- a/Engine/Source/Runtime/Core/Private/Mac/MacApplication.cpp
+++ b/Engine/Source/Runtime/Core/Private/Mac/MacApplication.cpp
@@ -499,35 +499,33 @@ void FMacApplication::ProcessEvent(const FDeferredMacEvent& Event)
 	TSharedPtr<FMacWindow> EventWindow = FindWindowByNSWindow(Event.Window);
 	if (Event.Type)
 	{
-		bool bModifiersValid = false;
-		
 		switch (Event.Type)
 		{
 			case NSMouseMoved:
 			case NSLeftMouseDragged:
 			case NSRightMouseDragged:
 			case NSOtherMouseDragged:
+				ConditionallyUpdateModifierKeys(Event);
 				ProcessMouseMovedEvent(Event, EventWindow);
-				bModifiersValid = true;
 				break;
 
 			case NSLeftMouseDown:
 			case NSRightMouseDown:
 			case NSOtherMouseDown:
+				ConditionallyUpdateModifierKeys(Event);
 				ProcessMouseDownEvent(Event, EventWindow);
-				bModifiersValid = true;
 				break;
 
 			case NSLeftMouseUp:
 			case NSRightMouseUp:
 			case NSOtherMouseUp:
+				ConditionallyUpdateModifierKeys(Event);
 				ProcessMouseUpEvent(Event, EventWindow);
-				bModifiersValid = true;
 				break;
 
 			case NSScrollWheel:
+				ConditionallyUpdateModifierKeys(Event);
 				ProcessScrollWheelEvent(Event, EventWindow);
-				bModifiersValid = true;
 				break;
 
 			case NSEventTypeMagnify:
@@ -535,46 +533,29 @@ void FMacApplication::ProcessEvent(const FDeferredMacEvent& Event)
 			case NSEventTypeRotate:
 			case NSEventTypeBeginGesture:
 			case NSEventTypeEndGesture:
+				ConditionallyUpdateModifierKeys(Event);
 				ProcessGestureEvent(Event);
-				bModifiersValid = true;
 				break;
 
 			case NSKeyDown:
+				ConditionallyUpdateModifierKeys(Event);
 				ProcessKeyDownEvent(Event, EventWindow);
-				bModifiersValid = true;
 				break;
 
 			case NSKeyUp:
+				ConditionallyUpdateModifierKeys(Event);
 				ProcessKeyUpEvent(Event);
-				bModifiersValid = true;
 				break;
 				
 			case NSFlagsChanged:
-				bModifiersValid = true;
+				ConditionallyUpdateModifierKeys(Event);
 				break;
 				
 			case NSMouseEntered:
 			case NSMouseExited:
-				bModifiersValid = true;
+				ConditionallyUpdateModifierKeys(Event);
 				break;
 		}
-		
-		if (bModifiersValid && CurrentModifierFlags != Event.ModifierFlags)
-		{
-			NSUInteger ModifierFlags = Event.ModifierFlags;
-			
-			HandleModifierChange(ModifierFlags, (1<<4), 7, MMK_RightCommand);
-			HandleModifierChange(ModifierFlags, (1<<3), 6, MMK_LeftCommand);
-			HandleModifierChange(ModifierFlags, (1<<1), 0, MMK_LeftShift);
-			HandleModifierChange(ModifierFlags, (1<<16), 8, MMK_CapsLock);
-			HandleModifierChange(ModifierFlags, (1<<5), 4, MMK_LeftAlt);
-			HandleModifierChange(ModifierFlags, (1<<0), 2, MMK_LeftControl);
-			HandleModifierChange(ModifierFlags, (1<<2), 1, MMK_RightShift);
-			HandleModifierChange(ModifierFlags, (1<<6), 5, MMK_RightAlt);
-			HandleModifierChange(ModifierFlags, (1<<13), 3, MMK_RightControl);
-			
-			CurrentModifierFlags = ModifierFlags;
-		}
 	}
 	else if (EventWindow.IsValid())
 	{
@@ -1141,6 +1122,26 @@ void FMacApplication::OnWindowsReordered(bool bIsAppInBackground)
 	}
 }
 
+void FMacApplication::ConditionallyUpdateModifierKeys(const FDeferredMacEvent& Event)
+{
+	if (CurrentModifierFlags != Event.ModifierFlags)
+	{
+		NSUInteger ModifierFlags = Event.ModifierFlags;
+		
+		HandleModifierChange(ModifierFlags, (1<<4), 7, MMK_RightCommand);
+		HandleModifierChange(ModifierFlags, (1<<3), 6, MMK_LeftCommand);
+		HandleModifierChange(ModifierFlags, (1<<1), 0, MMK_LeftShift);
+		HandleModifierChange(ModifierFlags, (1<<16), 8, MMK_CapsLock);
+		HandleModifierChange(ModifierFlags, (1<<5), 4, MMK_LeftAlt);
+		HandleModifierChange(ModifierFlags, (1<<0), 2, MMK_LeftControl);
+		HandleModifierChange(ModifierFlags, (1<<2), 1, MMK_RightShift);
+		HandleModifierChange(ModifierFlags, (1<<6), 5, MMK_RightAlt);
+		HandleModifierChange(ModifierFlags, (1<<13), 3, MMK_RightControl);
+		
+		CurrentModifierFlags = ModifierFlags;
+	}
+}
+
 void FMacApplication::HandleModifierChange(NSUInteger NewModifierFlags, NSUInteger FlagsShift, NSUInteger UE4Shift, EMacModifierKeys TranslatedCode)
 {
 	const bool CurrentPressed = (CurrentModifierFlags & FlagsShift) != 0;
diff --git a/Engine/Source/Runtime/Core/Public/Mac/MacApplication.h b/Engine/Source/Runtime/Core/Public/Mac/MacApplication.h
--- a/Engine/Source/Runtime/Core/Public/Mac/MacApplication.h
+++ b/Engine/Source/Runtime/Core/Public/Mac/MacApplication.h
@@ -216,6 +216,7 @@ private:
 	void OnApplicationWillResignActive();
 	void OnWindowsReordered(bool bIsAppInBackground);
 
+	void ConditionallyUpdateModifierKeys(const FDeferredMacEvent& Event);
 	void HandleModifierChange(NSUInteger NewModifierFlags, NSUInteger FlagsShift, NSUInteger UE4Shift, EMacModifierKeys TranslatedCode);
 
 	FCocoaWindow* FindEventWindow(NSEvent* CocoaEvent) const;
