PS4 - Allow dumping callstack during an assert

diff --git a/Engine/Source/Runtime/Core/Private/Misc/OutputDevice.cpp b/Engine/Source/Runtime/Core/Private/Misc/OutputDevice.cpp
--- a/Engine/Source/Runtime/Core/Private/Misc/OutputDevice.cpp
+++ b/Engine/Source/Runtime/Core/Private/Misc/OutputDevice.cpp
@@ -206,6 +206,18 @@ void VARARGS FDebug::LogAssertFailedMessage(const ANSICHAR* Expr, const ANSICHAR
 		TCHAR ErrorString[MAX_SPRINTF];
 		FCString::Sprintf( ErrorString, TEXT( "Assertion failed: %s" ), ANSI_TO_TCHAR( Expr ) );
 
+		if( FPlatformProperties::AllowsCallStackDumpDuringAssert() )
+		{
+			ANSICHAR StackTrace[4096];
+			if( StackTrace != NULL )
+			{
+				StackTrace[0] = 0;
+				FPlatformStackWalk::StackWalkAndDump( StackTrace, ARRAY_COUNT(StackTrace), CALLSTACK_IGNOREDEPTH );
+
+				FCString::Strncat( DescriptionString, ANSI_TO_TCHAR( StackTrace ), ARRAY_COUNT( DescriptionString ) - 1 );
+			}
+		}
+
 		StaticFailDebug( ErrorString, File, Line, DescriptionString );
 	}
 }
diff --git a/Engine/Source/Runtime/Core/Public/GenericPlatform/GenericPlatformProperties.h b/Engine/Source/Runtime/Core/Public/GenericPlatform/GenericPlatformProperties.h
--- a/Engine/Source/Runtime/Core/Public/GenericPlatform/GenericPlatformProperties.h
+++ b/Engine/Source/Runtime/Core/Public/GenericPlatform/GenericPlatformProperties.h
@@ -258,4 +258,10 @@ struct FGenericPlatformProperties
 	{
 		return false;
 	}
+
+	// Whether the platform allows the call stack to be dumped during an assert
+	static FORCEINLINE bool AllowsCallStackDumpDuringAssert()
+	{
+		return false;
+	}
 };
