Fix for FBitReference assignment.

diff --git a/Engine/Source/Runtime/Core/Public/Containers/BitArray.h b/Engine/Source/Runtime/Core/Public/Containers/BitArray.h
--- a/Engine/Source/Runtime/Core/Public/Containers/BitArray.h
+++ b/Engine/Source/Runtime/Core/Public/Containers/BitArray.h
@@ -98,8 +98,9 @@ public:
 	}
 	FORCEINLINE FBitReference& operator=(const FBitReference& Copy)
 	{
-		this->Data = Copy.Data;
-		this->Mask = Copy.Mask;
+		// As this is emulating a reference, assignment should not rebind,
+		// it should write to the referenced bit.
+		*this = (bool)Copy;
 		return *this;
 	}
 
