[INTEGRATE] merge ue4 <-ue4-launcheDev cl 2635304

diff --git a/Engine/Source/Runtime/Core/Private/Linux/LinuxPlatformMisc.cpp b/Engine/Source/Runtime/Core/Private/Linux/LinuxPlatformMisc.cpp
--- a/Engine/Source/Runtime/Core/Private/Linux/LinuxPlatformMisc.cpp
+++ b/Engine/Source/Runtime/Core/Private/Linux/LinuxPlatformMisc.cpp
@@ -770,3 +770,9 @@ FString FLinuxPlatformMisc::GetOperatingSystemId()
 
 	return CachedResult;
 }
+
+bool FLinuxPlatformMisc::GetDiskTotalAndFreeSpace(const FString& InPath, uint64& TotalNumberOfBytes, uint64& NumberOfFreeBytes)
+{
+	STUBBED("FLinuxPlatformMisc::GetDiskTotalAndFreeSpace");
+	return false;
+}
diff --git a/Engine/Source/Runtime/Core/Private/Mac/MacPlatformMisc.cpp b/Engine/Source/Runtime/Core/Private/Mac/MacPlatformMisc.cpp
--- a/Engine/Source/Runtime/Core/Private/Mac/MacPlatformMisc.cpp
+++ b/Engine/Source/Runtime/Core/Private/Mac/MacPlatformMisc.cpp
@@ -1066,6 +1066,19 @@ void FMacPlatformMisc::GetOSVersions( FString& out_OSVersionLabel, FString& out_
 	out_OSSubVersionLabel = GMacAppInfo.OSBuild;
 }
 
+bool FMacPlatformMisc::GetDiskTotalAndFreeSpace(const FString& InPath, uint64& TotalNumberOfBytes, uint64& NumberOfFreeBytes)
+{
+	struct statfs FSStat = { 0 };
+	FTCHARToUTF8 Converter(*InPath);
+	int Err = statfs((ANSICHAR*)Converter.Get(), &FSStat);
+	if (Err == 0)
+	{
+		TotalNumberOfBytes = FSStat.f_blocks * FSStat.f_bsize;
+		NumberOfFreeBytes = FSStat.f_bavail * FSStat.f_bsize;
+	}
+	return (Err == 0);
+}
+
 #include "ModuleManager.h"
 
 void FMacPlatformMisc::LoadPreInitModules()
diff --git a/Engine/Source/Runtime/Core/Public/Linux/LinuxPlatformMisc.h b/Engine/Source/Runtime/Core/Public/Linux/LinuxPlatformMisc.h
--- a/Engine/Source/Runtime/Core/Public/Linux/LinuxPlatformMisc.h
+++ b/Engine/Source/Runtime/Core/Public/Linux/LinuxPlatformMisc.h
@@ -90,6 +90,7 @@ struct CORE_API FLinuxPlatformMisc : public FGenericPlatformMisc
 	static void LoadPreInitModules();
 	static void LoadStartupModules();
 	static FString GetOperatingSystemId();
+	static bool GetDiskTotalAndFreeSpace(const FString& InPath, uint64& TotalNumberOfBytes, uint64& NumberOfFreeBytes);
 
 	/**
 	 * Determines the shader format for the platform
diff --git a/Engine/Source/Runtime/Core/Public/Mac/MacPlatformMisc.h b/Engine/Source/Runtime/Core/Public/Mac/MacPlatformMisc.h
--- a/Engine/Source/Runtime/Core/Public/Mac/MacPlatformMisc.h
+++ b/Engine/Source/Runtime/Core/Public/Mac/MacPlatformMisc.h
@@ -100,6 +100,7 @@ struct CORE_API FMacPlatformMisc : public FGenericPlatformMisc
 	static void NormalizePath(FString& InPath);
 	static FString GetPrimaryGPUBrand();
 	static void GetOSVersions( FString& out_OSVersionLabel, FString& out_OSSubVersionLabel );
+	static bool GetDiskTotalAndFreeSpace(const FString& InPath, uint64& TotalNumberOfBytes, uint64& NumberOfFreeBytes);
 
 
 	/** 
