Fixed mouse position setting/reporting in fullscreen mode on Mac

diff --git a/Engine/Source/Runtime/Core/Private/Mac/MacApplication.cpp b/Engine/Source/Runtime/Core/Private/Mac/MacApplication.cpp
--- a/Engine/Source/Runtime/Core/Private/Mac/MacApplication.cpp
+++ b/Engine/Source/Runtime/Core/Private/Mac/MacApplication.cpp
@@ -717,7 +717,7 @@ void FMacApplication::ProcessMouseMovedEvent(const FDeferredMacEvent& Event, TSh
 		}
 		else
 		{
-			MacCursor->UpdateCurrentPosition(CurrentPosition);
+			MacCursor->UpdateCurrentPosition(CurrentPosition / MouseScaling);
 		}
 
 		if (EventWindow.IsValid())
diff --git a/Engine/Source/Runtime/Core/Private/Mac/MacWindow.cpp b/Engine/Source/Runtime/Core/Private/Mac/MacWindow.cpp
--- a/Engine/Source/Runtime/Core/Private/Mac/MacWindow.cpp
+++ b/Engine/Source/Runtime/Core/Private/Mac/MacWindow.cpp
@@ -517,8 +517,7 @@ bool FMacWindow::IsPointInWindow( int32 X, int32 Y ) const
 		if([WindowHandle isOnActiveSpace])
 		{
 			FMacCursor* MacCursor = (FMacCursor*)MacApplication->Cursor.Get();
-			FVector2D MouseScale = MacCursor ? MacCursor->GetMouseScaling() : FVector2D(1.0f, 1.0f);
-			NSPoint CursorPoint = NSMakePoint((X / MouseScale.X), (WindowFrame.size.height - ((Y / MouseScale.Y) + 1)));
+			NSPoint CursorPoint = NSMakePoint(X, WindowFrame.size.height - (Y + 1));
 			PointInWindow = (NSPointInRect(CursorPoint, VisibleFrame) == YES);
 		}
 	}
