FindMatchingClosingParenthesis moved to UnrealString.h, return value fixed (was off-by-one) and checked upper bound of StartSearch.
UnrealHeaderTool now supports casts and brackets when recognizing UENUMs as array bounds.

diff --git a/Engine/Source/Programs/UnrealHeaderTool/Private/HeaderParser.cpp b/Engine/Source/Programs/UnrealHeaderTool/Private/HeaderParser.cpp
--- a/Engine/Source/Programs/UnrealHeaderTool/Private/HeaderParser.cpp
+++ b/Engine/Source/Programs/UnrealHeaderTool/Private/HeaderParser.cpp
@@ -3927,7 +3927,47 @@ UProperty* FHeaderParser::GetVarNameAndDim
 
 		if (*Dimensions.String)
 		{
-			UEnum::LookupEnumNameSlow(Dimensions.String, &Enum);
+			FString Temp = Dimensions.String;
+
+			bool bAgain;
+			do
+			{
+				bAgain = false;
+
+				// Remove any casts
+				static const TCHAR* Casts[] = {
+					TEXT("(uint32)"),
+					TEXT("(int32)"),
+					TEXT("(uint16)"),
+					TEXT("(int16)"),
+					TEXT("(uint8)"),
+					TEXT("(int8)")
+				};
+
+				// Remove any brackets
+				if (Temp[0] == TEXT('('))
+				{
+					int32 TempLen      = Temp.Len();
+					int32 ClosingParen = FindMatchingClosingParenthesis(Temp);
+					if (ClosingParen == TempLen - 1)
+					{
+						Temp = Temp.Mid(1, TempLen - 2);
+						bAgain = true;
+					}
+				}
+
+				for (const TCHAR* Cast : Casts)
+				{
+					if (Temp.StartsWith(Cast))
+					{
+						Temp = Temp.RightChop(FCString::Strlen(Cast));
+						bAgain = true;
+					}
+				}
+			}
+			while (bAgain);
+
+			UEnum::LookupEnumNameSlow(*Temp, &Enum);
 		}
 
 		if (!Enum)
diff --git a/Engine/Source/Runtime/Core/Private/Containers/String.cpp b/Engine/Source/Runtime/Core/Private/Containers/String.cpp
--- a/Engine/Source/Runtime/Core/Private/Containers/String.cpp
+++ b/Engine/Source/Runtime/Core/Private/Containers/String.cpp
@@ -1114,3 +1114,46 @@ FArchive& operator<<( FArchive& Ar, FString& A )
 
 	return Ar;
 }
+
+int32 FindMatchingClosingParenthesis(const FString& TargetString, const int32 StartSearch)
+{
+	check(StartSearch >= 0 && StartSearch <= TargetString.Len());// Check for usage, we do not accept INDEX_NONE like other string functions
+
+	const TCHAR* const StartPosition = (*TargetString) + StartSearch;
+	const TCHAR* CurrPosition = StartPosition;
+	int32 ParenthesisCount = 0;
+
+	// Move to first open parenthesis
+	while (*CurrPosition != 0 && *CurrPosition != TEXT('('))
+	{
+		++CurrPosition;
+	}
+
+	// Did we find the open parenthesis
+	if (*CurrPosition == TEXT('('))
+	{
+		++ParenthesisCount;
+		++CurrPosition;
+
+		while (*CurrPosition != 0 && ParenthesisCount > 0)
+		{
+			if (*CurrPosition == TEXT('('))
+			{
+				++ParenthesisCount;
+			}
+			else if (*CurrPosition == TEXT(')'))
+			{
+				--ParenthesisCount;
+			}
+			++CurrPosition;
+		}
+
+		// Did we find the matching close parenthesis
+		if (ParenthesisCount == 0 && *(CurrPosition - 1) == TEXT(')'))
+		{
+			return StartSearch + ((CurrPosition - 1) - StartPosition);
+		}
+	}
+
+	return INDEX_NONE;
+}
diff --git a/Engine/Source/Runtime/Core/Public/Containers/UnrealString.h b/Engine/Source/Runtime/Core/Public/Containers/UnrealString.h
--- a/Engine/Source/Runtime/Core/Public/Containers/UnrealString.h
+++ b/Engine/Source/Runtime/Core/Public/Containers/UnrealString.h
@@ -2019,3 +2019,13 @@ public:
 		return *this;
 	}
 };
+
+/**
+ * A helper function to find closing parenthesis that matches the first open parenthesis found. The open parenthesis
+ * referred to must be at or further up from the start index.
+ *
+ * @param TargetString      The string to search in
+ * @param StartSearch       The index to start searching at
+ * @return the index in the given string of the closing parenthesis
+ */
+CORE_API int32 FindMatchingClosingParenthesis(const FString& TargetString, const int32 StartSearch = 0);
