Core - Updated Kill() comment

diff --git a/Engine/Source/Runtime/Core/Public/HAL/ThreadingBase.h b/Engine/Source/Runtime/Core/Public/HAL/ThreadingBase.h
--- a/Engine/Source/Runtime/Core/Public/HAL/ThreadingBase.h
+++ b/Engine/Source/Runtime/Core/Public/HAL/ThreadingBase.h
@@ -325,16 +325,16 @@ public:
 	virtual void Suspend( bool bShouldPause = true ) = 0;
 
 	/**
-	 * Tells the thread to exit. If the caller needs to know when the thread
-	 * has exited, it should use the bShouldWait value and tell it how long
-	 * to wait before deciding that it is deadlocked and needs to be destroyed.
-	 * The implementation is responsible for calling Stop() on the runnable.
-	 * NOTE: having a thread forcibly destroyed can cause leaks in TLS, etc.
-	 *
-	 * @param bShouldWait If true, the call will wait for the thread to exit
-	 * @return True if the thread exited graceful, false otherwise
+	 * Tells the thread to exit. If the caller needs to know when the thread has exited, it should use the bShouldWait value.
+	 * It's highly recommended not to kill the thread without waiting for it.
+	 * Having a thread forcibly destroyed can cause leaks and deadlocks.
+	 * 
+	 * The kill method is calling Stop() on the runnable to kill the thread gracefully.
+	 * 
+	 * @param bShouldWait	If true, the call will wait infinitely for the thread to exit.					
+	 * @return Always true
 	 */
-	virtual bool Kill( bool bShouldWait = false ) = 0;
+	virtual bool Kill( bool bShouldWait = true ) = 0;
 
 	/** Halts the caller until this thread is has completed its work. */
 	virtual void WaitForCompletion() = 0;
