Linux: added missing implementation for GetDiskTotalAndFreeSpace().

diff --git a/Engine/Source/Runtime/Core/Private/Linux/LinuxPlatformMisc.cpp b/Engine/Source/Runtime/Core/Private/Linux/LinuxPlatformMisc.cpp
--- a/Engine/Source/Runtime/Core/Private/Linux/LinuxPlatformMisc.cpp
+++ b/Engine/Source/Runtime/Core/Private/Linux/LinuxPlatformMisc.cpp
@@ -13,6 +13,7 @@
 #include <sched.h>
 #include <fcntl.h>
 #include <signal.h>
+#include <sys/vfs.h>	// statfs()
 
 #include "ModuleManager.h"
 
@@ -773,6 +774,18 @@ FString FLinuxPlatformMisc::GetOperatingSystemId()
 
 bool FLinuxPlatformMisc::GetDiskTotalAndFreeSpace(const FString& InPath, uint64& TotalNumberOfBytes, uint64& NumberOfFreeBytes)
 {
-	STUBBED("FLinuxPlatformMisc::GetDiskTotalAndFreeSpace");
-	return false;
+	struct statfs FSStat = { 0 };
+	FTCHARToUTF8 Converter(*InPath);
+	int Err = statfs((ANSICHAR*)Converter.Get(), &FSStat);
+	if (Err == 0)
+	{
+		TotalNumberOfBytes = FSStat.f_blocks * FSStat.f_bsize;
+		NumberOfFreeBytes = FSStat.f_bavail * FSStat.f_bsize;
+	}
+	else
+	{
+		int ErrNo = errno;
+		UE_LOG(LogLinux, Warning, TEXT("Unable to statfs('%s'): errno=%d (%s)"), *InPath, ErrNo, ANSI_TO_TCHAR(strerror(ErrNo)));
+	}
+	return (Err == 0);
 }
