Added GetPreviousIndex to TCircularBuffer

--------
Integrated using branch UE4-To-UE4-Orion (reversed) of change#2636245 by [REDACTED_PERSON] on 2015/07/28 18:26:26.

diff --git a/Engine/Source/Runtime/Core/Public/Containers/CircularBuffer.h b/Engine/Source/Runtime/Core/Public/Containers/CircularBuffer.h
--- a/Engine/Source/Runtime/Core/Public/Containers/CircularBuffer.h
+++ b/Engine/Source/Runtime/Core/Public/Containers/CircularBuffer.h
@@ -89,6 +89,17 @@ public:
 		return ((CurrentIndex + 1) & IndexMask);
 	}
 
+	/**
+	 * Calculates the index previous to the given index.
+	 *
+	 * @param CurrentIndex The current index.
+	 * @return The previous index.
+	 */
+	FORCEINLINE uint32 GetPreviousIndex( uint32 CurrentIndex ) const
+	{
+		return ( ( CurrentIndex - 1 ) & IndexMask );
+	}
+
 private:
 
 	/** Holds the mask for indexing the buffer's elements. */
