Linux: fix opening projects with localized names (UE-18143).

- Can now open G┼╝eg┼╝├│┼éka.uproject/╨É╨▒╤Ç╨░╨║╨░╨┤╨░╨▒╤Ç╨░.uproject etc.
- And a few more fixes in other functions that misconverted paths to ASCII that may or may not have manifested yet.
- Also made CreateProc more spammy to aid catching this in the wild.

diff --git a/Engine/Source/Runtime/Core/Private/Linux/LinuxPlatformProcess.cpp b/Engine/Source/Runtime/Core/Private/Linux/LinuxPlatformProcess.cpp
--- a/Engine/Source/Runtime/Core/Private/Linux/LinuxPlatformProcess.cpp
+++ b/Engine/Source/Runtime/Core/Private/Linux/LinuxPlatformProcess.cpp
@@ -25,7 +25,7 @@ void* FLinuxPlatformProcess::GetDllHandle( const TCHAR* Filename )
 		DlOpenMode |= RTLD_LOCAL; //Local symbol resolution when loading shared objects - Needed for Hot-Reload
 	}
 
-	void *Handle = dlopen( TCHAR_TO_ANSI(*AbsolutePath), DlOpenMode );
+	void *Handle = dlopen( TCHAR_TO_UTF8(*AbsolutePath), DlOpenMode );
 	if (!Handle)
 	{
 		UE_LOG(LogLinux, Warning, TEXT("dlopen failed: %s"), ANSI_TO_TCHAR(dlerror()) );
@@ -452,7 +452,7 @@ void FLinuxPlatformProcess::LaunchURL(const TCHAR* URL, const TCHAR* Parms, FStr
 	UE_LOG(LogHAL, Verbose, TEXT("FLinuxPlatformProcess::LaunchURL: '%s'"), URL);
 	if (pid == 0)
 	{
-		exit(execl("/usr/bin/xdg-open", "xdg-open", TCHAR_TO_ANSI(URL), (char *)0));
+		exit(execl("/usr/bin/xdg-open", "xdg-open", TCHAR_TO_UTF8(URL), (char *)0));
 	}
 }
 
@@ -544,7 +544,7 @@ FProcHandle FLinuxPlatformProcess::CreateProc(const TCHAR* URL, const TCHAR* Par
 	Commandline += TEXT(" ");
 	Commandline += Parms;
 
-	UE_LOG(LogHAL, Verbose, TEXT("FLinuxPlatformProcess::CreateProc: '%s'"), *Commandline);
+	UE_LOG(LogHAL, Log, TEXT("FLinuxPlatformProcess::CreateProc: '%s'"), *Commandline);
 
 	TArray<FString> ArgvArray;
 	int Argc = Commandline.ParseIntoArray(ArgvArray, TEXT(" "), true);
@@ -638,15 +638,15 @@ FProcHandle FLinuxPlatformProcess::CreateProc(const TCHAR* URL, const TCHAR* Par
 
 		for (int Idx = 0; Idx < Argc; ++Idx)
 		{
-			auto AnsiBuffer = StringCast<ANSICHAR>(*NewArgvArray[Idx]);
+			FTCHARToUTF8 AnsiBuffer(*NewArgvArray[Idx]);
 			const char* Ansi = AnsiBuffer.Get();
-			size_t AnsiSize = FCStringAnsi::Strlen(Ansi) + 1;
+			size_t AnsiSize = FCStringAnsi::Strlen(Ansi) + 1;	// will work correctly with UTF-8
 			check(AnsiSize);
 
 			Argv[Idx] = reinterpret_cast< char* >( FMemory::Malloc(AnsiSize) );
 			check(Argv[Idx]);
 
-			FCStringAnsi::Strncpy(Argv[Idx], Ansi, AnsiSize);
+			FCStringAnsi::Strncpy(Argv[Idx], Ansi, AnsiSize);	// will work correctly with UTF-8
 		}
 
 		// last Argv should be NULL
@@ -666,7 +666,7 @@ FProcHandle FLinuxPlatformProcess::CreateProc(const TCHAR* URL, const TCHAR* Par
 		posix_spawn_file_actions_adddup2(&FileActions, PipeWriteHandle->GetHandle(), STDOUT_FILENO);
 	}
 
-	int PosixSpawnErrNo = posix_spawn(&ChildPid, TCHAR_TO_ANSI(*ProcessPath), &FileActions, nullptr, Argv, environ);
+	int PosixSpawnErrNo = posix_spawn(&ChildPid, TCHAR_TO_UTF8(*ProcessPath), &FileActions, nullptr, Argv, environ);
 	posix_spawn_file_actions_destroy(&FileActions);
 	if (PosixSpawnErrNo != 0)
 	{
@@ -999,7 +999,7 @@ bool FLinuxPlatformProcess::IsApplicationRunning( const TCHAR* ProcName )
 	FString Commandline = TEXT("pidof '");
 	Commandline += ProcName;
 	Commandline += TEXT("'  > /dev/null");
-	return !system(TCHAR_TO_ANSI(*Commandline));
+	return !system(TCHAR_TO_UTF8(*Commandline));
 }
 
 bool FLinuxPlatformProcess::ExecProcess( const TCHAR* URL, const TCHAR* Params, int32* OutReturnCode, FString* OutStdOut, FString* OutStdErr )
@@ -1080,7 +1080,7 @@ void FLinuxPlatformProcess::LaunchFileInDefaultExternalApplication( const TCHAR*
 	pid_t pid = fork();
 	if (pid == 0)
 	{
-		exit(execl("/usr/bin/xdg-open", "xdg-open", TCHAR_TO_ANSI(FileName), (char *)0));
+		exit(execl("/usr/bin/xdg-open", "xdg-open", TCHAR_TO_UTF8(FileName), (char *)0));
 	}
 }
 
