Linux: switch to jemalloc as the default for the editor and programs.

diff --git a/Engine/Source/Runtime/Core/Private/Linux/LinuxPlatformMemory.cpp b/Engine/Source/Runtime/Core/Private/Linux/LinuxPlatformMemory.cpp
--- a/Engine/Source/Runtime/Core/Private/Linux/LinuxPlatformMemory.cpp
+++ b/Engine/Source/Runtime/Core/Private/Linux/LinuxPlatformMemory.cpp
@@ -41,11 +41,23 @@ class FMalloc* FLinuxPlatformMemory::BaseAllocator()
 		Jemalloc,
 		Binned
 	}
-	AllocatorToUse = FORCE_ANSI_ALLOCATOR ? EAllocatorToUse::Ansi : EAllocatorToUse::Binned;
+	AllocatorToUse = EAllocatorToUse::Binned;
 
-	// we get here before main due to global ctors, so need to do some hackery to get command line args
-	if (!FORCE_ANSI_ALLOCATOR)
+	// Prefer jemalloc for the editor and programs as it saved ~20% RES usage in my (RCL) tests.
+	// Leave binned as the default for games and servers to keep runtime behavior consistent across platforms.
+	if (PLATFORM_SUPPORTS_JEMALLOC && (UE_EDITOR || IS_PROGRAM))
 	{
+		AllocatorToUse = EAllocatorToUse::Jemalloc;
+	}
+
+	if (FORCE_ANSI_ALLOCATOR)
+	{
+		AllocatorToUse = EAllocatorToUse::Ansi;
+	}
+	else
+	{
+		// Allow overriding on the command line.
+		// We get here before main due to global ctors, so need to do some hackery to get command line args
 		if (FILE* CmdLineFile = fopen("/proc/self/cmdline", "r"))
 		{
 			char * Arg = nullptr;
