Linux: fix mouse being offset in ShooterGame UI (UE-19072).

- Broken by PR #1369 - cached position did not account for frame extents (window borders).

diff --git a/Engine/Source/Runtime/Core/Private/Linux/LinuxApplication.cpp b/Engine/Source/Runtime/Core/Private/Linux/LinuxApplication.cpp
--- a/Engine/Source/Runtime/Core/Private/Linux/LinuxApplication.cpp
+++ b/Engine/Source/Runtime/Core/Private/Linux/LinuxApplication.cpp
@@ -312,7 +312,11 @@ void FLinuxApplication::ProcessDeferredMessage( SDL_Event Event )
 			{
 				int xOffset, yOffset;
 				SDL_GetWindowPosition( NativeWindow, &xOffset, &yOffset );
-				LinuxCursor->SetCachedPosition( motionEvent.x + xOffset, motionEvent.y + yOffset );
+
+				int32 BorderSizeX, BorderSizeY;
+				CurrentEventWindow->GetNativeBordersSize(BorderSizeX, BorderSizeY);
+
+				LinuxCursor->SetCachedPosition(motionEvent.x + xOffset + BorderSizeX, motionEvent.y + yOffset + BorderSizeY);
 
 				FVector2D CurrentPosition = LinuxCursor->GetPosition();
 				if( LinuxCursor->UpdateCursorClipping( CurrentPosition ) )
