Made Windows once again return file times that have time_t precision

The more precise timestamps had broken cook-on-the-fly as it was verifying that a more precise timestamp value had been accurately set on the underlying file system, which was failing for all of our non-Microsoft platforms as they have a 1 second precision on their file timestamps.

diff --git a/Engine/Source/Runtime/Core/Private/Windows/WindowsPlatformFile.cpp b/Engine/Source/Runtime/Core/Private/Windows/WindowsPlatformFile.cpp
--- a/Engine/Source/Runtime/Core/Private/Windows/WindowsPlatformFile.cpp
+++ b/Engine/Source/Runtime/Core/Private/Windows/WindowsPlatformFile.cpp
@@ -3,11 +3,6 @@
 #include "CorePrivatePCH.h"
 #include <sys/utime.h>
 
-
-// make an FTimeSpan object that represents the "epoch" for FILETIME
-// FILETIME contains a 64-bit value representing the number of 100-nanosecond intervals since January 1, 1601 (UTC)
-const FDateTime WindowsEpoch(1601, 1, 1);
-
 #include "AllowWindowsPlatformTypes.h"
 	namespace FileConstants
 	{
@@ -17,23 +12,62 @@ const FDateTime WindowsEpoch(1601, 1, 1);
 
 namespace
 {
+	FORCEINLINE int32 UEDayOfWeekToWindowsSystemTimeDayOfWeek(const EDayOfWeek InDayOfWeek)
+	{
+		switch (InDayOfWeek)
+		{
+			case EDayOfWeek::Monday:
+				return 1;
+			case EDayOfWeek::Tuesday:
+				return 2;
+			case EDayOfWeek::Wednesday:
+				return 3;
+			case EDayOfWeek::Thursday:
+				return 4;
+			case EDayOfWeek::Friday:
+				return 5;
+			case EDayOfWeek::Saturday:
+				return 6;
+			case EDayOfWeek::Sunday:
+				return 0;
+			default:
+				break;
+		}
+
+		return 0;
+	}
+
 	FORCEINLINE FDateTime WindowsFileTimeToUEDateTime(const FILETIME& InFileTime)
 	{
-		ULARGE_INTEGER FileTime64;
-		FileTime64.LowPart = InFileTime.dwLowDateTime;
-		FileTime64.HighPart = InFileTime.dwHighDateTime;
+		// This roundabout conversion clamps the precision of the returned time value to match that of time_t (1 second precision)
+		// This avoids issues when sending files over the network via cook-on-the-fly
+		SYSTEMTIME SysTime;
+		if (FileTimeToSystemTime(&InFileTime, &SysTime))
+		{
+			return FDateTime(SysTime.wYear, SysTime.wMonth, SysTime.wDay, SysTime.wHour, SysTime.wMinute, SysTime.wSecond);
+		}
 
-		return WindowsEpoch + FTimespan(static_cast<int64>(FileTime64.QuadPart));
+		// Failed to convert
+		return FDateTime::MinValue();
 	}
 
 	FORCEINLINE FILETIME UEDateTimeToWindowsFileTime(const FDateTime& InDateTime)
 	{
-		ULARGE_INTEGER FileTime64;
-		FileTime64.QuadPart = (InDateTime - WindowsEpoch).GetTicks();
+		// This roundabout conversion clamps the precision of the returned time value to match that of time_t (1 second precision)
+		// This avoids issues when sending files over the network via cook-on-the-fly
+		SYSTEMTIME SysTime;
+		SysTime.wYear = InDateTime.GetYear();
+		SysTime.wMonth = InDateTime.GetMonth();
+		SysTime.wDay = InDateTime.GetDay();
+		SysTime.wDayOfWeek = UEDayOfWeekToWindowsSystemTimeDayOfWeek(InDateTime.GetDayOfWeek());
+		SysTime.wHour = InDateTime.GetHour();
+		SysTime.wMinute = InDateTime.GetMinute();
+		SysTime.wSecond = InDateTime.GetSecond();
+		SysTime.wMilliseconds = 0;
 
 		FILETIME FileTime;
-		FileTime.dwLowDateTime = FileTime64.LowPart;
-		FileTime.dwHighDateTime = FileTime64.HighPart;
+		SystemTimeToFileTime(&SysTime, &FileTime);
+
 		return FileTime;
 	}
 }
