PR #1430: Fix "error: '&' within '|'" in EditorPerformanceTests.cpp when using clang-3.5.2 (Contributed by [REDACTED_PERSON])

diff --git a/Engine/Source/Runtime/Core/Public/Misc/AutomationTest.h b/Engine/Source/Runtime/Core/Public/Misc/AutomationTest.h
--- a/Engine/Source/Runtime/Core/Public/Misc/AutomationTest.h
+++ b/Engine/Source/Runtime/Core/Public/Misc/AutomationTest.h
@@ -1085,7 +1085,7 @@ public: \
 	public: \
 		TClass( const FString& InName ) \
 		:FAutomationTestBase( InName, true ) {} \
-		virtual uint32 GetTestFlags() const override { return (TFlags & ~(EAutomationTestFlags::ATF_SmokeTest)); } \
+		virtual uint32 GetTestFlags() const override { return ((TFlags) & ~(EAutomationTestFlags::ATF_SmokeTest)); } \
 		virtual bool IsStressTest() const { return true; } \
 		virtual uint32 GetRequiredDeviceNum() const override { return 1; } \
 	protected: \
@@ -1100,7 +1100,7 @@ public: \
 	public: \
 		TClass( const FString& InName ) \
 		:FAutomationTestBase( InName, false ) {} \
-		virtual uint32 GetTestFlags() const override { return (TFlags & ~(EAutomationTestFlags::ATF_Editor | EAutomationTestFlags::ATF_Commandlet | EAutomationTestFlags::ATF_SmokeTest)); } \
+		virtual uint32 GetTestFlags() const override { return ((TFlags) & ~(EAutomationTestFlags::ATF_Editor | EAutomationTestFlags::ATF_Commandlet | EAutomationTestFlags::ATF_SmokeTest)); } \
 		virtual uint32 GetRequiredDeviceNum() const override { return NumParticipants; } \
 	protected: \
 		virtual void GetTests(TArray<FString>& OutBeautifiedNames, TArray <FString>& OutTestCommands) const override \
