UE-17189 - Making changes in "Project Settings" now only edits one property in the default*.ini. With the exception of Array properties, which needs a whole other implementation. That will be incoming later.

diff --git a/Engine/Source/Developer/SettingsEditor/Private/Widgets/SSettingsEditor.cpp b/Engine/Source/Developer/SettingsEditor/Private/Widgets/SSettingsEditor.cpp
--- a/Engine/Source/Developer/SettingsEditor/Private/Widgets/SSettingsEditor.cpp
+++ b/Engine/Source/Developer/SettingsEditor/Private/Widgets/SSettingsEditor.cpp
@@ -244,7 +244,14 @@ void SSettingsEditor::NotifyPostChange( const FPropertyChangedEvent& PropertyCha
 	{
 		RecordPreferenceChangedAnalytics( SelectedSection, PropertyChangedEvent );
 
-		SelectedSection->Save();
+		if (SelectedSection->GetSettingsObject()->GetClass()->HasAnyClassFlags(CLASS_DefaultConfig) && !PropertyChangedEvent.Property->IsA(UArrayProperty::StaticClass()))
+		{
+			SelectedSection->GetSettingsObject()->UpdateSinglePropertyInConfigFile(PropertyThatChanged->GetActiveMemberNode()->GetValue(), SelectedSection->GetSettingsObject()->GetDefaultConfigFilename());
+		}
+		else
+		{
+			SelectedSection->Save();
+		}
 
 		static const FName ConfigRestartRequiredKey = "ConfigRestartRequired";
 		if (PropertyChangedEvent.Property->GetBoolMetaData(ConfigRestartRequiredKey))
diff --git a/Engine/Source/Runtime/Core/Private/Misc/ConfigCacheIni.cpp b/Engine/Source/Runtime/Core/Private/Misc/ConfigCacheIni.cpp
--- a/Engine/Source/Runtime/Core/Private/Misc/ConfigCacheIni.cpp
+++ b/Engine/Source/Runtime/Core/Private/Misc/ConfigCacheIni.cpp
@@ -3162,94 +3162,213 @@ void FConfigFile::UpdateSections(const TCHAR* DiskFilename, const TCHAR* IniRoot
 	Write(DiskFilename, true, NewFile);
 }
 
-bool FConfigFile::UpdateSinglePropertyInSection(const TCHAR* DiskFilename, const TCHAR* PropertyName, const TCHAR* SectionName)
+
+
+/**
+ * Functionality to assist with updating a config file with one property value change.
+ */
+class FSinglePropertyConfigHelper
 {
-	// Result of whether the file has been updated on disk.
-	bool bSuccessfullyUpdatedFile = false;
+public:
+	
 
-	// Since we don't want any modifications to other sections, or properties, we manually process the file, not read into sections
-	FString DiskFile;
-	if (FFileHelper::LoadFileToString(DiskFile, DiskFilename))
+	/**
+	 * We need certain information for the helper to be useful.
+	 *
+	 * @Param InIniFilename - The disk location of the file we wish to edit.
+	 * @Param InSectionName - The section the property belongs to.
+	 * @Param InPropertyName - The name of the property that has been edited.
+	 * @Param InPropertyValue - The new value of the property that has been edited.
+	 */
+	FSinglePropertyConfigHelper(const FString& InIniFilename, const FString& InSectionName, const FString& InPropertyName, const FString& InPropertyValue)
+		: IniFilename(InIniFilename)
+		, SectionName(InSectionName)
+		, PropertyName(InPropertyName)
+		, PropertyValue(InPropertyValue)
 	{
-		FString NewFile;
-		const FString DecoratedSectionName = FString::Printf(TEXT("[%s]"), SectionName);
+		// Split the file into the necessary parts.
+		PopulateFileContentHelper();
+	}
 
-		// Track if the these properties belong to the section we are looking for.
-		bool bIsCurrentlyWritingSectionProperty = false;
-		
-		// Track if the section was ever updated.
-		bool bSectionWasUpdated = false;
-		// Track if the property was ever updated.
-		bool bPropertyWasUpdated = false;
+	/**
+	 * Perform the action of updating the config file with the new property value.
+	 */
+	bool UpdateConfigFile()
+	{
+		UpdatePropertyInSection();
+		// Rebuild the file with the updated section.
+		const FString NewFile = IniFileMakeup.BeforeSection + IniFileMakeup.Section + IniFileMakeup.AfterSection;
+		return FFileHelper::SaveStringToFile(NewFile, *IniFilename);
+	}
 
-		// Iterate through the file, line by line, to find the property we would like to update.
-		FString TheLine;
-		const TCHAR* Ptr = DiskFile.Len() > 0 ? *DiskFile : NULL;
-		while (Ptr != NULL && FParse::Line(&Ptr, TheLine, true))
+
+private:
+
+	/**
+	 * Clear any trailing whitespace from the end of the output.
+	 */
+	void ClearTrailingWhitespace(FString& InStr)
+	{
+		const FString Endl(LINE_TERMINATOR);
+		while (InStr.EndsWith(LINE_TERMINATOR))
 		{
-			const bool bIsASectionHeader = (TheLine.Len() > 3) && (TheLine[0] == '[') && (TheLine[TheLine.Len() - 1] == ']');
-			if (bIsASectionHeader)
+			InStr = InStr.LeftChop(Endl.Len());
+		}
+	}
+	
+	/**
+	 * Update the section with the new value for the property.
+	 */
+	void UpdatePropertyInSection()
+	{
+		FString UpdatedSection;
+		if (IniFileMakeup.Section.IsEmpty())
+		{
+			const FString DecoratedSectionName = FString::Printf(TEXT("[%s]"), *SectionName);
+			
+			ClearTrailingWhitespace(IniFileMakeup.BeforeSection);
+			UpdatedSection += LINE_TERMINATOR;
+			UpdatedSection += LINE_TERMINATOR;
+			UpdatedSection += DecoratedSectionName;
+			AppendPropertyLine(UpdatedSection);
+		}
+		else
+		{
+			FString SectionLine;
+			const TCHAR* Ptr = *IniFileMakeup.Section;
+			bool bWrotePropertyOnPass = false;
+			while (Ptr != nullptr && FParse::Line(&Ptr, SectionLine, true))
 			{
-				// Cache off the old value of whether we had been writing a section.
-				const bool bWasWritingSection = bIsCurrentlyWritingSectionProperty;
-
-				if (!bSectionWasUpdated && bWasWritingSection && !bPropertyWasUpdated)
+				if (SectionLine.StartsWith(FString::Printf(TEXT("%s="), *PropertyName)))
 				{
-					bPropertyWasUpdated = true;
-					// If we have found the property in the file, let's update it with the new value.
-					if (const FConfigSection* LocalSection = this->Find(SectionName))
-					{
-						if (const FString* LocalSectionPropertyStr = LocalSection->Find(PropertyName))
-						{
-							NewFile += FString::Printf(TEXT("%s=%s"), PropertyName, *(*LocalSectionPropertyStr)) + LINE_TERMINATOR + LINE_TERMINATOR;
-						}
-					}
+					UpdatedSection += FString::Printf(TEXT("%s=%s"), *PropertyName, *PropertyValue);
+					UpdatedSection += LINE_TERMINATOR;
+					bWrotePropertyOnPass = true;
 				}
+				else
+				{
+					UpdatedSection += SectionLine;
+					UpdatedSection += LINE_TERMINATOR;
+				}
+			}
 
-				bIsCurrentlyWritingSectionProperty = (TheLine == DecoratedSectionName);
-				bSectionWasUpdated |= bIsCurrentlyWritingSectionProperty;
-
-				NewFile += TheLine + LINE_TERMINATOR;
+			// If the property wasnt found in the text of the existing section content,
+			// append it to the end of the section.
+			if (!bWrotePropertyOnPass)
+			{
+				AppendPropertyLine(UpdatedSection);
 			}
-			else if (bIsCurrentlyWritingSectionProperty && TheLine.StartsWith(FString::Printf(TEXT("%s="), PropertyName)))
+			else
+			{
+				UpdatedSection += LINE_TERMINATOR;
+			}
+		}
+
+		IniFileMakeup.Section = UpdatedSection;
+	}
+	
+	/**
+	 * Split the file up into parts:
+	 * -> Before the section we wish to edit, which will remain unaltered,
+	 * ->-> The section we wish to edit, we only seek to edit the single property,
+	 * ->->-> After the section we wish to edit, which will remain unaltered.
+	 */
+	void PopulateFileContentHelper()
+	{
+		FString UnprocessedFileContents;
+		if (FFileHelper::LoadFileToString(UnprocessedFileContents, *IniFilename))
+		{
+			// Find the section in the file text.
+			const FString DecoratedSectionName = FString::Printf(TEXT("[%s]"), *SectionName);
+
+			const int32 DecoratedSectionNameStartIndex = UnprocessedFileContents.Find(DecoratedSectionName);
+			if (DecoratedSectionNameStartIndex != INDEX_NONE)
 			{
-				bPropertyWasUpdated = true;
-				// If we have found the property in the file, let's update it with the new value.
-				if (const FConfigSection* LocalSection = this->Find(SectionName))
+				// If we found the section, cache off the file text before the section.
+				IniFileMakeup.BeforeSection = UnprocessedFileContents.Left(DecoratedSectionNameStartIndex);
+				UnprocessedFileContents.RemoveAt(0, IniFileMakeup.BeforeSection.Len());
+
+				// For the rest of the file, split it into the section we are editing and the rest of the file after.
+				const TCHAR* Ptr = UnprocessedFileContents.Len() > 0 ? *UnprocessedFileContents : nullptr;
+				FString NextUnprocessedLine;
+				bool bReachedNextSection = false;
+				while (Ptr != nullptr && FParse::Line(&Ptr, NextUnprocessedLine, true))
 				{
-					if (const FString* LocalSectionPropertyStr = LocalSection->Find(PropertyName))
+					if (bReachedNextSection |= (NextUnprocessedLine.StartsWith(TEXT("[")) && NextUnprocessedLine != DecoratedSectionName))
+					{
+						IniFileMakeup.AfterSection += NextUnprocessedLine;
+						IniFileMakeup.AfterSection += LINE_TERMINATOR;
+					}
+					else
 					{
-						NewFile += FString::Printf(TEXT("%s=%s"), PropertyName, *(*LocalSectionPropertyStr)) + LINE_TERMINATOR;
+						IniFileMakeup.Section += NextUnprocessedLine;
+						IniFileMakeup.Section += LINE_TERMINATOR;
 					}
 				}
 			}
 			else
 			{
-				NewFile += TheLine + LINE_TERMINATOR;
+				IniFileMakeup.BeforeSection = UnprocessedFileContents;
 			}
 		}
+	}
+	
+	/**
+	 * Append the property entry to the section
+	 */
+	void AppendPropertyLine(FString& PreText)
+	{
+		// Make sure we dont leave much whitespace, and append the property name/value entry
+		ClearTrailingWhitespace(PreText);
+		PreText += LINE_TERMINATOR;
+		PreText += FString::Printf(TEXT("%s=%s"), *PropertyName, *PropertyValue);
+		PreText += LINE_TERMINATOR;
+		PreText += LINE_TERMINATOR;
+	}
 
-		if (bSectionWasUpdated == false)
-		{
-			// If we reached the eof without updating the section, write the header
-			NewFile += LINE_TERMINATOR + DecoratedSectionName + LINE_TERMINATOR;
-		}
 
-		if (bPropertyWasUpdated == false)
+private:
+	// The disk location of the ini file we seek to edit
+	FString IniFilename;
+
+	// The section in the config file
+	FString SectionName;
+
+	// The name of the property that has been changed
+	FString PropertyName;
+
+	// The new value, in string format, of the property that has been changed
+	FString PropertyValue;
+
+	// Helper struct that holds the makeup of the ini file.
+	struct IniFileContent
+	{
+		// The section we wish to edit
+		FString Section;
+
+		// The file contents before the section we are editing
+		FString BeforeSection;
+
+		// The file contents after the section we are editing
+		FString AfterSection;
+	} IniFileMakeup; // Instance of the helper to maintain file structure.
+};
+
+
+bool FConfigFile::UpdateSinglePropertyInSection(const TCHAR* DiskFilename, const TCHAR* PropertyName, const TCHAR* SectionName)
+{
+	// Result of whether the file has been updated on disk.
+	bool bSuccessfullyUpdatedFile = false;
+
+	FString PropertyValue;
+	if (const FConfigSection* LocalSection = this->Find(SectionName))
+	{
+		if (const FString* LocalSectionPropertyStr = LocalSection->Find(PropertyName))
 		{
-			// If we reached the eof without updating the property, we should write it out.
-			bPropertyWasUpdated = true;
-			if (const FConfigSection* LocalSection = this->Find(SectionName))
-			{
-				if (const FString* LocalSectionPropertyStr = LocalSection->Find(PropertyName))
-				{
-					NewFile += FString::Printf(TEXT("%s=%s"), PropertyName, *(*LocalSectionPropertyStr)) + LINE_TERMINATOR;
-				}
-			}
+			PropertyValue = *LocalSectionPropertyStr;
+			FSinglePropertyConfigHelper SinglePropertyConfigHelper(DiskFilename, SectionName, PropertyName, PropertyValue);
+			bSuccessfullyUpdatedFile = SinglePropertyConfigHelper.UpdateConfigFile();
 		}
-
-		// Write the updated file to disk.
-		bSuccessfullyUpdatedFile = FFileHelper::SaveStringToFile(NewFile, DiskFilename);
 	}
 
 	return bSuccessfullyUpdatedFile;
diff --git a/Engine/Source/Runtime/CoreUObject/Private/UObject/Obj.cpp b/Engine/Source/Runtime/CoreUObject/Private/UObject/Obj.cpp
--- a/Engine/Source/Runtime/CoreUObject/Private/UObject/Obj.cpp
+++ b/Engine/Source/Runtime/CoreUObject/Private/UObject/Obj.cpp
@@ -1936,8 +1936,18 @@ void UObject::UpdateSinglePropertyInConfigFile(const UProperty* InProperty, cons
 		NewFile.GetKeys(Keys);
 
 		const FString SectionName = Keys[0];
-		const FString PropertyName = InProperty->GetFName().ToString();
-		NewFile.UpdateSinglePropertyInSection(*InConfigIniName, *PropertyName, *SectionName);
+		FString PropertyKey = InProperty->GetFName().ToString();
+
+#if WITH_EDITOR
+		static FName ConsoleVariableFName(TEXT("ConsoleVariable"));
+		FString CVarName = InProperty->GetMetaData(ConsoleVariableFName);
+		if (!CVarName.IsEmpty())
+		{
+			PropertyKey = CVarName;
+		}
+#endif // #if WITH_EDITOR
+
+		NewFile.UpdateSinglePropertyInSection(*InConfigIniName, *PropertyKey, *SectionName);
 
 		// reload the file, so that it refresh the cache internally.
 		FString FinalIniFileName;
