UE-17189 - Fix for clang issue introduced by my latest change.

diff --git a/Engine/Source/Runtime/Core/Private/Misc/ConfigCacheIni.cpp b/Engine/Source/Runtime/Core/Private/Misc/ConfigCacheIni.cpp
--- a/Engine/Source/Runtime/Core/Private/Misc/ConfigCacheIni.cpp
+++ b/Engine/Source/Runtime/Core/Private/Misc/ConfigCacheIni.cpp
@@ -3294,7 +3294,8 @@ private:
 				bool bReachedNextSection = false;
 				while (Ptr != nullptr && FParse::Line(&Ptr, NextUnprocessedLine, true))
 				{
-					if (bReachedNextSection |= (NextUnprocessedLine.StartsWith(TEXT("[")) && NextUnprocessedLine != DecoratedSectionName))
+					bReachedNextSection |= (NextUnprocessedLine.StartsWith(TEXT("[")) && NextUnprocessedLine != DecoratedSectionName);
+					if (bReachedNextSection)
 					{
 						IniFileMakeup.AfterSection += NextUnprocessedLine;
 						IniFileMakeup.AfterSection += LINE_TERMINATOR;
