UE-17430: GitHub 1290 : Disable hot-reload feature in game builds

#1290

diff --git a/Engine/Source/Runtime/Core/Public/Misc/Build.h b/Engine/Source/Runtime/Core/Public/Misc/Build.h
--- a/Engine/Source/Runtime/Core/Public/Misc/Build.h
+++ b/Engine/Source/Runtime/Core/Public/Misc/Build.h
@@ -129,7 +129,7 @@
 * Whether we support hot-reload. Currently requires a non-monolithic build and non-shipping configuration.
 */
 #ifndef WITH_HOT_RELOAD
-	#define WITH_HOT_RELOAD (!IS_MONOLITHIC && !UE_BUILD_SHIPPING)
+	#define WITH_HOT_RELOAD (!IS_MONOLITHIC && !UE_BUILD_SHIPPING && !UE_GAME)
 #endif
 
 /**
