Only modify the PATH variable around the code that we use load other DLLs. Prevents passing really long PATH variables to child processes, like UBT or UAT. Fix for hitting environment variable length limit on Windows.

diff --git a/Engine/Source/Runtime/Core/Private/Windows/WindowsPlatformProcess.cpp b/Engine/Source/Runtime/Core/Private/Windows/WindowsPlatformProcess.cpp
--- a/Engine/Source/Runtime/Core/Private/Windows/WindowsPlatformProcess.cpp
+++ b/Engine/Source/Runtime/Core/Private/Windows/WindowsPlatformProcess.cpp
@@ -27,53 +27,48 @@
 
 // static variables
 TArray<FString> FWindowsPlatformProcess::DllDirectoryStack;
+TArray<FString> FWindowsPlatformProcess::DllDirectories;
 
 
 void FWindowsPlatformProcess::AddDllDirectory(const TCHAR* Directory)
 {
-	// Normalize the input directory
 	FString NormalizedDirectory = Directory;
 	FPaths::NormalizeDirectoryName(NormalizedDirectory);
 	FPaths::MakePlatformFilename(NormalizedDirectory);
+	DllDirectories.AddUnique(NormalizedDirectory);
+}
 
-	// Get the size of the PATH variable
-	TArray<TCHAR> PathVariable;
-	PathVariable.AddUninitialized(::GetEnvironmentVariable(TEXT("PATH"), NULL, 0));
+void* FWindowsPlatformProcess::GetDllHandle( const TCHAR* Filename )
+{
+	check(Filename);
 
-	// Get the actual value of variable.
-	if (::GetEnvironmentVariable(TEXT("PATH"), PathVariable.GetData(), PathVariable.Num()) == 0)
+	// In order to load the DLL and resolve its imports correctly, we update the PATH environment variable before the load, and restore it when we're done.
+	TArray<TCHAR> InitialPathVariable;
+	InitialPathVariable.AddUninitialized(::GetEnvironmentVariable(TEXT("PATH"), NULL, 0));
+	if (::GetEnvironmentVariable(TEXT("PATH"), InitialPathVariable.GetData(), InitialPathVariable.Num()) == 0)
 	{
-		// Log a warning if reading value fails, but continue anyway.
 		UE_LOG(LogWindows, Warning, TEXT("Failed to load PATH environment variable. It either doesn't exist, or is too long."));
-		PathVariable.Add(TEXT(';'));
 	}
 
 	// Set the new path variable with the input directory at the start. Skip over any existing instances of the input directory.
-	FString NewPathVariable = NormalizedDirectory;
-	for(const TCHAR* PathPos = PathVariable.GetData(); PathPos < PathVariable.GetData() + PathVariable.Num(); )
+	FString NewPathVariable;
+	for(const FString& DllDirectory: DllDirectories)
 	{
-		// Scan to the end of this directory
-		const TCHAR* PathEnd = PathPos;
-		while(*PathEnd != ';' && *PathEnd != 0) PathEnd++;
-
-		// Add it to the new path variable if it doesn't match the input directory
-		if(PathEnd - PathPos != NormalizedDirectory.Len() || FCString::Strnicmp(*NormalizedDirectory, PathPos, PathEnd - PathPos) != 0)
+		if(NewPathVariable.Len() > 0)
 		{
 			NewPathVariable.AppendChar(TEXT(';'));
-			NewPathVariable.AppendChars(PathPos, PathEnd - PathPos);
 		}
-
-		// Move to the next string
-		PathPos = PathEnd + 1;
+		NewPathVariable.Append(DllDirectory);
 	}
 	SetEnvironmentVariable(TEXT("PATH"), *NewPathVariable);
-}
 
-void* FWindowsPlatformProcess::GetDllHandle( const TCHAR* Filename )
-{
-	check(Filename);
+	// Load the DLL
 	::SetErrorMode(SEM_NOOPENFILEERRORBOX);
-	return ::LoadLibraryW(Filename);
+	void* Handle = ::LoadLibraryW(Filename);
+
+	// Restore the PATH variable back to normal
+	::SetEnvironmentVariable(TEXT("PATH"), InitialPathVariable.GetData());
+	return Handle;
 }
 
 void FWindowsPlatformProcess::FreeDllHandle( void* DllHandle )
diff --git a/Engine/Source/Runtime/Core/Public/Windows/WindowsPlatformProcess.h b/Engine/Source/Runtime/Core/Public/Windows/WindowsPlatformProcess.h
--- a/Engine/Source/Runtime/Core/Public/Windows/WindowsPlatformProcess.h
+++ b/Engine/Source/Runtime/Core/Public/Windows/WindowsPlatformProcess.h
@@ -191,6 +191,11 @@ private:
 	 * this stack is used to reset the previous DllDirectory.
 	 */
 	static TArray<FString> DllDirectoryStack;
+
+	/**
+	 * All the DLL directories we want to load from. 
+	 */
+	static TArray<FString> DllDirectories;
 };
 
 
