Merge fix for threading bugs with MallocBinned stats

diff --git a/Engine/Source/Runtime/Core/Public/HAL/MallocBinned.h b/Engine/Source/Runtime/Core/Public/HAL/MallocBinned.h
--- a/Engine/Source/Runtime/Core/Public/HAL/MallocBinned.h
+++ b/Engine/Source/Runtime/Core/Public/HAL/MallocBinned.h
@@ -31,6 +31,39 @@
 #	define USE_FINE_GRAIN_LOCKS
 #endif
 
+//when modifying the global allocator stats, if we are using COARSE locks, then all callsites for stat modification are covered by the allocator-wide access guard. Thus the stats can be modified directly.
+//If we are using FINE locks, then we must modify the stats through atomics as the locks are either not actually covering the stat callsites, or are locking specific table locks which is not sufficient for stats.
+#if STATS
+#	if USE_COARSE_GRAIN_LOCKS
+#		define BINNED_STAT int64
+#		define BINNED_INCREMENT_STATCOUNTER(counter) (++(counter))
+#		define BINNED_DECREMENT_STATCOUNTER(counter) (--(counter))
+#		define BINNED_ADD_STATCOUNTER(counter, value) ((counter) += (value))
+#		define BINNED_PEAK_STATCOUNTER(PeakCounter, CompareVal) ((PeakCounter) = FMath::Max((PeakCounter), (CompareVal)))
+#	else
+#		define BINNED_STAT volatile int64
+#		define BINNED_INCREMENT_STATCOUNTER(counter) (FPlatformAtomics::InterlockedIncrement(&(counter)))
+#		define BINNED_DECREMENT_STATCOUNTER(counter) (FPlatformAtomics::InterlockedDecrement(&(counter)))
+#		define BINNED_ADD_STATCOUNTER(counter, value) (FPlatformAtomics::InterlockedAdd(&counter, (value)))
+#		define BINNED_PEAK_STATCOUNTER(PeakCounter, CompareVal)	{																												\
+																	int64 NewCompare;																							\
+																	int64 NewPeak;																								\
+																	do																											\
+																	{																											\
+																		NewCompare = (PeakCounter);																				\
+																		NewPeak = FMath::Max((PeakCounter), (CompareVal));														\
+																	}																											\
+																	while (FPlatformAtomics::InterlockedCompareExchange(&(PeakCounter), NewPeak, NewCompare) != NewCompare);	\
+																}
+#	endif
+#else
+#	define BINNED_STAT int64
+#	define BINNED_INCREMENT_STATCOUNTER(counter)
+#	define BINNED_DECREMENT_STATCOUNTER(counter)
+#	define BINNED_ADD_STATCOUNTER(counter, value)
+#	define BINNED_PEAK_STATCOUNTER(PeakCounter, CompareVal)
+#endif
+
 #include "LockFreeList.h"
 #include "Array.h"
 
@@ -305,16 +338,16 @@ private:
 #endif
 
 #if STATS
-	SIZE_T		OsCurrent;
-	SIZE_T		OsPeak;
-	SIZE_T		WasteCurrent;
-	SIZE_T		WastePeak;
-	SIZE_T		UsedCurrent;
-	SIZE_T		UsedPeak;
-	SIZE_T		CurrentAllocs;
-	SIZE_T		TotalAllocs;
+	BINNED_STAT		OsCurrent;
+	BINNED_STAT		OsPeak;
+	BINNED_STAT		WasteCurrent;
+	BINNED_STAT		WastePeak;
+	BINNED_STAT		UsedCurrent;
+	BINNED_STAT		UsedPeak;
+	BINNED_STAT		CurrentAllocs;
+	BINNED_STAT		TotalAllocs;
 	/** OsCurrent - WasteCurrent - UsedCurrent. */
-	SIZE_T		SlackCurrent;
+	BINNED_STAT		SlackCurrent;
 	double		MemTime;
 #endif
 
@@ -350,8 +383,8 @@ private:
 			OutOfMemory(IndirectPoolBlockSize * sizeof(FPoolInfo));
 		}
 		FMemory::Memset(Indirect, 0, IndirectPoolBlockSize*sizeof(FPoolInfo));
-		STAT(OsPeak = FMath::Max(OsPeak, OsCurrent += Align(IndirectPoolBlockSize * sizeof(FPoolInfo), PageSize)));
-		STAT(WastePeak = FMath::Max(WastePeak, WasteCurrent += Align(IndirectPoolBlockSize * sizeof(FPoolInfo), PageSize)));
+		BINNED_PEAK_STATCOUNTER(OsPeak, BINNED_ADD_STATCOUNTER(OsCurrent, (int64)(Align(IndirectPoolBlockSize * sizeof(FPoolInfo), PageSize))));
+		BINNED_PEAK_STATCOUNTER(WastePeak, BINNED_ADD_STATCOUNTER(WasteCurrent, (int64)(Align(IndirectPoolBlockSize * sizeof(FPoolInfo), PageSize))));
 		return Indirect;
 	}
 
@@ -475,8 +508,9 @@ private:
 		if (!HashBucketFreeList) 
 		{
 			HashBucketFreeList=(PoolHashBucket*)FPlatformMemory::BinnedAllocFromOS(PageSize);
-			STAT(OsPeak = FMath::Max(OsPeak, OsCurrent += PageSize));
-			STAT(WastePeak = FMath::Max(WastePeak, WasteCurrent += PageSize));
+			BINNED_PEAK_STATCOUNTER(OsPeak, BINNED_ADD_STATCOUNTER(OsCurrent, PageSize));
+			BINNED_PEAK_STATCOUNTER(WastePeak, BINNED_ADD_STATCOUNTER(WasteCurrent, PageSize));
+			
 			for (UPTRINT i=0, n=(PageSize/sizeof(PoolHashBucket)); i<n; ++i)
 			{
 				HashBucketFreeList->Link(new (HashBucketFreeList+i) PoolHashBucket());
@@ -527,13 +561,15 @@ private:
 				//Set trailing pools to point back to first pool
 				TrailingPool->SetAllocationSizes(0, 0, Offset, BinnedOSTableIndex);
 			}
+
+			
+			BINNED_PEAK_STATCOUNTER(OsPeak, BINNED_ADD_STATCOUNTER(OsCurrent,OsBytes));
+			BINNED_PEAK_STATCOUNTER(WastePeak, BINNED_ADD_STATCOUNTER(WasteCurrent, (OsBytes - Bytes)));			
 		}
 
 		// Init pool.
 		Pool->Link( Table->FirstPool );
 		Pool->SetAllocationSizes(Bytes, OsBytes, TableIndex, BinnedOSTableIndex);
-		STAT(OsPeak = FMath::Max(OsPeak, OsCurrent += OsBytes));
-		STAT(WastePeak = FMath::Max(WastePeak, WasteCurrent += OsBytes - Bytes));
 		Pool->Taken		 = 0;
 		Pool->FirstMem   = Free;
 
@@ -567,7 +603,7 @@ private:
 				Pool->Link( Table->ExhaustedPool );
 			}
 		}
-		STAT(UsedPeak = FMath::Max(UsedPeak, UsedCurrent += Table->BlockSize));
+		BINNED_PEAK_STATCOUNTER(UsedPeak, BINNED_ADD_STATCOUNTER(UsedCurrent, Table->BlockSize));
 		return Align(Free, Alignment);
 	}
 
@@ -578,7 +614,7 @@ private:
 	void FreeInternal( void* Ptr )
 	{
 		MEM_TIME(MemTime -= FPlatformTime::Seconds());
-		STAT(CurrentAllocs--);
+		BINNED_DECREMENT_STATCOUNTER(CurrentAllocs);
 
 		UPTRINT BasePtr;
 		FPoolInfo* Pool = FindPoolInfo((UPTRINT)Ptr, BasePtr);
@@ -621,7 +657,7 @@ private:
 			Free->NumFreeBlocks	= 1;
 			Free->Next			= Pool->FirstMem;
 			Pool->FirstMem		= Free;
-			STAT(UsedCurrent -= Table->BlockSize);
+			BINNED_ADD_STATCOUNTER(UsedCurrent, -(int64)(Table->BlockSize));
 
 			// Free this pool.
 			checkSlow(Pool->Taken >= 1);
@@ -632,8 +668,8 @@ private:
 #endif
 				// Free the OS memory.
 				SIZE_T OsBytes = Pool->GetOsBytes(PageSize, BinnedOSTableIndex);
-				STAT(OsCurrent -= OsBytes);
-				STAT(WasteCurrent -= OsBytes - Pool->GetBytes());
+				BINNED_ADD_STATCOUNTER(OsCurrent, -(int64)(OsBytes));
+				BINNED_ADD_STATCOUNTER(WasteCurrent, -(int64)(OsBytes - Pool->GetBytes()));
 				Pool->Unlink();
 				Pool->SetAllocationSizes(0, 0, 0, BinnedOSTableIndex);
 				OSFree((void*)BasePtr, OsBytes);
@@ -644,9 +680,10 @@ private:
 			// Free an OS allocation.
 			checkSlow(!((UPTRINT)Ptr & (PageSize-1)));
 			SIZE_T OsBytes = Pool->GetOsBytes(PageSize, BinnedOSTableIndex);
-			STAT(UsedCurrent -= Pool->GetBytes());
-			STAT(OsCurrent -= OsBytes);
-			STAT(WasteCurrent -= OsBytes - Pool->GetBytes());
+
+			BINNED_ADD_STATCOUNTER(UsedCurrent, -(int64)Pool->GetBytes());
+			BINNED_ADD_STATCOUNTER(OsCurrent, -(int64)OsBytes);
+			BINNED_ADD_STATCOUNTER(WasteCurrent, -(int64)(OsBytes - Pool->GetBytes()));
 			OSFree((void*)BasePtr, OsBytes);
 		}
 
@@ -953,8 +990,10 @@ public:
 		SIZE_T SpareBytesCount = FMath::Min<SIZE_T>(DEFAULT_BINNED_ALLOCATOR_ALIGNMENT, Size);
 		Size = FMath::Max<SIZE_T>(PoolTable[0].BlockSize, Size + (Alignment - SpareBytesCount));
 		MEM_TIME(MemTime -= FPlatformTime::Seconds());
-		STAT(CurrentAllocs++);
-		STAT(TotalAllocs++);
+		
+		BINNED_INCREMENT_STATCOUNTER(CurrentAllocs);
+		BINNED_INCREMENT_STATCOUNTER(TotalAllocs);
+		
 		FFreeMem* Free;
 		if( Size < BinnedSizeLimit )
 		{
@@ -1031,10 +1070,9 @@ public:
 			}
 			}
 			Free = (FFreeMem*)AlignedFree;
-			Pool->SetAllocationSizes(Size, AlignedSize, BinnedOSTableIndex, BinnedOSTableIndex);
-			STAT(OsPeak = FMath::Max(OsPeak, OsCurrent += AlignedSize));
-			STAT(UsedPeak = FMath::Max(UsedPeak, UsedCurrent += Size));
-			STAT(WastePeak = FMath::Max(WastePeak, WasteCurrent += AlignedSize - Size));
+			BINNED_PEAK_STATCOUNTER(OsPeak, BINNED_ADD_STATCOUNTER(OsCurrent, AlignedSize));
+			BINNED_PEAK_STATCOUNTER(UsedPeak, BINNED_ADD_STATCOUNTER(UsedCurrent, Size));
+			BINNED_PEAK_STATCOUNTER(WastePeak, BINNED_ADD_STATCOUNTER(WasteCurrent, (int64)(AlignedSize - Size)));
 		}
 
 		MEM_TIME(MemTime += FPlatformTime::Seconds());
@@ -1094,10 +1132,16 @@ public:
 				}
 				else
 				{
+//need a lock to cover the SetAllocationSizes()
+#ifdef USE_FINE_GRAIN_LOCKS
+					FScopeLock PoolInfoLock(&AccessGuard);
+#endif
+					int32 UsedChange = (NewSize - Pool->GetBytes());
+					
 					// Keep as-is, reallocation isn't worth the overhead.
-					STAT(UsedCurrent += NewSize - Pool->GetBytes());
-					STAT(UsedPeak = FMath::Max(UsedPeak, UsedCurrent));
-					STAT(WasteCurrent += Pool->GetBytes() - NewSize);
+					BINNED_ADD_STATCOUNTER(UsedCurrent, UsedChange);
+					BINNED_PEAK_STATCOUNTER(UsedPeak, UsedCurrent);
+					BINNED_ADD_STATCOUNTER(WasteCurrent, (Pool->GetBytes() - NewSize));
 					Pool->SetAllocationSizes(NewSizeUnmodified, Pool->GetOsBytes(PageSize, BinnedOSTableIndex), BinnedOSTableIndex, BinnedOSTableIndex);
 				}
 			}
