Linux: more logging to catch a mysterious bug (UE-18812).

- Some buildfarm-built binaries exhibit strange behavior that cannot be reliably reproduced; adding more logs and exec commands.
- Also some minor clean up.

diff --git a/Engine/Source/Runtime/Core/Private/Linux/LinuxApplication.cpp b/Engine/Source/Runtime/Core/Private/Linux/LinuxApplication.cpp
--- a/Engine/Source/Runtime/Core/Private/Linux/LinuxApplication.cpp
+++ b/Engine/Source/Runtime/Core/Private/Linux/LinuxApplication.cpp
@@ -1602,3 +1602,60 @@ TSharedPtr< FLinuxWindow > FLinuxApplication::GetRootWindow(const TSharedPtr< FL
 	}
 	return ParentWindow;
 }
+
+bool FLinuxApplication::Exec(UWorld* InWorld, const TCHAR* Cmd, FOutputDevice& Ar)
+{
+	// Ignore any execs that doesn't start with LinuxApp
+	if (!FParse::Command(&Cmd, TEXT("LinuxApp")))
+	{
+		return false;
+	}
+
+	if (FParse::Command(&Cmd, TEXT("Cursor")))
+	{
+		return HandleCursorCommand(Cmd, Ar);
+	}
+	else if (FParse::Command(&Cmd, TEXT("Window")))
+	{
+		return HandleWindowCommand(Cmd, Ar);
+	}
+
+	return false;
+}
+
+bool FLinuxApplication::HandleCursorCommand(const TCHAR* Cmd, FOutputDevice& Ar)
+{
+	if (FParse::Command(&Cmd, TEXT("Status")))
+	{
+		FLinuxCursor *LinuxCursor = static_cast<FLinuxCursor*>(Cursor.Get());
+		FVector2D CurrentPosition = LinuxCursor->GetPosition();
+
+		Ar.Logf(TEXT("Cursor status:"));
+		Ar.Logf(TEXT("Position: (%f, %f)"), CurrentPosition.X, CurrentPosition.Y);
+		Ar.Logf(TEXT("IsHidden: %s"), LinuxCursor->IsHidden() ? TEXT("true") : TEXT("false"));
+		Ar.Logf(TEXT("bIsMouseCaptureEnabled: %s"), bIsMouseCaptureEnabled ? TEXT("true") : TEXT("false"));
+		Ar.Logf(TEXT("bUsingHighPrecisionMouseInput: %s"), bUsingHighPrecisionMouseInput ? TEXT("true") : TEXT("false"));
+		Ar.Logf(TEXT("bIsMouseCaptureEnabled: %s"), bIsMouseCaptureEnabled ? TEXT("true") : TEXT("false"));
+
+		return true;
+	}
+
+	return false;
+}
+
+bool FLinuxApplication::HandleWindowCommand(const TCHAR* Cmd, FOutputDevice& Ar)
+{
+	if (FParse::Command(&Cmd, TEXT("List")))
+	{
+		Ar.Logf(TEXT("Window list:"));
+		for (int WindowIdx = 0, NumWindows = Windows.Num(); WindowIdx < NumWindows; ++WindowIdx)
+		{
+			Ar.Logf(TEXT("%d: native handle: %p, debugging ID: %d"), WindowIdx, Windows[WindowIdx]->GetHWnd(), Windows[WindowIdx]->GetID());
+		}
+
+		return true;
+	}
+
+	return false;
+}
+
diff --git a/Engine/Source/Runtime/Core/Private/Linux/LinuxPlatformMisc.cpp b/Engine/Source/Runtime/Core/Private/Linux/LinuxPlatformMisc.cpp
--- a/Engine/Source/Runtime/Core/Private/Linux/LinuxPlatformMisc.cpp
+++ b/Engine/Source/Runtime/Core/Private/Linux/LinuxPlatformMisc.cpp
@@ -181,13 +181,22 @@ bool FLinuxPlatformMisc::PlatformInitMultimedia()
 			return false;
 		}
 
+		// print out version information
+		SDL_version CompileTimeSDLVersion;
+		SDL_version RunTimeSDLVersion;
+		SDL_VERSION(&CompileTimeSDLVersion);
+		SDL_GetVersion(&RunTimeSDLVersion);
+		UE_LOG(LogInit, Log, TEXT("Initialized SDL %d.%d.%d (compiled against %d.%d.%d)"),
+			CompileTimeSDLVersion.major, CompileTimeSDLVersion.minor, CompileTimeSDLVersion.patch,
+			RunTimeSDLVersion.major, RunTimeSDLVersion.minor, RunTimeSDLVersion.patch
+			);
+
 		// Used to make SDL push SDL_TEXTINPUT events.
 		SDL_StartTextInput();
 
 		GInitializedSDL = true;
 
 		// needs to come after GInitializedSDL, otherwise it will recurse here
-		// @TODO [RCL] 2014-09-30 - move to FDisplayMetrics itself sometime after 4.5
 		if (!UE_BUILD_SHIPPING)
 		{
 			// dump information about screens for debug
diff --git a/Engine/Source/Runtime/Core/Private/Linux/LinuxWindow.cpp b/Engine/Source/Runtime/Core/Private/Linux/LinuxWindow.cpp
--- a/Engine/Source/Runtime/Core/Private/Linux/LinuxWindow.cpp
+++ b/Engine/Source/Runtime/Core/Private/Linux/LinuxWindow.cpp
@@ -15,9 +15,6 @@ FLinuxWindow::~FLinuxWindow()
 	//       Use Destroy() instead.
 }
 
-const TCHAR FLinuxWindow::AppWindowClass[] = TEXT("UnrealWindow");
-
-
 TSharedRef< FLinuxWindow > FLinuxWindow::Make()
 {
 	// First, allocate the new native window object.  This doesn't actually create a native window or anything,
diff --git a/Engine/Source/Runtime/Core/Public/Linux/LinuxApplication.h b/Engine/Source/Runtime/Core/Public/Linux/LinuxApplication.h
--- a/Engine/Source/Runtime/Core/Public/Linux/LinuxApplication.h
+++ b/Engine/Source/Runtime/Core/Public/Linux/LinuxApplication.h
@@ -10,7 +10,7 @@ class FLinuxWindow;
 class FGenericApplicationMessageHandler;
 
 
-class FLinuxApplication : public GenericApplication
+class FLinuxApplication : public GenericApplication, public FSelfRegisteringExec
 {
 
 public:
@@ -26,6 +26,9 @@ public:
 	
 	virtual void DestroyApplication() override;
 
+	// FSelfRegisteringExec
+	virtual bool Exec(UWorld* InWorld, const TCHAR* Cmd, FOutputDevice& Ar) override;
+
 public:
 	virtual void SetMessageHandler( const TSharedRef< class FGenericApplicationMessageHandler >& InMessageHandler ) override;
 
@@ -128,6 +131,12 @@ private:
 
 	TSharedPtr< FLinuxWindow > GetRootWindow(const TSharedPtr< FLinuxWindow >& Window);
 
+	/** Handles "Cursor" exec commands" */
+	bool HandleCursorCommand(const TCHAR* Cmd, FOutputDevice& Ar);
+
+	/** Handles "Window" exec commands" */
+	bool HandleWindowCommand(const TCHAR* Cmd, FOutputDevice& Ar);
+
 private:
 
 	void RefreshDisplayCache();
@@ -154,8 +163,6 @@ private:
 		}
 	};
 
-	bool bUsingHighPrecisionMouseInput;
-
 	TArray< SDL_Event > PendingEvents;
 
 	TArray< TSharedRef< FLinuxWindow > > Windows;
@@ -168,7 +175,13 @@ private:
 
 	int32 bAllowedToDeferMessageProcessing;
 
+	/** Using high precision mouse input */
+	bool bUsingHighPrecisionMouseInput;
+
+	/** TODO: describe */
 	bool bIsMouseCursorLocked;
+
+	/** TODO: describe */
 	bool bIsMouseCaptureEnabled;
 
 	/** Window that we think has been activated last. */
diff --git a/Engine/Source/Runtime/Core/Public/Linux/LinuxWindow.h b/Engine/Source/Runtime/Core/Public/Linux/LinuxWindow.h
--- a/Engine/Source/Runtime/Core/Public/Linux/LinuxWindow.h
+++ b/Engine/Source/Runtime/Core/Public/Linux/LinuxWindow.h
@@ -23,9 +23,6 @@ class CORE_API FLinuxWindow : public FGenericWindow
 public:
 	~FLinuxWindow();
 
-	/** Win32 requirement: see CreateWindowEx and RegisterClassEx. */
-	static const TCHAR AppWindowClass[];
-
 	/** Create a new SDLWindow.
 	 *
 	 * @param SlateWindows		List of all top level Slate windows.  This function will add the owner window to this list.
