Fix 32bit platforms CIS

diff --git a/Engine/Source/Runtime/Core/Public/HAL/MallocBinned.h b/Engine/Source/Runtime/Core/Public/HAL/MallocBinned.h
--- a/Engine/Source/Runtime/Core/Public/HAL/MallocBinned.h
+++ b/Engine/Source/Runtime/Core/Public/HAL/MallocBinned.h
@@ -31,23 +31,29 @@
 #	define USE_FINE_GRAIN_LOCKS
 #endif
 
+#if PLATFORM_64BITS
+typedef int64 BINNED_STAT_TYPE;
+#else
+typedef int32 BINNED_STAT_TYPE;
+#endif
+
 //when modifying the global allocator stats, if we are using COARSE locks, then all callsites for stat modification are covered by the allocator-wide access guard. Thus the stats can be modified directly.
 //If we are using FINE locks, then we must modify the stats through atomics as the locks are either not actually covering the stat callsites, or are locking specific table locks which is not sufficient for stats.
 #if STATS
 #	if USE_COARSE_GRAIN_LOCKS
-#		define BINNED_STAT int64
+#		define BINNED_STAT BINNED_STAT_TYPE
 #		define BINNED_INCREMENT_STATCOUNTER(counter) (++(counter))
 #		define BINNED_DECREMENT_STATCOUNTER(counter) (--(counter))
 #		define BINNED_ADD_STATCOUNTER(counter, value) ((counter) += (value))
 #		define BINNED_PEAK_STATCOUNTER(PeakCounter, CompareVal) ((PeakCounter) = FMath::Max((PeakCounter), (CompareVal)))
 #	else
-#		define BINNED_STAT volatile int64
+#		define BINNED_STAT volatile BINNED_STAT_TYPE
 #		define BINNED_INCREMENT_STATCOUNTER(counter) (FPlatformAtomics::InterlockedIncrement(&(counter)))
 #		define BINNED_DECREMENT_STATCOUNTER(counter) (FPlatformAtomics::InterlockedDecrement(&(counter)))
 #		define BINNED_ADD_STATCOUNTER(counter, value) (FPlatformAtomics::InterlockedAdd(&counter, (value)))
 #		define BINNED_PEAK_STATCOUNTER(PeakCounter, CompareVal)	{																												\
-																	int64 NewCompare;																							\
-																	int64 NewPeak;																								\
+																	BINNED_STAT_TYPE NewCompare;																							\
+																	BINNED_STAT_TYPE NewPeak;																								\
 																	do																											\
 																	{																											\
 																		NewCompare = (PeakCounter);																				\
@@ -57,7 +63,7 @@
 																}
 #	endif
 #else
-#	define BINNED_STAT int64
+#	define BINNED_STAT BINNED_STAT_TYPE
 #	define BINNED_INCREMENT_STATCOUNTER(counter)
 #	define BINNED_DECREMENT_STATCOUNTER(counter)
 #	define BINNED_ADD_STATCOUNTER(counter, value)
