Fix pool sizes not getting set

diff --git a/Engine/Source/Runtime/Core/Public/HAL/MallocBinned.h b/Engine/Source/Runtime/Core/Public/HAL/MallocBinned.h
--- a/Engine/Source/Runtime/Core/Public/HAL/MallocBinned.h
+++ b/Engine/Source/Runtime/Core/Public/HAL/MallocBinned.h
@@ -1072,10 +1072,11 @@ public:
 						check(TrailingPool);
 						//Set trailing pools to point back to first pool
 						TrailingPool->SetAllocationSizes(0, 0, Offset, BinnedOSTableIndex);
-			}
-			}
+					}
+				}
 			}
 			Free = (FFreeMem*)AlignedFree;
+			Pool->SetAllocationSizes(Size, AlignedSize, BinnedOSTableIndex, BinnedOSTableIndex);
 			BINNED_PEAK_STATCOUNTER(OsPeak, BINNED_ADD_STATCOUNTER(OsCurrent, AlignedSize));
 			BINNED_PEAK_STATCOUNTER(UsedPeak, BINNED_ADD_STATCOUNTER(UsedCurrent, Size));
 			BINNED_PEAK_STATCOUNTER(WastePeak, BINNED_ADD_STATCOUNTER(WasteCurrent, (int64)(AlignedSize - Size)));
