Fix "Under My Blueprints Right click Rename does not work on Mac" on Mac by ensuring that we issue a window deactivation event when the current 'main' window is closed not the current 'key' window.

diff --git a/Engine/Source/Runtime/Core/Private/Mac/MacApplication.cpp b/Engine/Source/Runtime/Core/Private/Mac/MacApplication.cpp
--- a/Engine/Source/Runtime/Core/Private/Mac/MacApplication.cpp
+++ b/Engine/Source/Runtime/Core/Private/Mac/MacApplication.cpp
@@ -978,7 +978,7 @@ void FMacApplication::OnWindowDidResize(TSharedRef<FMacWindow> Window)
 bool FMacApplication::OnWindowDestroyed(TSharedRef<FMacWindow> Window)
 {
 	SCOPED_AUTORELEASE_POOL;
-	if ([Window->GetWindowHandle() isKeyWindow])
+	if ([Window->GetWindowHandle() isMainWindow])
 	{
 		MessageHandler->OnWindowActivationChanged(Window, EWindowActivation::Deactivate);
 	}
