Fix shadow variable warnings on Mac.

diff --git a/Engine/Source/Runtime/Core/Private/Modules/ModuleManager.cpp b/Engine/Source/Runtime/Core/Private/Modules/ModuleManager.cpp
--- a/Engine/Source/Runtime/Core/Private/Modules/ModuleManager.cpp
+++ b/Engine/Source/Runtime/Core/Private/Modules/ModuleManager.cpp
@@ -943,11 +943,11 @@ void FModuleManager::FindModulePathsInDirectory(const FString& InDirectoryName,
 	if(QueryModulesDelegate.IsBound())
 	{
 		// Use the delegate to query all the modules in this directory
-		TMap<FString, FString> Modules;
-		QueryModulesDelegate.Execute(InDirectoryName, bIsGameDirectory, Modules);
+		TMap<FString, FString> ValidModules;
+		QueryModulesDelegate.Execute(InDirectoryName, bIsGameDirectory, ValidModules);
 
 		// Fill the output map with modules that match the wildcard
-		for(const TPair<FString, FString>& Pair: Modules)
+		for(const TPair<FString, FString>& Pair: ValidModules)
 		{
 			if(Pair.Key.MatchesWildcard(NamePattern))
 			{
