Added string formatting methods for internationalization metadata.

diff --git a/Engine/Source/Runtime/Core/Private/Internationalization/InternationalizationMetadata.cpp b/Engine/Source/Runtime/Core/Private/Internationalization/InternationalizationMetadata.cpp
--- a/Engine/Source/Runtime/Core/Private/Internationalization/InternationalizationMetadata.cpp
+++ b/Engine/Source/Runtime/Core/Private/Internationalization/InternationalizationMetadata.cpp
@@ -260,6 +260,19 @@ bool FLocMetadataObject::IsMetadataExactMatch( const FLocMetadataObject* const M
 	return MetadataA->IsExactMatch( *(MetadataB) );
 }
 
+FString FLocMetadataObject::ToString() const
+{
+	FString MemberList;
+	TMap< FString, TSharedPtr<FLocMetadataValue> > Values;
+	for (const auto Pair : Values)
+	{
+		const FString MemberName = Pair.Key;
+		const TSharedPtr<FLocMetadataValue> MemberValue = Pair.Value;
+		MemberList += (MemberList.IsEmpty() ? TEXT("") : TEXT(",")) + MemberValue->ToString();
+	}
+	return FString::Printf(TEXT("{%s}"), *MemberList);
+}
+
 namespace
 {
 	void SerializeLocMetadataValue(FArchive& Archive, FLocMetadataValue*& Value)
@@ -524,6 +537,16 @@ TSharedRef<FLocMetadataValue> FLocMetadataValueArray::Clone() const
 	return MakeShareable( new FLocMetadataValueArray( NewValue ) );
 }
 
+FString FLocMetadataValueArray::ToString() const
+{
+	FString ElementList;
+	for (const TSharedPtr<FLocMetadataValue> Element : Value)
+	{
+		ElementList += (ElementList.IsEmpty() ? TEXT("") : TEXT(",")) + Element->ToString();
+	}
+	return FString::Printf(TEXT("[%s]"), *ElementList);
+}
+
 FLocMetadataValueArray::FLocMetadataValueArray( FArchive& Archive )
 {
 	check(Archive.IsLoading());
@@ -589,6 +612,11 @@ TSharedRef<FLocMetadataValue> FLocMetadataValueObject::Clone() const
 	return MakeShareable( new FLocMetadataValueObject( NewLocMetadataObject ) );
 }
 
+FString FLocMetadataValueObject::ToString() const
+{
+	return Value->ToString();
+}
+
 FLocMetadataValueObject::FLocMetadataValueObject( FArchive& Archive )
 	: Value(new FLocMetadataObject)
 {
diff --git a/Engine/Source/Runtime/Core/Public/Internationalization/InternationalizationMetadata.h b/Engine/Source/Runtime/Core/Public/Internationalization/InternationalizationMetadata.h
--- a/Engine/Source/Runtime/Core/Public/Internationalization/InternationalizationMetadata.h
+++ b/Engine/Source/Runtime/Core/Public/Internationalization/InternationalizationMetadata.h
@@ -25,7 +25,9 @@ class CORE_API FLocMetadataValue
 public:
 	virtual ~FLocMetadataValue() = 0;
 
-	/** Returns this value as a number, throwing an error if this is not a Metadata String */
+	virtual FString ToString() const = 0;
+
+	/** Returns this value as a string, throwing an error if this is not a Metadata String */
 	virtual FString AsString() {ErrorMessage(TEXT("String")); return FString();}
 
 	/** Returns this value as a boolean, throwing an error if this is not an Metadata Bool */
@@ -152,6 +154,8 @@ public:
 
 	static bool IsMetadataExactMatch( const FLocMetadataObject* const MetadataA, const FLocMetadataObject* const MetadataB );
 
+	FString ToString() const;
+
 public:
 	/** Stores the name/value pairs for the metadata object */
 	TMap< FString, TSharedPtr<FLocMetadataValue> > Values;
@@ -175,6 +179,8 @@ public:
 	virtual TSharedRef<FLocMetadataValue> Clone() const override;
 	void SetString( const FString& InString ) { Value = InString; }
 
+	virtual FString ToString() const override {return Value;}
+
 protected:
 
 	FString Value;
@@ -198,6 +204,8 @@ public:
 
 	virtual TSharedRef<FLocMetadataValue> Clone() const override;
 	
+	virtual FString ToString() const override {return Value ? TEXT("true") : TEXT("false");}
+
 protected:
 	bool Value;
 
@@ -219,6 +227,8 @@ public:
 
 	virtual TSharedRef<FLocMetadataValue> Clone() const override;
 	
+	virtual FString ToString() const override;
+
 protected:
 	TArray< TSharedPtr<FLocMetadataValue> > Value;
 
@@ -240,6 +250,8 @@ public:
 
 	virtual TSharedRef<FLocMetadataValue> Clone() const override;
 	
+	virtual FString ToString() const override;
+
 protected:
 	TSharedPtr<FLocMetadataObject> Value;
 
