Linux: added missing misc platform function (GetSystemErrorMessage).

diff --git a/Engine/Source/Runtime/Core/Private/Linux/LinuxPlatformMisc.cpp b/Engine/Source/Runtime/Core/Private/Linux/LinuxPlatformMisc.cpp
--- a/Engine/Source/Runtime/Core/Private/Linux/LinuxPlatformMisc.cpp
+++ b/Engine/Source/Runtime/Core/Private/Linux/LinuxPlatformMisc.cpp
@@ -357,6 +357,20 @@ uint32 FLinuxPlatformMisc::GetKeyMap( uint32* KeyCodes, FString* KeyNames, uint3
 	return NumMappings;
 }
 
+const TCHAR* FLinuxPlatformMisc::GetSystemErrorMessage(TCHAR* OutBuffer, int32 BufferCount, int32 Error)
+{
+	check(OutBuffer && BufferCount);
+	if (Error == 0)
+	{
+		Error = errno;
+	}
+
+	FString Message = FString::Printf(TEXT("errno=%d (%s)"), Error, UTF8_TO_TCHAR(strerror(Error)));
+	FCString::Strncpy(OutBuffer, *Message, BufferCount);
+
+	return OutBuffer;
+}
+
 void FLinuxPlatformMisc::ClipboardCopy(const TCHAR* Str)
 {
 	if (SDL_HasClipboardText() == SDL_TRUE)
diff --git a/Engine/Source/Runtime/Core/Public/Linux/LinuxPlatformMisc.h b/Engine/Source/Runtime/Core/Public/Linux/LinuxPlatformMisc.h
--- a/Engine/Source/Runtime/Core/Public/Linux/LinuxPlatformMisc.h
+++ b/Engine/Source/Runtime/Core/Public/Linux/LinuxPlatformMisc.h
@@ -67,6 +67,7 @@ struct CORE_API FLinuxPlatformMisc : public FGenericPlatformMisc
 	static void LowLevelOutputDebugString(const TCHAR *Message);
 	static bool ControlScreensaver(EScreenSaverAction Action);
 
+	static const TCHAR* GetSystemErrorMessage(TCHAR* OutBuffer, int32 BufferCount, int32 Error);
 	static void ClipboardCopy(const TCHAR* Str);
 	static void ClipboardPaste(class FString& Dest);
 
