Fixed compilation warning-as-error due to shadowing variable in InternationalizationMetadata.cpp.

diff --git a/Engine/Source/Runtime/Core/Private/Internationalization/InternationalizationMetadata.cpp b/Engine/Source/Runtime/Core/Private/Internationalization/InternationalizationMetadata.cpp
--- a/Engine/Source/Runtime/Core/Private/Internationalization/InternationalizationMetadata.cpp
+++ b/Engine/Source/Runtime/Core/Private/Internationalization/InternationalizationMetadata.cpp
@@ -263,7 +263,6 @@ bool FLocMetadataObject::IsMetadataExactMatch( const FLocMetadataObject* const M
 FString FLocMetadataObject::ToString() const
 {
 	FString MemberList;
-	TMap< FString, TSharedPtr<FLocMetadataValue> > Values;
 	for (const auto Pair : Values)
 	{
 		const FString MemberName = Pair.Key;
