Fix for VC when trying to move construct a TTuple.

diff --git a/Engine/Source/Runtime/Core/Public/Delegates/Tuple.h b/Engine/Source/Runtime/Core/Public/Delegates/Tuple.h
--- a/Engine/Source/Runtime/Core/Public/Delegates/Tuple.h
+++ b/Engine/Source/Runtime/Core/Public/Delegates/Tuple.h
@@ -4,6 +4,21 @@
 
 #include "IntegerSequence.h"
 
+template <int32 N, typename... Types>
+struct TNthTypeFromParameterPack;
+
+template <int32 N, typename T, typename... OtherTypes>
+struct TNthTypeFromParameterPack<N, T, OtherTypes...>
+{
+	typedef typename TNthTypeFromParameterPack<N - 1, OtherTypes...>::Type Type;
+};
+
+template <typename T, typename... OtherTypes>
+struct TNthTypeFromParameterPack<0, T, OtherTypes...>
+{
+	typedef T Type;
+};
+
 template <typename T, uint32 Index>
 struct TTupleElement
 {
@@ -122,12 +137,19 @@ struct TTupleImpl<TIntegerSequence<Indices...>, Types...> : TTupleElement<Types,
 
 #endif
 
+template <typename T, typename... Types>
+struct TDecayedFrontOfParameterPackIsSameType
+{
+	enum { Value = TAreTypesEqual<T, typename TDecay<typename TNthTypeFromParameterPack<0, Types...>::Type>::Type>::Value };
+};
+
 template <typename... Types>
 struct TTuple : TTupleImpl<TMakeIntegerSequence<sizeof...(Types)>, Types...>
 {
-	template <typename... ArgTypes>
+	template <typename... ArgTypes, typename = typename TEnableIf<!TAnd<TBoolConstant<sizeof...(ArgTypes) == 1>, TDecayedFrontOfParameterPackIsSameType<TTuple, ArgTypes...>>::Value>::Type>
 	explicit TTuple(ArgTypes&&... Args)
 		: TTupleImpl<TMakeIntegerSequence<sizeof...(Types)>, Types...>(Forward<ArgTypes>(Args)...)
 	{
+		// This constructor is disabled for TTuple because VC is incorrectly instantiating it as a move/copy constructor.
 	}
 };
diff --git a/Engine/Source/Runtime/Core/Public/Templates/UnrealTemplate.h b/Engine/Source/Runtime/Core/Public/Templates/UnrealTemplate.h
--- a/Engine/Source/Runtime/Core/Public/Templates/UnrealTemplate.h
+++ b/Engine/Source/Runtime/Core/Public/Templates/UnrealTemplate.h
@@ -550,6 +550,14 @@ struct TOr : TOrValue<LHS::Value, RHS>
 {
 };
 
+/**
+ * Metafunction which returns the specified boolean value.
+ */
+template <bool bValue>
+struct TBoolConstant
+{
+	enum { Value = bValue };
+};
 
 /**
  * Equivalent to std::declval.  
