Include subdirectories when searching for module manifests. Fixes missing target platforms in the editor.

diff --git a/Engine/Source/Runtime/Core/Private/Modules/ModuleManager.cpp b/Engine/Source/Runtime/Core/Private/Modules/ModuleManager.cpp
--- a/Engine/Source/Runtime/Core/Private/Modules/ModuleManager.cpp
+++ b/Engine/Source/Runtime/Core/Private/Modules/ModuleManager.cpp
@@ -942,16 +942,25 @@ void FModuleManager::FindModulePathsInDirectory(const FString& InDirectoryName,
 {
 	if(QueryModulesDelegate.IsBound())
 	{
-		// Use the delegate to query all the modules in this directory
-		TMap<FString, FString> ValidModules;
-		QueryModulesDelegate.Execute(InDirectoryName, bIsGameDirectory, ValidModules);
+		// Find all the directories to search through, including the base directory
+		TArray<FString> SearchDirectoryNames;
+		IFileManager::Get().FindFilesRecursive(SearchDirectoryNames, *InDirectoryName, TEXT("*"), false, true);
+		SearchDirectoryNames.Insert(InDirectoryName, 0);
 
-		// Fill the output map with modules that match the wildcard
-		for(const TPair<FString, FString>& Pair: ValidModules)
+		// Find the modules in each directory
+		for(const FString& SearchDirectoryName: SearchDirectoryNames)
 		{
-			if(Pair.Key.MatchesWildcard(NamePattern))
+			// Use the delegate to query all the modules in this directory
+			TMap<FString, FString> ValidModules;
+			QueryModulesDelegate.Execute(SearchDirectoryName, bIsGameDirectory, ValidModules);
+
+			// Fill the output map with modules that match the wildcard
+			for(const TPair<FString, FString>& Pair: ValidModules)
 			{
-				OutModulePaths.Add(FName(*Pair.Key), FPaths::Combine(*InDirectoryName, *Pair.Value));
+				if(Pair.Key.MatchesWildcard(NamePattern))
+				{
+					OutModulePaths.Add(FName(*Pair.Key), *FPaths::Combine(*SearchDirectoryName, *Pair.Value));
+				}
 			}
 		}
 	}
diff --git a/Engine/Source/Runtime/Launch/Private/LaunchEngineLoop.cpp b/Engine/Source/Runtime/Launch/Private/LaunchEngineLoop.cpp
--- a/Engine/Source/Runtime/Launch/Private/LaunchEngineLoop.cpp
+++ b/Engine/Source/Runtime/Launch/Private/LaunchEngineLoop.cpp
@@ -1031,7 +1031,7 @@ int32 FEngineLoop::PreInit( const TCHAR* CmdLine )
 	}
 
 	// Set up the module list and version information, if it's not compiled-in
-#if 0 && !IS_MONOLITHIC && BUILT_FROM_CHANGELIST == 0
+#if !IS_MONOLITHIC && BUILT_FROM_CHANGELIST == 0
 	static FVersionedModuleEnumerator ModuleEnumerator;
 	if(ModuleEnumerator.RegisterWithModuleManager())
 	{
