Removed type trait macros in favour of the type trait templates
(approved/recommended by [REDACTED_PERSON])

diff --git a/Engine/Source/Runtime/Core/Public/Templates/UnrealTypeTraits.h b/Engine/Source/Runtime/Core/Public/Templates/UnrealTypeTraits.h
--- a/Engine/Source/Runtime/Core/Public/Templates/UnrealTypeTraits.h
+++ b/Engine/Source/Runtime/Core/Public/Templates/UnrealTypeTraits.h
@@ -182,7 +182,7 @@ template<typename T> struct TIsLValueReferenceType     { enum { Value = false };
 template<typename T> struct TIsLValueReferenceType<T&> { enum { Value = true  }; };
 
 /**
- * TIsLValueReferenceType
+ * TIsRValueReferenceType
  */
 template<typename T> struct TIsRValueReferenceType      { enum { Value = false }; };
 template<typename T> struct TIsRValueReferenceType<T&&> { enum { Value = true  }; };
@@ -606,3 +606,14 @@ struct TIsEnum
 {
 	enum { Value = IS_ENUM(T) };
 };
+
+/*-----------------------------------------------------------------------------
+ * Undef Macros abstracting the presence of certain compiler intrinsic type traits
+ -----------------------------------------------------------------------------*/
+#undef IS_EMPTY
+#undef IS_ENUM
+#undef IS_POD
+#undef HAS_TRIVIAL_COPY
+#undef HAS_TRIVIAL_ASSIGN
+#undef HAS_TRIVIAL_DESTRUCTOR
+#undef HAS_TRIVIAL_CONSTRUCTOR
diff --git a/Engine/Source/Runtime/CoreUObject/Public/UObject/Class.h b/Engine/Source/Runtime/CoreUObject/Public/UObject/Class.h
--- a/Engine/Source/Runtime/CoreUObject/Public/UObject/Class.h
+++ b/Engine/Source/Runtime/CoreUObject/Public/UObject/Class.h
@@ -1646,7 +1646,7 @@ public:
 	FORCEINLINE static FString GetValueAsString( const TCHAR* EnumPath, const T EnumValue )
 	{
 		// For the C++ enum.
-		static_assert(IS_ENUM(T), "Should only call this with enum types");
+		static_assert(TIsEnum<T>::Value, "Should only call this with enum types");
 		return GetValueAsString_Internal(EnumPath, (int32)EnumValue);
 	}
 
@@ -1672,7 +1672,7 @@ public:
 	FORCEINLINE static FText GetDisplayValueAsText( const TCHAR* EnumPath, const T EnumValue )
 	{
 		// For the C++ enum.
-		static_assert(IS_ENUM(T), "Should only call this with enum types");
+		static_assert(TIsEnum<T>::Value, "Should only call this with enum types");
 		return GetDisplayValueAsText_Internal(EnumPath, (int32)EnumValue);
 	}
 
