FCustomVersion optimizations:

- Friendly name will no longer be serialized
- Reduced memory footprint of friendly names
- Faster custom version look up when serializing many custom versions

diff --git a/Engine/Source/Developer/DirectoryWatcher/Private/FileCache.cpp b/Engine/Source/Developer/DirectoryWatcher/Private/FileCache.cpp
--- a/Engine/Source/Developer/DirectoryWatcher/Private/FileCache.cpp
+++ b/Engine/Source/Developer/DirectoryWatcher/Private/FileCache.cpp
@@ -9,7 +9,7 @@ namespace DirectoryWatcher
 {
 
 template<typename T>
-void ReadWithCustomVersions(FArchive& Ar, T& Data)
+void ReadWithCustomVersions(FArchive& Ar, T& Data, ECustomVersionSerializationFormat::Type CustomVersionFormat)
 {
 	int64 CustomVersionsOffset = 0;
 	Ar << CustomVersionsOffset;
@@ -20,7 +20,7 @@ void ReadWithCustomVersions(FArchive& Ar, T& Data)
 
 	// Serialize the custom versions
 	FCustomVersionContainer Vers = Ar.GetCustomVersions();
-	Vers.Serialize(Ar);
+	Vers.Serialize(Ar, CustomVersionFormat);
 	Ar.SetCustomVersions(Vers);
 
 	Ar.Seek(DataStart);
@@ -65,7 +65,20 @@ EFileAction ToFileAction(FFileChangeData::EFileChangeAction InAction)
 const FGuid FFileCacheCustomVersion::Key(0x8E7DDCB3, 0x80DA47BB, 0x9FD346A2, 0x93984DF6);
 FCustomVersionRegistration GRegisterFileCacheVersion(FFileCacheCustomVersion::Key, FFileCacheCustomVersion::Latest, TEXT("FileCacheVersion"));
 
-static const uint32 CacheFileMagicNumber = 0x03DCCB00;
+static const uint32 CacheFileMagicNumberOldCustomVersionFormat = 0x03DCCB00;
+static const uint32 CacheFileMagicNumber = 0x03DCCB03;
+
+static ECustomVersionSerializationFormat::Type GetCustomVersionFormatForFileCache(uint32 MagicNumber)
+{
+	if (MagicNumber == CacheFileMagicNumberOldCustomVersionFormat)
+	{
+		return ECustomVersionSerializationFormat::Guids;
+	}
+	else
+	{
+		return ECustomVersionSerializationFormat::Optimized;
+	}
+}
 
 /** Single runnable thread used to parse file cache directories without blocking the main thread */
 struct FAsyncTaskThread : public FRunnable
@@ -465,10 +478,10 @@ TOptional<FDirectoryState> FFileCache::ReadCache() const
 			uint32 MagicNumber = 0;
 			*Ar << MagicNumber;
 
-			if (MagicNumber == CacheFileMagicNumber)
+			if (MagicNumber == CacheFileMagicNumber || MagicNumber == CacheFileMagicNumberOldCustomVersionFormat)
 			{
 				FDirectoryState Result;
-				ReadWithCustomVersions(*Ar, Result);
+				ReadWithCustomVersions(*Ar, Result, GetCustomVersionFormatForFileCache(MagicNumber));
 
 				Optional.Emplace(MoveTemp(Result));
 			}
diff --git a/Engine/Source/Editor/UnrealEd/Private/CookOnTheFlyServer.cpp b/Engine/Source/Editor/UnrealEd/Private/CookOnTheFlyServer.cpp
--- a/Engine/Source/Editor/UnrealEd/Private/CookOnTheFlyServer.cpp
+++ b/Engine/Source/Editor/UnrealEd/Private/CookOnTheFlyServer.cpp
@@ -2847,7 +2847,7 @@ bool UCookOnTheFlyServer::GetCurrentIniVersionStrings( const ITargetPlatform* Ta
 	static const FCustomVersionContainer& CustomVersionContainer = FCustomVersionContainer::GetRegistered();
 	for (const auto& CustomVersion : CustomVersionContainer.GetAllVersions())
 	{
-		FString CustomVersionString = FString::Printf(TEXT("%s:%s:%d"), *CustomVersion.FriendlyName, *CustomVersion.Key.ToString(), CustomVersion.Version);
+		FString CustomVersionString = FString::Printf(TEXT("%s:%s:%d"), *CustomVersion.GetFriendlyName().ToString(), *CustomVersion.Key.ToString(), CustomVersion.Version);
 		IniVersionStrings.Emplace(MoveTemp(CustomVersionString));
 	}
 
diff --git a/Engine/Source/Runtime/AssetRegistry/Private/PackageReader.cpp b/Engine/Source/Runtime/AssetRegistry/Private/PackageReader.cpp
--- a/Engine/Source/Runtime/AssetRegistry/Private/PackageReader.cpp
+++ b/Engine/Source/Runtime/AssetRegistry/Private/PackageReader.cpp
@@ -54,11 +54,9 @@ bool FPackageReader::OpenPackageFile(const FString& InPackageFilename)
 
 	// Check serialized custom versions against latest custom versions.
 	const FCustomVersionContainer& LatestCustomVersions  = FCustomVersionContainer::GetRegistered();
-	const TArray<FCustomVersion>&  PackageCustomVersions = PackageFileSummary.GetCustomVersionContainer().GetAllVersions();
-	for (TArray<FCustomVersion>::TConstIterator It(PackageCustomVersions); It; ++It)
+	const FCustomVersionSet&  PackageCustomVersions = PackageFileSummary.GetCustomVersionContainer().GetAllVersions();
+	for (const FCustomVersion& SerializedCustomVersion : PackageCustomVersions)
 	{
-		const FCustomVersion& SerializedCustomVersion = *It;
-
 		auto* LatestVersion = LatestCustomVersions.GetVersion(SerializedCustomVersion.Key);
 		if (!LatestVersion || SerializedCustomVersion.Version > LatestVersion->Version)
 		{
diff --git a/Engine/Source/Runtime/Core/Private/Serialization/Archive.cpp b/Engine/Source/Runtime/Core/Private/Serialization/Archive.cpp
--- a/Engine/Source/Runtime/Core/Private/Serialization/Archive.cpp
+++ b/Engine/Source/Runtime/Core/Private/Serialization/Archive.cpp
@@ -208,7 +208,7 @@ void FArchive::UsingCustomVersion(const FGuid& Key)
 	// If this fails, you probably don't have an FCustomVersionRegistration variable defined for this GUID.
 	check(RegisteredVersion);
 
-	const_cast<FCustomVersionContainer&>(GetCustomVersions()).SetVersion(Key, RegisteredVersion->Version, RegisteredVersion->FriendlyName);
+	const_cast<FCustomVersionContainer&>(GetCustomVersions()).SetVersion(Key, RegisteredVersion->Version, RegisteredVersion->GetFriendlyName());
 }
 
 int32 FArchive::CustomVer(const FGuid& Key) const
@@ -222,7 +222,7 @@ int32 FArchive::CustomVer(const FGuid& Key) const
 	return CustomVersion ? CustomVersion->Version : -1;
 }
 
-void FArchive::SetCustomVersion(const FGuid&  Key, int32 Version, FString FriendlyName)
+void FArchive::SetCustomVersion(const FGuid& Key, int32 Version, FName FriendlyName)
 {
 	const_cast<FCustomVersionContainer&>(GetCustomVersions()).SetVersion(Key, Version, FriendlyName);
 }
diff --git a/Engine/Source/Runtime/Core/Private/Serialization/CustomVersion.cpp b/Engine/Source/Runtime/Core/Private/Serialization/CustomVersion.cpp
--- a/Engine/Source/Runtime/Core/Private/Serialization/CustomVersion.cpp
+++ b/Engine/Source/Runtime/Core/Private/Serialization/CustomVersion.cpp
@@ -19,7 +19,7 @@ namespace
 		FCustomVersion ToCustomVersion() const
 		{
 			// We'll invent a GUID from three zeroes and the original tag
-			return FCustomVersion(FGuid(0, 0, 0, Tag), Version, FString::Printf(TEXT("Old enum-based tag: %u"), Tag));
+			return FCustomVersion(FGuid(0, 0, 0, Tag), Version, *FString::Printf(TEXT("EnumTag%u"), Tag));
 		}
 	};
 
@@ -31,6 +31,36 @@ namespace
 
 		return Ar;
 	}
+
+	struct FGuidCustomVersion_DEPRECATED
+	{
+		FGuid Key;
+		int32 Version;
+		FString FriendlyName;
+
+		FCustomVersion ToCustomVersion() const
+		{
+			// We'll invent a GUID from three zeroes and the original tag
+			return FCustomVersion(Key, Version, *FriendlyName);
+		}
+	};
+
+	FArchive& operator<<(FArchive& Ar, FGuidCustomVersion_DEPRECATED& Version)
+	{
+		Ar << Version.Key;
+		Ar << Version.Version;
+		Ar << Version.FriendlyName;
+		return Ar;
+	}
+}
+
+const FName FCustomVersion::GetFriendlyName() const
+{
+	if (FriendlyName == NAME_None)
+	{
+		FriendlyName = FCustomVersionContainer::GetRegistered().GetFriendlyName(Key);
+	}
+	return FriendlyName;
 }
 
 const FCustomVersionContainer& FCustomVersionContainer::GetRegistered()
@@ -46,11 +76,10 @@ void FCustomVersionContainer::Empty()
 FString FCustomVersionContainer::ToString(const FString& Indent) const
 {
 	FString VersionsAsString;
-	for (int32 i=0; i<Versions.Num(); i++)
+	for (const FCustomVersion& SomeVersion : Versions)
 	{
-		const FCustomVersion& SomeVersion = Versions[i];
 		VersionsAsString += Indent;
-		VersionsAsString += FString::Printf(TEXT("Key=%s  Version=%d  Friendly Name=%s \n"), *SomeVersion.Key.ToString(), SomeVersion.Version, *SomeVersion.FriendlyName );
+		VersionsAsString += FString::Printf(TEXT("Key=%s  Version=%d  Friendly Name=%s \n"), *SomeVersion.Key.ToString(), SomeVersion.Version, *SomeVersion.GetFriendlyName().ToString() );
 	}
 
 	return VersionsAsString;
@@ -67,8 +96,6 @@ FArchive& operator<<(FArchive& Ar, FCustomVersion& Version)
 {
 	Ar << Version.Key;
 	Ar << Version.Version;
-	Ar << Version.FriendlyName;
-
 	return Ar;
 }
 
@@ -95,8 +122,25 @@ void FCustomVersionContainer::Serialize(FArchive& Ar, ECustomVersionSerializatio
 		break;
 
 		case ECustomVersionSerializationFormat::Guids:
+		{
+			// We should only ever be loading old versions.  They should never be saved - they only exist for backward compatibility.
+			check(Ar.IsLoading());
+
+			TArray<FGuidCustomVersion_DEPRECATED> VersionArray;
+			Ar << VersionArray;
+			Versions.Empty(VersionArray.Num());
+			for (FGuidCustomVersion_DEPRECATED& OldVer : VersionArray)
+			{
+				Versions.Add(OldVer.ToCustomVersion());
+			}
+		}
+		break;
+
+		case ECustomVersionSerializationFormat::Optimized:
+		{
 			Ar << Versions;
-			break;
+		}
+		break;
 	}
 }
 
@@ -105,33 +149,46 @@ const FCustomVersion* FCustomVersionContainer::GetVersion(FGuid Key) const
 	// A testing tag was written out to a few archives during testing so we need to
 	// handle the existence of it to ensure that those archives can still be loaded.
 	if (Key == UnusedCustomVersion.Key)
+	{
 		return &UnusedCustomVersion;
+	}
 
-	return Versions.FindByKey(Key);
+	return Versions.Find(Key);
+}
+
+const FName FCustomVersionContainer::GetFriendlyName(FGuid Key) const
+{
+	FName FriendlyName = NAME_Name;
+	const FCustomVersion* CustomVersion = GetVersion(Key);
+	if (CustomVersion)
+	{
+		FriendlyName = CustomVersion->FriendlyName;
+	}
+	return FriendlyName;
 }
 
-void FCustomVersionContainer::SetVersion(FGuid CustomKey, int32 Version, FString FriendlyName)
+void FCustomVersionContainer::SetVersion(FGuid CustomKey, int32 Version, FName FriendlyName)
 {
 	if (CustomKey == UnusedCustomVersion.Key)
 		return;
 
-	if (FCustomVersion* Found = Versions.FindByKey(CustomKey))
+	if (FCustomVersion* Found = Versions.Find(CustomKey))
 	{
 		Found->Version      = Version;
-		Found->FriendlyName = MoveTemp(FriendlyName);
+		Found->FriendlyName = FriendlyName;
 	}
 	else
 	{
-		Versions.Add(FCustomVersion(CustomKey, Version, MoveTemp(FriendlyName)));
+		Versions.Add(FCustomVersion(CustomKey, Version, FriendlyName));
 	}
 }
 
-FCustomVersionRegistration::FCustomVersionRegistration(FGuid InKey, int32 InVersion, const TCHAR* InFriendlyName)
+FCustomVersionRegistration::FCustomVersionRegistration(FGuid InKey, int32 InVersion, FName InFriendlyName)
 {
-	TArray<FCustomVersion>& Versions = FCustomVersionContainer::GetInstance().Versions;
+	FCustomVersionSet& Versions = FCustomVersionContainer::GetInstance().Versions;
 
 	// Check if this tag hasn't already been registered
-	if (FCustomVersion* ExistingRegistration = Versions.FindByKey(InKey))
+	if (FCustomVersion* ExistingRegistration = Versions.Find(InKey))
 	{
 		// We don't allow the registration details to change across registrations - this code path only exists to support hotreload
 
@@ -139,9 +196,9 @@ FCustomVersionRegistration::FCustomVersionRegistration(FGuid InKey, int32 InVers
 		// * Changed registration details during hotreload.
 		// * Accidentally copy-and-pasted an FCustomVersionRegistration object.
 		ensureMsgf(
-			ExistingRegistration->Version == InVersion && ExistingRegistration->FriendlyName == InFriendlyName,
+			ExistingRegistration->Version == InVersion && ExistingRegistration->GetFriendlyName() == InFriendlyName,
 			TEXT("Custom version registrations cannot change between hotreloads - \"%s\" version %d is being reregistered as \"%s\" version %d"),
-			*ExistingRegistration->FriendlyName,
+			*ExistingRegistration->GetFriendlyName().ToString(),
 			ExistingRegistration->Version,
 			InFriendlyName,
 			InVersion
@@ -159,9 +216,9 @@ FCustomVersionRegistration::FCustomVersionRegistration(FGuid InKey, int32 InVers
 
 FCustomVersionRegistration::~FCustomVersionRegistration()
 {
-	TArray<FCustomVersion>& Versions = FCustomVersionContainer::GetInstance().Versions;
+	FCustomVersionSet& Versions = FCustomVersionContainer::GetInstance().Versions;
 
-	FCustomVersion* FoundKey = Versions.FindByKey(Key);
+	FCustomVersion* FoundKey = Versions.Find(Key);
 
 	// Ensure this tag has been registered
 	check(FoundKey);
@@ -169,6 +226,6 @@ FCustomVersionRegistration::~FCustomVersionRegistration()
 	--FoundKey->ReferenceCount;
 	if (FoundKey->ReferenceCount == 0)
 	{
-		Versions.RemoveAt(FoundKey - Versions.GetData());
+		Versions.Remove(Key);
 	}
 }
diff --git a/Engine/Source/Runtime/Core/Public/Serialization/ArchiveBase.h b/Engine/Source/Runtime/Core/Public/Serialization/ArchiveBase.h
--- a/Engine/Source/Runtime/Core/Public/Serialization/ArchiveBase.h
+++ b/Engine/Source/Runtime/Core/Public/Serialization/ArchiveBase.h
@@ -767,7 +767,7 @@ public:
 	 * @param Version - The version number to set key to
 	 * @param FriendlyName - Friendly name corresponding to the key
 	 */
-	void SetCustomVersion(const struct FGuid& Key, int32 Version, FString FriendlyName);
+	void SetCustomVersion(const struct FGuid& Key, int32 Version, FName FriendlyName);
 
 	/**
 	 * Toggle saving as Unicode. This is needed when we need to make sure ANSI strings are saved as Unicode
diff --git a/Engine/Source/Runtime/Core/Public/Serialization/CustomVersion.h b/Engine/Source/Runtime/Core/Public/Serialization/CustomVersion.h
--- a/Engine/Source/Runtime/Core/Public/Serialization/CustomVersion.h
+++ b/Engine/Source/Runtime/Core/Public/Serialization/CustomVersion.h
@@ -9,8 +9,14 @@ struct ECustomVersionSerializationFormat
 {
 	enum Type
 	{
+		Unknown,
 		Guids,
-		Enums
+		Enums,
+		Optimized,
+
+		// Add new versions above this comment
+		CustomVersion_Automatic_Plus_One,
+		Latest = CustomVersion_Automatic_Plus_One - 1
 	};
 };
 
@@ -20,15 +26,14 @@ struct ECustomVersionSerializationFormat
  */
 struct CORE_API FCustomVersion
 {
+	friend class FCustomVersionContainer;
+
 	/** Unique custom key. */
 	FGuid Key;
 
 	/** Custom version. */
 	int32 Version;
 
-	/** Friendly name for error messages or whatever. */
-	FString FriendlyName;
-
 	/** Number of times this GUID has been registered */
 	int32 ReferenceCount;
 
@@ -38,10 +43,10 @@ struct CORE_API FCustomVersion
 	}
 
 	/** Helper constructor. */
-	FORCEINLINE FCustomVersion(FGuid InKey, int32 InVersion, FString InFriendlyName)
+	FORCEINLINE FCustomVersion(FGuid InKey, int32 InVersion, FName InFriendlyName)
 	: Key           (InKey)
 	, Version       (InVersion)
-	, FriendlyName  (MoveTemp(InFriendlyName))
+	, FriendlyName  (InFriendlyName)
 	, ReferenceCount(1)
 	{
 	}
@@ -59,8 +64,33 @@ struct CORE_API FCustomVersion
 	}
 
 	CORE_API friend FArchive& operator<<(FArchive& Ar, FCustomVersion& Version);
+
+	/** Gets the friendly name for error messages or whatever */
+	const FName GetFriendlyName() const;
+
+private:
+
+	/** Friendly name for error messages or whatever. Lazy initialized for serialized versions */
+	mutable FName FriendlyName;
+};
+
+struct FCustomVersionKeyFuncs : BaseKeyFuncs<FCustomVersion, FGuid, false>
+{
+	static FORCEINLINE const FGuid& GetSetKey(const FCustomVersion& Element)
+	{
+		return Element.Key;
+	}
+	static FORCEINLINE bool Matches(FGuid A, FGuid B)
+	{
+		return A == B;
+	}
+	static FORCEINLINE uint32 GetKeyHash(FGuid Key)
+	{
+		return FCrc::MemCrc_DEPRECATED(&Key, sizeof(FGuid));
+	}
 };
 
+typedef TSet<FCustomVersion, FCustomVersionKeyFuncs> FCustomVersionSet;
 
 class CORE_API FCustomVersionRegistration;
 
@@ -74,7 +104,7 @@ class CORE_API FCustomVersionContainer
 
 public:
 	/** Gets available custom versions in this container. */
-	FORCEINLINE const TArray<FCustomVersion>& GetAllVersions() const
+	FORCEINLINE const FCustomVersionSet& GetAllVersions() const
 	{
 		return Versions;
 	}
@@ -87,6 +117,14 @@ public:
 	 */
 	const FCustomVersion* GetVersion(FGuid CustomKey) const;
 
+	/**
+	* Gets a custom version friendly name from the container.
+	*
+	* @param CustomKey Custom key for which to retrieve the version.
+	* @return The friendly name for the specified custom key, or NAME_None if the key doesn't exist in the container.
+	*/
+	const FName GetFriendlyName(FGuid CustomKey) const;
+
 	/**
 	 * Sets a specific custom version in the container.
 	 *
@@ -94,10 +132,10 @@ public:
 	 * @param Version The version number for the specified custom key
 	 * @param FriendlyName A friendly name to assign to this version
 	 */
-	void SetVersion(FGuid CustomKey, int32 Version, FString FriendlyName);
+	void SetVersion(FGuid CustomKey, int32 Version, FName FriendlyName);
 
 	/** Serialization. */
-	void Serialize(FArchive& Ar, ECustomVersionSerializationFormat::Type Format = ECustomVersionSerializationFormat::Guids);
+	void Serialize(FArchive& Ar, ECustomVersionSerializationFormat::Type Format = ECustomVersionSerializationFormat::Latest);
 
 	/**
 	 * Gets a singleton with the registered versions.
@@ -120,7 +158,7 @@ private:
 	static FCustomVersionContainer& GetInstance();
 
 	/** Array containing custom versions. */
-	TArray<FCustomVersion> Versions;
+	FCustomVersionSet Versions;
 };
 
 
@@ -133,7 +171,7 @@ class FCustomVersionRegistration : FNoncopyable
 {
 public:
 
-	FCustomVersionRegistration(FGuid InKey, int32 Version, const TCHAR*);
+	FCustomVersionRegistration(FGuid InKey, int32 Version, FName InFriendlyName);
 	~FCustomVersionRegistration();
 
 private:
diff --git a/Engine/Source/Runtime/CoreUObject/Private/UObject/LinkerLoad.cpp b/Engine/Source/Runtime/CoreUObject/Private/UObject/LinkerLoad.cpp
--- a/Engine/Source/Runtime/CoreUObject/Private/UObject/LinkerLoad.cpp
+++ b/Engine/Source/Runtime/CoreUObject/Private/UObject/LinkerLoad.cpp
@@ -1054,12 +1054,12 @@ FLinkerLoad::ELinkerStatus FLinkerLoad::SerializePackageFileSummary()
 
 		// Check custom versions.
 		const FCustomVersionContainer& LatestCustomVersions  = FCustomVersionContainer::GetRegistered();
-		const TArray<FCustomVersion>&  PackageCustomVersions = Summary.GetCustomVersionContainer().GetAllVersions();
+		const FCustomVersionSet&  PackageCustomVersions = Summary.GetCustomVersionContainer().GetAllVersions();
 		for (auto It = PackageCustomVersions.CreateConstIterator(); It; ++It)
 		{
 			const FCustomVersion& SerializedCustomVersion = *It;
 
-			auto* LatestVersion = LatestCustomVersions.GetVersion(SerializedCustomVersion.Key);
+			const FCustomVersion* LatestVersion = LatestCustomVersions.GetVersion(SerializedCustomVersion.Key);
 			if (!LatestVersion)
 			{
 				// Loading a package with custom integration that we don't know about!
diff --git a/Engine/Source/Runtime/CoreUObject/Private/UObject/PackageFileSummary.cpp b/Engine/Source/Runtime/CoreUObject/Private/UObject/PackageFileSummary.cpp
--- a/Engine/Source/Runtime/CoreUObject/Private/UObject/PackageFileSummary.cpp
+++ b/Engine/Source/Runtime/CoreUObject/Private/UObject/PackageFileSummary.cpp
@@ -8,6 +8,26 @@ FPackageFileSummary::FPackageFileSummary()
 	FMemory::Memzero( this, sizeof(*this) );
 }
 
+/** Converts file version to custom version system version */
+static ECustomVersionSerializationFormat::Type GetCustomVersionFormatForArchive(int32 LegacyFileVersion)
+{
+	ECustomVersionSerializationFormat::Type CustomVersionFormat = ECustomVersionSerializationFormat::Unknown;
+	if (LegacyFileVersion == -2)
+	{
+		CustomVersionFormat = ECustomVersionSerializationFormat::Enums;
+	}
+	else if (LegacyFileVersion < -2 && LegacyFileVersion >= -5)
+	{
+		CustomVersionFormat = ECustomVersionSerializationFormat::Guids;
+	}
+	else if (LegacyFileVersion < -5)
+	{
+		CustomVersionFormat = ECustomVersionSerializationFormat::Optimized;
+	}
+	check(CustomVersionFormat != ECustomVersionSerializationFormat::Unknown);
+	return CustomVersionFormat;
+}
+
 FArchive& operator<<( FArchive& Ar, FPackageFileSummary& Sum )
 {
 	bool bCanStartSerializing = true;
@@ -50,13 +70,14 @@ FArchive& operator<<( FArchive& Ar, FPackageFileSummary& Sum )
 		 *
 		 * Lower 16 bits stores the UE3 engine version
 		 * Upper 16 bits stores the UE4/licensee version
-		 * For newer packages this is -5
+		 * For newer packages this is -6
 		 *		-2 indicates presence of enum-based custom versions
 		 *		-3 indicates guid-based custom versions
 		 *		-4 indicates removal of the UE3 version. Packages saved with this ID cannot be loaded in older engine versions 
-		 *      -5 indicates the replacement of writing out the "UE3 version" so older versions of engine can gracefully fail to open newer packages
+		 *		-5 indicates the replacement of writing out the "UE3 version" so older versions of engine can gracefully fail to open newer packages
+		 *		-6 indicates optimizations to how custom versions are being serialized
 		 */
-		int32 LegacyFileVersion = -5;
+		int32 LegacyFileVersion = -6;
 		Ar << LegacyFileVersion;
 
 		if (Ar.IsLoading())
@@ -73,7 +94,7 @@ FArchive& operator<<( FArchive& Ar, FPackageFileSummary& Sum )
 
 				if (LegacyFileVersion <= -2)
 				{
-					Sum.CustomVersionContainer.Serialize(Ar, (LegacyFileVersion == -2) ? ECustomVersionSerializationFormat::Enums : ECustomVersionSerializationFormat::Guids);
+					Sum.CustomVersionContainer.Serialize(Ar, GetCustomVersionFormatForArchive(LegacyFileVersion));
 				}
 
 				if (!Sum.FileVersionUE4 && !Sum.FileVersionLicenseeUE4)
