Fix resolution cache RequestedContentScaleFactor log message
#ue4
#android

diff --git a/Engine/Source/Runtime/Core/Private/Android/AndroidWindow.cpp b/Engine/Source/Runtime/Core/Private/Android/AndroidWindow.cpp
--- a/Engine/Source/Runtime/Core/Private/Android/AndroidWindow.cpp
+++ b/Engine/Source/Runtime/Core/Private/Android/AndroidWindow.cpp
@@ -98,7 +98,7 @@ FPlatformRect FAndroidWindow::GetScreenRect()
 
 	if (RequestedContentScaleFactor != ContentScaleFactor)
 	{
-		FPlatformMisc::LowLevelOutputDebugStringf(TEXT("***** RequestedContentScaleFactor different %d != %d, not using res cache"), RequestedContentScaleFactor, ContentScaleFactor);
+		FPlatformMisc::LowLevelOutputDebugStringf(TEXT("***** RequestedContentScaleFactor different %f != %f, not using res cache"), RequestedContentScaleFactor, ContentScaleFactor);
 	}
 
 	if (Window != LastWindow)
