Fix Linux running out of file handles over time (UE-16483).

- Child processes inherited open handles; added CLOEXEC to both file operations and sockets.

diff --git a/Engine/Source/Runtime/Core/Private/Linux/LinuxPlatformFile.cpp b/Engine/Source/Runtime/Core/Private/Linux/LinuxPlatformFile.cpp
--- a/Engine/Source/Runtime/Core/Private/Linux/LinuxPlatformFile.cpp
+++ b/Engine/Source/Runtime/Core/Private/Linux/LinuxPlatformFile.cpp
@@ -216,7 +216,7 @@ private:
 			{
 				ReserveSlot();
 				
-				FileHandle = open(TCHAR_TO_UTF8(*Filename), O_RDONLY);
+				FileHandle = open(TCHAR_TO_UTF8(*Filename), O_RDONLY | O_CLOEXEC);
 				if( FileHandle != -1 )
 				{
 					lseek(FileHandle, FileOffset, SEEK_SET);
@@ -530,7 +530,7 @@ public:
 		}
 
 		// try opening right away
-		int32 Handle = open(TCHAR_TO_UTF8(*Filename), O_RDONLY);
+		int32 Handle = open(TCHAR_TO_UTF8(*Filename), O_RDONLY | O_CLOEXEC);
 		if (Handle != -1)
 		{
 			MappedToFilename = Filename;
@@ -541,7 +541,7 @@ public:
 			if (ENOENT != errno)
 			{
 				int ErrNo = errno;
-				UE_LOG(LogLinuxPlatformFile, Warning, TEXT( "open('%s', ORDONLY) failed: errno=%d (%s)" ), *Filename, ErrNo, ANSI_TO_TCHAR(strerror(ErrNo)));
+				UE_LOG(LogLinuxPlatformFile, Warning, TEXT( "open('%s', O_RDONLY | O_CLOEXEC) failed: errno=%d (%s)" ), *Filename, ErrNo, ANSI_TO_TCHAR(strerror(ErrNo)));
 			}
 			else
 			{
@@ -555,7 +555,7 @@ public:
 					FString FoundFilename(TEXT("/"));	// start with root
 					if (MapFileRecursively(Filename, 0, MaxPathComponents, FoundFilename))
 					{
-						Handle = open(TCHAR_TO_UTF8(*FoundFilename), O_RDONLY);
+						Handle = open(TCHAR_TO_UTF8(*FoundFilename), O_RDONLY | O_CLOEXEC);
 						if (Handle != -1)
 						{
 							MappedToFilename = FoundFilename;
diff --git a/Engine/Source/Runtime/Core/Public/HAL/Platform.h b/Engine/Source/Runtime/Core/Public/HAL/Platform.h
--- a/Engine/Source/Runtime/Core/Public/HAL/Platform.h
+++ b/Engine/Source/Runtime/Core/Public/HAL/Platform.h
@@ -228,6 +228,9 @@
 #ifndef PLATFORM_HAS_BSD_SOCKET_FEATURE_GETHOSTNAME
 	#define PLATFORM_HAS_BSD_SOCKET_FEATURE_GETHOSTNAME	1
 #endif
+#ifndef PLATFORM_HAS_BSD_SOCKET_FEATURE_CLOSE_ON_EXEC
+	#define PLATFORM_HAS_BSD_SOCKET_FEATURE_CLOSE_ON_EXEC	0
+#endif
 #ifndef PLATFORM_HAS_NO_EPROCLIM
 	#define PLATFORM_HAS_NO_EPROCLIM			0
 #endif
diff --git a/Engine/Source/Runtime/Core/Public/Linux/LinuxPlatform.h b/Engine/Source/Runtime/Core/Public/Linux/LinuxPlatform.h
--- a/Engine/Source/Runtime/Core/Public/Linux/LinuxPlatform.h
+++ b/Engine/Source/Runtime/Core/Public/Linux/LinuxPlatform.h
@@ -6,6 +6,8 @@
 
 #pragma once
 
+#include <linux/version.h>	// for checking if SOCK_CLOEXEC is defined - including sys/socket.h at this stage pollutes the namespaces with macros which cause collisions (e.g. PF_MAX)
+
 /**
 * Linux specific types
 **/
@@ -43,6 +45,11 @@ typedef FLinuxPlatformTypes FPlatformTypes;
 #define PLATFORM_HAS_BSD_SOCKET_FEATURE_IOCTL		1
 #define PLATFORM_HAS_BSD_IPV6_SOCKETS				1
 
+// SOCK_CLOEXEC is available on Linux since 2.6.27
+#if LINUX_VERSION_CODE >= KERNEL_VERSION(2,6,27)
+	#define PLATFORM_HAS_BSD_SOCKET_FEATURE_CLOSE_ON_EXEC	1
+#endif // LINUX_VERSION_CODE >= KERNEL_VERSION(2,6,27)
+
 // only enable vectorintrinsics on x86(-64) for now
 #if defined(_M_IX86) || defined(__i386__) || defined(_M_X64) || defined(__x86_64__) || defined (__amd64__) 
 	#define PLATFORM_ENABLE_VECTORINTRINSICS		1
diff --git a/Engine/Source/Runtime/Sockets/Private/BSDSockets/SocketSubsystemBSD.cpp b/Engine/Source/Runtime/Sockets/Private/BSDSockets/SocketSubsystemBSD.cpp
--- a/Engine/Source/Runtime/Sockets/Private/BSDSockets/SocketSubsystemBSD.cpp
+++ b/Engine/Source/Runtime/Sockets/Private/BSDSockets/SocketSubsystemBSD.cpp
@@ -19,18 +19,23 @@ FSocket* FSocketSubsystemBSD::CreateSocket(const FName& SocketType, const FStrin
 {
 	SOCKET Socket = INVALID_SOCKET;
 	FSocket* NewSocket = nullptr;
+	int PlatformSpecificTypeFlags = 0;
+
+#if PLATFORM_HAS_BSD_SOCKET_FEATURE_CLOSE_ON_EXEC
+	PlatformSpecificTypeFlags = SOCK_CLOEXEC;
+#endif // PLATFORM_HAS_BSD_SOCKET_FEATURE_CLOSE_ON_EXEC
 
 	switch (SocketType.GetComparisonIndex())
 	{
 	case NAME_DGram:
 		// Creates a data gram (UDP) socket
-		Socket = socket(AF_INET, SOCK_DGRAM, IPPROTO_UDP);
+		Socket = socket(AF_INET, SOCK_DGRAM | PlatformSpecificTypeFlags, IPPROTO_UDP);
 		NewSocket = (Socket != INVALID_SOCKET) ? InternalBSDSocketFactory(Socket, SOCKTYPE_Datagram, SocketDescription) : nullptr;
 		break;
 
 	case NAME_Stream:
 		// Creates a stream (TCP) socket
-		Socket = socket(AF_INET, SOCK_STREAM, IPPROTO_TCP);
+		Socket = socket(AF_INET, SOCK_STREAM | PlatformSpecificTypeFlags, IPPROTO_TCP);
 		NewSocket = (Socket != INVALID_SOCKET) ? InternalBSDSocketFactory(Socket, SOCKTYPE_Streaming, SocketDescription) : nullptr;
 		break;
 
