Linux: prevent servers from trying to poll SDL events.

diff --git a/Engine/Source/Runtime/Core/Private/Linux/LinuxPlatformMisc.cpp b/Engine/Source/Runtime/Core/Private/Linux/LinuxPlatformMisc.cpp
--- a/Engine/Source/Runtime/Core/Private/Linux/LinuxPlatformMisc.cpp
+++ b/Engine/Source/Runtime/Core/Private/Linux/LinuxPlatformMisc.cpp
@@ -255,7 +255,7 @@ void FLinuxPlatformMisc::SetEnvironmentVar(const TCHAR* InVariableName, const TC
 
 void FLinuxPlatformMisc::PumpMessages( bool bFromMainLoop )
 {
-	if( bFromMainLoop )
+	if (GInitializedSDL && bFromMainLoop)
 	{
 		SDL_Event event;
 
