Various changes for VS2015:
Printf format string fix.
Cast for type size mismatch.

diff --git a/Engine/Source/Runtime/Core/Public/HAL/MallocBinned.h b/Engine/Source/Runtime/Core/Public/HAL/MallocBinned.h
--- a/Engine/Source/Runtime/Core/Public/HAL/MallocBinned.h
+++ b/Engine/Source/Runtime/Core/Public/HAL/MallocBinned.h
@@ -1063,7 +1063,7 @@ public:
 #endif
 				Pool = GetPoolInfo((UPTRINT)Free);
 
-				if ((UPTRINT)Free != ((UPTRINT)AlignedFree & ~(PageSize - 1)))
+				if ((UPTRINT)Free != ((UPTRINT)AlignedFree & ~((UPTRINT)PageSize - 1)))
 				{
 					// Mark the FPoolInfo for AlignedFree to jump back to the FPoolInfo for ptr.
 					for (UPTRINT i = (UPTRINT)PageSize, Offset = 0; i < AlignedSize; i += PageSize, ++Offset)
