VS2015 requires some early header includes similar to clang

diff --git a/Engine/Source/Runtime/Core/Public/Windows/WindowsSystemIncludes.h b/Engine/Source/Runtime/Core/Public/Windows/WindowsSystemIncludes.h
--- a/Engine/Source/Runtime/Core/Public/Windows/WindowsSystemIncludes.h
+++ b/Engine/Source/Runtime/Core/Public/Windows/WindowsSystemIncludes.h
@@ -17,11 +17,11 @@ extern "C" CORE_API HINSTANCE hInstance;
 #include <stdint.h>
 #include <tchar.h>
 
-// @todo clang: When compiling under Clang for Windows, these headers cause us particular problems.  We need to make sure they're included
-// before we pull in our 'DoNotUseOldUE4Type' namespace.  This is because these headers will redeclare various numeric typedefs, but under
-// the Clang compiler it is not allowed to define a typedef with a global scope operator in it (such as ::INT).  So we'll get these headers
-// included early on to avoid compiler errors with that.
-#if defined(__clang__)
+// @todo clang: When compiling under Clang for Windows or with VS2015, these headers cause us particular problems.  We need to make sure
+// they're included before we pull in our 'DoNotUseOldUE4Type' namespace.  This is because these headers will redeclare various numeric
+// typedefs, but under the Clang and 2015 compilers it is not allowed to define a typedef with a global scope operator in it (such as ::INT).
+// So we'll get these headers included early on to avoid compiler errors with that.
+#if defined(__clang__) || (defined(_MSC_VER) && (_MSC_VER >= 1900))
 #include <intsafe.h>
 #include <strsafe.h>
 #endif
