Re-disable default functions in VS2015 for now

diff --git a/Engine/Source/Runtime/Core/Public/Windows/WIndowsPlatform.h b/Engine/Source/Runtime/Core/Public/Windows/WIndowsPlatform.h
--- a/Engine/Source/Runtime/Core/Public/Windows/WIndowsPlatform.h
+++ b/Engine/Source/Runtime/Core/Public/Windows/WIndowsPlatform.h
@@ -58,7 +58,7 @@ typedef FWindowsPlatformTypes FPlatformTypes;
 #define PLATFORM_USES_MICROSOFT_LIBC_FUNCTIONS				1
 #define PLATFORM_SUPPORTS_TBB								1
 #define PLATFORM_SUPPORTS_NAMED_PIPES						1
-#if _MSC_VER >= 1900
+#if 0 //@todo: VS2015 supports defaulted functions but we have errors in some of our classes we need to fix up before we can enable it
 	#define PLATFORM_COMPILER_HAS_DEFAULTED_FUNCTIONS		1
 #else
 	#define PLATFORM_COMPILER_HAS_DEFAULTED_FUNCTIONS		0
