Ensure that the launch path used in Mac ExecProcess and CreateProc isn't both released and auto-released which will crash.

diff --git a/Engine/Source/Runtime/Core/Private/Mac/MacPlatformProcess.cpp b/Engine/Source/Runtime/Core/Private/Mac/MacPlatformProcess.cpp
--- a/Engine/Source/Runtime/Core/Private/Mac/MacPlatformProcess.cpp
+++ b/Engine/Source/Runtime/Core/Private/Mac/MacPlatformProcess.cpp
@@ -210,7 +210,7 @@ bool FMacPlatformProcess::ExecProcess( const TCHAR* URL, const TCHAR* Params, in
 	SCOPED_AUTORELEASE_POOL;
 
 	FString ProcessPath = URL;
-	NSString* LaunchPath = (NSString*)FPlatformString::TCHARToCFString(*ProcessPath);
+	NSString* LaunchPath = ProcessPath.GetNSString();
 	
 	if (![[NSFileManager defaultManager] fileExistsAtPath: LaunchPath])
 	{
@@ -377,7 +377,7 @@ FProcHandle FMacPlatformProcess::CreateProc( const TCHAR* URL, const TCHAR* Parm
 		ProcessPath = FString(BaseDir()) + ProcessPath;
 	}
 
-	NSString* LaunchPath = (NSString*)FPlatformString::TCHARToCFString(*ProcessPath);
+	NSString* LaunchPath = ProcessPath.GetNSString();
 
 	if (![[NSFileManager defaultManager] fileExistsAtPath: LaunchPath])
 	{
@@ -519,8 +519,6 @@ FProcHandle FMacPlatformProcess::CreateProc( const TCHAR* URL, const TCHAR* Parm
 		*OutProcessID = ProcessHandle ? [ProcessHandle processIdentifier] : 0;
 	}
 
-	CFRelease((CFStringRef)LaunchPath);
-
 	return FProcHandle(ProcessHandle);
 }
 
