[INTEGRATE] Minor optimization to ICU code
- Avoids some unnecessary smart pointer conversions

diff --git a/Engine/Source/Runtime/Core/Private/Internationalization/ICUText.h b/Engine/Source/Runtime/Core/Private/Internationalization/ICUText.h
--- a/Engine/Source/Runtime/Core/Private/Internationalization/ICUText.h
+++ b/Engine/Source/Runtime/Core/Private/Internationalization/ICUText.h
@@ -15,9 +15,9 @@ FText FText::AsNumberTemplate(T1 Val, const FNumberFormattingOptions* const Opti
 {
 	FInternationalization& I18N = FInternationalization::Get();
 	checkf(I18N.IsInitialized() == true, TEXT("FInternationalization is not initialized. An FText formatting method was likely used in static object initialization - this is not supported."));
-	const FCultureRef Culture = TargetCulture.IsValid() ? TargetCulture.ToSharedRef() : I18N.GetCurrentCulture();
+	const FCulture& Culture = TargetCulture.IsValid() ? *TargetCulture : *I18N.GetCurrentCulture();
 	UErrorCode ICUStatus = U_ZERO_ERROR;
-	const TSharedRef<const icu::DecimalFormat, ESPMode::ThreadSafe> ICUDecimalFormat( Culture->Implementation->GetDecimalFormatter(Options) );
+	const TSharedRef<const icu::DecimalFormat, ESPMode::ThreadSafe>& ICUDecimalFormat( Culture.Implementation->GetDecimalFormatter(Options) );
 	icu::Formattable FormattableVal(static_cast<T2>(Val));
 	icu::UnicodeString FormattedString;
 	ICUDecimalFormat->format(FormattableVal, FormattedString, ICUStatus);
@@ -35,9 +35,9 @@ FText FText::AsCurrencyTemplate(T1 Val, const FString& CurrencyCode, const FNumb
 {
 	FInternationalization& I18N = FInternationalization::Get();
 	checkf(I18N.IsInitialized() == true, TEXT("FInternationalization is not initialized. An FText formatting method was likely used in static object initialization - this is not supported."));
-	const FCultureRef Culture = TargetCulture.IsValid() ? TargetCulture.ToSharedRef() : I18N.GetCurrentCulture();
+	const FCulture& Culture = TargetCulture.IsValid() ? *TargetCulture : *I18N.GetCurrentCulture();
 	UErrorCode ICUStatus = U_ZERO_ERROR;
-	const TSharedRef<const icu::DecimalFormat> ICUDecimalFormat(Culture->Implementation->GetCurrencyFormatter(CurrencyCode, Options));
+	const TSharedRef<const icu::DecimalFormat>& ICUDecimalFormat(Culture.Implementation->GetCurrencyFormatter(CurrencyCode, Options));
 	icu::Formattable FormattableVal(static_cast<T2>(Val));
 	icu::UnicodeString FormattedString;
 	ICUDecimalFormat->format(FormattableVal, FormattedString, ICUStatus);
@@ -55,9 +55,9 @@ FText FText::AsPercentTemplate(T1 Val, const FNumberFormattingOptions* const Opt
 {
 	FInternationalization& I18N = FInternationalization::Get();
 	checkf(I18N.IsInitialized() == true, TEXT("FInternationalization is not initialized. An FText formatting method was likely used in static object initialization - this is not supported."));
-	const FCultureRef Culture = TargetCulture.IsValid() ? TargetCulture.ToSharedRef() : I18N.GetCurrentCulture();
+	const FCulture& Culture = TargetCulture.IsValid() ? *TargetCulture : *I18N.GetCurrentCulture();
 	UErrorCode ICUStatus = U_ZERO_ERROR;
-	const TSharedRef<const icu::DecimalFormat> ICUDecimalFormat( Culture->Implementation->GetPercentFormatter(Options) );
+	const TSharedRef<const icu::DecimalFormat>& ICUDecimalFormat( Culture.Implementation->GetPercentFormatter(Options) );
 	icu::Formattable FormattableVal(static_cast<T2>(Val));
 	icu::UnicodeString FormattedString;
 	ICUDecimalFormat->format(FormattableVal, FormattedString, ICUStatus);
diff --git a/Engine/Source/Runtime/Core/Private/Internationalization/LegacyText.h b/Engine/Source/Runtime/Core/Private/Internationalization/LegacyText.h
--- a/Engine/Source/Runtime/Core/Private/Internationalization/LegacyText.h
+++ b/Engine/Source/Runtime/Core/Private/Internationalization/LegacyText.h
@@ -7,8 +7,8 @@ FText FText::AsNumberTemplate(T1 Val, const FNumberFormattingOptions* const Opti
 {
 	FInternationalization& I18N = FInternationalization::Get();
 	checkf(I18N.IsInitialized() == true, TEXT("FInternationalization is not initialized. An FText formatting method was likely used in static object initialization - this is not supported."));
-	const FCultureRef Culture = TargetCulture.IsValid() ? TargetCulture.ToSharedRef() : I18N.GetCurrentCulture();
-	return FText::CreateNumericalText( Culture->NumberFormattingRule.AsNumber(Val) );
+	const FCulture& Culture = TargetCulture.IsValid() ? *TargetCulture : *I18N.GetCurrentCulture();
+	return FText::CreateNumericalText( Culture.NumberFormattingRule.AsNumber(Val) );
 }
 
 template<typename T1, typename T2>
@@ -16,8 +16,8 @@ FText FText::AsCurrencyTemplate(T1 Val, const FString& CurrencyCode, const FNumb
 {
 	FInternationalization& I18N = FInternationalization::Get();
 	checkf(I18N.IsInitialized() == true, TEXT("FInternationalization is not initialized. An FText formatting method was likely used in static object initialization - this is not supported."));
-	const FCultureRef Culture = TargetCulture.IsValid() ? TargetCulture.ToSharedRef() : I18N.GetCurrentCulture();
-	return FText::CreateNumericalText( Culture->NumberFormattingRule.AsCurrency(Val) );
+	const FCulture& Culture = TargetCulture.IsValid() ? *TargetCulture : *I18N.GetCurrentCulture();
+	return FText::CreateNumericalText( Culture.NumberFormattingRule.AsCurrency(Val) );
 }
 
 template<typename T1, typename T2>
@@ -25,6 +25,6 @@ FText FText::AsPercentTemplate(T1 Val, const FNumberFormattingOptions* const Opt
 {
 	FInternationalization& I18N = FInternationalization::Get();
 	checkf(I18N.IsInitialized() == true, TEXT("FInternationalization is not initialized. An FText formatting method was likely used in static object initialization - this is not supported."));
-	const FCultureRef Culture = TargetCulture.IsValid() ? TargetCulture.ToSharedRef() : I18N.GetCurrentCulture();
-	return FText::CreateNumericalText( Culture->NumberFormattingRule.AsPercent(Val) );
+	const FCulture& Culture = TargetCulture.IsValid() ? *TargetCulture : *I18N;
+	return FText::CreateNumericalText( Culture.NumberFormattingRule.AsPercent(Val) );
 }
\ No newline at end of file
