[INTEGRATE] Minor optimization to ICU code (part 2)
- Fixed compile error with LegacyText

diff --git a/Engine/Source/Runtime/Core/Private/Internationalization/LegacyText.h b/Engine/Source/Runtime/Core/Private/Internationalization/LegacyText.h
--- a/Engine/Source/Runtime/Core/Private/Internationalization/LegacyText.h
+++ b/Engine/Source/Runtime/Core/Private/Internationalization/LegacyText.h
@@ -25,6 +25,6 @@ FText FText::AsPercentTemplate(T1 Val, const FNumberFormattingOptions* const Opt
 {
 	FInternationalization& I18N = FInternationalization::Get();
 	checkf(I18N.IsInitialized() == true, TEXT("FInternationalization is not initialized. An FText formatting method was likely used in static object initialization - this is not supported."));
-	const FCulture& Culture = TargetCulture.IsValid() ? *TargetCulture : *I18N;
+	const FCulture& Culture = TargetCulture.IsValid() ? *TargetCulture : *I18N.GetCurrentCulture();
 	return FText::CreateNumericalText( Culture.NumberFormattingRule.AsPercent(Val) );
 }
\ No newline at end of file
