UE-19812 - Updated the UIRemoteNotificationType usage to UIUserNotificationType for IOS 8 onwards. The former has since been deprecated.

diff --git a/Engine/Source/Runtime/Core/Private/IOS/IOSPlatformMisc.cpp b/Engine/Source/Runtime/Core/Private/IOS/IOSPlatformMisc.cpp
--- a/Engine/Source/Runtime/Core/Private/IOS/IOSPlatformMisc.cpp
+++ b/Engine/Source/Runtime/Core/Private/IOS/IOSPlatformMisc.cpp
@@ -768,13 +768,19 @@ void FIOSPlatformMisc::RegisterForRemoteNotifications()
 	if ([application respondsToSelector : @selector(registerUserNotifcationSettings:)])
 	{
 #ifdef __IPHONE_8_0
-		UIUserNotificationSettings * settings = [UIUserNotificationSettings settingsForTypes : (UIRemoteNotificationTypeBadge | UIRemoteNotificationTypeSound | UIRemoteNotificationTypeAlert) categories:nil];
+		UIUserNotificationSettings * settings = [UIUserNotificationSettings settingsForTypes : (UIUserNotificationTypeBadge | UIUserNotificationTypeSound | UIUserNotificationTypeAlert) categories:nil];
 		[application registerUserNotificationSettings : settings];
 #endif
 	}
 	else
 	{
+        
+#ifdef __IPHONE_8_0
+        UIUserNotificationSettings * settings = [UIUserNotificationSettings settingsForTypes : (UIUserNotificationTypeBadge | UIUserNotificationTypeAlert) categories:nil];
+        [application registerUserNotificationSettings : settings];
+#else
 		UIRemoteNotificationType myTypes = UIRemoteNotificationTypeBadge | UIRemoteNotificationTypeAlert | UIRemoteNotificationTypeSound;
 		[application registerForRemoteNotificationTypes : myTypes];
+#endif
 	}
 }
