OR-5427: Crash callstacks are populating inconsistently in logs

diff --git a/Engine/Source/Runtime/Core/Private/Windows/WindowsPlatformCrashContext.cpp b/Engine/Source/Runtime/Core/Private/Windows/WindowsPlatformCrashContext.cpp
--- a/Engine/Source/Runtime/Core/Private/Windows/WindowsPlatformCrashContext.cpp
+++ b/Engine/Source/Runtime/Core/Private/Windows/WindowsPlatformCrashContext.cpp
@@ -493,7 +493,7 @@ int32 ReportCrash( LPEXCEPTION_POINTERS ExceptionInfo )
 	// Then try run time crash processing and broadcast information about a crash.
 	FCoreDelegates::OnHandleSystemError.Broadcast();
 
-	const bool bGenerateRuntimeCallstack = FEngineBuildSettings::IsInternalBuild() || FEngineBuildSettings::IsPerforceBuild() || FEngineBuildSettings::IsSourceDistribution();
+	const bool bGenerateRuntimeCallstack = FParse::Param(FCommandLine::Get(), TEXT("ForceLogCallstacks")) || FEngineBuildSettings::IsInternalBuild() || FEngineBuildSettings::IsPerforceBuild() || FEngineBuildSettings::IsSourceDistribution();
 	if (bGenerateRuntimeCallstack)
 	{
 		const SIZE_T StackTraceSize = 65535;
