Added support for detecting the reading direction of text

This uses the ICU BiDi functionality (where available) to detect whether text is left-to-right, right-to-left, or mixed, and in the mixed case, can give you a breakdown of the direction of each sub-section of text.

diff --git a/Engine/Source/Runtime/Core/Private/Internationalization/ICUText.cpp b/Engine/Source/Runtime/Core/Private/Internationalization/ICUText.cpp
--- a/Engine/Source/Runtime/Core/Private/Internationalization/ICUText.cpp
+++ b/Engine/Source/Runtime/Core/Private/Internationalization/ICUText.cpp
@@ -14,6 +14,7 @@ PRAGMA_DISABLE_SHADOW_VARIABLE_WARNINGS
 #include <unicode/msgfmt.h>
 PRAGMA_ENABLE_SHADOW_VARIABLE_WARNINGS
 #include <unicode/uniset.h>
+#include <unicode/ubidi.h>
 
 #include "ICUUtilities.h"
 #include "ICUTextCharacterIterator.h"
@@ -231,4 +232,205 @@ bool FText::IsLetter( const TCHAR Char )
 	return Uniscode.contains(Char) != 0;
 }
 
+namespace TextBiDi
+{
+
+namespace Internal
+{
+
+FORCEINLINE ETextDirection ICUToUE(const UBiDiDirection InDirection)
+{
+	switch (InDirection)
+	{
+	case UBIDI_LTR:
+		return ETextDirection::LeftToRight;
+	case UBIDI_RTL:
+		return ETextDirection::RightToLeft;
+	case UBIDI_MIXED:
+		return ETextDirection::Mixed;
+	default:
+		break;
+	}
+
+	return ETextDirection::LeftToRight;
+}
+
+ETextDirection ComputeTextDirection(UBiDi* InICUBiDi, const icu::UnicodeString& InICUString)
+{
+	UErrorCode ICUStatus = U_ZERO_ERROR;
+
+	ubidi_setPara(InICUBiDi, InICUString.getBuffer(), InICUString.length(), 0, nullptr, &ICUStatus);
+
+	if (U_SUCCESS(ICUStatus))
+	{
+		return Internal::ICUToUE(ubidi_getDirection(InICUBiDi));
+	}
+	else
+	{
+		UE_LOG(LogCore, Warning, TEXT("Failed to set the string data on the ICU BiDi object (error code: %d). Text will assumed to be left-to-right"), static_cast<int32>(ICUStatus));
+	}
+
+	return ETextDirection::LeftToRight;
+}
+
+ETextDirection ComputeTextDirection(UBiDi* InICUBiDi, const icu::UnicodeString& InICUString, TArray<FTextDirectionInfo>& OutTextDirectionInfo)
+{
+	UErrorCode ICUStatus = U_ZERO_ERROR;
+
+	ubidi_setPara(InICUBiDi, InICUString.getBuffer(), InICUString.length(), 0, nullptr, &ICUStatus);
+
+	if (U_SUCCESS(ICUStatus))
+	{
+		const ETextDirection ReturnDirection = Internal::ICUToUE(ubidi_getDirection(InICUBiDi));
+
+		const int32 RunCount = ubidi_countRuns(InICUBiDi, &ICUStatus);
+		OutTextDirectionInfo.AddZeroed(RunCount);
+		for (int32 RunIndex = 0; RunIndex < RunCount; ++RunIndex)
+		{
+			FTextDirectionInfo& CurTextDirectionInfo = OutTextDirectionInfo[RunIndex];
+			CurTextDirectionInfo.TextDirection = Internal::ICUToUE(ubidi_getVisualRun(InICUBiDi, RunIndex, &CurTextDirectionInfo.StartIndex, &CurTextDirectionInfo.Length));
+		}
+
+		return ReturnDirection;
+	}
+	else
+	{
+		UE_LOG(LogCore, Warning, TEXT("Failed to set the string data on the ICU BiDi object (error code: %d). Text will assumed to be left-to-right"), static_cast<int32>(ICUStatus));
+	}
+
+	return ETextDirection::LeftToRight;
+}
+
+class FICUTextBiDi : public ITextBiDi
+{
+public:
+	FICUTextBiDi()
+		: ICUBiDi(ubidi_open())
+	{
+	}
+
+	~FICUTextBiDi()
+	{
+		ubidi_close(ICUBiDi);
+		ICUBiDi = nullptr;
+	}
+
+	virtual ETextDirection ComputeTextDirection(const FText& InText) override
+	{
+		return FICUTextBiDi::ComputeTextDirection(InText.ToString());
+	}
+
+	virtual ETextDirection ComputeTextDirection(const FString& InString) override
+	{
+		if (InString.IsEmpty())
+		{
+			return ETextDirection::LeftToRight;
+		}
+
+		ICUUtilities::ConvertString(InString, ICUString);
+
+		return Internal::ComputeTextDirection(ICUBiDi, ICUString);
+	}
+
+	virtual ETextDirection ComputeTextDirection(const FText& InText, TArray<FTextDirectionInfo>& OutTextDirectionInfo) override
+	{
+		return FICUTextBiDi::ComputeTextDirection(InText.ToString(), OutTextDirectionInfo);
+	}
+
+	virtual ETextDirection ComputeTextDirection(const FString& InString, TArray<FTextDirectionInfo>& OutTextDirectionInfo) override
+	{
+		OutTextDirectionInfo.Reset();
+
+		if (InString.IsEmpty())
+		{
+			return ETextDirection::LeftToRight;
+		}
+
+		ICUUtilities::ConvertString(InString, ICUString);
+
+		return Internal::ComputeTextDirection(ICUBiDi, ICUString, OutTextDirectionInfo);
+	}
+
+private:
+	UBiDi* ICUBiDi;
+	icu::UnicodeString ICUString;
+};
+
+} // namespace Internal
+
+TUniquePtr<ITextBiDi> CreateTextBiDi()
+{
+	return MakeUnique<Internal::FICUTextBiDi>();
+}
+
+ETextDirection ComputeTextDirection(const FText& InText)
+{
+	return ComputeTextDirection(InText.ToString());
+}
+
+ETextDirection ComputeTextDirection(const FString& InString)
+{
+	if (InString.IsEmpty())
+	{
+		return ETextDirection::LeftToRight;
+	}
+
+	icu::UnicodeString ICUString = ICUUtilities::ConvertString(InString);
+
+	UErrorCode ICUStatus = U_ZERO_ERROR;
+	UBiDi* ICUBiDi = ubidi_openSized(ICUString.length(), 0, &ICUStatus);
+	if (ICUBiDi && U_SUCCESS(ICUStatus))
+	{
+		const ETextDirection ReturnDirection = Internal::ComputeTextDirection(ICUBiDi, ICUString);
+
+		ubidi_close(ICUBiDi);
+		ICUBiDi = nullptr;
+
+		return ReturnDirection;
+	}
+	else
+	{
+		UE_LOG(LogCore, Warning, TEXT("Failed to create ICU BiDi object (error code: %d). Text will assumed to be left-to-right"), static_cast<int32>(ICUStatus));
+	}
+
+	return ETextDirection::LeftToRight;
+}
+
+ETextDirection ComputeTextDirection(const FText& InText, TArray<FTextDirectionInfo>& OutTextDirectionInfo)
+{
+	return ComputeTextDirection(InText.ToString(), OutTextDirectionInfo);
+}
+
+ETextDirection ComputeTextDirection(const FString& InString, TArray<FTextDirectionInfo>& OutTextDirectionInfo)
+{
+	OutTextDirectionInfo.Reset();
+
+	if (InString.IsEmpty())
+	{
+		return ETextDirection::LeftToRight;
+	}
+
+	icu::UnicodeString ICUString = ICUUtilities::ConvertString(InString);
+
+	UErrorCode ICUStatus = U_ZERO_ERROR;
+	UBiDi* ICUBiDi = ubidi_openSized(ICUString.length(), 0, &ICUStatus);
+	if (ICUBiDi && U_SUCCESS(ICUStatus))
+	{
+		const ETextDirection ReturnDirection = Internal::ComputeTextDirection(ICUBiDi, ICUString, OutTextDirectionInfo);
+
+		ubidi_close(ICUBiDi);
+		ICUBiDi = nullptr;
+
+		return ReturnDirection;
+	}
+	else
+	{
+		UE_LOG(LogCore, Warning, TEXT("Failed to create ICU BiDi object (error code: %d). Text will assumed to be left-to-right"), static_cast<int32>(ICUStatus));
+	}
+
+	return ETextDirection::LeftToRight;
+}
+
+} // namespace TextBiDi
+
 #endif
diff --git a/Engine/Source/Runtime/Core/Private/Internationalization/LegacyText.cpp b/Engine/Source/Runtime/Core/Private/Internationalization/LegacyText.cpp
--- a/Engine/Source/Runtime/Core/Private/Internationalization/LegacyText.cpp
+++ b/Engine/Source/Runtime/Core/Private/Internationalization/LegacyText.cpp
@@ -96,4 +96,85 @@ bool FText::IsLetter( const TCHAR Char )
 	return (Char>='A' && Char<='Z') || (Char>='a' && Char<='z');
 }
 
+namespace TextBiDi
+{
+
+namespace Internal
+{
+
+class FLegacyTextBiDi : public ITextBiDi
+{
+public:
+	virtual ETextDirection ComputeTextDirection(const FText& InText) override
+	{
+		return FLegacyTextBiDi::ComputeTextDirection(InText.ToString());
+	}
+
+	virtual ETextDirection ComputeTextDirection(const FString& InString) override
+	{
+		return ETextDirection::LeftToRight;
+	}
+
+	virtual ETextDirection ComputeTextDirection(const FText& InText, TArray<FTextDirectionInfo>& OutTextDirectionInfo) override
+	{
+		return FLegacyTextBiDi::ComputeTextDirection(InText.ToString(), OutTextDirectionInfo);
+	}
+
+	virtual ETextDirection ComputeTextDirection(const FString& InString, TArray<FTextDirectionInfo>& OutTextDirectionInfo) override
+	{
+		OutTextDirectionInfo.Reset();
+
+		if (!InString.IsEmpty())
+		{
+			FTextDirectionInfo TextDirectionInfo;
+			TextDirectionInfo.StartIndex = 0;
+			TextDirectionInfo.Length = InString.Len();
+			TextDirectionInfo.TextDirection = ETextDirection::LeftToRight;
+			OutTextDirectionInfo.Add(MoveTemp(TextDirectionInfo));
+		}
+
+		return ETextDirection::LeftToRight;
+	}
+};
+
+} // namespace Internal
+
+TUniquePtr<ITextBiDi> CreateTextBiDi()
+{
+	return MakeUnique<Internal::FLegacyTextBiDi>();
+}
+
+ETextDirection ComputeTextDirection(const FText& InText)
+{
+	return ComputeTextDirection(InText.ToString());
+}
+
+ETextDirection ComputeTextDirection(const FString& InString)
+{
+	return ETextDirection::LeftToRight;
+}
+
+ETextDirection ComputeTextDirection(const FText& InText, TArray<FTextDirectionInfo>& OutTextDirectionInfo)
+{
+	return ComputeTextDirection(InText.ToString(), OutTextDirectionInfo);
+}
+
+ETextDirection ComputeTextDirection(const FString& InString, TArray<FTextDirectionInfo>& OutTextDirectionInfo)
+{
+	OutTextDirectionInfo.Reset();
+
+	if (!InString.IsEmpty())
+	{
+		FTextDirectionInfo TextDirectionInfo;
+		TextDirectionInfo.StartIndex = 0;
+		TextDirectionInfo.Length = InString.Len();
+		TextDirectionInfo.TextDirection = ETextDirection::LeftToRight;
+		OutTextDirectionInfo.Add(MoveTemp(TextDirectionInfo));
+	}
+
+	return ETextDirection::LeftToRight;
+}
+
+} // namespace TextBiDi
+
 #endif
diff --git a/Engine/Source/Runtime/Core/Public/Internationalization/Text.h b/Engine/Source/Runtime/Core/Public/Internationalization/Text.h
--- a/Engine/Source/Runtime/Core/Public/Internationalization/Text.h
+++ b/Engine/Source/Runtime/Core/Public/Internationalization/Text.h
@@ -6,6 +6,7 @@
 #include "UObject/NameTypes.h"
 #include "TextLocalizationManager.h"
 #include "Optional.h"
+#include "UniquePtr.h"
 
 struct FTimespan;
 struct FDateTime;
@@ -603,4 +604,66 @@ private:
 	FText& TextToPersist;
 };
 
+/**
+ * Unicode Bidirectional text support 
+ * http://www.unicode.org/reports/tr9/
+ */
+namespace TextBiDi
+{
+	/** Lists the potential reading directions for text */
+	enum class ETextDirection : uint8
+	{
+		/** Contains only LTR text - requires simple LTR layout */
+		LeftToRight,
+		/** Contains only RTL text - requires simple RTL layout */
+		RightToLeft,
+		/** Contains both LTR and RTL text - requires more complex layout using multiple runs of text */
+		Mixed,
+	};
+
+	/** A single complex layout entry. Defines the starting position, length, and reading direction for a sub-section of text */
+	struct FTextDirectionInfo
+	{
+		int32 StartIndex;
+		int32 Length;
+		ETextDirection TextDirection;
+	};
+
+	/** Defines the interface for a re-usable BiDi object */
+	class CORE_API ITextBiDi
+	{
+	public:
+		virtual ~ITextBiDi() {}
+
+		/** See TextBiDi::ComputeTextDirection */
+		virtual ETextDirection ComputeTextDirection(const FText& InText) = 0;
+		virtual ETextDirection ComputeTextDirection(const FString& InString) = 0;
+
+		/** See TextBiDi::ComputeTextDirection */
+		virtual ETextDirection ComputeTextDirection(const FText& InText, TArray<FTextDirectionInfo>& OutTextDirectionInfo) = 0;
+		virtual ETextDirection ComputeTextDirection(const FString& InString, TArray<FTextDirectionInfo>& OutTextDirectionInfo) = 0;
+	};
+
+	/**
+	 * Create a re-usable BiDi object.
+	 * This may yield better performance than the utility functions if you're performing a lot of BiDi requests, as this object can re-use allocated data between requests.
+	 */
+	CORE_API TUniquePtr<ITextBiDi> CreateTextBiDi();
+
+	/**
+	 * Utility function which will compute the reading direction of the given text.
+	 * @note You may want to use the version that returns you the advanced layout data in the Mixed case.
+	 * @return LeftToRight if all of the text is LTR, RightToLeft if all of the text is RTL, or Mixed if the text contains both LTR and RTL text.
+	 */
+	CORE_API ETextDirection ComputeTextDirection(const FText& InText);
+	CORE_API ETextDirection ComputeTextDirection(const FString& InString);
+
+	/**
+	 * Utility function which will compute the reading direction of the given text, as well as populate any advanced layout data for the text.
+	 * @return LeftToRight if all of the text is LTR, RightToLeft if all of the text is RTL, or Mixed if the text contains both LTR and RTL text.
+	 */
+	CORE_API ETextDirection ComputeTextDirection(const FText& InText, TArray<FTextDirectionInfo>& OutTextDirectionInfo);
+	CORE_API ETextDirection ComputeTextDirection(const FString& InString, TArray<FTextDirectionInfo>& OutTextDirectionInfo);
+} // namespace TextBiDi
+
 Expose_TNameOf(FText)
