More verbose log when the file is currupted

diff --git a/Engine/Source/Runtime/Core/Private/HAL/FileManagerGeneric.cpp b/Engine/Source/Runtime/Core/Private/HAL/FileManagerGeneric.cpp
--- a/Engine/Source/Runtime/Core/Private/HAL/FileManagerGeneric.cpp
+++ b/Engine/Source/Runtime/Core/Private/HAL/FileManagerGeneric.cpp
@@ -658,7 +658,7 @@ bool FArchiveFileReaderGeneric::InternalPrecache( int64 PrecacheOffset, int64 Pr
 
 		// Read data from device via Win32 ReadFile API.
 		{
-			UE_CLOG( BufferCount > ARRAY_COUNT( Buffer ) || BufferCount <= 0, LogFileManager, Fatal, TEXT("Invalid BufferCount=%lld while reading %s. Pos=%lld, Size=%lld, PrecacheSize=%lld, PrecacheOffset=%lld"),
+			UE_CLOG( BufferCount > ARRAY_COUNT( Buffer ) || BufferCount <= 0, LogFileManager, Fatal, TEXT("Invalid BufferCount=%lld while reading %s. File is most likely corrupted, try deleting it if possible. Pos=%lld, Size=%lld, PrecacheSize=%lld, PrecacheOffset=%lld"),
 				BufferCount, *Filename, Pos, Size, PrecacheSize, PrecacheOffset );
 
 			ReadLowLevel( Buffer, BufferCount, Count );
