UE-15045: FWindowsOSVersionHelper::GetOSVersions will not return OS information for newer OS's.

diff --git a/Engine/Source/Runtime/Core/Private/Windows/WindowsPlatformMisc.cpp b/Engine/Source/Runtime/Core/Private/Windows/WindowsPlatformMisc.cpp
--- a/Engine/Source/Runtime/Core/Private/Windows/WindowsPlatformMisc.cpp
+++ b/Engine/Source/Runtime/Core/Private/Windows/WindowsPlatformMisc.cpp
@@ -46,6 +46,105 @@
 #define SM_CONVERTIBLESLATEMODE			0x2003
 #endif
 
+namespace
+{
+	/**
+	 * According to MSDN GetVersionEx without special targeting works to 6.2
+	 * version only. To retrive proper version for later version we can check
+	 * version of system libraries e.g. kernel32.dll.
+	 */
+	int32 GetWindowsGT62Versions(bool bIsWorkstation, FString& out_OSVersionLabel)
+	{
+		const int BufferSize = 256;
+		TCHAR Buffer[BufferSize];
+		
+		if (!GetSystemDirectory(Buffer, BufferSize))
+		{
+			return (int32)FWindowsOSVersionHelper::ERROR_GETWINDOWSGT62VERSIONS_FAILED;
+		}
+
+		FString SystemDir(Buffer);
+		FString KernelPath = FPaths::Combine(*SystemDir, TEXT("kernel32.dll"));
+
+		DWORD Size = GetFileVersionInfoSize(*KernelPath, nullptr);
+
+		if (Size <= 0)
+		{
+			return (int32)FWindowsOSVersionHelper::ERROR_GETWINDOWSGT62VERSIONS_FAILED;
+		}
+
+		TArray<uint8> VerBlock;
+		VerBlock.Reserve(Size);
+
+		if (!GetFileVersionInfo(*KernelPath, 0, Size, VerBlock.GetData()))
+		{
+			return (int32)FWindowsOSVersionHelper::ERROR_GETWINDOWSGT62VERSIONS_FAILED;
+		}
+
+		VS_FIXEDFILEINFO* FileInfo = nullptr;
+		uint32 Len;
+
+		if (!VerQueryValue(VerBlock.GetData(), TEXT("\\"), (void **)&FileInfo, &Len))
+		{
+			return (int32)FWindowsOSVersionHelper::ERROR_GETWINDOWSGT62VERSIONS_FAILED;
+		}
+
+		int Major = FileInfo->dwProductVersionMS >> 16;
+		int Minor = FileInfo->dwProductVersionMS & 0xFFFF;
+
+		switch (Major)
+		{
+		case 6:
+			switch (Minor)
+			{
+			case 3:
+				if (bIsWorkstation)
+				{
+					out_OSVersionLabel = TEXT("Windows 8.1");
+				}
+				else
+				{
+					out_OSVersionLabel = TEXT("Windows Server 2012 R2");
+				}
+				break;
+			case 2:
+				if (bIsWorkstation)
+				{
+					out_OSVersionLabel = TEXT("Windows 8");
+				}
+				else
+				{
+					out_OSVersionLabel = TEXT("Windows Server 2012");
+				}
+				break;
+			default:
+				return (int32)FWindowsOSVersionHelper::ERROR_UNKNOWNVERSION;
+			}
+			break;
+		case 10:
+			switch (Minor)
+			{
+			case 0:
+				if (bIsWorkstation)
+				{
+					out_OSVersionLabel = TEXT("Windows 10");
+				}
+				else
+				{
+					out_OSVersionLabel = TEXT("Windows Server Technical Preview");
+				}
+				break;
+			default:
+				return (int32)FWindowsOSVersionHelper::ERROR_UNKNOWNVERSION;
+			}
+			break;
+		default:
+			return (int32)FWindowsOSVersionHelper::ERROR_UNKNOWNVERSION;
+		}
+
+		return (int32)FWindowsOSVersionHelper::SUCCEEDED;
+	}
+}
 
 int32 FWindowsOSVersionHelper::GetOSVersions( FString& out_OSVersionLabel, FString& out_OSSubVersionLabel )
 {
@@ -163,14 +262,7 @@ int32 FWindowsOSVersionHelper::GetOSVersions( FString& out_OSVersionLabel, FStri
 						}
 						break;
 					case 2:
-						if( OsVersionInfo.wProductType == VER_NT_WORKSTATION )
-						{
-							out_OSVersionLabel = TEXT( "Windows 8" );
-						}
-						else
-						{
-							out_OSVersionLabel = TEXT( "Windows Server 2012" );
-						}
+						ErrorCode |= GetWindowsGT62Versions(OsVersionInfo.wProductType == VER_NT_WORKSTATION, out_OSVersionLabel);
 						break;
 					default:
 						ErrorCode |= (int32)ERROR_UNKNOWNVERSION;
diff --git a/Engine/Source/Runtime/Core/Public/Windows/WindowsPlatformMisc.h b/Engine/Source/Runtime/Core/Public/Windows/WindowsPlatformMisc.h
--- a/Engine/Source/Runtime/Core/Public/Windows/WindowsPlatformMisc.h
+++ b/Engine/Source/Runtime/Core/Public/Windows/WindowsPlatformMisc.h
@@ -15,6 +15,7 @@ struct FWindowsOSVersionHelper
 		ERROR_UNKNOWNVERSION = 1,
 		ERROR_GETPRODUCTINFO_FAILED = 2,
 		ERROR_GETVERSIONEX_FAILED = 4,
+		ERROR_GETWINDOWSGT62VERSIONS_FAILED = 8,
 	};
 
 	static int32 GetOSVersions( FString& out_OSVersion, FString& out_OSSubVersion );
