Use the dispatch semaphore to reduce the impact of PerformBlockOnThread by waiting on it with a short timeout. This loop can't wait forever because the run-loop source may need to be rescheduled to complete the outstanding work.

diff --git a/Engine/Source/Runtime/Core/Private/Mac/CocoaThread.cpp b/Engine/Source/Runtime/Core/Private/Mac/CocoaThread.cpp
--- a/Engine/Source/Runtime/Core/Private/Mac/CocoaThread.cpp
+++ b/Engine/Source/Runtime/Core/Private/Mac/CocoaThread.cpp
@@ -393,13 +393,11 @@ static void PerformBlockOnThread(FCocoaRunLoopSource& ThreadSource, dispatch_blo
 		dispatch_block_t ExecuteBlock = Block_copy(^{ CopiedBlock(); dispatch_semaphore_signal(Semaphore); });
 		
 		ThreadSource.Schedule(ExecuteBlock, SendModes);
-		ThreadSource.Wake();
 		
-		while (dispatch_semaphore_wait(Semaphore, DISPATCH_TIME_NOW))
-		{
-			ThreadSource.RunInMode( (CFStringRef)WaitMode );
+		do {
 			ThreadSource.Wake();
-		}
+			ThreadSource.RunInMode( (CFStringRef)WaitMode );
+		} while (dispatch_semaphore_wait(Semaphore, dispatch_time(0, 100000ull)));
 		
 		Block_release(ExecuteBlock);
 		dispatch_release(Semaphore);
