UE4 - merge from orion, patch up morph target disable cvar threading. CL 2654433 in orion

diff --git a/Engine/Source/Runtime/Core/Public/HAL/IConsoleManager.h b/Engine/Source/Runtime/Core/Public/HAL/IConsoleManager.h
--- a/Engine/Source/Runtime/Core/Public/HAL/IConsoleManager.h
+++ b/Engine/Source/Runtime/Core/Public/HAL/IConsoleManager.h
@@ -757,10 +757,7 @@ public:
 	T GetValueOnGameThread() const
 	{
 		// compiled out in shipping for performance (we can change in development later), if this get triggered you need to call GetValueOnRenderThread() or GetValueOnAnyThread(), the last one is a bit slower
-	#if DO_CHECK
 		ensure(GetShadowIndex() == 0);	// ensure to not block content creators, #if to optimize in shipping
-	#endif
-
 		return ShadowedValue[0];
 	}
 
@@ -769,18 +766,15 @@ public:
 	{
 #if !defined(__clang__) // @todo Mac: figure out how to make this compile
 		// compiled out in shipping for performance (we can change in development later), if this get triggered you need to call GetValueOnGameThread() or GetValueOnAnyThread(), the last one is a bit slower
-	#if DO_CHECK
 		ensure(IsInParallelRenderingThread());	// ensure to not block content creators, #if to optimize in shipping
-	#endif
-
 #endif
 		return ShadowedValue[1];
 	}
 
 	// convenient, for better performance consider using GetValueOnGameThread() or GetValueOnRenderThread()
-	T GetValueOnAnyThread() const
+	T GetValueOnAnyThread(bool bForceGameThread = false) const
 	{
-		return ShadowedValue[GetShadowIndex()];
+		return ShadowedValue[GetShadowIndex(bForceGameThread)];
 	}
 
 private: // ----------------------------------------------------
@@ -789,15 +783,20 @@ private: // ----------------------------------------------------
 	T ShadowedValue[2];
 
 	// @return 0:main thread, 1: render thread, later more
-	static uint32 GetShadowIndex()
+	static uint32 GetShadowIndex(bool bForceGameThread = false)
 	{	
+		if (bForceGameThread)
+		{
+			ensure(!IsInActualRenderingThread());
+			return 0;
+		}
 		return IsInGameThread() ? 0 : 1;
 	}
 
 	// needed for FConsoleVariable and FConsoleVariableRef2, intentionally not public
-	T& GetReferenceOnAnyThread()
+	T& GetReferenceOnAnyThread(bool bForceGameThread = false)
 	{
-		return ShadowedValue[GetShadowIndex()];
+		return ShadowedValue[GetShadowIndex(bForceGameThread)];
 	}
 
 	template<class T2> friend class FConsoleVariable;
@@ -829,9 +828,9 @@ public:
 		return Ref->GetValueOnRenderThread();
 	}
 
-	T GetValueOnAnyThread() const
+	T GetValueOnAnyThread(bool bForceGameThread = false) const
 	{
-		return Ref->GetValueOnAnyThread();
+		return Ref->GetValueOnAnyThread(bForceGameThread);
 	}
 
 private:
diff --git a/Engine/Source/Runtime/Engine/Private/Components/SkinnedMeshComponent.cpp b/Engine/Source/Runtime/Engine/Private/Components/SkinnedMeshComponent.cpp
--- a/Engine/Source/Runtime/Engine/Private/Components/SkinnedMeshComponent.cpp
+++ b/Engine/Source/Runtime/Engine/Private/Components/SkinnedMeshComponent.cpp
@@ -319,7 +319,7 @@ void USkinnedMeshComponent::CreateRenderState_Concurrent()
 			// verifies vertex animations are valid
 			RefreshActiveVertexAnims();
 
-			const bool bMorphTargetsAllowed = CVarEnableMorphTargets.GetValueOnAnyThread() != 0;
+			const bool bMorphTargetsAllowed = CVarEnableMorphTargets.GetValueOnAnyThread(true) != 0;
 
 			// Are morph targets disabled for this LOD?
 			if (SkeletalMesh->LODInfo[UseLOD].bHasBeenSimplified || bDisableMorphTarget || !bMorphTargetsAllowed)
@@ -382,7 +382,7 @@ void USkinnedMeshComponent::SendRenderDynamicData_Concurrent()
 
 		int32 UseLOD = PredictedLODLevel;
 
-		const bool bMorphTargetsAllowed = CVarEnableMorphTargets.GetValueOnAnyThread() != 0;
+		const bool bMorphTargetsAllowed = CVarEnableMorphTargets.GetValueOnAnyThread(true) != 0;
 
 		// Are morph targets disabled for this LOD?
 		if (SkeletalMesh->LODInfo[UseLOD].bHasBeenSimplified || bDisableMorphTarget || !bMorphTargetsAllowed)
