Move GEngineVersion into FEngineVersion::Current(), and GCompatibleWithEngineVersion to FEngineVersion::CompatibleWith().

diff --git a/Engine/Plugins/Editor/EpicSurvey/Source/EpicSurvey/Private/Survey.cpp b/Engine/Plugins/Editor/EpicSurvey/Source/EpicSurvey/Private/Survey.cpp
--- a/Engine/Plugins/Editor/EpicSurvey/Source/EpicSurvey/Private/Survey.cpp
+++ b/Engine/Plugins/Editor/EpicSurvey/Source/EpicSurvey/Private/Survey.cpp
@@ -49,8 +49,8 @@ TSharedPtr< FSurvey > FSurvey::Create( const TSharedRef< class FEpicSurvey >& In
 		FEngineVersion::Parse( JsonConfig->GetStringField( TEXT("max_engine_version") ), NewSurvey->MaxEngineVersion );
 	}
 	if( (NewSurvey->SurveyVersion != CurrentSurveyVersion) 
-		|| (!NewSurvey->MinEngineVersion.IsEmpty() && GEngineVersion.IsCompatibleWith(NewSurvey->MinEngineVersion))
-		|| (!NewSurvey->MaxEngineVersion.IsEmpty() && NewSurvey->MaxEngineVersion.IsCompatibleWith(GEngineVersion)) )
+		|| (!NewSurvey->MinEngineVersion.IsEmpty() && FEngineVersion::Current().IsCompatibleWith(NewSurvey->MinEngineVersion))
+		|| (!NewSurvey->MaxEngineVersion.IsEmpty() && NewSurvey->MaxEngineVersion.IsCompatibleWith(FEngineVersion::Current())) )
 	{
 		return NULL;
 	}
diff --git a/Engine/Plugins/NetcodeUnitTest/NetcodeUnitTest/Source/NetcodeUnitTest/Private/Net/NUTUtilNet.cpp b/Engine/Plugins/NetcodeUnitTest/NetcodeUnitTest/Source/NetcodeUnitTest/Private/Net/NUTUtilNet.cpp
--- a/Engine/Plugins/NetcodeUnitTest/NetcodeUnitTest/Source/NetcodeUnitTest/Private/Net/NUTUtilNet.cpp
+++ b/Engine/Plugins/NetcodeUnitTest/NetcodeUnitTest/Source/NetcodeUnitTest/Private/Net/NUTUtilNet.cpp
@@ -362,7 +362,7 @@ bool NUTNet::CreateFakePlayer(UWorld* InWorld, UNetDriver*& InNetDriver, FString
 #if TARGET_UE4_CL >= CL_FNETWORKVERSION
 			// Starting with 4.8.0, the network protocol has changed slightly
 			// @todo JohnB: Refactor this, to toggle at compile time only, based on CL (might require more accurate UT integrate CLs)
-			FString VersionStr = GEngineVersion.ToString(EVersionComponent::Minor);
+			FString VersionStr = FEngineVersion::Current().ToString(EVersionComponent::Minor);
 			int32 VersionDelim = VersionStr.Find(TEXT("."));
 			int32 MajorVersion = FCString::Atoi(*VersionStr.Left(VersionDelim));
 			int32 MinorVersion = FCString::Atoi(*VersionStr.Mid(VersionDelim+1));
diff --git a/Engine/Plugins/Runtime/GearVR/Source/GearVR/Private/GearVR.cpp b/Engine/Plugins/Runtime/GearVR/Source/GearVR/Private/GearVR.cpp
--- a/Engine/Plugins/Runtime/GearVR/Source/GearVR/Private/GearVR.cpp
+++ b/Engine/Plugins/Runtime/GearVR/Source/GearVR/Private/GearVR.cpp
@@ -359,7 +359,7 @@ bool FGearVR::Exec( UWorld* InWorld, const TCHAR* Cmd, FOutputDevice& Ar )
 FString FGearVR::GetVersionString() const
 {
 	static const char* Results = OVR_VERSION_STRING;
-	FString s = FString::Printf(TEXT("%s, VrLib: %s, built %s, %s"), *GEngineVersion.ToString(), UTF8_TO_TCHAR(Results),
+	FString s = FString::Printf(TEXT("%s, VrLib: %s, built %s, %s"), *FEngineVersion::Current().ToString(), UTF8_TO_TCHAR(Results),
 		UTF8_TO_TCHAR(__DATE__), UTF8_TO_TCHAR(__TIME__));
 	return s;
 }
diff --git a/Engine/Plugins/Runtime/OculusRift/Source/OculusRift/Private/0.5/OculusRiftHMD_05.cpp b/Engine/Plugins/Runtime/OculusRift/Source/OculusRift/Private/0.5/OculusRiftHMD_05.cpp
--- a/Engine/Plugins/Runtime/OculusRift/Source/OculusRift/Private/0.5/OculusRiftHMD_05.cpp
+++ b/Engine/Plugins/Runtime/OculusRift/Source/OculusRift/Private/0.5/OculusRiftHMD_05.cpp
@@ -623,7 +623,7 @@ bool FOculusRiftHMD::Exec( UWorld* InWorld, const TCHAR* Cmd, FOutputDevice& Ar
 FString FOculusRiftHMD::GetVersionString() const
 {
 	const char* Results = ovr_GetVersionString();
-	FString s = FString::Printf(TEXT("%s, LibOVR: %s, built %s, %s"), *GEngineVersion.ToString(), UTF8_TO_TCHAR(Results),
+	FString s = FString::Printf(TEXT("%s, LibOVR: %s, built %s, %s"), *FEngineVersion::Current().ToString(), UTF8_TO_TCHAR(Results),
 		UTF8_TO_TCHAR(__DATE__), UTF8_TO_TCHAR(__TIME__));
 	return s;
 }
diff --git a/Engine/Plugins/Runtime/OculusRift/Source/OculusRift/Private/OculusRiftHMD.cpp b/Engine/Plugins/Runtime/OculusRift/Source/OculusRift/Private/OculusRiftHMD.cpp
--- a/Engine/Plugins/Runtime/OculusRift/Source/OculusRift/Private/OculusRiftHMD.cpp
+++ b/Engine/Plugins/Runtime/OculusRift/Source/OculusRift/Private/OculusRiftHMD.cpp
@@ -589,7 +589,7 @@ bool FOculusRiftHMD::Exec( UWorld* InWorld, const TCHAR* Cmd, FOutputDevice& Ar
 FString FOculusRiftHMD::GetVersionString() const
 {
 	const char* Results = ovr_GetVersionString();
-	FString s = FString::Printf(TEXT("%s, LibOVR: %s, built %s, %s"), *GEngineVersion.ToString(), UTF8_TO_TCHAR(Results),
+	FString s = FString::Printf(TEXT("%s, LibOVR: %s, built %s, %s"), *FEngineVersion::Current().ToString(), UTF8_TO_TCHAR(Results),
 		UTF8_TO_TCHAR(__DATE__), UTF8_TO_TCHAR(__TIME__));
 	return s;
 }
diff --git a/Engine/Source/Developer/CrashDebugHelper/Private/CrashDebugHelper.cpp b/Engine/Source/Developer/CrashDebugHelper/Private/CrashDebugHelper.cpp
--- a/Engine/Source/Developer/CrashDebugHelper/Private/CrashDebugHelper.cpp
+++ b/Engine/Source/Developer/CrashDebugHelper/Private/CrashDebugHelper.cpp
@@ -71,7 +71,7 @@ bool ICrashDebugHelper::Init()
 		// Default to BUILT_FROM_CHANGELIST.
 		else
 		{
-			BuiltFromCL = GEngineVersion.GetChangelist();
+			BuiltFromCL = FEngineVersion::Current().GetChangelist();
 		}
 
 		CrashInfo.BuiltFromCL = BuiltFromCL;
diff --git a/Engine/Source/Editor/GameProjectGeneration/Private/GameProjectUtils.cpp b/Engine/Source/Editor/GameProjectGeneration/Private/GameProjectUtils.cpp
--- a/Engine/Source/Editor/GameProjectGeneration/Private/GameProjectUtils.cpp
+++ b/Engine/Source/Editor/GameProjectGeneration/Private/GameProjectUtils.cpp
@@ -2265,7 +2265,7 @@ GameProjectUtils::EProjectDuplicateResult GameProjectUtils::DuplicateProjectForU
 	}
 
 	// Append the new version number
-	NewDirectoryName += FString::Printf(TEXT(" %s"), *GEngineVersion.ToString(EVersionComponent::Minor));
+	NewDirectoryName += FString::Printf(TEXT(" %s"), *FEngineVersion::Current().ToString(EVersionComponent::Minor));
 
 	// Find a directory name that doesn't exist
 	FString BaseDirectoryName = NewDirectoryName;
diff --git a/Engine/Source/Editor/GameProjectGeneration/Private/SProjectBrowser.cpp b/Engine/Source/Editor/GameProjectGeneration/Private/SProjectBrowser.cpp
--- a/Engine/Source/Editor/GameProjectGeneration/Private/SProjectBrowser.cpp
+++ b/Engine/Source/Editor/GameProjectGeneration/Private/SProjectBrowser.cpp
@@ -970,7 +970,7 @@ bool SProjectBrowser::OpenProject( const FString& InProjectFile )
 	FEngineVersion EngineVersion;
 	if (FDesktopPlatformModule::Get()->TryParseStockEngineVersion(ProjectIdentifier, EngineVersion))
 	{
-		if (FEngineVersion::GetNewest(EngineVersion, GEngineVersion, nullptr) == EVersionComparison::First)
+		if (FEngineVersion::GetNewest(EngineVersion, FEngineVersion::Current(), nullptr) == EVersionComparison::First)
 		{
 			FMessageDialog::Open(EAppMsgType::Ok, LOCTEXT("CantLoadNewerProject", "Unable to open this project, as it was made with a newer version of the Unreal Engine."));
 			return false;
diff --git a/Engine/Source/Editor/LandscapeEditor/Private/Tests/LandscapeAutomationTests.cpp b/Engine/Source/Editor/LandscapeEditor/Private/Tests/LandscapeAutomationTests.cpp
--- a/Engine/Source/Editor/LandscapeEditor/Private/Tests/LandscapeAutomationTests.cpp
+++ b/Engine/Source/Editor/LandscapeEditor/Private/Tests/LandscapeAutomationTests.cpp
@@ -161,7 +161,7 @@ bool FLandscapeEditorTest::RunTest(const FString& Parameters)
 		const FString TestName = TEXT("NewLandscapeTest");
 		FString PathName = FPaths::AutomationDir() + TestName / FPlatformProperties::PlatformName();
 		FPaths::MakePathRelativeTo(PathName, *FPaths::RootDir());
-		ScreenshotParameters.ScreenshotName = FString::Printf(TEXT("%s/%d.png"), *PathName, GEngineVersion.GetChangelist());
+		ScreenshotParameters.ScreenshotName = FString::Printf(TEXT("%s/%d.png"), *PathName, FEngineVersion::Current().GetChangelist());
 
 		//Take a screenshot
 		ADD_LATENT_AUTOMATION_COMMAND(FTakeEditorScreenshotCommand(ScreenshotParameters));
diff --git a/Engine/Source/Editor/LevelEditor/Private/LevelEditor.cpp b/Engine/Source/Editor/LevelEditor/Private/LevelEditor.cpp
--- a/Engine/Source/Editor/LevelEditor/Private/LevelEditor.cpp
+++ b/Engine/Source/Editor/LevelEditor/Private/LevelEditor.cpp
@@ -59,7 +59,7 @@ public:
 		FColor BadgeTextColor = FColor(128,128,128,255);
 		GConfig->GetColor(TEXT("LevelEditor"), TEXT("ProjectBadgeTextColor"), /*out*/ BadgeTextColor, GEditorPerProjectIni);
 
-		const FString EngineVersionString = GEngineVersion.ToString(GEngineVersion.HasChangelist() ? EVersionComponent::Changelist : EVersionComponent::Patch);
+		const FString EngineVersionString = FEngineVersion::Current().ToString(FEngineVersion::Current().HasChangelist() ? EVersionComponent::Changelist : EVersionComponent::Patch);
 		
 		FFormatNamedArguments Args;
 		Args.Add(TEXT("Branch"), FText::FromString(OptionalBranchPrefix));
diff --git a/Engine/Source/Editor/StatsViewer/Private/SStatsViewer.cpp b/Engine/Source/Editor/StatsViewer/Private/SStatsViewer.cpp
--- a/Engine/Source/Editor/StatsViewer/Private/SStatsViewer.cpp
+++ b/Engine/Source/Editor/StatsViewer/Private/SStatsViewer.cpp
@@ -427,7 +427,7 @@ FReply SStatsViewer::OnExportClicked()
 	CSVFilename /= FString::Printf(
 		TEXT("%s-%i-%s.csv"),
 		FApp::GetGameName(),
-		GEngineVersion.GetChangelist(),
+		FEngineVersion::Current().GetChangelist(),
 		*FDateTime::Now().ToString() );
 
 	// Create the CSV (can't use ToNumber or FormatIntToHumanReadable as it'll break the CSV!)
diff --git a/Engine/Source/Editor/UnrealEd/Private/AboutScreen.cpp b/Engine/Source/Editor/UnrealEd/Private/AboutScreen.cpp
--- a/Engine/Source/Editor/UnrealEd/Private/AboutScreen.cpp
+++ b/Engine/Source/Editor/UnrealEd/Private/AboutScreen.cpp
@@ -20,7 +20,7 @@ void SAboutScreen::Construct(const FArguments& InArgs)
 #pragma warning(pop)
 #endif
 
-	FText Version = FText::Format( LOCTEXT("VersionLabel", "Version: {0}"), FText::FromString( GEngineVersion.ToString( ) ) );
+	FText Version = FText::Format( LOCTEXT("VersionLabel", "Version: {0}"), FText::FromString( FEngineVersion::Current().ToString( ) ) );
 
 	ChildSlot
 		[
diff --git a/Engine/Source/Editor/UnrealEd/Private/Commandlets/UpdateGameProjectCommandlet.cpp b/Engine/Source/Editor/UnrealEd/Private/Commandlets/UpdateGameProjectCommandlet.cpp
--- a/Engine/Source/Editor/UnrealEd/Private/Commandlets/UpdateGameProjectCommandlet.cpp
+++ b/Engine/Source/Editor/UnrealEd/Private/Commandlets/UpdateGameProjectCommandlet.cpp
@@ -68,7 +68,7 @@ int32 UUpdateGameProjectCommandlet::Main( const FString& InParams )
 		SourceControlProvider.Init();
 	}
 
-	FString EngineIdentifier = GEngineVersion.ToString(EVersionComponent::Minor);
+	FString EngineIdentifier = FEngineVersion::Current().ToString(EVersionComponent::Minor);
 
 	UE_LOG(LogUpdateGameProjectCommandlet, Display, TEXT("Updating project file %s to %s..."), *ProjectFilePath, *EngineIdentifier);
 
diff --git a/Engine/Source/Editor/UnrealEd/Private/CreditsScreen.cpp b/Engine/Source/Editor/UnrealEd/Private/CreditsScreen.cpp
--- a/Engine/Source/Editor/UnrealEd/Private/CreditsScreen.cpp
+++ b/Engine/Source/Editor/UnrealEd/Private/CreditsScreen.cpp
@@ -16,7 +16,7 @@ void SCreditsScreen::Construct(const FArguments& InArgs)
 
 	ActiveTimerHandle = RegisterActiveTimer( 0.f, FWidgetActiveTimerDelegate::CreateSP( this, &SCreditsScreen::RollCredits ) );
 
-	const FString Version = GEngineVersion.ToString(FEngineBuildSettings::IsPerforceBuild() ? EVersionComponent::Branch : EVersionComponent::Patch);
+	const FString Version = FEngineVersion::Current().ToString(FEngineBuildSettings::IsPerforceBuild() ? EVersionComponent::Branch : EVersionComponent::Patch);
 
 	FString CreditsText;
 	FFileHelper::LoadFileToString(CreditsText, *( FPaths::EngineContentDir() + TEXT("Editor/Credits.rt") ));
diff --git a/Engine/Source/Editor/UnrealEd/Private/Fbx/FbxMainExport.cpp b/Engine/Source/Editor/UnrealEd/Private/Fbx/FbxMainExport.cpp
--- a/Engine/Source/Editor/UnrealEd/Private/Fbx/FbxMainExport.cpp
+++ b/Engine/Source/Editor/UnrealEd/Private/Fbx/FbxMainExport.cpp
@@ -127,10 +127,10 @@ void FFbxExporter::CreateDocument()
 	SceneInfo->mSubject = "Export FBX meshes from Unreal";
 	SceneInfo->Original_ApplicationVendor.Set( "Epic Games" );
 	SceneInfo->Original_ApplicationName.Set( "Unreal Engine" );
-	SceneInfo->Original_ApplicationVersion.Set( TCHAR_TO_UTF8(*GEngineVersion.ToString()) );
+	SceneInfo->Original_ApplicationVersion.Set( TCHAR_TO_UTF8(*FEngineVersion::Current().ToString()) );
 	SceneInfo->LastSaved_ApplicationVendor.Set( "Epic Games" );
 	SceneInfo->LastSaved_ApplicationName.Set( "Unreal Engine" );
-	SceneInfo->LastSaved_ApplicationVersion.Set( TCHAR_TO_UTF8(*GEngineVersion.ToString()) );
+	SceneInfo->LastSaved_ApplicationVersion.Set( TCHAR_TO_UTF8(*FEngineVersion::Current().ToString()) );
 
 	Scene->SetSceneInfo(SceneInfo);
 	
diff --git a/Engine/Source/Editor/UnrealEd/Private/Tests/EditorBuildPromotionTests.cpp b/Engine/Source/Editor/UnrealEd/Private/Tests/EditorBuildPromotionTests.cpp
--- a/Engine/Source/Editor/UnrealEd/Private/Tests/EditorBuildPromotionTests.cpp
+++ b/Engine/Source/Editor/UnrealEd/Private/Tests/EditorBuildPromotionTests.cpp
@@ -3715,11 +3715,11 @@ bool FBuildPromotionSettingsTest::RunTest(const FString& Parameters)
 	UE_LOG(LogEditorBuildPromotionTests, Display, TEXT("Exporting Current keybindings and editor settings"));
 
 	//Export the original keybindings
-	const FString TargetOriginalKeybindFile = FString::Printf(TEXT("%s/BuildPromotion/OriginalKeybindings-%d.ini"), *FPaths::AutomationDir(), GEngineVersion.GetChangelist());
+	const FString TargetOriginalKeybindFile = FString::Printf(TEXT("%s/BuildPromotion/OriginalKeybindings-%d.ini"), *FPaths::AutomationDir(), FEngineVersion::Current().GetChangelist());
 	EditorBuildPromotionTestUtils::ExportKeybindings(TargetOriginalKeybindFile);
 
 	//Export the original preferences
-	const FString TargetOriginalPreferenceFile = FString::Printf(TEXT("%s/BuildPromotion/OriginalPreferences-%d.ini"), *FPaths::AutomationDir(), GEngineVersion.GetChangelist());
+	const FString TargetOriginalPreferenceFile = FString::Printf(TEXT("%s/BuildPromotion/OriginalPreferences-%d.ini"), *FPaths::AutomationDir(), FEngineVersion::Current().GetChangelist());
 	EditorBuildPromotionTestUtils::ExportEditorSettings(TargetOriginalPreferenceFile);
 
 //New Editor Settings
@@ -3736,7 +3736,7 @@ bool FBuildPromotionSettingsTest::RunTest(const FString& Parameters)
 	FInputChord NewPIEChord = EditorBuildPromotionTestUtils::SetKeybinding(TEXT("PlayWorld"), TEXT("RepeatLastPlay"), EKeys::L, EModifierKey::Control);
 	
 	//Export the keybindings
-	const FString TargetKeybindFile = FString::Printf(TEXT("%s/BuildPromotion/Keybindings-%d.ini"), *FPaths::AutomationDir(), GEngineVersion.GetChangelist());
+	const FString TargetKeybindFile = FString::Printf(TEXT("%s/BuildPromotion/Keybindings-%d.ini"), *FPaths::AutomationDir(), FEngineVersion::Current().GetChangelist());
 	UE_LOG(LogEditorBuildPromotionTests, Display, TEXT("Exporting keybind"));
 	EditorBuildPromotionTestUtils::ExportKeybindings(TargetKeybindFile);
 
@@ -3748,7 +3748,7 @@ bool FBuildPromotionSettingsTest::RunTest(const FString& Parameters)
 	UE_LOG(LogEditorBuildPromotionTests, Display, TEXT("Set UseSmallToolBarIcons"));
 
 	//Export the preferences
-	const FString TargetPreferenceFile = FString::Printf(TEXT("%s/BuildPromotion/Preferences-%d.ini"), *FPaths::AutomationDir(), GEngineVersion.GetChangelist());
+	const FString TargetPreferenceFile = FString::Printf(TEXT("%s/BuildPromotion/Preferences-%d.ini"), *FPaths::AutomationDir(), FEngineVersion::Current().GetChangelist());
 	EditorBuildPromotionTestUtils::ExportEditorSettings(TargetPreferenceFile);
 
 	//Take a screenshot of the small icons
diff --git a/Engine/Source/Editor/UnrealEd/Private/Tests/EditorPerformanceTests.cpp b/Engine/Source/Editor/UnrealEd/Private/Tests/EditorPerformanceTests.cpp
--- a/Engine/Source/Editor/UnrealEd/Private/Tests/EditorPerformanceTests.cpp
+++ b/Engine/Source/Editor/UnrealEd/Private/Tests/EditorPerformanceTests.cpp
@@ -82,7 +82,7 @@ void EditorPerfDump(EditorPerfCaptureParameters& EditorPerfStats)
 		//If the raw file isn't available to write to then we'll fail back this test.
 		if (AutomationEditorCommonUtils::IsArchiveWriteable(RAWCSVFilePath, RAWCSVArchive))
 		{
-			RAWCSVLine = FString::Printf(TEXT("%s,%s,%s,%.3f,%.1f,%.1f,%.0f,%.0f,%.0f,%.0f,%.0f,%.0f%s"), *EditorPerfStats.MapName, *GEngineVersion.ToString(EVersionComponent::Changelist), *EditorPerfStats.FormattedTimeStamp[i], EditorPerfStats.MapLoadTime, EditorPerfStats.AverageFPS[i], EditorPerfStats.AverageFrameTime[i], EditorPerfStats.UsedPhysical[i], EditorPerfStats.UsedVirtual[i], EditorPerfStats.PeakUsedPhysical[i], EditorPerfStats.PeakUsedVirtual[i], EditorPerfStats.AvailablePhysical[i], EditorPerfStats.AvailableVirtual[i], LINE_TERMINATOR);
+			RAWCSVLine = FString::Printf(TEXT("%s,%s,%s,%.3f,%.1f,%.1f,%.0f,%.0f,%.0f,%.0f,%.0f,%.0f%s"), *EditorPerfStats.MapName, *FEngineVersion::Current().ToString(EVersionComponent::Changelist), *EditorPerfStats.FormattedTimeStamp[i], EditorPerfStats.MapLoadTime, EditorPerfStats.AverageFPS[i], EditorPerfStats.AverageFrameTime[i], EditorPerfStats.UsedPhysical[i], EditorPerfStats.UsedVirtual[i], EditorPerfStats.PeakUsedPhysical[i], EditorPerfStats.PeakUsedVirtual[i], EditorPerfStats.AvailablePhysical[i], EditorPerfStats.AvailableVirtual[i], LINE_TERMINATOR);
 			RAWCSVArchive->Serialize(TCHAR_TO_ANSI(*RAWCSVLine), RAWCSVLine.Len());
 		}
 	}
@@ -123,7 +123,7 @@ void EditorPerfDump(EditorPerfCaptureParameters& EditorPerfStats)
 		FinalCSVArchive->Serialize(TCHAR_TO_ANSI(*OldPerformanceCSVFile), OldPerformanceCSVFile.Len());
 
 		//Dump the pretty stats to the Performance CSV file and then close it so we can edit it while the engine is still running.
-		FString FinalCSVLine = FString::Printf(TEXT("%s,%s,%s,%.0f,%.3f,%.1f,%.1f,%.0f,%.0f,%.0f,%.0f,%.0f,%.0f%s"), *FDateTime::Now().ToString(), *EditorPerfStats.MapName, *GEngineVersion.ToString(EVersionComponent::Changelist), TestRunDuration.GetTotalSeconds(), EditorPerfStats.MapLoadTime, AverageFPS, AverageFrameTime, MemoryUsedPhysical, MemoryUsedVirtualAvg, MemoryUsedPeak, MemoryUsedPeakVirtual, MemoryAvailPhysAvg, MemoryAvailVirtualAvg, LINE_TERMINATOR);
+		FString FinalCSVLine = FString::Printf(TEXT("%s,%s,%s,%.0f,%.3f,%.1f,%.1f,%.0f,%.0f,%.0f,%.0f,%.0f,%.0f%s"), *FDateTime::Now().ToString(), *EditorPerfStats.MapName, *FEngineVersion::Current().ToString(EVersionComponent::Changelist), TestRunDuration.GetTotalSeconds(), EditorPerfStats.MapLoadTime, AverageFPS, AverageFrameTime, MemoryUsedPhysical, MemoryUsedVirtualAvg, MemoryUsedPeak, MemoryUsedPeakVirtual, MemoryAvailPhysAvg, MemoryAvailVirtualAvg, LINE_TERMINATOR);
 		FinalCSVArchive->Serialize(TCHAR_TO_ANSI(*FinalCSVLine), FinalCSVLine.Len());
 		FinalCSVArchive->Close();
 	}
diff --git a/Engine/Source/Editor/UnrealEd/Private/UnrealEdEngine.cpp b/Engine/Source/Editor/UnrealEd/Private/UnrealEdEngine.cpp
--- a/Engine/Source/Editor/UnrealEd/Private/UnrealEdEngine.cpp
+++ b/Engine/Source/Editor/UnrealEd/Private/UnrealEdEngine.cpp
@@ -455,7 +455,7 @@ void UUnrealEdEngine::OnPackageDirtyStateUpdated( UPackage* Pkg)
 				FPackageFileSummary Summary;
 				*PackageReader << Summary;
 
-				if ( Summary.GetFileVersionUE4() > GPackageFileUE4Version || !GEngineVersion.IsCompatibleWith(Summary.CompatibleWithEngineVersion) )
+				if ( Summary.GetFileVersionUE4() > GPackageFileUE4Version || !FEngineVersion::Current().IsCompatibleWith(Summary.CompatibleWithEngineVersion) )
 				{
 					WarningStateToSet = WDWS_PendingWarn;
 					bNeedWarningForPkgEngineVer = true;
diff --git a/Engine/Source/Programs/CrashReporter/CrashReportClient/Private/CrashDescription.h b/Engine/Source/Programs/CrashReporter/CrashReportClient/Private/CrashDescription.h
--- a/Engine/Source/Programs/CrashReporter/CrashReportClient/Private/CrashDescription.h
+++ b/Engine/Source/Programs/CrashReporter/CrashReportClient/Private/CrashDescription.h
@@ -52,7 +52,7 @@ struct FCrashDescription
 	 * BuildVersion-BuiltFromCL-BranchName
 	 * @EngineVersion	varchar(64)	
 	 * 
-	 * GEngineVersion.ToString()
+	 * FEngineVersion::Current().ToString()
 	 * ENGINE_VERSION_STRING
 	 */
 	FEngineVersion EngineVersion;
diff --git a/Engine/Source/Runtime/Analytics/AnalyticsET/Private/AnalyticsET.cpp b/Engine/Source/Runtime/Analytics/AnalyticsET/Private/AnalyticsET.cpp
--- a/Engine/Source/Runtime/Analytics/AnalyticsET/Private/AnalyticsET.cpp
+++ b/Engine/Source/Runtime/Analytics/AnalyticsET/Private/AnalyticsET.cpp
@@ -162,13 +162,13 @@ FAnalyticsProviderET::FAnalyticsProviderET(const FAnalyticsET::Config& ConfigVal
 		? FAnalyticsET::Config::GetDefaultAPIServer()
 		: ConfigValues.APIServerET;
 
-	// default to GEngineVersion if one is not provided, append GEngineVersion otherwise.
+	// default to FEngineVersion::Current() if one is not provided, append FEngineVersion::Current() otherwise.
 	FString ConfigAppVersion = ConfigValues.AppVersionET;
 	// Allow the cmdline to force a specific AppVersion so it can be set dynamically.
 	FParse::Value(FCommandLine::Get(), TEXT("ANALYTICSAPPVERSION="), ConfigAppVersion, false);
 	AppVersion = ConfigAppVersion.IsEmpty() 
-		? GEngineVersion.ToString() 
-		: ConfigAppVersion.Replace(TEXT("%VERSION%"), *GEngineVersion.ToString(), ESearchCase::CaseSensitive);
+		? FEngineVersion::Current().ToString() 
+		: ConfigAppVersion.Replace(TEXT("%VERSION%"), *FEngineVersion::Current().ToString(), ESearchCase::CaseSensitive);
 
 	UE_LOG(LogAnalytics, Log, TEXT("[%s] APIServer = %s. AppVersion = %s"), *APIKey, *APIServer, *AppVersion);
 
diff --git a/Engine/Source/Runtime/Analytics/AnalyticsET/Public/AnalyticsET.h b/Engine/Source/Runtime/Analytics/AnalyticsET/Public/AnalyticsET.h
--- a/Engine/Source/Runtime/Analytics/AnalyticsET/Public/AnalyticsET.h
+++ b/Engine/Source/Runtime/Analytics/AnalyticsET/Public/AnalyticsET.h
@@ -41,9 +41,9 @@ public:
 		/** ET API Server - Defaults if empty to GetDefaultAPIServer. */
 		FString APIServerET;
 		/** 
-		 * AppVersion - defines the app version passed to the provider. By default this will be GEngineVersion, but you can supply your own. 
+		 * AppVersion - defines the app version passed to the provider. By default this will be FEngineVersion::Current(), but you can supply your own. 
 		 * As a convenience, you can use -AnalyticsAppVersion=XXX to force the AppVersion to a specific value. Useful for playtest etc where you want to define a specific version string dynamically.
-		 * If you supply your own Version string, occurrences of "%VERSION%" are replaced with GEngineVersion. ie, -AnalyticsAppVersion=MyCustomID-%VERSION%.
+		 * If you supply your own Version string, occurrences of "%VERSION%" are replaced with FEngineVersion::Current(). ie, -AnalyticsAppVersion=MyCustomID-%VERSION%.
 		 */
 		FString AppVersionET;
 
diff --git a/Engine/Source/Runtime/Analytics/AnalyticsSwrve/Private/AnalyticsSwrve.cpp b/Engine/Source/Runtime/Analytics/AnalyticsSwrve/Private/AnalyticsSwrve.cpp
--- a/Engine/Source/Runtime/Analytics/AnalyticsSwrve/Private/AnalyticsSwrve.cpp
+++ b/Engine/Source/Runtime/Analytics/AnalyticsSwrve/Private/AnalyticsSwrve.cpp
@@ -197,10 +197,10 @@ FAnalyticsProviderSwrve::FAnalyticsProviderSwrve(const FAnalyticsSwrve::Config&
 			: FAnalyticsSwrve::Config::GetDefaultAPIServer()
 		: ConfigValues.APIServerSwrve;
 
-	// default to GEngineVersion if one is not provided, append GEngineVersion otherwise.
+	// default to FEngineVersion::Current() if one is not provided, append FEngineVersion::Current() otherwise.
 	AppVersion = ConfigValues.AppVersionSwrve.IsEmpty()
-		? FString::Printf(TEXT("%d"), GEngineVersion.GetChangelist())
-		: FString::Printf(TEXT("%s.%d"), *ConfigValues.AppVersionSwrve, GEngineVersion.GetChangelist());
+		? FString::Printf(TEXT("%d"), FEngineVersion::Current().GetChangelist())
+		: FString::Printf(TEXT("%s.%d"), *ConfigValues.AppVersionSwrve, FEngineVersion::Current().GetChangelist());
 
 	UE_LOG(LogAnalytics, Log, TEXT("Swrve APIKey = %s. APIServer = %s. AppVersion = %s"), *APIKey, *APIServer, *AppVersion);
 
diff --git a/Engine/Source/Runtime/Analytics/AnalyticsSwrve/Public/AnalyticsSwrve.h b/Engine/Source/Runtime/Analytics/AnalyticsSwrve/Public/AnalyticsSwrve.h
--- a/Engine/Source/Runtime/Analytics/AnalyticsSwrve/Public/AnalyticsSwrve.h
+++ b/Engine/Source/Runtime/Analytics/AnalyticsSwrve/Public/AnalyticsSwrve.h
@@ -38,7 +38,7 @@ public:
 		FString APIKeySwrve;
 		/** Swrve API Server - Defaults if empty to GetDefaultAPIServer. FAnalytics::Debug build types default to GetDefaultAPIServerDebug. */
 		FString APIServerSwrve;
-		/** AppVersion - defines the app version passed to the provider. By default this will be GEngineVersion.GetChangelist(). If you provide your own, ".<GEngineVersion.GetChangelist()>" is appended to it. */
+		/** AppVersion - defines the app version passed to the provider. By default this will be FEngineVersion::Current().GetChangelist(). If you provide your own, ".<FEngineVersion::Current().GetChangelist()>" is appended to it. */
 		FString AppVersionSwrve;
 
 		/** KeyName required for APIKey configuration. */
diff --git a/Engine/Source/Runtime/Core/Private/GenericPlatform/GenericPlatformCrashContext.cpp b/Engine/Source/Runtime/Core/Private/GenericPlatform/GenericPlatformCrashContext.cpp
--- a/Engine/Source/Runtime/Core/Private/GenericPlatform/GenericPlatformCrashContext.cpp
+++ b/Engine/Source/Runtime/Core/Private/GenericPlatform/GenericPlatformCrashContext.cpp
@@ -129,7 +129,7 @@ void FGenericCrashContext::SerializeContentToBuffer()
 	AddCrashProperty( TEXT( "PlatformName" ), *NCachedCrashContextProperties::PlatformName );
 	AddCrashProperty( TEXT( "PlatformNameIni" ), *NCachedCrashContextProperties::PlatformNameIni );
 	AddCrashProperty( TEXT( "EngineMode" ), FPlatformMisc::GetEngineMode() );
-	AddCrashProperty( TEXT( "EngineVersion" ), *GEngineVersion.ToString() );
+	AddCrashProperty( TEXT( "EngineVersion" ), *FEngineVersion::Current().ToString() );
 	AddCrashProperty( TEXT( "CommandLine" ), FCommandLine::IsInitialized() ? FCommandLine::Get() : TEXT( "" ) );
 	AddCrashProperty( TEXT( "LanguageLCID" ), FInternationalization::Get().GetCurrentCulture()->GetLCID() );
 	AddCrashProperty( TEXT( "AppDefaultLocale" ), *NCachedCrashContextProperties::DefaultLocale );
diff --git a/Engine/Source/Runtime/Core/Private/GenericPlatform/GenericPlatformProcess.cpp b/Engine/Source/Runtime/Core/Private/GenericPlatform/GenericPlatformProcess.cpp
--- a/Engine/Source/Runtime/Core/Private/GenericPlatform/GenericPlatformProcess.cpp
+++ b/Engine/Source/Runtime/Core/Private/GenericPlatform/GenericPlatformProcess.cpp
@@ -35,7 +35,7 @@ void* FGenericPlatformProcess::GetDllExport( void* DllHandle, const TCHAR* ProcN
 int32 FGenericPlatformProcess::GetDllApiVersion( const TCHAR* Filename )
 {
 	UE_LOG(LogHAL, Fatal, TEXT("FPlatformProcess::GetBinaryFileVersion not implemented on this platform"));
-	return GCompatibleWithEngineVersion.GetChangelist();
+	return FEngineVersion::CompatibleWith().GetChangelist();
 }
 
 uint32 FGenericPlatformProcess::GetCurrentProcessId()
diff --git a/Engine/Source/Runtime/Core/Private/Linux/LinuxPlatformCrashContext.cpp b/Engine/Source/Runtime/Core/Private/Linux/LinuxPlatformCrashContext.cpp
--- a/Engine/Source/Runtime/Core/Private/Linux/LinuxPlatformCrashContext.cpp
+++ b/Engine/Source/Runtime/Core/Private/Linux/LinuxPlatformCrashContext.cpp
@@ -514,10 +514,10 @@ void FLinuxCrashContext::GenerateReport(const FString & DiagnosticsPath) const
 		WriteLine(ReportFile, "Generating report for minidump");
 		WriteLine(ReportFile);
 
-		Line = FString::Printf(TEXT("Application version %d.%d.%d.0" ), GEngineVersion.GetMajor(), GEngineVersion.GetMinor(), GEngineVersion.GetPatch());
+		Line = FString::Printf(TEXT("Application version %d.%d.%d.0" ), FEngineVersion::Current().GetMajor(), FEngineVersion::Current().GetMinor(), FEngineVersion::Current().GetPatch());
 		WriteLine(ReportFile, TCHAR_TO_UTF8(*Line));
 
-		Line = FString::Printf(TEXT(" ... built from changelist %d"), GEngineVersion.GetChangelist());
+		Line = FString::Printf(TEXT(" ... built from changelist %d"), FEngineVersion::Current().GetChangelist());
 		WriteLine(ReportFile, TCHAR_TO_UTF8(*Line));
 		WriteLine(ReportFile);
 
@@ -597,7 +597,7 @@ void GenerateWindowsErrorReport(const FString & WERPath)
 		WriteLine(ReportFile, TEXT("\t<ProblemSignatures>"));
 		WriteLine(ReportFile, TEXT("\t\t<EventType>APPCRASH</EventType>"));
 		WriteLine(ReportFile, *FString::Printf(TEXT("\t\t<Parameter0>UE4-%s</Parameter0>"), FApp::GetGameName()));
-		WriteLine(ReportFile, *FString::Printf(TEXT("\t\t<Parameter1>%d.%d.%d</Parameter1>"), GEngineVersion.GetMajor(), GEngineVersion.GetMinor(), GEngineVersion.GetPatch()));
+		WriteLine(ReportFile, *FString::Printf(TEXT("\t\t<Parameter1>%d.%d.%d</Parameter1>"), FEngineVersion::Current().GetMajor(), FEngineVersion::Current().GetMinor(), FEngineVersion::Current().GetPatch()));
 		WriteLine(ReportFile, TEXT("\t\t<Parameter2>0</Parameter2>"));													// FIXME: supply valid?
 		WriteLine(ReportFile, TEXT("\t\t<Parameter3>Unknown Fault Module</Parameter3>"));										// FIXME: supply valid?
 		WriteLine(ReportFile, TEXT("\t\t<Parameter4>0.0.0.0</Parameter4>"));													// FIXME: supply valid?
@@ -605,7 +605,7 @@ void GenerateWindowsErrorReport(const FString & WERPath)
 		WriteLine(ReportFile, TEXT("\t\t<Parameter6>00000000</Parameter6>"));													// FIXME: supply valid?
 		WriteLine(ReportFile, TEXT("\t\t<Parameter7>0000000000000000</Parameter7>"));											// FIXME: supply valid?
 		WriteLine(ReportFile, *FString::Printf(TEXT("\t\t<Parameter8>!%s!</Parameter8>"), FCommandLine::Get()));				// FIXME: supply valid? Only partially valid
-		WriteLine(ReportFile, *FString::Printf(TEXT("\t\t<Parameter9>%s!%s!%s!%d</Parameter9>"), *FApp::GetBranchName(), FPlatformProcess::BaseDir(), FPlatformMisc::GetEngineMode(), GEngineVersion.GetChangelist()));
+		WriteLine(ReportFile, *FString::Printf(TEXT("\t\t<Parameter9>%s!%s!%s!%d</Parameter9>"), *FApp::GetBranchName(), FPlatformProcess::BaseDir(), FPlatformMisc::GetEngineMode(), FEngineVersion::Current().GetChangelist()));
 		WriteLine(ReportFile, TEXT("\t</ProblemSignatures>"));
 
 		WriteLine(ReportFile, TEXT("\t<DynamicSignatures>"));
diff --git a/Engine/Source/Runtime/Core/Private/Linux/LinuxPlatformProcess.cpp b/Engine/Source/Runtime/Core/Private/Linux/LinuxPlatformProcess.cpp
--- a/Engine/Source/Runtime/Core/Private/Linux/LinuxPlatformProcess.cpp
+++ b/Engine/Source/Runtime/Core/Private/Linux/LinuxPlatformProcess.cpp
@@ -49,7 +49,7 @@ void* FLinuxPlatformProcess::GetDllExport( void* DllHandle, const TCHAR* ProcNam
 int32 FLinuxPlatformProcess::GetDllApiVersion( const TCHAR* Filename )
 {
 	check(Filename);
-	return GCompatibleWithEngineVersion.GetChangelist();
+	return FEngineVersion::CompatibleWith().GetChangelist();
 }
 
 const TCHAR* FLinuxPlatformProcess::GetModulePrefix()
diff --git a/Engine/Source/Runtime/Core/Private/Linux/LinuxPlatformSplash.cpp b/Engine/Source/Runtime/Core/Private/Linux/LinuxPlatformSplash.cpp
--- a/Engine/Source/Runtime/Core/Private/Linux/LinuxPlatformSplash.cpp
+++ b/Engine/Source/Runtime/Core/Private/Linux/LinuxPlatformSplash.cpp
@@ -906,7 +906,7 @@ void FLinuxPlatformSplash::Show( )
 
 		// Set version info
 		{
-			const FText Version = FText::FromString( GEngineVersion.ToString( FEngineBuildSettings::IsPerforceBuild() ? EVersionComponent::Branch : EVersionComponent::Patch ) );
+			const FText Version = FText::FromString( FEngineVersion::Current().ToString( FEngineBuildSettings::IsPerforceBuild() ? EVersionComponent::Branch : EVersionComponent::Patch ) );
 
 			FText VersionInfo;
 			FText AppName;
diff --git a/Engine/Source/Runtime/Core/Private/Mac/MacPlatformMisc.cpp b/Engine/Source/Runtime/Core/Private/Mac/MacPlatformMisc.cpp
--- a/Engine/Source/Runtime/Core/Private/Mac/MacPlatformMisc.cpp
+++ b/Engine/Source/Runtime/Core/Private/Mac/MacPlatformMisc.cpp
@@ -120,7 +120,7 @@ struct FMacApplicationInfo
 		
 		FString CrashVideoPath = FPaths::GameLogDir() + TEXT("CrashVideo.avi");
 		
-		BranchBaseDir = FString::Printf( TEXT( "%s!%s!%s!%d" ), *FApp::GetBranchName(), FPlatformProcess::BaseDir(), FPlatformMisc::GetEngineMode(), GEngineVersion.GetChangelist() );
+		BranchBaseDir = FString::Printf( TEXT( "%s!%s!%s!%d" ), *FApp::GetBranchName(), FPlatformProcess::BaseDir(), FPlatformMisc::GetEngineMode(), FEngineVersion::Current().GetChangelist() );
 		
 		// Get the paths that the files will actually have been saved to
 		FString LogDirectory = FPaths::GameLogDir();
@@ -1475,11 +1475,11 @@ void FMacCrashContext::GenerateWindowsErrorReport(char const* WERPath) const
 		WriteLine(ReportFile, TEXT("</Parameter0>"));
 		
 		WriteUTF16String(ReportFile, TEXT("\t\t<Parameter1>"));
-		WriteUTF16String(ReportFile, ItoTCHAR(GEngineVersion.GetMajor(), 10));
+		WriteUTF16String(ReportFile, ItoTCHAR(FEngineVersion::Current().GetMajor(), 10));
 		WriteUTF16String(ReportFile, TEXT("."));
-		WriteUTF16String(ReportFile, ItoTCHAR(GEngineVersion.GetMinor(), 10));
+		WriteUTF16String(ReportFile, ItoTCHAR(FEngineVersion::Current().GetMinor(), 10));
 		WriteUTF16String(ReportFile, TEXT("."));
-		WriteUTF16String(ReportFile, ItoTCHAR(GEngineVersion.GetPatch(), 10));
+		WriteUTF16String(ReportFile, ItoTCHAR(FEngineVersion::Current().GetPatch(), 10));
 		WriteLine(ReportFile, TEXT("</Parameter1>"));
 
 		// App time stamp
@@ -1673,9 +1673,9 @@ void FMacCrashContext::GenerateCrashInfoAndLaunchReporter() const
 			WriteLine(ReportFile, *GMacAppInfo.AppName);
 			
 			WriteUTF16String(ReportFile, TEXT("BuildVersion 1.0."));
-			WriteUTF16String(ReportFile, ItoTCHAR(GEngineVersion.GetChangelist() >> 16, 10));
+			WriteUTF16String(ReportFile, ItoTCHAR(FEngineVersion::Current().GetChangelist() >> 16, 10));
 			WriteUTF16String(ReportFile, TEXT("."));
-			WriteLine(ReportFile, ItoTCHAR(GEngineVersion.GetChangelist() & 0xffff, 10));
+			WriteLine(ReportFile, ItoTCHAR(FEngineVersion::Current().GetChangelist() & 0xffff, 10));
 			
 			WriteUTF16String(ReportFile, TEXT("CommandLine "));
 			WriteLine(ReportFile, *GMacAppInfo.CommandLine);
diff --git a/Engine/Source/Runtime/Core/Private/Mac/MacPlatformSplash.cpp b/Engine/Source/Runtime/Core/Private/Mac/MacPlatformSplash.cpp
--- a/Engine/Source/Runtime/Core/Private/Mac/MacPlatformSplash.cpp
+++ b/Engine/Source/Runtime/Core/Private/Mac/MacPlatformSplash.cpp
@@ -239,7 +239,7 @@ void FMacPlatformSplash::Show()
 
 				// Set version info
 				{
-					const FText Version = FText::FromString( GEngineVersion.ToString( FEngineBuildSettings::IsPerforceBuild() ? EVersionComponent::Branch : EVersionComponent::Patch ) );
+					const FText Version = FText::FromString( FEngineVersion::Current().ToString( FEngineBuildSettings::IsPerforceBuild() ? EVersionComponent::Branch : EVersionComponent::Patch ) );
 
 					FText VersionInfo;
 					FText AppName;
diff --git a/Engine/Source/Runtime/Core/Private/Misc/ConfigManifest.cpp b/Engine/Source/Runtime/Core/Private/Misc/ConfigManifest.cpp
--- a/Engine/Source/Runtime/Core/Private/Misc/ConfigManifest.cpp
+++ b/Engine/Source/Runtime/Core/Private/Misc/ConfigManifest.cpp
@@ -63,12 +63,12 @@ void MigratePreviousEngineInis()
 		return;
 	}
 
-	int32 MinorVersion = GEngineVersion.GetMinor() - 1;
+	int32 MinorVersion = FEngineVersion::Current().GetMinor() - 1;
 
 	IPlatformFile& PlatformFile = FPlatformFileManager::Get().GetPlatformFile();
 	while(MinorVersion >= 0)
 	{
-		const FEngineVersion PreviousVersion(GEngineVersion.GetMajor(), MinorVersion--, 0, 0, FString());
+		const FEngineVersion PreviousVersion(FEngineVersion::Current().GetMajor(), MinorVersion--, 0, 0, FString());
 	
 		const FString Directory = FString(FPlatformProcess::UserSettingsDir()) / ENGINE_VERSION_TEXT(EPIC_PRODUCT_IDENTIFIER) / PreviousVersion.ToString(EVersionComponent::Minor) / TEXT("Saved") / TEXT("Config") / ANSI_TO_TCHAR(FPlatformProperties::PlatformName());
 		if (FPaths::DirectoryExists(Directory))
diff --git a/Engine/Source/Runtime/Core/Private/Misc/EngineVersion.cpp b/Engine/Source/Runtime/Core/Private/Misc/EngineVersion.cpp
--- a/Engine/Source/Runtime/Core/Private/Misc/EngineVersion.cpp
+++ b/Engine/Source/Runtime/Core/Private/Misc/EngineVersion.cpp
@@ -183,6 +183,16 @@ bool FEngineVersion::Parse(const FString &Text, FEngineVersion &OutVersion)
 	return true;
 }
 
+const FEngineVersion& FEngineVersion::Current()
+{
+	return CurrentVersion;
+}
+
+const FEngineVersion& FEngineVersion::CompatibleWith()
+{
+	return CompatibleWithVersion;
+}
+
 void operator<<(FArchive &Ar, FEngineVersion &Version)
 {
 	Ar << Version.Major;
@@ -202,7 +212,7 @@ const int32 GEngineMinNetVersion		= 7038;
 const int32 GEngineNegotiationVersion	= 3077;
 
 // Global instance of the current engine version
-const FEngineVersion GEngineVersion(ENGINE_MAJOR_VERSION, ENGINE_MINOR_VERSION, ENGINE_PATCH_VERSION, (GEngineVersionChangelist | (ENGINE_IS_LICENSEE_VERSION << 31)), BRANCH_NAME);
+FEngineVersion FEngineVersion::CurrentVersion(ENGINE_MAJOR_VERSION, ENGINE_MINOR_VERSION, ENGINE_PATCH_VERSION, (GEngineVersionChangelist | (ENGINE_IS_LICENSEE_VERSION << 31)), BRANCH_NAME);
 
 // Version which this engine maintains strict API and package compatibility with
-const FEngineVersion GCompatibleWithEngineVersion(ENGINE_MAJOR_VERSION, ENGINE_MINOR_VERSION, ENGINE_PATCH_VERSION, (GCompatibleWithEngineVersionChangelist | (ENGINE_IS_LICENSEE_VERSION << 31)), BRANCH_NAME);
+FEngineVersion FEngineVersion::CompatibleWithVersion(ENGINE_MAJOR_VERSION, ENGINE_MINOR_VERSION, ENGINE_PATCH_VERSION, (GCompatibleWithEngineVersionChangelist | (ENGINE_IS_LICENSEE_VERSION << 31)), BRANCH_NAME);
diff --git a/Engine/Source/Runtime/Core/Private/Misc/Paths.cpp b/Engine/Source/Runtime/Core/Private/Misc/Paths.cpp
--- a/Engine/Source/Runtime/Core/Private/Misc/Paths.cpp
+++ b/Engine/Source/Runtime/Core/Private/Misc/Paths.cpp
@@ -38,7 +38,7 @@ FString FPaths::EngineUserDir()
 {
 	if (ShouldSaveToUserDir() || FApp::IsEngineInstalled())
 	{
-		return FPaths::Combine(FPlatformProcess::UserSettingsDir(), *FApp::GetEpicProductIdentifier(), *GEngineVersion.ToString(EVersionComponent::Minor)) + TEXT("/");
+		return FPaths::Combine(FPlatformProcess::UserSettingsDir(), *FApp::GetEpicProductIdentifier(), *FEngineVersion::Current().ToString(EVersionComponent::Minor)) + TEXT("/");
 	}
 	else
 	{
diff --git a/Engine/Source/Runtime/Core/Private/ProfilingDebugging/ProfilingHelpers.cpp b/Engine/Source/Runtime/Core/Private/ProfilingDebugging/ProfilingHelpers.cpp
--- a/Engine/Source/Runtime/Core/Private/ProfilingDebugging/ProfilingHelpers.cpp
+++ b/Engine/Source/Runtime/Core/Private/ProfilingDebugging/ProfilingHelpers.cpp
@@ -20,7 +20,7 @@ FGetMapNameDelegate GGetMapNameDelegate;
 
 /**
  * This will get the changelist that should be used with the Automated Performance Testing
- * If one is passed in we use that otherwise we use the GEngineVersion.GetChangelist().  This allows
+ * If one is passed in we use that otherwise we use the FEngineVersion::Current().GetChangelist().  This allows
  * us to have build machine built .exe and test them. 
  *
  * NOTE: had to use AutomatedBenchmarking as the parsing code is flawed and doesn't match
@@ -28,7 +28,7 @@ FGetMapNameDelegate GGetMapNameDelegate;
  **/
 int32 GetChangeListNumberForPerfTesting()
 {
-	int32 Retval = GEngineVersion.GetChangelist();
+	int32 Retval = FEngineVersion::Current().GetChangelist();
 
 	// if we have passed in the changelist to use then use it
 	int32 FromCommandLine = 0;
diff --git a/Engine/Source/Runtime/Core/Private/Serialization/Archive.cpp b/Engine/Source/Runtime/Core/Private/Serialization/Archive.cpp
--- a/Engine/Source/Runtime/Core/Private/Serialization/Archive.cpp
+++ b/Engine/Source/Runtime/Core/Private/Serialization/Archive.cpp
@@ -63,7 +63,7 @@ void FArchive::Reset()
 	ArNetVer							= GEngineNegotiationVersion;
 	ArUE4Ver							= GPackageFileUE4Version;
 	ArLicenseeUE4Ver					= GPackageFileLicenseeUE4Version;
-	ArEngineVer = GEngineVersion;
+	ArEngineVer = FEngineVersion::Current();
 	ArIsLoading							= false;
 	ArIsSaving							= false;
 	ArIsTransacting						= false;
diff --git a/Engine/Source/Runtime/Core/Private/Windows/WindowsPlatformCrashContext.cpp b/Engine/Source/Runtime/Core/Private/Windows/WindowsPlatformCrashContext.cpp
--- a/Engine/Source/Runtime/Core/Private/Windows/WindowsPlatformCrashContext.cpp
+++ b/Engine/Source/Runtime/Core/Private/Windows/WindowsPlatformCrashContext.cpp
@@ -195,7 +195,7 @@ void SetReportParameters( HREPORT ReportHandle, EXCEPTION_POINTERS* ExceptionInf
 	StringCchPrintf( StringBuffer, MAX_SPRINTF, TEXT( "!%s!AssertLog=\"%s\"" ), FCommandLine::GetOriginal(), LocalBuffer );
 	Result = WerReportSetParameter( ReportHandle, WER_P8, TEXT( "Commandline" ), StringBuffer );
 
-	StringCchPrintf( StringBuffer, MAX_SPRINTF, TEXT( "%s!%s!%s!%d" ), *FApp::GetBranchName(), FPlatformProcess::BaseDir(), FPlatformMisc::GetEngineMode(), GEngineVersion.GetChangelist() );
+	StringCchPrintf( StringBuffer, MAX_SPRINTF, TEXT( "%s!%s!%s!%d" ), *FApp::GetBranchName(), FPlatformProcess::BaseDir(), FPlatformMisc::GetEngineMode(), FEngineVersion::Current().GetChangelist() );
 	Result = WerReportSetParameter( ReportHandle, WER_P9, TEXT( "BranchBaseDir" ), StringBuffer );
 }
 
diff --git a/Engine/Source/Runtime/Core/Private/Windows/WindowsPlatformMisc.cpp b/Engine/Source/Runtime/Core/Private/Windows/WindowsPlatformMisc.cpp
--- a/Engine/Source/Runtime/Core/Private/Windows/WindowsPlatformMisc.cpp
+++ b/Engine/Source/Runtime/Core/Private/Windows/WindowsPlatformMisc.cpp
@@ -703,7 +703,7 @@ void FWindowsPlatformMisc::SubmitErrorReport( const TCHAR* InErrorHist, EErrorRe
 			TCHAR SystemTime[MAX_STRING_LEN];
 			FCString::Strncpy(SystemTime, *FDateTime::Now().ToString(), MAX_STRING_LEN);
 			TCHAR EngineVersionStr[MAX_STRING_LEN];
-			FCString::Strncpy(EngineVersionStr, *GEngineVersion.ToString(), 256 );
+			FCString::Strncpy(EngineVersionStr, *FEngineVersion::Current().ToString(), 256 );
 
 			TCHAR ChangelistVersionStr[MAX_STRING_LEN];
 			int32 ChangelistFromCommandLine = 0;
@@ -715,7 +715,7 @@ void FWindowsPlatformMisc::SubmitErrorReport( const TCHAR* InErrorHist, EErrorRe
 			// we are not passing in the changelist to use so use the one that was stored in the ObjectVersion
 			else
 			{
-				FCString::Strncpy(ChangelistVersionStr, *FString::FromInt(GEngineVersion.GetChangelist()), MAX_STRING_LEN);
+				FCString::Strncpy(ChangelistVersionStr, *FString::FromInt(FEngineVersion::Current().GetChangelist()), MAX_STRING_LEN);
 			}
 
 			TCHAR CmdLine[2048];
@@ -2074,7 +2074,7 @@ void FWindowsPlatformMisc::PromptForRemoteDebugging(bool bIsEnsure)
 			TEXT("Once he confirms he is connected to the machine,\n")
 			TEXT("hit YES to allow him to debug the crash.\n")
 			TEXT("[Changelist = %d]"),
-			GEngineVersion.GetChangelist());
+			FEngineVersion::Current().GetChangelist());
 		if (MessageBox(0, GErrorRemoteDebugPromptMessage, TEXT("CRASHED"), MB_YESNO|MB_SYSTEMMODAL) == IDYES)
 		{
 			::DebugBreak();
diff --git a/Engine/Source/Runtime/Core/Private/Windows/WindowsPlatformSplash.cpp b/Engine/Source/Runtime/Core/Private/Windows/WindowsPlatformSplash.cpp
--- a/Engine/Source/Runtime/Core/Private/Windows/WindowsPlatformSplash.cpp
+++ b/Engine/Source/Runtime/Core/Private/Windows/WindowsPlatformSplash.cpp
@@ -647,7 +647,7 @@ void FWindowsPlatformSplash::Show()
 
 				// Set version info
 				{
-					const FText Version = FText::FromString( GEngineVersion.ToString( FEngineBuildSettings::IsPerforceBuild() ? EVersionComponent::Branch : EVersionComponent::Patch ) );
+					const FText Version = FText::FromString( FEngineVersion::Current().ToString( FEngineBuildSettings::IsPerforceBuild() ? EVersionComponent::Branch : EVersionComponent::Patch ) );
 
 					FText VersionInfo;
 					FText AppName;
diff --git a/Engine/Source/Runtime/Core/Public/Misc/EngineVersion.h b/Engine/Source/Runtime/Core/Public/Misc/EngineVersion.h
--- a/Engine/Source/Runtime/Core/Public/Misc/EngineVersion.h
+++ b/Engine/Source/Runtime/Core/Public/Misc/EngineVersion.h
@@ -31,6 +31,12 @@ public:
 	/** Parses a version object from a string. Returns true on success. */
 	static bool Parse(const FString &Text, FEngineVersion &OutVersion);
 
+	/** Gets the current engine version */
+	static const FEngineVersion& Current();
+
+	/** Gets the earliest version which this engine maintains strict API and package compatibility with */
+	static const FEngineVersion& CompatibleWith();
+
 	/** Serialization function */
 	friend CORE_API void operator<<(class FArchive &Ar, FEngineVersion &Version);
 
@@ -44,6 +50,12 @@ private:
 
 	/** Branch name. */
 	FString Branch;
+
+	/** Global instance of the current engine version. */
+	static FEngineVersion CurrentVersion;
+
+	/** Earliest version which this engine maintains strict API and package compatibility with */
+	static FEngineVersion CompatibleWithVersion;
 };
 
 /** Version used for networking; the P4 changelist number. */
@@ -55,8 +67,3 @@ CORE_API extern const int32 GEngineMinNetVersion;
 /** Base protocol version to negotiate in network play. */
 CORE_API extern const int32 GEngineNegotiationVersion;		
 
-/** Global instance of the current engine version. */
-CORE_API extern const FEngineVersion GEngineVersion;
-
-/** Earliest version which this engine maintains strict API and package compatibility with */
-CORE_API extern const FEngineVersion GCompatibleWithEngineVersion;
diff --git a/Engine/Source/Runtime/CoreUObject/Private/UObject/LinkerLoad.cpp b/Engine/Source/Runtime/CoreUObject/Private/UObject/LinkerLoad.cpp
--- a/Engine/Source/Runtime/CoreUObject/Private/UObject/LinkerLoad.cpp
+++ b/Engine/Source/Runtime/CoreUObject/Private/UObject/LinkerLoad.cpp
@@ -909,7 +909,7 @@ FLinkerLoad::ELinkerStatus FLinkerLoad::CreateLoader()
 		// Set status info.
 		ArUE4Ver = GPackageFileUE4Version;
 		ArLicenseeUE4Ver = GPackageFileLicenseeUE4Version;
-		ArEngineVer = GEngineVersion;
+		ArEngineVer = FEngineVersion::Current();
 		ArIsLoading = true;
 		ArIsPersistent = true;
 
@@ -965,12 +965,12 @@ FLinkerLoad::ELinkerStatus FLinkerLoad::SerializePackageFileSummary()
 		}
 
 		// Don't load packages that are only compatible with an engine version newer than the current one.
-		if( !GEngineVersion.IsCompatibleWith(Summary.CompatibleWithEngineVersion) )
+		if( !FEngineVersion::Current().IsCompatibleWith(Summary.CompatibleWithEngineVersion) )
 		{
-			UE_LOG(LogLinker, Warning, TEXT("Asset '%s' has been saved with engine version newer than current and therefore can't be loaded. CurrEngineVersion: %s AssetEngineVersion: %s"), *Filename, *GEngineVersion.ToString(), *Summary.CompatibleWithEngineVersion.ToString() );
+			UE_LOG(LogLinker, Warning, TEXT("Asset '%s' has been saved with engine version newer than current and therefore can't be loaded. CurrEngineVersion: %s AssetEngineVersion: %s"), *Filename, *FEngineVersion::Current().ToString(), *Summary.CompatibleWithEngineVersion.ToString() );
 			return LINKER_Failed;
 		}
-		else if( !FPlatformProperties::RequiresCookedData() && !Summary.SavedByEngineVersion.HasChangelist() && GEngineVersion.HasChangelist() )
+		else if( !FPlatformProperties::RequiresCookedData() && !Summary.SavedByEngineVersion.HasChangelist() && FEngineVersion::Current().HasChangelist() )
 		{
 			// This warning can be disabled in ini with [Core.System] ZeroEngineVersionWarning=False
 			static struct FInitZeroEngineVersionWarning
diff --git a/Engine/Source/Runtime/CoreUObject/Private/UObject/LinkerSave.cpp b/Engine/Source/Runtime/CoreUObject/Private/UObject/LinkerSave.cpp
--- a/Engine/Source/Runtime/CoreUObject/Private/UObject/LinkerSave.cpp
+++ b/Engine/Source/Runtime/CoreUObject/Private/UObject/LinkerSave.cpp
@@ -27,8 +27,8 @@ FLinkerSave::FLinkerSave(UPackage* InParent, const TCHAR* InFilename, bool bForc
 		// Set main summary info.
 		Summary.Tag           = PACKAGE_FILE_TAG;
 		Summary.SetFileVersions( GPackageFileUE4Version, GPackageFileLicenseeUE4Version, bInSaveUnversioned );
-		Summary.SavedByEngineVersion = GEngineVersion;
-		Summary.CompatibleWithEngineVersion = GCompatibleWithEngineVersion;
+		Summary.SavedByEngineVersion = FEngineVersion::Current();
+		Summary.CompatibleWithEngineVersion = FEngineVersion::CompatibleWith();
 		Summary.PackageFlags = Package ? (Package->PackageFlags & ~PKG_NewlyCreated) : 0;
 
 		if (Package)
@@ -59,8 +59,8 @@ FLinkerSave::FLinkerSave(UPackage* InParent, bool bForceByteSwapping, bool bInSa
 		// Set main summary info.
 		Summary.Tag           = PACKAGE_FILE_TAG;
 		Summary.SetFileVersions( GPackageFileUE4Version, GPackageFileLicenseeUE4Version, bInSaveUnversioned );
-		Summary.SavedByEngineVersion = GEngineVersion;
-		Summary.CompatibleWithEngineVersion = GCompatibleWithEngineVersion;
+		Summary.SavedByEngineVersion = FEngineVersion::Current();
+		Summary.CompatibleWithEngineVersion = FEngineVersion::CompatibleWith();
 		Summary.PackageFlags = Package ? (Package->PackageFlags & ~PKG_NewlyCreated) : 0;
 
 		if (Package)
diff --git a/Engine/Source/Runtime/CoreUObject/Private/UObject/PackageFileSummary.cpp b/Engine/Source/Runtime/CoreUObject/Private/UObject/PackageFileSummary.cpp
--- a/Engine/Source/Runtime/CoreUObject/Private/UObject/PackageFileSummary.cpp
+++ b/Engine/Source/Runtime/CoreUObject/Private/UObject/PackageFileSummary.cpp
@@ -189,7 +189,7 @@ FArchive& operator<<( FArchive& Ar, FPackageFileSummary& Sum )
 
 		if( Sum.GetFileVersionUE4() >= VER_UE4_ENGINE_VERSION_OBJECT )
 		{
-			if(Ar.IsCooking() || (Ar.IsSaving() && !GEngineVersion.HasChangelist()))
+			if(Ar.IsCooking() || (Ar.IsSaving() && !FEngineVersion::Current().HasChangelist()))
 			{
 				FEngineVersion EmptyEngineVersion;
 				Ar << EmptyEngineVersion;
@@ -212,7 +212,7 @@ FArchive& operator<<( FArchive& Ar, FPackageFileSummary& Sum )
 
 		if (Sum.GetFileVersionUE4() >= VER_UE4_PACKAGE_SUMMARY_HAS_COMPATIBLE_ENGINE_VERSION )
 		{
-			if(Ar.IsCooking() || (Ar.IsSaving() && !GEngineVersion.HasChangelist()))
+			if(Ar.IsCooking() || (Ar.IsSaving() && !FEngineVersion::Current().HasChangelist()))
 			{
 				FEngineVersion EmptyEngineVersion;
 				Ar << EmptyEngineVersion;
diff --git a/Engine/Source/Runtime/Engine/Private/ChartCreation.cpp b/Engine/Source/Runtime/Engine/Private/ChartCreation.cpp
--- a/Engine/Source/Runtime/Engine/Private/ChartCreation.cpp
+++ b/Engine/Source/Runtime/Engine/Private/ChartCreation.cpp
@@ -697,7 +697,7 @@ void UEngine::DumpFPSChartToStatsLog( float TotalTime, float DeltaTime, int32 Nu
 
 	if( OutputFile )
 	{
-		OutputFile->Logf(TEXT("Dumping FPS chart at %s using build %s built from changelist %i"), *FDateTime::Now().ToString(), *GEngineVersion.ToString(), GetChangeListNumberForPerfTesting() );
+		OutputFile->Logf(TEXT("Dumping FPS chart at %s using build %s built from changelist %i"), *FDateTime::Now().ToString(), *FEngineVersion::Current().ToString(), GetChangeListNumberForPerfTesting() );
 
 		// Get OS info
 		FString OSMajor;
diff --git a/Engine/Source/Runtime/Engine/Private/CheatManager.cpp b/Engine/Source/Runtime/Engine/Private/CheatManager.cpp
--- a/Engine/Source/Runtime/Engine/Private/CheatManager.cpp
+++ b/Engine/Source/Runtime/Engine/Private/CheatManager.cpp
@@ -1168,14 +1168,14 @@ void UCheatManager::LogOutBugItGoToLogFile( const FString& InScreenShotDesc, con
 	const FString DescPlusExtension = FString::Printf( TEXT("%s%i.txt"), *InScreenShotDesc, GScreenshotBitmapIndex );
 	const FString TxtFileName = CreateProfileFilename( DescPlusExtension, false );
 
-	//FString::Printf( TEXT("BugIt%s-%s%05i"), *GEngineVersion.ToString(), *InScreenShotDesc, GScreenshotBitmapIndex+1 ) + TEXT( ".txt" );
+	//FString::Printf( TEXT("BugIt%s-%s%05i"), *FEngineVersion::Current().ToString(), *InScreenShotDesc, GScreenshotBitmapIndex+1 ) + TEXT( ".txt" );
 	const FString FullFileName = OutputDir + TxtFileName;
 
 	FOutputDeviceFile OutputFile(*FullFileName);
 	//FArchive* OutputFile = IFileManager::Get().CreateDebugFileWriter( *(FullFileName), FILEWRITE_Append );
 
 
-	OutputFile.Logf( TEXT("Dumping BugIt data chart at %s using build %s built from changelist %i"), *FDateTime::Now().ToString(), *GEngineVersion.ToString(), GetChangeListNumberForPerfTesting() );
+	OutputFile.Logf( TEXT("Dumping BugIt data chart at %s using build %s built from changelist %i"), *FDateTime::Now().ToString(), *FEngineVersion::Current().ToString(), GetChangeListNumberForPerfTesting() );
 
 	const FString MapNameStr = GetWorld()->GetMapName();
 
diff --git a/Engine/Source/Runtime/Engine/Private/EngineUtils.cpp b/Engine/Source/Runtime/Engine/Private/EngineUtils.cpp
--- a/Engine/Source/Runtime/Engine/Private/EngineUtils.cpp
+++ b/Engine/Source/Runtime/Engine/Private/EngineUtils.cpp
@@ -153,7 +153,7 @@ bool FContentComparisonHelper::CompareClasses(const FString& InBaseClassName, co
 
 	FString BaseCSVName = (
 		FString(TEXT("ContentComparison/")) + 
-		FString::Printf(TEXT("ContentCompare-%s/"), *GEngineVersion.ToString()) +
+		FString::Printf(TEXT("ContentCompare-%s/"), *FEngineVersion::Current().ToString()) +
 		FString::Printf(TEXT("%s"), *InBaseClassName)
 		);
 
@@ -170,7 +170,7 @@ bool FContentComparisonHelper::CompareClasses(const FString& InBaseClassName, co
 		}
 		BaseCSVName = (
 			FString(TEXT("ContentComparison/")) + 
-			FString::Printf(TEXT("ContentCompare-%s/"), *GEngineVersion.ToString()) +
+			FString::Printf(TEXT("ContentCompare-%s/"), *FEngineVersion::Current().ToString()) +
 			FString::Printf(TEXT("%s"), *EditedBaseClassName)
 			);
 	}
diff --git a/Engine/Source/Runtime/Engine/Private/GameplayStatics.cpp b/Engine/Source/Runtime/Engine/Private/GameplayStatics.cpp
--- a/Engine/Source/Runtime/Engine/Private/GameplayStatics.cpp
+++ b/Engine/Source/Runtime/Engine/Private/GameplayStatics.cpp
@@ -1242,7 +1242,7 @@ bool UGameplayStatics::SaveGameToSlot(USaveGame* SaveGameObject, const FString&
 		// Write out engine and UE4 version information
 		int32 PackageFileUE4Version = GPackageFileUE4Version;
 		MemoryWriter << PackageFileUE4Version;
-		FEngineVersion SavedEngineVersion = GEngineVersion;
+		FEngineVersion SavedEngineVersion = FEngineVersion::Current();
 		MemoryWriter << SavedEngineVersion;
 
 		// Write the class name so we know what class to load to
@@ -1313,7 +1313,7 @@ USaveGame* UGameplayStatics::LoadGameFromSlot(const FString& SlotName, const int
 				// you don't want to delete it, try uncommenting these lines and changing them to use the version 
 				// information from your previous build. Then load and resave your savegame file.
 				//MemoryReader.SetUE4Ver(MyPreviousUE4Version);				// @see GPackageFileUE4Version
-				//MemoryReader.SetEngineVer(MyPreviousEngineVersion);		// @see GEngineVersion
+				//MemoryReader.SetEngineVer(MyPreviousEngineVersion);		// @see FEngineVersion::Current()
 			}
 			else
 			{
diff --git a/Engine/Source/Runtime/Engine/Private/KismetSystemLibrary.cpp b/Engine/Source/Runtime/Engine/Private/KismetSystemLibrary.cpp
--- a/Engine/Source/Runtime/Engine/Private/KismetSystemLibrary.cpp
+++ b/Engine/Source/Runtime/Engine/Private/KismetSystemLibrary.cpp
@@ -71,7 +71,7 @@ FString UKismetSystemLibrary::GetClassDisplayName(UClass* Class)
 
 FString UKismetSystemLibrary::GetEngineVersion()
 {
-	return GEngineVersion.ToString();
+	return FEngineVersion::Current().ToString();
 }
 
 FString UKismetSystemLibrary::GetGameName()
diff --git a/Engine/Source/Runtime/Engine/Private/UnrealClient.cpp b/Engine/Source/Runtime/Engine/Private/UnrealClient.cpp
--- a/Engine/Source/Runtime/Engine/Private/UnrealClient.cpp
+++ b/Engine/Source/Runtime/Engine/Private/UnrealClient.cpp
@@ -694,7 +694,7 @@ FViewport::FViewport(FViewportClient* InViewportClient):
 	}
 #endif
 
-	AppVersionString = FString::Printf( TEXT( "Version: %s" ), *GEngineVersion.ToString() );
+	AppVersionString = FString::Printf( TEXT( "Version: %s" ), *FEngineVersion::Current().ToString() );
 
 	bIsPlayInEditorViewport = false;
 }
diff --git a/Engine/Source/Runtime/Engine/Private/UnrealEngine.cpp b/Engine/Source/Runtime/Engine/Private/UnrealEngine.cpp
--- a/Engine/Source/Runtime/Engine/Private/UnrealEngine.cpp
+++ b/Engine/Source/Runtime/Engine/Private/UnrealEngine.cpp
@@ -3060,7 +3060,7 @@ bool UEngine::HandleStopMovieCaptureCommand( const TCHAR* Cmd, FOutputDevice& Ar
 bool UEngine::HandleGameVerCommand( const TCHAR* Cmd, FOutputDevice& Ar )
 {
 	FString VersionString = FString::Printf(TEXT("GameVersion Branch: %s, Configuration: %s, Version: %s, CommandLine: %s"), 
-		*FApp::GetBranchName(), EBuildConfigurations::ToString(FApp::GetBuildConfiguration()), *GEngineVersion.ToString(), FCommandLine::Get());
+		*FApp::GetBranchName(), EBuildConfigurations::ToString(FApp::GetBuildConfiguration()), *FEngineVersion::Current().ToString(), FCommandLine::Get());
 
 	Ar.Logf( *VersionString );
 	FPlatformMisc::ClipboardCopy( *VersionString );
@@ -6893,7 +6893,7 @@ void UEngine::PerformanceCapture(const FString& CaptureName)
 	PathName = FPaths::AutomationDir() + PathName;
 	FPaths::MakePathRelativeTo(PathName,*FPaths::RootDir());
 	
-	FString ScreenshotName = FString::Printf(TEXT("%s%d.png"), *PathName, GEngineVersion.GetChangelist());
+	FString ScreenshotName = FString::Printf(TEXT("%s%d.png"), *PathName, FEngineVersion::Current().GetChangelist());
 	
 	const bool bShowUI = false;
 	const bool bAddFilenameSuffix = false;
diff --git a/Engine/Source/Runtime/Engine/Private/UnrealNetwork.cpp b/Engine/Source/Runtime/Engine/Private/UnrealNetwork.cpp
--- a/Engine/Source/Runtime/Engine/Private/UnrealNetwork.cpp
+++ b/Engine/Source/Runtime/Engine/Private/UnrealNetwork.cpp
@@ -46,7 +46,7 @@ uint32 FNetworkVersion::GetLocalNetworkVersion(bool AllowOverrideDelegate /*=tru
 	uint32 LocalNetworkVersion = FCrc::MemCrc32( &InternalProtocolVersion, sizeof( InternalProtocolVersion ), VersionHash );
 
 #if 0//!(UE_BUILD_SHIPPING || UE_BUILD_TEST)	// DISABLED FOR NOW, MESSES UP COPIED BUILDS
-	if ( !GEngineVersion.HasChangelist() )
+	if ( !FEngineVersion::Current().HasChangelist() )
 	{
 		// Further hash with machine id if this is a non promoted build
 		const FString MachineId = FPlatformMisc::GetMachineId().ToString( EGuidFormats::Digits ).ToLower();
@@ -76,7 +76,7 @@ bool FNetworkVersion::IsNetworkCompatible( const uint32 LocalNetworkVersion, con
 
 FNetworkReplayVersion FNetworkVersion::GetReplayVersion()
 {
-	return FNetworkReplayVersion( FApp::GetGameName(), GetLocalNetworkVersion(), GEngineVersion.GetChangelist() );
+	return FNetworkReplayVersion( FApp::GetGameName(), GetLocalNetworkVersion(), FEngineVersion::Current().GetChangelist() );
 }
 
 
diff --git a/Engine/Source/Runtime/Engine/Public/Tests/AutomationCommon.h b/Engine/Source/Runtime/Engine/Public/Tests/AutomationCommon.h
--- a/Engine/Source/Runtime/Engine/Public/Tests/AutomationCommon.h
+++ b/Engine/Source/Runtime/Engine/Public/Tests/AutomationCommon.h
@@ -76,7 +76,7 @@ namespace AutomationCommon
 
 		FPaths::MakePathRelativeTo(PathName, *FPaths::RootDir());
 
-		OutScreenshotName = FString::Printf(TEXT("%s/%d.png"), *PathName, GEngineVersion.GetChangelist());
+		OutScreenshotName = FString::Printf(TEXT("%s/%d.png"), *PathName, FEngineVersion::Current().GetChangelist());
 	}
 }
 
diff --git a/Engine/Source/Runtime/Launch/Private/Launch.cpp b/Engine/Source/Runtime/Launch/Private/Launch.cpp
--- a/Engine/Source/Runtime/Launch/Private/Launch.cpp
+++ b/Engine/Source/Runtime/Launch/Private/Launch.cpp
@@ -96,7 +96,7 @@ int32 GuardedMain( const TCHAR* CmdLine )
 	// These names will be updated as soon as the Filemanager is set up so we can write to the log file.
 	// That will also use the user folder for installed builds so we don't write into program files or whatever.
 #if PLATFORM_WINDOWS
-	FCString::Strcpy(MiniDumpFilenameW, *FString::Printf(TEXT("unreal-v%i-%s.dmp"), GEngineVersion.GetChangelist(), *FDateTime::Now().ToString()));
+	FCString::Strcpy(MiniDumpFilenameW, *FString::Printf(TEXT("unreal-v%i-%s.dmp"), FEngineVersion::Current().GetChangelist(), *FDateTime::Now().ToString()));
 
 	const TCHAR* OrgCmdLine = CmdLine;
 
diff --git a/Engine/Source/Runtime/Launch/Private/LaunchEngineLoop.cpp b/Engine/Source/Runtime/Launch/Private/LaunchEngineLoop.cpp
--- a/Engine/Source/Runtime/Launch/Private/LaunchEngineLoop.cpp
+++ b/Engine/Source/Runtime/Launch/Private/LaunchEngineLoop.cpp
@@ -2693,7 +2693,7 @@ bool FEngineLoop::AppInit( )
 	IFileManager::Get().MakeDirectory( *FPaths::GameLogDir() );
 
 	// update the mini dump filename now that we have enough info to point it to the log folder even in installed builds
-	FCString::Strcpy(MiniDumpFilenameW, *IFileManager::Get().ConvertToAbsolutePathForExternalAppForWrite(*FString::Printf(TEXT("%sunreal-v%i-%s.dmp"), *FPaths::GameLogDir(), GEngineVersion.GetChangelist(), *FDateTime::Now().ToString())));
+	FCString::Strcpy(MiniDumpFilenameW, *IFileManager::Get().ConvertToAbsolutePathForExternalAppForWrite(*FString::Printf(TEXT("%sunreal-v%i-%s.dmp"), *FPaths::GameLogDir(), FEngineVersion::Current().GetChangelist(), *FDateTime::Now().ToString())));
 #endif
 
 	// Init logging to disk
@@ -2833,8 +2833,8 @@ bool FEngineLoop::AppInit( )
 #endif // !UE_BUILD_SHIPPING
 
 	//// Command line.
-	UE_LOG(LogInit, Log, TEXT("Version: %s"), *GEngineVersion.ToString());
-	UE_LOG(LogInit, Log, TEXT("API Version: %u"), GCompatibleWithEngineVersion.GetChangelist());
+	UE_LOG(LogInit, Log, TEXT("Version: %s"), *FEngineVersion::Current().ToString());
+	UE_LOG(LogInit, Log, TEXT("API Version: %u"), FEngineVersion::CompatibleWith().GetChangelist());
 
 #if PLATFORM_64BITS
 	UE_LOG(LogInit, Log, TEXT("Compiled (64-bit): %s %s"), ANSI_TO_TCHAR(__DATE__), ANSI_TO_TCHAR(__TIME__));
@@ -2918,7 +2918,7 @@ bool FEngineLoop::AppInit( )
 	check(sizeof(bool) == 1);
 	check(sizeof(float) == 4);
 	check(sizeof(double) == 8);
-	check(GEngineNetVersion == 0 || GEngineNetVersion >= GEngineMinNetVersion || GEngineVersion.IsLicenseeVersion());
+	check(GEngineNetVersion == 0 || GEngineNetVersion >= GEngineMinNetVersion || FEngineVersion::Current().IsLicenseeVersion());
 
 	// Init list of common colors.
 	GColorList.CreateColorMap();
diff --git a/Engine/Source/Runtime/Launch/Private/Linux/LaunchLinux.cpp b/Engine/Source/Runtime/Launch/Private/Linux/LaunchLinux.cpp
--- a/Engine/Source/Runtime/Launch/Private/Linux/LaunchLinux.cpp
+++ b/Engine/Source/Runtime/Launch/Private/Linux/LaunchLinux.cpp
@@ -135,7 +135,7 @@ int main(int argc, char *argv[])
 
 #if UE_BUILD_SHIPPING
 	// only printed in shipping
-	printf("%s %d %d %d %d\n", StringCast<ANSICHAR>(*GEngineVersion.ToString()).Get(), GEngineMinNetVersion, GEngineNegotiationVersion, GPackageFileUE4Version, GPackageFileLicenseeUE4Version);
+	printf("%s %d %d %d %d\n", StringCast<ANSICHAR>(*FEngineVersion::Current().ToString()).Get(), GEngineMinNetVersion, GEngineNegotiationVersion, GPackageFileUE4Version, GPackageFileLicenseeUE4Version);
 #endif // UE_BUILD_SHIPPING
 	int ErrorLevel = 0;
 
diff --git a/Engine/Source/Runtime/Launch/Private/WinRT/WinRTLaunch.cpp b/Engine/Source/Runtime/Launch/Private/WinRT/WinRTLaunch.cpp
--- a/Engine/Source/Runtime/Launch/Private/WinRT/WinRTLaunch.cpp
+++ b/Engine/Source/Runtime/Launch/Private/WinRT/WinRTLaunch.cpp
@@ -636,7 +636,7 @@ int32 GuardedMain( const TCHAR* CmdLine, HINSTANCE hInInstance, HINSTANCE hPrevI
 	// destruction and main uses SEH.
 	// These names will be updated as soon as the Filemanager is set up so we can write to the log file.
 	// That will also use the user folder for installed builds so we don't write into program files or whatever.
-	FCString::Strcpy(MiniDumpFilenameW, *FString::Printf( TEXT("unreal-v%i-%s.dmp"), GEngineVersion.GetChangelist(), *FDateTime::Now().ToString()));
+	FCString::Strcpy(MiniDumpFilenameW, *FString::Printf( TEXT("unreal-v%i-%s.dmp"), FEngineVersion::Current().GetChangelist(), *FDateTime::Now().ToString()));
 	CmdLine = FCommandLine::RemoveExeName(CmdLine);
 
 	// Call PreInit and exit if failed.
diff --git a/Engine/Source/Runtime/NetworkReplayStreaming/NullNetworkReplayStreaming/Private/NullNetworkReplayStreaming.cpp b/Engine/Source/Runtime/NetworkReplayStreaming/NullNetworkReplayStreaming/Private/NullNetworkReplayStreaming.cpp
--- a/Engine/Source/Runtime/NetworkReplayStreaming/NullNetworkReplayStreaming/Private/NullNetworkReplayStreaming.cpp
+++ b/Engine/Source/Runtime/NetworkReplayStreaming/NullNetworkReplayStreaming/Private/NullNetworkReplayStreaming.cpp
@@ -43,7 +43,7 @@ static FString GetStreamBaseFilename(const FString& StreamName)
 	DemoName.ReplaceInline( TEXT( "%td" ), *FDateTime::Now().ToString() );
 	DemoName.ReplaceInline( TEXT( "%d" ), *FString::Printf( TEXT( "%i-%i-%i" ), Month, Day, Year ) );
 	DemoName.ReplaceInline( TEXT( "%t" ), *FString::Printf( TEXT( "%i" ), ( ( Hour * 3600 ) + ( Min * 60 ) + Sec ) * 1000 + MSec ) );
-	DemoName.ReplaceInline( TEXT( "%v" ), *FString::Printf( TEXT( "%i" ), GEngineVersion.GetChangelist() ) );
+	DemoName.ReplaceInline( TEXT( "%v" ), *FString::Printf( TEXT( "%i" ), FEngineVersion::Current().GetChangelist() ) );
 
 	// replace bad characters with underscores
 	DemoName.ReplaceInline( TEXT( "\\" ),	TEXT( "_" ) );
diff --git a/Engine/Source/Runtime/Online/HTTP/Private/Curl/CurlHttp.cpp b/Engine/Source/Runtime/Online/HTTP/Private/Curl/CurlHttp.cpp
--- a/Engine/Source/Runtime/Online/HTTP/Private/Curl/CurlHttp.cpp
+++ b/Engine/Source/Runtime/Online/HTTP/Private/Curl/CurlHttp.cpp
@@ -540,7 +540,7 @@ bool FCurlHttpRequest::StartRequest()
 	// set up headers
 	if (GetHeader("User-Agent").IsEmpty())
 	{
-		SetHeader(TEXT("User-Agent"), FString::Printf(TEXT("game=%s, engine=UE4, version=%s"), FApp::GetGameName(), *GEngineVersion.ToString()));
+		SetHeader(TEXT("User-Agent"), FString::Printf(TEXT("game=%s, engine=UE4, version=%s"), FApp::GetGameName(), *FEngineVersion::Current().ToString()));
 	}
 
 	// content-length should be present http://www.w3.org/Protocols/rfc2616/rfc2616-sec4.html#sec4.4
diff --git a/Engine/Source/Runtime/Online/HTTP/Private/HTML5/HTML5HTTP.cpp b/Engine/Source/Runtime/Online/HTTP/Private/HTML5/HTML5HTTP.cpp
--- a/Engine/Source/Runtime/Online/HTTP/Private/HTML5/HTML5HTTP.cpp
+++ b/Engine/Source/Runtime/Online/HTTP/Private/HTML5/HTML5HTTP.cpp
@@ -295,7 +295,7 @@ bool FHTML5HttpRequest::StartRequest()
 	//"User-Agent" && "Content-Length" are automatically set by the browser xhr request. We can't do much.
 
 	// make a fake header, so server has some idea this is UE
-    SetHeader(TEXT("X-UnrealEngine-Agent"), FString::Printf(TEXT("game=%s, engine=UE4, version=%s"), FApp::GetGameName(), *GEngineVersion.ToString()));
+    SetHeader(TEXT("X-UnrealEngine-Agent"), FString::Printf(TEXT("game=%s, engine=UE4, version=%s"), FApp::GetGameName(), *FEngineVersion::Current().ToString()));
 
 	// Add "Pragma: no-cache" to mimic WinInet behavior
 	if (GetHeader("Pragma").IsEmpty())
diff --git a/Engine/Source/Runtime/Online/HTTP/Private/IOS/IOSHTTP.cpp b/Engine/Source/Runtime/Online/HTTP/Private/IOS/IOSHTTP.cpp
--- a/Engine/Source/Runtime/Online/HTTP/Private/IOS/IOSHTTP.cpp
+++ b/Engine/Source/Runtime/Online/HTTP/Private/IOS/IOSHTTP.cpp
@@ -221,7 +221,7 @@ bool FIOSHttpRequest::StartRequest()
 		[Request setValue:[NSString stringWithFormat:@"%d", GetContentLength()] forHTTPHeaderField:@"Content-Length"];
 	}
 
-	NSString* Tag = [NSString stringWithFString:FString::Printf(TEXT("UE4-%s,UE4Ver(%s)"), FApp::GetGameName(), *GEngineVersion.ToString())];
+	NSString* Tag = [NSString stringWithFString:FString::Printf(TEXT("UE4-%s,UE4Ver(%s)"), FApp::GetGameName(), *FEngineVersion::Current().ToString())];
 	[Request addValue:Tag forHTTPHeaderField:@"User-Agent"];
 
 	Response = MakeShareable( new FIOSHttpResponse( *this ) );
diff --git a/Engine/Source/Runtime/Online/HTTP/Private/Mac/MacHTTP.cpp b/Engine/Source/Runtime/Online/HTTP/Private/Mac/MacHTTP.cpp
--- a/Engine/Source/Runtime/Online/HTTP/Private/Mac/MacHTTP.cpp
+++ b/Engine/Source/Runtime/Online/HTTP/Private/Mac/MacHTTP.cpp
@@ -220,7 +220,7 @@ bool FMacHttpRequest::StartRequest()
 	const FString UserAgent = GetHeader("User-Agent");
 	if(UserAgent.IsEmpty())
 	{
-		NSString* Tag = FString::Printf(TEXT("UE4-%s,UE4Ver(%s)"), FApp::GetGameName(), *GEngineVersion.ToString()).GetNSString();
+		NSString* Tag = FString::Printf(TEXT("UE4-%s,UE4Ver(%s)"), FApp::GetGameName(), *FEngineVersion::Current().ToString()).GetNSString();
 		[Request addValue:Tag forHTTPHeaderField:@"User-Agent"];
 	}
 	else
