UE4 - End of frame updates now optionally can do the gamethread updates while we are doing parallel updates. The relevant cvars are AllowAsyncRenderThreadUpdates and AllowAsyncRenderThreadUpdatesDuringGamethreadUpdates. The latter is disabled by default.

diff --git a/Engine/Source/Runtime/Core/Public/Async/ParallelFor.h b/Engine/Source/Runtime/Core/Public/Async/ParallelFor.h
--- a/Engine/Source/Runtime/Core/Public/Async/ParallelFor.h
+++ b/Engine/Source/Runtime/Core/Public/Async/ParallelFor.h
@@ -14,16 +14,28 @@ struct FParallelForData
 {
 	int32 Num;
 	TFunctionRef<void(int32)> Body;
-	FScopedEvent& DoneEvent;
+	FEvent* Event;
 	FThreadSafeCounter IndexToDo;
 	FThreadSafeCounter NumCompleted;
-	FParallelForData(int32 InNum, TFunctionRef<void(int32)> InBody, FScopedEvent& InDoneEvent)
+	bool bExited;
+	bool bTriggered;
+	FParallelForData(int32 InNum, TFunctionRef<void(int32)> InBody)
 		: Num(InNum)
 		, Body(InBody)
-		, DoneEvent(InDoneEvent)
+		, Event(FPlatformProcess::GetSynchEventFromPool(false))
+		, bExited(false)
+		, bTriggered(false)
 	{
 	}
-	void Process(int32 TasksToSpawn, TSharedRef<FParallelForData, ESPMode::ThreadSafe>& Data);
+	~FParallelForData()
+	{
+		FPlatformMisc::MemoryBarrier();
+		check(IndexToDo.GetValue() >= Num);
+		check(NumCompleted.GetValue() == Num);
+		check(bExited);
+		FPlatformProcess::ReturnSynchEventToPool(Event);
+	}
+	bool Process(int32 TasksToSpawn, TSharedRef<FParallelForData, ESPMode::ThreadSafe>& Data);
 };
 
 class FParallelForTask
@@ -50,11 +62,16 @@ public:
 	}
 	void DoTask(ENamedThreads::Type CurrentThread, const FGraphEventRef& MyCompletionGraphEvent)
 	{
-		Data->Process(TasksToSpawn, Data);
+		if (Data->Process(TasksToSpawn, Data))
+		{
+			check(!Data->bTriggered);
+			Data->bTriggered = true;
+			Data->Event->Trigger();
+		}
 	}
 };
 
-inline void FParallelForData::Process(int32 TasksToSpawn, TSharedRef<FParallelForData, ESPMode::ThreadSafe>& Data)
+inline bool FParallelForData::Process(int32 TasksToSpawn, TSharedRef<FParallelForData, ESPMode::ThreadSafe>& Data)
 {
 	int32 MaybeTasksLeft = Num - IndexToDo.GetValue();
 	if (TasksToSpawn && MaybeTasksLeft > 0)
@@ -68,19 +85,24 @@ inline void FParallelForData::Process(int32 TasksToSpawn, TSharedRef<FParallelFo
 		if (MyIndex < Num)
 		{
 			Body(MyIndex);
-			if (NumCompleted.Increment() == Num)
+			check(!bExited);
+			int32 LocalNumCompleted = NumCompleted.Increment();
+			if (LocalNumCompleted == Num)
 			{
-				DoneEvent.Trigger(); // I was the last one; let parallelfor exit
+				FPlatformMisc::MemoryBarrier();
+				check(IndexToDo.GetValue() >= Num);
+				return true;
 			}
+			check(LocalNumCompleted < Num);
 		}
 		else
 		{
 			break;
 		}
 	}
+	return false;
 }
 
-
 /** 
 	*	General purpose parallel for that uses the taskgraph
 	*	@param Num; number of calls of Body; Body(0), Body(1)....Body(Num - 1)
@@ -107,12 +129,24 @@ inline void ParallelFor(int32 Num, TFunctionRef<void(int32)> Body, bool bForceSi
 		}
 		return;
 	}
-
-	FScopedEvent DoneEvent;
-	TSharedRef<FParallelForData, ESPMode::ThreadSafe> Data = MakeShareable(new FParallelForData(Num, Body, DoneEvent));
+	check(Num);
+	FParallelForData* DataPtr = new FParallelForData(Num, Body);
+	TSharedRef<FParallelForData, ESPMode::ThreadSafe> Data = MakeShareable(DataPtr);
 	TGraphTask<FParallelForTask>::CreateTask().ConstructAndDispatchWhenReady(Data, AnyThreadTasks - 1);		
 	// this thread can help too and this is important to prevent deadlock on recursion 
-	Data->Process(0, Data);
+	if (!Data->Process(0, Data))
+	{
+		FPlatformMisc::MemoryBarrier();
+		Data->Event->Wait();
+		check(Data->bTriggered);
+	}
+	else
+	{
+		FPlatformMisc::MemoryBarrier();
+		check(!Data->bTriggered);
+	}
+	check(Data->NumCompleted.GetValue() == Num);
+	Data->bExited = true;
 	// DoneEvent waits here if some other thread finishes the last item
 	// Data must live on until all of the tasks are cleared which might be long after this function exits
 }
@@ -128,12 +162,11 @@ inline void ParallelFor(int32 Num, TFunctionRef<void(int32)> Body, bool bForceSi
 inline void ParallelForWithPreWork(int32 Num, TFunctionRef<void(int32)> Body, TFunctionRef<void()> CurrentThreadWorkToDoBeforeHelping, bool bForceSingleThread = false)
 {
 	SCOPE_CYCLE_COUNTER(STAT_ParallelFor);
-	check(Num >= 0);
 
 	int32 AnyThreadTasks = 0;
-	if (Num > 1 && !bForceSingleThread && FApp::ShouldUseThreadingForPerformance())
+	if (!bForceSingleThread && FApp::ShouldUseThreadingForPerformance())
 	{
-		AnyThreadTasks = FMath::Min<int32>(FTaskGraphInterface::Get().GetNumWorkerThreads(), Num - 1);
+		AnyThreadTasks = FMath::Min<int32>(FTaskGraphInterface::Get().GetNumWorkerThreads(), Num);
 	}
 	if (!AnyThreadTasks)
 	{
@@ -146,15 +179,26 @@ inline void ParallelForWithPreWork(int32 Num, TFunctionRef<void(int32)> Body, TF
 		}
 		return;
 	}
-
-	FScopedEvent DoneEvent;
-	TSharedRef<FParallelForData, ESPMode::ThreadSafe> Data = MakeShareable(new FParallelForData(Num, Body, DoneEvent));
+	check(Num);
+	FParallelForData* DataPtr = new FParallelForData(Num, Body);
+	TSharedRef<FParallelForData, ESPMode::ThreadSafe> Data = MakeShareable(DataPtr);
 	TGraphTask<FParallelForTask>::CreateTask().ConstructAndDispatchWhenReady(Data, AnyThreadTasks - 1);		
 	// do the prework
 	CurrentThreadWorkToDoBeforeHelping();
 	// this thread can help too and this is important to prevent deadlock on recursion 
-	Data->Process(0, Data);
-	// DoneEvent waits here if some other thread finishes the last item
+	if (!Data->Process(0, Data))
+	{
+		FPlatformMisc::MemoryBarrier();
+		Data->Event->Wait();
+		check(Data->bTriggered);
+	}
+	else
+	{
+		FPlatformMisc::MemoryBarrier();
+		check(!Data->bTriggered);
+	}
+	check(Data->NumCompleted.GetValue() == Num);
+	Data->bExited = true;
 	// Data must live on until all of the tasks are cleared which might be long after this function exits
 }
 
diff --git a/Engine/Source/Runtime/Engine/Private/LevelTick.cpp b/Engine/Source/Runtime/Engine/Private/LevelTick.cpp
--- a/Engine/Source/Runtime/Engine/Private/LevelTick.cpp
+++ b/Engine/Source/Runtime/Engine/Private/LevelTick.cpp
@@ -706,6 +706,11 @@ static TAutoConsoleVariable<int32> CVarAllowAsyncRenderThreadUpdates(
 	1,
 	TEXT("Used to control async renderthread updates. Also gated on FApp::ShouldUseThreadingForPerformance()."));
 
+static TAutoConsoleVariable<int32> CVarAllowAsyncRenderThreadUpdatesDuringGamethreadUpdates(
+	TEXT("AllowAsyncRenderThreadUpdatesDuringGamethreadUpdates"),
+	0,
+	TEXT("If > 0 then we do the gamethread updates _while_ doing parallel updates."));
+
 static TAutoConsoleVariable<int32> CVarAllowAsyncRenderThreadUpdatesEditor(
 	TEXT("AllowAsyncRenderThreadUpdatesEditor"),
 	0,
@@ -800,47 +805,57 @@ void UWorld::SendAllEndOfFrameUpdates()
 	// update all dirty components. 
 	TGuardValue<bool> GuardIsFlushedGlobal( bPostTickComponentUpdate, true ); 
 
-	// Game thread updates need to happen before we go wide on the other threads.
-	// These updates are things that have said that they are NOT SAFE to run concurrently.
-	for (TArray<TWeakObjectPtr<UActorComponent> >::TIterator It(ComponentsThatNeedEndOfFrameUpdate_OnGameThread); It; ++It)
-	{
-		UActorComponent* Component = It->Get();
-		if (Component)
+	TArray<TWeakObjectPtr<class UActorComponent> >& LocalComponentsThatNeedEndOfFrameUpdate = ComponentsThatNeedEndOfFrameUpdate;
+	TArray<TWeakObjectPtr<class UActorComponent> >& LocalComponentsThatNeedEndOfFrameUpdate_OnGameThread = ComponentsThatNeedEndOfFrameUpdate_OnGameThread;
+
+	TFunctionRef<void(int32)> ParallelWork = 
+		[&LocalComponentsThatNeedEndOfFrameUpdate](int32 Index)
 		{
-			if ( !Component->IsPendingKill() && Component->IsRegistered() && !Component->IsTemplate())
+			UActorComponent* NextComponent = LocalComponentsThatNeedEndOfFrameUpdate[Index].Get();
+			if (NextComponent)
 			{
-				FScopeCycleCounterUObject ComponentScope(Component);
-				FScopeCycleCounterUObject AdditionalScope(STATS ? Component->AdditionalStatObject() : NULL);
-				Component->DoDeferredRenderUpdates_Concurrent();
+				if (!NextComponent->IsPendingKill() && NextComponent->IsRegistered() && !NextComponent->IsTemplate())
+				{
+					FScopeCycleCounterUObject ComponentScope(NextComponent);
+					FScopeCycleCounterUObject AdditionalScope(STATS ? NextComponent->AdditionalStatObject() : NULL);
+					NextComponent->DoDeferredRenderUpdates_Concurrent();
+				}
+				check(NextComponent->GetMarkedForEndOfFrameUpdateState() == EComponentMarkedForEndOfFrameUpdateState::Marked);
+				FMarkComponentEndOfFrameUpdateState::Set(NextComponent, EComponentMarkedForEndOfFrameUpdateState::Unmarked);
 			}
-			check(Component->GetMarkedForEndOfFrameUpdateState() == EComponentMarkedForEndOfFrameUpdateState::MarkedForGameThread);
-			FMarkComponentEndOfFrameUpdateState::Set(Component, EComponentMarkedForEndOfFrameUpdateState::Unmarked);
-		}
-	}
-	ComponentsThatNeedEndOfFrameUpdate_OnGameThread.Reset();
-
-	if (ComponentsThatNeedEndOfFrameUpdate.Num())
-	{
-		TArray<TWeakObjectPtr<class UActorComponent> >& LocalComponentsThatNeedEndOfFrameUpdate = ComponentsThatNeedEndOfFrameUpdate;
-		ParallelFor(ComponentsThatNeedEndOfFrameUpdate.Num(), 
-			[&LocalComponentsThatNeedEndOfFrameUpdate](int32 Index)
+		};
+	TFunctionRef<void()> GTWork = 
+		[&LocalComponentsThatNeedEndOfFrameUpdate_OnGameThread]()
+		{
+			QUICK_SCOPE_CYCLE_COUNTER(STAT_PostTickComponentUpdate_ForcedGameThread);
+			for (TArray<TWeakObjectPtr<UActorComponent> >::TIterator It(LocalComponentsThatNeedEndOfFrameUpdate_OnGameThread); It; ++It)
 			{
-				UActorComponent* NextComponent = LocalComponentsThatNeedEndOfFrameUpdate[Index].Get();
-				if (NextComponent)
+				UActorComponent* Component = It->Get();
+				if (Component)
 				{
-					if (!NextComponent->IsPendingKill() && NextComponent->IsRegistered() && !NextComponent->IsTemplate())
+					if ( !Component->IsPendingKill() && Component->IsRegistered() && !Component->IsTemplate())
 					{
-						FScopeCycleCounterUObject ComponentScope(NextComponent);
-						FScopeCycleCounterUObject AdditionalScope(STATS ? NextComponent->AdditionalStatObject() : NULL);
-						NextComponent->DoDeferredRenderUpdates_Concurrent();
+						FScopeCycleCounterUObject ComponentScope(Component);
+						FScopeCycleCounterUObject AdditionalScope(STATS ? Component->AdditionalStatObject() : NULL);
+						Component->DoDeferredRenderUpdates_Concurrent();
 					}
-					check(NextComponent->GetMarkedForEndOfFrameUpdateState() == EComponentMarkedForEndOfFrameUpdateState::Marked);
-					FMarkComponentEndOfFrameUpdateState::Set(NextComponent, EComponentMarkedForEndOfFrameUpdateState::Unmarked);
+					check(Component->GetMarkedForEndOfFrameUpdateState() == EComponentMarkedForEndOfFrameUpdateState::MarkedForGameThread);
+					FMarkComponentEndOfFrameUpdateState::Set(Component, EComponentMarkedForEndOfFrameUpdateState::Unmarked);
 				}
 			}
-		);
-		ComponentsThatNeedEndOfFrameUpdate.Reset();
+			LocalComponentsThatNeedEndOfFrameUpdate_OnGameThread.Reset();
+		};
+
+	if (CVarAllowAsyncRenderThreadUpdatesDuringGamethreadUpdates.GetValueOnGameThread() > 0)
+	{
+		ParallelForWithPreWork(ComponentsThatNeedEndOfFrameUpdate.Num(), ParallelWork, GTWork);
+	}
+	else
+	{
+		GTWork();
+		ParallelFor(ComponentsThatNeedEndOfFrameUpdate.Num(), ParallelWork);
 	}
+	ComponentsThatNeedEndOfFrameUpdate.Reset();
 }
 
 
