Allow setting the build changelist at runtime in modular builds, through the version manifests.

diff --git a/Engine/Programs/UnrealGameSync/UnrealGameSync.ini b/Engine/Programs/UnrealGameSync/UnrealGameSync.ini
--- a/Engine/Programs/UnrealGameSync/UnrealGameSync.ini
+++ b/Engine/Programs/UnrealGameSync/UnrealGameSync.ini
@@ -1,2 +1,2 @@
 [Options]
-UseFastModularVersioning=true
+UseFastModularVersioning=1
diff --git a/Engine/Source/Runtime/Core/Private/Misc/EngineVersion.cpp b/Engine/Source/Runtime/Core/Private/Misc/EngineVersion.cpp
--- a/Engine/Source/Runtime/Core/Private/Misc/EngineVersion.cpp
+++ b/Engine/Source/Runtime/Core/Private/Misc/EngineVersion.cpp
@@ -6,6 +6,18 @@
 #include "EngineBuildSettings.h"
 #include "ModuleVersion.h"
 
+/** Version numbers for networking */
+int32 GEngineNetVersion			= BUILT_FROM_CHANGELIST;
+const int32 GEngineMinNetVersion		= 7038;
+const int32 GEngineNegotiationVersion	= 3077;
+
+// Global instance of the current engine version
+FEngineVersion FEngineVersion::CurrentVersion(ENGINE_MAJOR_VERSION, ENGINE_MINOR_VERSION, ENGINE_PATCH_VERSION, (BUILT_FROM_CHANGELIST | (ENGINE_IS_LICENSEE_VERSION << 31)), BRANCH_NAME);
+
+// Version which this engine maintains strict API and package compatibility with
+FEngineVersion FEngineVersion::CompatibleWithVersion(ENGINE_MAJOR_VERSION, ENGINE_MINOR_VERSION, ENGINE_PATCH_VERSION, (MODULE_API_VERSION | (ENGINE_IS_LICENSEE_VERSION << 31)), BRANCH_NAME);
+
+
 FEngineVersionBase::FEngineVersionBase()
 : Major(0)
 , Minor(0)
@@ -193,6 +205,19 @@ const FEngineVersion& FEngineVersion::CompatibleWith()
 	return CompatibleWithVersion;
 }
 
+bool FEngineVersion::OverrideCurrentVersionChangelist(int32 NewChangelist)
+{
+	if(CurrentVersion.GetChangelist() != 0 || CompatibleWithVersion.GetChangelist() != 0)
+	{
+		return false;
+	}
+
+	CurrentVersion.Set(CurrentVersion.Major, CurrentVersion.Minor, CurrentVersion.Patch, NewChangelist | (ENGINE_IS_LICENSEE_VERSION << 31), CurrentVersion.Branch);
+	CompatibleWithVersion.Set(CompatibleWithVersion.Major, CompatibleWithVersion.Minor, CompatibleWithVersion.Patch, NewChangelist | (ENGINE_IS_LICENSEE_VERSION << 31), CompatibleWithVersion.Branch);
+	GEngineNetVersion = NewChangelist;
+	return true;
+}
+
 void operator<<(FArchive &Ar, FEngineVersion &Version)
 {
 	Ar << Version.Major;
@@ -201,18 +226,3 @@ void operator<<(FArchive &Ar, FEngineVersion &Version)
 	Ar << Version.Changelist;
 	Ar << Version.Branch;
 }
-
-/** Changelist constants. The two lines below may be modified by tools. */
-const int32 GEngineVersionChangelist = BUILT_FROM_CHANGELIST;
-const int32 GCompatibleWithEngineVersionChangelist = MODULE_API_VERSION;
-
-/** Version numbers for networking */
-const int32 GEngineNetVersion			= GEngineVersionChangelist;
-const int32 GEngineMinNetVersion		= 7038;
-const int32 GEngineNegotiationVersion	= 3077;
-
-// Global instance of the current engine version
-FEngineVersion FEngineVersion::CurrentVersion(ENGINE_MAJOR_VERSION, ENGINE_MINOR_VERSION, ENGINE_PATCH_VERSION, (GEngineVersionChangelist | (ENGINE_IS_LICENSEE_VERSION << 31)), BRANCH_NAME);
-
-// Version which this engine maintains strict API and package compatibility with
-FEngineVersion FEngineVersion::CompatibleWithVersion(ENGINE_MAJOR_VERSION, ENGINE_MINOR_VERSION, ENGINE_PATCH_VERSION, (GCompatibleWithEngineVersionChangelist | (ENGINE_IS_LICENSEE_VERSION << 31)), BRANCH_NAME);
diff --git a/Engine/Source/Runtime/Core/Public/Misc/EngineVersion.h b/Engine/Source/Runtime/Core/Public/Misc/EngineVersion.h
--- a/Engine/Source/Runtime/Core/Public/Misc/EngineVersion.h
+++ b/Engine/Source/Runtime/Core/Public/Misc/EngineVersion.h
@@ -37,6 +37,9 @@ public:
 	/** Gets the earliest version which this engine maintains strict API and package compatibility with */
 	static const FEngineVersion& CompatibleWith();
 
+	/** Overrides the current changelist in the verison */
+	static bool OverrideCurrentVersionChangelist(int32 NewChangelist);
+
 	/** Serialization function */
 	friend CORE_API void operator<<(class FArchive &Ar, FEngineVersion &Version);
 
@@ -59,7 +62,7 @@ private:
 };
 
 /** Version used for networking; the P4 changelist number. */
-CORE_API extern const int32 GEngineNetVersion;
+CORE_API extern int32 GEngineNetVersion;
 
 /** Earliest engine build that is network compatible with this one. */
 CORE_API extern const int32 GEngineMinNetVersion;
diff --git a/Engine/Source/Runtime/Launch/Private/LaunchEngineLoop.cpp b/Engine/Source/Runtime/Launch/Private/LaunchEngineLoop.cpp
--- a/Engine/Source/Runtime/Launch/Private/LaunchEngineLoop.cpp
+++ b/Engine/Source/Runtime/Launch/Private/LaunchEngineLoop.cpp
@@ -1052,6 +1052,10 @@ int32 FEngineLoop::PreInit( const TCHAR* CmdLine )
 	if(ModuleEnumerator.RegisterWithModuleManager())
 	{
 		const FVersionManifest& Manifest = ModuleEnumerator.GetInitialManifest();
+		if(Manifest.Changelist != 0 && !FEngineVersion::OverrideCurrentVersionChangelist(Manifest.Changelist))
+		{
+			UE_LOG(LogInit, Fatal, TEXT("Couldn't update engine changelist to %d."), Manifest.Changelist);
+		}
 		UE_LOG(LogInit, Log, TEXT("Using version manifest at CL %d with build ID '%s'"), Manifest.Changelist, *Manifest.BuildId);
 	}
 #endif
