Linux: improvements and modal dialog fix on Gnome WM (UE-18767).

- PR #1388 contributed by [REDACTED_PERSON].

diff --git a/Engine/Source/Runtime/Core/Private/Linux/LinuxApplication.cpp b/Engine/Source/Runtime/Core/Private/Linux/LinuxApplication.cpp
--- a/Engine/Source/Runtime/Core/Private/Linux/LinuxApplication.cpp
+++ b/Engine/Source/Runtime/Core/Private/Linux/LinuxApplication.cpp
@@ -735,17 +735,6 @@ void FLinuxApplication::ProcessDeferredMessage( SDL_Event Event )
 
 				case SDL_WINDOWEVENT_RESIZED:
 					{
-						int NewWidth  = windowEvent.data1;
-						int NewHeight = windowEvent.data2;
-
-						MessageHandler->OnSizeChanged(
-							CurrentEventWindow.ToSharedRef(),
-							NewWidth,
-							NewHeight,
-							//	bWasMinimized
-							false
-						);
-
 						MessageHandler->OnResizingWindow( CurrentEventWindow.ToSharedRef() );
 					}
 					break;
@@ -788,17 +777,6 @@ void FLinuxApplication::ProcessDeferredMessage( SDL_Event Event )
 								SDL_SetWindowInputFocus(CurrentEventWindow->GetHWnd());
 							}
 						}
-
-						int Width, Height;
-
-						SDL_GetWindowSize(NativeWindow, &Width, &Height);
-
-						MessageHandler->OnSizeChanged(
-							CurrentEventWindow.ToSharedRef(),
-							Width,
-							Height,
-							false
-						);
 					}
 					break;
 
@@ -893,7 +871,10 @@ void FLinuxApplication::ProcessDeferredMessage( SDL_Event Event )
 							ActivateApplication();
 						}
 
-						if(CurrentFocusWindow != CurrentEventWindow)
+						// Some windows like notification windows may popup without needing the focus. That is handled in the SDL_WINDOWEVENT_SHOWN case.
+						// The problem would be that the WM will send the Take Focus event and wants to set the focus. We don't want it to set it
+						// for notifications because they are already handled in the above mentioned event.
+						if ((CurrentFocusWindow != CurrentEventWindow) && !CurrentEventWindow->IsNotificationWindow())
 						{
 							SDL_SetWindowInputFocus(CurrentEventWindow->GetHWnd());
 						}
