Support png/jpg splash screen images (in addition to bmp).

- PR #1444: Add support for PNG/JPG including BMP in Splash screen image (Contributed by [REDACTED_PERSON], with changes)
- Also incorporates PR #1399 that fixes fonts loading on Linux that depended on a particular path and converted paths to ASCII instead of UTF-8.

diff --git a/Engine/Source/Runtime/Core/Core.Build.cs b/Engine/Source/Runtime/Core/Core.Build.cs
--- a/Engine/Source/Runtime/Core/Core.Build.cs
+++ b/Engine/Source/Runtime/Core/Core.Build.cs
@@ -130,6 +130,7 @@ public class Core : ModuleRules
             {
                 AddThirdPartyPrivateStaticDependencies(Target, "FreeType2");
 				AddThirdPartyPrivateStaticDependencies(Target, "OpenGL");
+				PrivateIncludePathModuleNames.Add("ImageWrapper");
 			}
         }
 		else if (Target.Platform == UnrealTargetPlatform.HTML5 && Target.Architecture == "-win32")
diff --git a/Engine/Source/Runtime/Core/Private/GenericPlatform/GenericPlatformSplash.cpp b/Engine/Source/Runtime/Core/Private/GenericPlatform/GenericPlatformSplash.cpp
new file mode 100644
--- /dev/null
+++ b/Engine/Source/Runtime/Core/Private/GenericPlatform/GenericPlatformSplash.cpp
@@ -0,0 +1,94 @@
+// Copyright 1998-2015 Epic Games, Inc. All Rights Reserved.
+
+#include "CorePrivatePCH.h"
+#include "GenericPlatformSplash.h"
+
+// Supported file extension for splash image
+const TCHAR* SupportedSplashImageExt[] = 
+{
+	TEXT(".png"),
+	TEXT(".jpg"),
+	nullptr
+};
+
+/**
+* Return the filename found (look for PNG, JPG and BMP in that order, try to avoid BMP, use more space...)
+*/
+FString GetSplashFilename(const FString& ContentDir, const FString& Filename)
+{
+	int index = 0;
+	const FString ImageName = ContentDir / Filename;
+	FString Path;
+
+	while (SupportedSplashImageExt[index])
+	{
+		Path = ImageName + SupportedSplashImageExt[index++];
+
+		if (FPaths::FileExists(Path))
+			return Path;
+	}
+
+	// if no image was found, we assume it's a BMP (default)
+	return ImageName + TEXT(".bmp");
+}
+
+/**
+* Finds a usable splash pathname for the given filename
+*
+* @param SplashFilename Name of the desired splash name("Splash")
+* @param IconFilename Name of the desired icon name("Splash")
+* @param OutPath String containing the path to the file, if this function returns true
+* @param OutIconPath String containing the path to the icon, if this function returns true
+*
+* @return true if a splash screen was found
+*/
+bool FGenericPlatformSplash::GetSplashPath(const TCHAR* SplashFilename, FString& OutPath, bool& OutIsCustom)
+{
+	FString Filename = FString(TEXT("Splash/")) + SplashFilename;
+
+	// first look in game's splash directory
+	OutPath = GetSplashFilename(FPaths::GameContentDir(), Filename);
+	OutIsCustom = true;
+
+	// if this was found, then we're done
+	if (IFileManager::Get().FileSize(*OutPath) != -1)
+		return true;
+
+	// next look in Engine/Splash
+	OutPath = FPaths::ConvertRelativePathToFull(GetSplashFilename(FPaths::EngineContentDir(), Filename));
+	OutIsCustom = false;
+
+	// if this was found, then we're done
+	if (IFileManager::Get().FileSize(*OutPath) != -1)
+		return true;
+
+	// if not found yet, then return failure
+	return false;
+}
+
+bool FGenericPlatformSplash::GetSplashPath(const TCHAR* SplashFilename, const TCHAR* IconFilename, FString& OutPath, FString& OutIconPath, bool& OutIsCustom)
+{
+	FString Filename = FString(TEXT("Splash/")) + SplashFilename;
+	FString IconName = FString(TEXT("Splash/")) + IconFilename;
+
+	// first look in game's splash directory
+	OutPath = FPaths::ConvertRelativePathToFull(GetSplashFilename(FPaths::GameContentDir(), Filename));
+	OutIconPath = GetSplashFilename(FPaths::GameContentDir(), IconName);
+	OutIsCustom = true;
+
+	// if this was found, then we're done
+	if (IFileManager::Get().FileSize(*OutPath) != -1)
+		return true;
+
+	// next look in Engine/Splash
+	OutPath = FPaths::ConvertRelativePathToFull(GetSplashFilename(FPaths::EngineContentDir(), Filename));
+	OutIconPath = GetSplashFilename(FPaths::EngineContentDir(), IconName);
+	OutIsCustom = false;
+
+	// if this was found, then we're done
+	if (IFileManager::Get().FileSize(*OutPath) != -1)
+		return true;
+
+	// if not found yet, then return failure
+	return false;
+}
diff --git a/Engine/Source/Runtime/Core/Private/Linux/LinuxPlatformSplash.cpp b/Engine/Source/Runtime/Core/Private/Linux/LinuxPlatformSplash.cpp
--- a/Engine/Source/Runtime/Core/Private/Linux/LinuxPlatformSplash.cpp
+++ b/Engine/Source/Runtime/Core/Private/Linux/LinuxPlatformSplash.cpp
@@ -17,6 +17,7 @@
 #include "SDL.h"
 #include "SDL_thread.h"
 #include "SDL_timer.h"
+#include "ImageWrapper.h"
 
 #include "ft2build.h"
 #include FT_FREETYPE_H
@@ -64,6 +65,8 @@ static FText GSplashScreenText[ SplashTextType::NumTextTypes ];
 static Rect GSplashScreenTextRects[ SplashTextType::NumTextTypes ];
 static FString GSplashPath;
 static FString GIconPath;
+static IImageWrapperPtr GSplashImageWrapper;
+static IImageWrapperPtr GIconImageWrapper;
 
 
 static int32 SplashWidth = 0, SplashHeight = 0;
@@ -282,9 +285,9 @@ static int32 OpenFonts ( )
 	}
 	
 	// small font face
-	FString FontPath = FPaths::EngineContentDir() / TEXT("Slate/Fonts/Roboto-Light.ttf");
+	FString FontPath = FPaths::ConvertRelativePathToFull(FPaths::EngineContentDir() / TEXT("Slate/Fonts/Roboto-Light.ttf"));
 	
-	if (FT_New_Face(FontLibrary, TCHAR_TO_ANSI(*FontPath), 0, &FontSmall.Face))
+	if (FT_New_Face(FontLibrary, TCHAR_TO_UTF8(*FontPath), 0, &FontSmall.Face))
 	{
 		UE_LOG(LogHAL, Error, TEXT("*** Unable to open small font face for splash screen."));		
 	}
@@ -296,9 +299,9 @@ static int32 OpenFonts ( )
 
 
 	// normal font face
-	FontPath = FPaths::EngineContentDir() / TEXT("Slate/Fonts/Roboto-Regular.ttf");
+	FontPath = FPaths::ConvertRelativePathToFull(FPaths::EngineContentDir() / TEXT("Slate/Fonts/Roboto-Regular.ttf"));
 	
-	if (FT_New_Face(FontLibrary, TCHAR_TO_ANSI(*FontPath), 0, &FontNormal.Face))
+	if (FT_New_Face(FontLibrary, TCHAR_TO_UTF8(*FontPath), 0, &FontNormal.Face))
 	{
 		UE_LOG(LogHAL, Error, TEXT("*** Unable to open normal font face for splash screen."));		
 	}
@@ -309,9 +312,9 @@ static int32 OpenFonts ( )
 	}	
 	
 	// large font face
-	FontPath = FPaths::EngineContentDir() / TEXT("Slate/Fonts/Roboto-Bold.ttf");
+	FontPath = FPaths::ConvertRelativePathToFull(FPaths::EngineContentDir() / TEXT("Slate/Fonts/Roboto-Bold.ttf"));
 	
-	if (FT_New_Face(FontLibrary, TCHAR_TO_ANSI(*FontPath), 0, &FontLarge.Face))
+	if (FT_New_Face(FontLibrary, TCHAR_TO_UTF8(*FontPath), 0, &FontLarge.Face))
 	{
 		UE_LOG(LogHAL, Error, TEXT("*** Unable to open large font face for splash screen."));		
 	}
@@ -501,6 +504,51 @@ static int RenderString (GLuint tex_idx)
 	return 0;
 }
 
+/**
+ * @brief Helper function to load an image in any format.
+ *
+ * @param ImagePath an (absolute) path to the image
+ * @param OutImageWrapper a pointer to IImageWrapper which may be needed to hold the data (should outlive returned SDL_Surface)
+ *
+ * @return Splash surface
+ */
+static SDL_Surface* LinuxSplash_LoadImage(const TCHAR* ImagePath, IImageWrapperPtr & OutImageWrapper)
+{
+	TArray<uint8> RawFileData;
+
+	// Load the image buffer first
+	if (!FFileHelper::LoadFileToArray(RawFileData, ImagePath))
+	{
+		// If for some reason the image cannot be loaded, use the default BMP function
+		return SDL_LoadBMP(TCHAR_TO_UTF8(ImagePath));
+	}
+
+	auto& ImageWrapperModule = FModuleManager::LoadModuleChecked<IImageWrapperModule>(FName("ImageWrapper"));
+	auto Format = ImageWrapperModule.DetectImageFormat(RawFileData.GetData(), RawFileData.Num());
+	OutImageWrapper = ImageWrapperModule.CreateImageWrapper(Format);
+
+	if (OutImageWrapper.IsValid() && OutImageWrapper->SetCompressed(RawFileData.GetData(), RawFileData.Num()))
+	{
+		const TArray<uint8>* RawData = nullptr;
+
+		if (OutImageWrapper->GetRaw(ERGBFormat::BGRA, 8, RawData))
+		{
+			uint32 Bmask = 0x000000ff;
+			uint32 Gmask = 0x0000ff00;
+			uint32 Rmask = 0x00ff0000;
+			uint32 Amask = 0xff000000;
+
+			return SDL_CreateRGBSurfaceFrom((void*)RawData->GetData(),
+				OutImageWrapper->GetWidth(), OutImageWrapper->GetHeight(),
+				32, OutImageWrapper->GetWidth() * 4,
+				Rmask, Gmask, Bmask, Amask);
+		}
+	}
+
+	// If for some reason the image cannot be load, use the default BMP function
+	return SDL_LoadBMP(TCHAR_TO_UTF8(ImagePath));
+}
+
 /** Helper function to init resources used by the splash thread */
 bool LinuxSplash_InitSplashResources()
 {
@@ -511,7 +559,7 @@ bool LinuxSplash_InitSplashResources()
 	}
 
 	// load splash .bmp image
-	GSplashScreenImage = SDL_LoadBMP(TCHAR_TO_UTF8(*GSplashPath));
+	GSplashScreenImage = LinuxSplash_LoadImage(*GSplashPath, GSplashImageWrapper);
 	
 	if (GSplashScreenImage == nullptr)
 	{
@@ -535,7 +583,7 @@ bool LinuxSplash_InitSplashResources()
 
 	SDL_SetWindowTitle( GSplashWindow, TCHAR_TO_UTF8(*GSplashScreenAppName.ToString()) );
 
-	GSplashIconImage = SDL_LoadBMP(TCHAR_TO_UTF8(*(FString(FPlatformProcess::BaseDir()) / GIconPath)));
+	GSplashIconImage = LinuxSplash_LoadImage(*(FString(FPlatformProcess::BaseDir()) / GIconPath), GIconImageWrapper);
 
 	if (GSplashIconImage != nullptr)
 	{
@@ -558,9 +606,11 @@ void LinuxSplash_TearDownSplashResources()
 
 	SDL_FreeSurface(GSplashIconImage);
 	GSplashIconImage = nullptr;
+	GIconImageWrapper.Reset();
 
 	SDL_FreeSurface(GSplashScreenImage);
 	GSplashScreenImage = nullptr;
+	GSplashImageWrapper.Reset();
 
 	// do not deinit SDL here
 }
@@ -807,41 +857,6 @@ static int StartSplashScreenThread(void *ptr)
 	return 0;
 }
 
-/**
- * Finds a usable splash pathname for the given filename
- *
- * @param SplashFilename Name of the desired splash name ("Splash.bmp")
- * @param OutPath String containing the path to the file, if this function returns true
- *
- * @return true if a splash screen was found
- */
-static bool GetSplashPath(const TCHAR* SplashFilename, const TCHAR* IconFilename, bool& OutIsCustom)
-{
-	// first look in game's splash directory
-	GSplashPath = FPaths::GameContentDir() + TEXT("Splash/") + SplashFilename;
-	GIconPath = FPaths::GameContentDir() + TEXT("Splash/") + IconFilename;
-	OutIsCustom = true;
-
-	// if this was found, then we're done
-	if (IFileManager::Get().FileSize(*GSplashPath) != -1)
-	{
-		return true;
-	}
-
-	// next look in Engine/Splash
-	GSplashPath = FPaths::ConvertRelativePathToFull(FPaths::EngineContentDir() + TEXT("Splash/") + SplashFilename);
-	GIconPath = FPaths::EngineContentDir() + TEXT("Splash/") + IconFilename;
-	OutIsCustom = false;
-
-	// if this was found, then we're done
-	if (IFileManager::Get().FileSize(*GSplashPath) != -1)
-	{
-		return true;
-	}
-
-	// if not found yet, then return failure
-	return false;
-}
 #endif //WITH_EDITOR
 
 /**
@@ -875,15 +890,15 @@ void FLinuxPlatformSplash::Show( )
 	// decide on which splash screen to show
 	const FText GameName = FText::FromString(FApp::GetGameName());
 
-	const TCHAR* SplashImage = GIsEditor ? ( GameName.IsEmpty() ? TEXT("EdSplashDefault.bmp") : TEXT("EdSplash.bmp") ) : ( GameName.IsEmpty() ? TEXT("SplashDefault.bmp") : TEXT("Splash.bmp") );
+	const TCHAR* SplashImage = GIsEditor ? ( GameName.IsEmpty() ? TEXT("EdSplashDefault") : TEXT("EdSplash") ) : ( GameName.IsEmpty() ? TEXT("SplashDefault") : TEXT("Splash") );
 	// TODO: add an icon
-	//const TCHAR* IconImage = GIsEditor ? ( GameName.IsEmpty() ? TEXT("EdIconDefault.bmp") : TEXT("EdIcon.bmp") ) : ( GameName.IsEmpty() ? TEXT("IconDefault.bmp") : TEXT("Icon.bmp") );
-	const TCHAR* IconImage = GIsEditor ? ( GameName.IsEmpty() ? TEXT("EdSplashDefault.bmp") : TEXT("EdSplash.bmp") ) : ( GameName.IsEmpty() ? TEXT("SplashDefault.bmp") : TEXT("Splash.bmp") );
+	//const TCHAR* IconImage = GIsEditor ? ( GameName.IsEmpty() ? TEXT("EdIconDefault") : TEXT("EdIcon") ) : ( GameName.IsEmpty() ? TEXT("IconDefault") : TEXT("Icon.bmp") );
+	const TCHAR* IconImage = GIsEditor ? ( GameName.IsEmpty() ? TEXT("EdSplashDefault") : TEXT("EdSplash") ) : ( GameName.IsEmpty() ? TEXT("SplashDefault") : TEXT("Splash") );
 
 	// make sure a splash was found
 	bool IsCustom;
 	
-	if ( GetSplashPath(SplashImage, IconImage, IsCustom) == false )
+	if (GetSplashPath(SplashImage, IconImage, GSplashPath, GIconPath, IsCustom) == false)
 	{
 		UE_LOG(LogHAL, Warning, TEXT("Splash screen image not found."));
 	}
diff --git a/Engine/Source/Runtime/Core/Private/Mac/MacPlatformSplash.cpp b/Engine/Source/Runtime/Core/Private/Mac/MacPlatformSplash.cpp
--- a/Engine/Source/Runtime/Core/Private/Mac/MacPlatformSplash.cpp
+++ b/Engine/Source/Runtime/Core/Private/Mac/MacPlatformSplash.cpp
@@ -37,41 +37,6 @@ static FText GSplashScreenAppName;
 static FText GSplashScreenText[ SplashTextType::NumTextTypes ];
 static NSRect GSplashScreenTextRects[ SplashTextType::NumTextTypes ];
 
-
-/**
- * Finds a usable splash pathname for the given filename
- *
- * @param SplashFilename Name of the desired splash name ("Splash.bmp")
- * @param OutPath String containing the path to the file, if this function returns true
- *
- * @return true if a splash screen was found
- */
-static bool GetSplashPath(const TCHAR* SplashFilename, FString& OutPath, bool& OutIsCustom)
-{
-	// first look in game's splash directory
-	OutPath = FPaths::GameContentDir() + TEXT("Splash/") + SplashFilename;
-	OutIsCustom = true;
-	
-	// if this was found, then we're done
-	if (IFileManager::Get().FileSize(*OutPath) != -1)
-	{
-		return true;
-	}
-	
-	// next look in Engine/Splash
-	OutPath = FPaths::ConvertRelativePathToFull(FPaths::EngineContentDir() + TEXT("Splash/") + SplashFilename);
-	OutIsCustom = false;
-	
-	// if this was found, then we're done
-	if (IFileManager::Get().FileSize(*OutPath) != -1)
-	{
-		return true;
-	}
-	
-	// if not found yet, then return failure
-	return false;
-}
-
 /**
  * Sets the text displayed on the splash screen (for startup/loading progress)
  *
@@ -215,7 +180,7 @@ void FMacPlatformSplash::Show()
 
 		const FText GameName = FText::FromString(FApp::GetGameName());
 
-		const TCHAR* SplashImage = GIsEditor ? ( GameName.IsEmpty() ? TEXT("EdSplashDefault.bmp") : TEXT("EdSplash.bmp") ) : ( GameName.IsEmpty() ? TEXT("SplashDefault.bmp") : TEXT("Splash.bmp") );
+		const TCHAR* SplashImage = GIsEditor ? ( GameName.IsEmpty() ? TEXT("EdSplashDefault") : TEXT("EdSplash") ) : ( GameName.IsEmpty() ? TEXT("SplashDefault") : TEXT("Splash") );
 
 		// make sure a splash was found
 		FString SplashPath;
diff --git a/Engine/Source/Runtime/Core/Private/Windows/WindowsPlatformSplash.cpp b/Engine/Source/Runtime/Core/Private/Windows/WindowsPlatformSplash.cpp
--- a/Engine/Source/Runtime/Core/Private/Windows/WindowsPlatformSplash.cpp
+++ b/Engine/Source/Runtime/Core/Private/Windows/WindowsPlatformSplash.cpp
@@ -571,40 +571,6 @@ uint32 WINAPI StartSplashScreenThread( LPVOID unused )
 	return 0; 
 }
 
-/**
- * Finds a usable splash pathname for the given filename
- * 
- * @param SplashFilename Name of the desired splash name ("Splash.bmp")
- * @param OutPath String containing the path to the file, if this function returns true
- *
- * @return true if a splash screen was found
- */
-static bool GetSplashPath(const TCHAR* SplashFilename, FString& OutPath, bool& OutIsCustom)
-{
-	// first look in game's splash directory
-	OutPath = FPaths::GameContentDir() + TEXT("Splash/") + SplashFilename;
-	OutIsCustom = true;
-	
-	// if this was found, then we're done
-	if (IFileManager::Get().FileSize(*OutPath) != -1)
-	{
-		return true;
-	}
-
-	// next look in Engine/Splash
-	OutPath = FPaths::ConvertRelativePathToFull(FPaths::EngineContentDir() + TEXT("Splash/") + SplashFilename);
-	OutIsCustom = false;
-
-	// if this was found, then we're done
-	if (IFileManager::Get().FileSize(*OutPath) != -1)
-	{
-		return true;
-	}
-
-	// if not found yet, then return failure
-	return false;
-}
-
 /**
  * Sets the text displayed on the splash screen (for startup/loading progress)
  *
@@ -623,7 +589,7 @@ void FWindowsPlatformSplash::Show()
 	{
 		const FText GameName = FText::FromString( FApp::GetGameName() );
 
-		const TCHAR* SplashImage = GIsEditor ? ( GameName.IsEmpty() ? TEXT("EdSplashDefault.bmp") : TEXT("EdSplash.bmp") ) : ( GameName.IsEmpty() ? TEXT("SplashDefault.bmp") : TEXT("Splash.bmp") );
+		const TCHAR* SplashImage = GIsEditor ? ( GameName.IsEmpty() ? TEXT("EdSplashDefault") : TEXT("EdSplash") ) : ( GameName.IsEmpty() ? TEXT("SplashDefault") : TEXT("Splash") );
 
 		// make sure a splash was found
 		FString SplashPath;
diff --git a/Engine/Source/Runtime/Core/Public/GenericPlatform/GenericPlatformSplash.h b/Engine/Source/Runtime/Core/Public/GenericPlatform/GenericPlatformSplash.h
--- a/Engine/Source/Runtime/Core/Public/GenericPlatform/GenericPlatformSplash.h
+++ b/Engine/Source/Runtime/Core/Public/GenericPlatform/GenericPlatformSplash.h
@@ -31,6 +31,7 @@ namespace SplashTextType
 	};
 }
 
+class FString;
 
 /**
  * Generic implementation for most platforms
@@ -61,4 +62,18 @@ struct CORE_API FGenericPlatformSplash
 	{
 		return true;
 	}
+
+protected:
+	/**
+	* Finds a usable splash pathname for the given filename
+	*
+	* @param SplashFilename Name of the desired splash name("Splash")
+	* @param IconFilename Name of the desired icon name("Splash")
+	* @param OutPath String containing the path to the file, if this function returns true
+	* @param OutIconPath String containing the path to the icon, if this function returns true
+	*
+	* @return true if a splash screen was found
+	*/
+	static bool GetSplashPath(const TCHAR* SplashFilename, FString& OutPath, bool& OutIsCustom);
+	static bool GetSplashPath(const TCHAR* SplashFilename, const TCHAR* IconFilename, FString& OutPath, FString& OutIconPath, bool& OutIsCustom);
 };
