Core: Added bAllowShrinking to various TArray heap methods
[INTEGRATE] Change 2657135 by [REDACTED_PERSON][REDACTED_HANDLE] on 2015/08/14 22:48:55

diff --git a/Engine/Source/Runtime/Core/Public/Containers/Array.h b/Engine/Source/Runtime/Core/Public/Containers/Array.h
--- a/Engine/Source/Runtime/Core/Public/Containers/Array.h
+++ b/Engine/Source/Runtime/Core/Public/Containers/Array.h
@@ -2382,10 +2382,10 @@ public:
 	 * @param Predicate Predicate class instance.
 	 */
 	template <class PREDICATE_CLASS>
-	void HeapPop(ElementType& OutItem, const PREDICATE_CLASS& Predicate)
+	void HeapPop(ElementType& OutItem, const PREDICATE_CLASS& Predicate, bool bAllowShrinking = true)
 	{
 		OutItem = (*this)[0];
-		RemoveAtSwap(0);
+		RemoveAtSwap(0, 1, bAllowShrinking);
 
 		TDereferenceWrapper< ElementType, PREDICATE_CLASS> PredicateWrapper(Predicate);
 		SiftDown(0, Num(), PredicateWrapper);
@@ -2400,10 +2400,11 @@ public:
 	 * the template type.
 	 *
 	 * @param OutItem The removed item.
+	 * @param bAllowShrinking (Optional) Tells if this call can shrink the array allocation if suitable after the pop. Default is true.
 	 */
-	void HeapPop(ElementType& OutItem)
+	void HeapPop(ElementType& OutItem, bool bAllowShrinking = true)
 	{
-		HeapPop(OutItem, TLess<ElementType>());
+		HeapPop(OutItem, TLess<ElementType>(), bAllowShrinking);
 	}
 
 	/**
@@ -2430,11 +2431,12 @@ public:
 	 * Removes the top element from the heap.
 	 *
 	 * @param Predicate Predicate class instance.
+	 * @param bAllowShrinking (Optional) Tells if this call can shrink the array allocation if suitable after the discard. Default is true.
 	 */
 	template <class PREDICATE_CLASS>
-	void HeapPopDiscard(const PREDICATE_CLASS& Predicate)
+	void HeapPopDiscard(const PREDICATE_CLASS& Predicate, bool bAllowShrinking = true)
 	{
-		RemoveAtSwap(0);
+		RemoveAtSwap(0, 1, bAllowShrinking);
 		TDereferenceWrapper< ElementType, PREDICATE_CLASS> PredicateWrapper(Predicate);
 		SiftDown(0, Num(), PredicateWrapper);
 
@@ -2446,10 +2448,11 @@ public:
 	/** 
 	 * Removes the top element from the heap. Assumes < operator is defined for
 	 * the template type.
+	 * @param bAllowShrinking (Optional) Tells if this call can shrink the array allocation if suitable after the discard. Default is true.
 	 */
-	void HeapPopDiscard()
+	void HeapPopDiscard(bool bAllowShrinking = true)
 	{
-		HeapPopDiscard(TLess<ElementType>());
+		HeapPopDiscard(TLess<ElementType>(), bAllowShrinking);
 	}
 
 	/** 
@@ -2479,11 +2482,12 @@ public:
 	 *
 	 * @param Index Position at which to remove item.
 	 * @param Predicate Predicate class instance.
+	 * @param bAllowShrinking (Optional) Tells if this call can shrink the array allocation if suitable after the remove. Default is true.
 	 */
 	template <class PREDICATE_CLASS>
-	void HeapRemoveAt(int32 Index, const PREDICATE_CLASS& Predicate)
+	void HeapRemoveAt(int32 Index, const PREDICATE_CLASS& Predicate, bool bAllowShrinking = true)
 	{
-		RemoveAtSwap(Index);
+		RemoveAtSwap(Index, 1, bAllowShrinking);
 
 		TDereferenceWrapper< ElementType, PREDICATE_CLASS> PredicateWrapper(Predicate);
 		SiftDown(Index, Num(), PredicateWrapper);
@@ -2499,10 +2503,11 @@ public:
 	 * template type.
 	 *
 	 * @param Index Position at which to remove item.
+	 * @param bAllowShrinking (Optional) Tells if this call can shrink the array allocation if suitable after the remove. Default is true.
 	 */
-	void HeapRemoveAt(int32 Index)
+	void HeapRemoveAt(int32 Index, bool bAllowShrinking = true)
 	{
-		HeapRemoveAt(Index, TLess< ElementType >());
+		HeapRemoveAt(Index, TLess< ElementType >(), bAllowShrinking);
 	}
 
 	/**
