Mac: Fixed window reordering issues on app activation and cursor behavior when locked to a region

diff --git a/Engine/Source/Runtime/Core/Private/Mac/MacApplication.cpp b/Engine/Source/Runtime/Core/Private/Mac/MacApplication.cpp
--- a/Engine/Source/Runtime/Core/Private/Mac/MacApplication.cpp
+++ b/Engine/Source/Runtime/Core/Private/Mac/MacApplication.cpp
@@ -47,6 +47,7 @@ FMacApplication::FMacApplication()
 ,	ModifierKeysFlags(0)
 ,	CurrentModifierFlags(0)
 ,	bEmulatingRightClick(false)
+,	bIgnoreMouseMoveDelta(false)
 ,	bIsWorkspaceSessionActive(true)
 {
 	TextInputMethodSystem = MakeShareable(new FMacTextInputMethodSystem);
@@ -319,7 +320,7 @@ void FMacApplication::DeferEvent(NSObject* Object)
 			case NSRightMouseDragged:
 			case NSOtherMouseDragged:
 			case NSEventTypeSwipe:
-				DeferredEvent.Delta = FVector2D([Event deltaX], [Event deltaY]);
+				DeferredEvent.Delta = bIgnoreMouseMoveDelta ? FVector2D::ZeroVector : FVector2D([Event deltaX], [Event deltaY]);
 				break;
 
 			case NSLeftMouseDown:
@@ -507,6 +508,7 @@ void FMacApplication::ProcessEvent(const FDeferredMacEvent& Event)
 			case NSOtherMouseDragged:
 				ConditionallyUpdateModifierKeys(Event);
 				ProcessMouseMovedEvent(Event, EventWindow);
+				bIgnoreMouseMoveDelta = false;
 				break;
 
 			case NSLeftMouseDown:
@@ -659,7 +661,7 @@ void FMacApplication::ProcessMouseMovedEvent(const FDeferredMacEvent& Event, TSh
 
 	FVector2D const MouseScaling = MacCursor->GetMouseScaling();
 
-	if (bUsingHighPrecisionMouseInput)
+	if (bUsingHighPrecisionMouseInput || MacCursor->IsLocked())
 	{
 		// Under OS X we disassociate the cursor and mouse position during hi-precision mouse input.
 		// The game snaps the mouse cursor back to the starting point when this is disabled, which
@@ -1048,7 +1050,7 @@ void FMacApplication::OnApplicationWillResignActive()
 	if (SavedWindowsOrder.Num() > 0)
 	{
 		NSWindow* TopWindow = [NSApp windowWithWindowNumber:SavedWindowsOrder[0].WindowNumber];
-		if ( TopWindow )
+		if (TopWindow)
 		{
 			[TopWindow orderWindow:NSWindowAbove relativeTo:0];
 		}
@@ -1058,14 +1060,7 @@ void FMacApplication::OnApplicationWillResignActive()
 			NSWindow* Window = [NSApp windowWithWindowNumber:Info.WindowNumber];
 			if (Window)
 			{
-				if (SavedWindowsOrder[Index - 1].Level == Info.Level && TopWindow)
-				{
-					[Window orderWindow:NSWindowBelow relativeTo:[TopWindow windowNumber]];
-				}
-				else
-				{
-					[Window orderWindow:NSWindowAbove relativeTo:0];
-				}
+				[Window orderWindow:NSWindowBelow relativeTo:[TopWindow windowNumber]];
 				TopWindow = Window;
 			}
 		}
@@ -1122,6 +1117,30 @@ void FMacApplication::OnWindowsReordered(bool bIsAppInBackground)
 	}
 }
 
+void FMacApplication::OnCursorLock()
+{
+	NSWindow* NativeWindow = [NSApp keyWindow];
+	if (NativeWindow)
+	{
+		const bool bIsCursorLocked = ((FMacCursor*)Cursor.Get())->IsLocked();
+		if (bIsCursorLocked)
+		{
+			[NativeWindow setMinSize:NSMakeSize(NativeWindow.frame.size.width, NativeWindow.frame.size.height)];
+			[NativeWindow setMaxSize:NSMakeSize(NativeWindow.frame.size.width, NativeWindow.frame.size.height)];
+		}
+		else
+		{
+			TSharedPtr<FMacWindow> Window = FindWindowByNSWindow((FCocoaWindow*)NativeWindow);
+			if (Window.IsValid())
+			{
+				const FGenericWindowDefinition& Definition = Window->GetDefinition();
+				[NativeWindow setMinSize:NSMakeSize(Definition.SizeLimits.GetMinWidth().Get(10.0f), Definition.SizeLimits.GetMinHeight().Get(10.0f))];
+				[NativeWindow setMaxSize:NSMakeSize(Definition.SizeLimits.GetMaxWidth().Get(10000.0f), Definition.SizeLimits.GetMaxHeight().Get(10000.0f))];
+			}
+		}
+	}
+}
+
 void FMacApplication::ConditionallyUpdateModifierKeys(const FDeferredMacEvent& Event)
 {
 	if (CurrentModifierFlags != Event.ModifierFlags)
diff --git a/Engine/Source/Runtime/Core/Private/Mac/MacCursor.cpp b/Engine/Source/Runtime/Core/Private/Mac/MacCursor.cpp
--- a/Engine/Source/Runtime/Core/Private/Mac/MacCursor.cpp
+++ b/Engine/Source/Runtime/Core/Private/Mac/MacCursor.cpp
@@ -289,48 +289,54 @@ void FMacCursor::Lock( const RECT* const Bounds )
 	// Lock/Unlock the cursor
 	if ( Bounds == NULL )
 	{
-		CusorClipRect = FIntRect();
+		CursorClipRect = FIntRect();
 	}
 	else
 	{
-		CusorClipRect.Min.X = FMath::TruncToInt(Bounds->left);
-		CusorClipRect.Min.Y = FMath::TruncToInt(Bounds->top);
-		CusorClipRect.Max.X = FMath::TruncToInt(Bounds->right) - 1;
-		CusorClipRect.Max.Y = FMath::TruncToInt(Bounds->bottom) - 1;
+		CursorClipRect.Min.X = FMath::TruncToInt(Bounds->left);
+		CursorClipRect.Min.Y = FMath::TruncToInt(Bounds->top);
+		CursorClipRect.Max.X = FMath::TruncToInt(Bounds->right) - 1;
+		CursorClipRect.Max.Y = FMath::TruncToInt(Bounds->bottom) - 1;
 	}
 
+	CGAssociateMouseAndMouseCursorPosition( !bUseHighPrecisionMode && !IsLocked() );
+
+	MacApplication->OnCursorLock();
+
 	FVector2D Position = GetPosition();
 	if( UpdateCursorClipping( Position ) )
 	{
 		SetPosition( Position.X, Position.Y );
 	}
+
+	MouseWarpDelta = FVector2D::ZeroVector;
 }
 
 bool FMacCursor::UpdateCursorClipping( FVector2D& CursorPosition )
 {
 	bool bAdjusted = false;
 
-	if (CusorClipRect.Area() > 0)
+	if (CursorClipRect.Area() > 0)
 	{
-		if (CursorPosition.X < CusorClipRect.Min.X)
+		if (CursorPosition.X < CursorClipRect.Min.X)
 		{
-			CursorPosition.X = CusorClipRect.Min.X;
+			CursorPosition.X = CursorClipRect.Min.X;
 			bAdjusted = true;
 		}
-		else if (CursorPosition.X > CusorClipRect.Max.X)
+		else if (CursorPosition.X > CursorClipRect.Max.X)
 		{
-			CursorPosition.X = CusorClipRect.Max.X;
+			CursorPosition.X = CursorClipRect.Max.X;
 			bAdjusted = true;
 		}
 
-		if (CursorPosition.Y < CusorClipRect.Min.Y)
+		if (CursorPosition.Y < CursorClipRect.Min.Y)
 		{
-			CursorPosition.Y = CusorClipRect.Min.Y;
+			CursorPosition.Y = CursorClipRect.Min.Y;
 			bAdjusted = true;
 		}
-		else if (CursorPosition.Y > CusorClipRect.Max.Y)
+		else if (CursorPosition.Y > CursorClipRect.Max.Y)
 		{
-			CursorPosition.Y = CusorClipRect.Max.Y;
+			CursorPosition.Y = CursorClipRect.Max.Y;
 			bAdjusted = true;
 		}
 	}
@@ -390,7 +396,7 @@ void FMacCursor::WarpCursor( const int32 X, const int32 Y )
 	// Previously there was CGSetLocalEventsSuppressionInterval to explicitly control this behaviour but that is deprecated.
 	// The replacement CGEventSourceSetLocalEventsSuppressionInterval isn't useful because it is unclear how to obtain the correct event source.
 	// Instead, when we want the warp to be visible we need to disassociate mouse & cursor...
-	if( !bUseHighPrecisionMode )
+	if( !bUseHighPrecisionMode && !IsLocked() )
 	{
 		CGAssociateMouseAndMouseCursorPosition( false );
 	}
@@ -399,12 +405,14 @@ void FMacCursor::WarpCursor( const int32 X, const int32 Y )
 	CGWarpMouseCursorPosition( NSMakePoint( FMath::TruncToInt( X ), FMath::TruncToInt( Y ) ) );
 	
 	// And then reassociate the mouse cursor, which forces the mouse events to come through.
-	if( !bUseHighPrecisionMode )
+	if( !bUseHighPrecisionMode && !IsLocked() )
 	{
 		CGAssociateMouseAndMouseCursorPosition( true );
 	}
 	
 	UpdateCurrentPosition( FVector2D(X, Y) );
+
+	MacApplication->IgnoreMouseMoveDelta();
 }
 
 FVector2D FMacCursor::GetMouseWarpDelta()
@@ -420,7 +428,7 @@ void FMacCursor::SetHighPrecisionMouseMode( bool const bEnable )
 	{
 		bUseHighPrecisionMode = bEnable;
 		
-		CGAssociateMouseAndMouseCursorPosition( !bUseHighPrecisionMode );
+		CGAssociateMouseAndMouseCursorPosition( !bUseHighPrecisionMode && !IsLocked() );
 		
 		if ( GMacDisableMouseCoalescing )
 		{
diff --git a/Engine/Source/Runtime/Core/Public/Mac/MacApplication.h b/Engine/Source/Runtime/Core/Public/Mac/MacApplication.h
--- a/Engine/Source/Runtime/Core/Public/Mac/MacApplication.h
+++ b/Engine/Source/Runtime/Core/Public/Mac/MacApplication.h
@@ -187,6 +187,10 @@ public:
 
 	const TArray<TSharedRef<FMacWindow>>& GetAllWindows() const { return Windows; }
 
+	void OnCursorLock();
+
+	void IgnoreMouseMoveDelta() { bIgnoreMouseMoveDelta = true; }
+
 private:
 
 	static NSEvent* HandleNSEvent(NSEvent* Event);
@@ -280,6 +284,8 @@ private:
 
 	bool bEmulatingRightClick;
 
+	bool bIgnoreMouseMoveDelta;
+
 	TArray<FCocoaWindow*> WindowsToClose;
 
 	TArray<FCocoaWindow*> WindowsRequiringTextInvalidation;
diff --git a/Engine/Source/Runtime/Core/Public/Mac/MacCursor.h b/Engine/Source/Runtime/Core/Public/Mac/MacCursor.h
--- a/Engine/Source/Runtime/Core/Public/Mac/MacCursor.h
+++ b/Engine/Source/Runtime/Core/Public/Mac/MacCursor.h
@@ -59,13 +59,15 @@ public:
 
 	void UpdateVisibility();
 
+	bool IsLocked() const { return CursorClipRect.Area() > 0; }
+
 private:
 	EMouseCursor::Type CurrentType;
 
 	/** Cursors */
 	NSCursor* CursorHandles[EMouseCursor::TotalCursorCount];
 
-	FIntRect CusorClipRect;
+	FIntRect CursorClipRect;
 
 	bool bIsVisible;
 	bool bUseHighPrecisionMode;
diff --git a/Engine/Source/Runtime/Slate/Private/Framework/Application/SlateApplication.cpp b/Engine/Source/Runtime/Slate/Private/Framework/Application/SlateApplication.cpp
--- a/Engine/Source/Runtime/Slate/Private/Framework/Application/SlateApplication.cpp
+++ b/Engine/Source/Runtime/Slate/Private/Framework/Application/SlateApplication.cpp
@@ -4239,16 +4239,19 @@ FReply FSlateApplication::RoutePointerDownEvent(FWidgetPath& WidgetsUnderPointer
 
 #if PLATFORM_MAC
 		TSharedPtr<SWindow> TopLevelWindow = WidgetsUnderPointer.TopLevelWindow;
-		if ( bNeedToActivateWindow && TopLevelWindow.IsValid() )
+		if ( bNeedToActivateWindow )
 		{
 			// Clicking on a context menu should not activate anything
 			// @todo: This needs to be updated when we have window type in SWindow and we no longer have to guess if WidgetsUnderCursor.TopLevelWindow is a menu
-			const bool bIsContextMenu = !TopLevelWindow->IsRegularWindow() && TopLevelWindow->HasMinimizeBox() && TopLevelWindow->HasMaximizeBox();
+			const bool bIsContextMenu = TopLevelWindow.IsValid() && !TopLevelWindow->IsRegularWindow() && TopLevelWindow->HasMinimizeBox() && TopLevelWindow->HasMaximizeBox();
 			if ( !bIsContextMenu && PointerEvent.GetEffectingButton() == EKeys::LeftMouseButton && !DragDetector.DetectDragForWidget.IsValid() && ActiveWindow == [NSApp keyWindow] )
 			{
 				MouseCaptorHelper Captor = MouseCaptor;
 				FPlatformMisc::ActivateApplication();
-				TopLevelWindow->BringToFront(true);
+				if (TopLevelWindow.IsValid())
+				{
+					TopLevelWindow->BringToFront(true);
+				}
 				MouseCaptor = Captor;
 			}
 		}
