Linux: use correct window coordinates when processing event queue.

- PR #1466: [GNU/Linux] Fixed a mouse input bug in the graph editor when using SDL events. (Contributed by [REDACTED_PERSON])

diff --git a/Engine/Source/Runtime/Core/Private/Linux/LinuxApplication.cpp b/Engine/Source/Runtime/Core/Private/Linux/LinuxApplication.cpp
--- a/Engine/Source/Runtime/Core/Private/Linux/LinuxApplication.cpp
+++ b/Engine/Source/Runtime/Core/Private/Linux/LinuxApplication.cpp
@@ -299,7 +299,7 @@ void FLinuxApplication::ProcessDeferredMessage( SDL_Event Event )
 					if (motionEvent.x != (width / 2) || motionEvent.y != (height / 2))
 					{
 						int xOffset, yOffset;
-						SDL_GetWindowPosition(NativeWindow, &xOffset, &yOffset);
+						GetWindowPositionInEventLoop(NativeWindow, &xOffset, &yOffset);
 						LinuxCursor->SetPosition(width / 2 + xOffset, height / 2 + yOffset);
 					}
 					else
@@ -311,7 +311,7 @@ void FLinuxApplication::ProcessDeferredMessage( SDL_Event Event )
 			else
 			{
 				int xOffset, yOffset;
-				SDL_GetWindowPosition( NativeWindow, &xOffset, &yOffset );
+				GetWindowPositionInEventLoop(NativeWindow, &xOffset, &yOffset);
 
 				int32 BorderSizeX, BorderSizeY;
 				CurrentEventWindow->GetNativeBordersSize(BorderSizeX, BorderSizeY);
@@ -1640,3 +1640,46 @@ bool FLinuxApplication::HandleWindowCommand(const TCHAR* Cmd, FOutputDevice& Ar)
 	return false;
 }
 
+void FLinuxApplication::SaveWindowLocationsForEventLoop(void)
+{
+	for (int32 WindowIndex = 0; WindowIndex < Windows.Num(); ++WindowIndex)
+	{
+		TSharedRef< FLinuxWindow > Window = Windows[WindowIndex];
+		int x = 0;
+		int y = 0;
+		SDL_HWindow NativeWindow = Window->GetHWnd();
+		SDL_GetWindowPosition(NativeWindow, &x, &y);
+		SavedWindowLocationsForEventLoop.FindOrAdd(NativeWindow) = FVector2D(x, y);
+	}
+}
+
+void FLinuxApplication::ClearWindowLocationsAfterEventLoop(void)
+{
+	SavedWindowLocationsForEventLoop.Empty();
+}
+
+void FLinuxApplication::GetWindowPositionInEventLoop(SDL_HWindow NativeWindow, int *x, int *y)
+{
+	FVector2D *Position = SavedWindowLocationsForEventLoop.Find(NativeWindow);
+	if(Position)
+	{
+		// Found saved location.
+		*x = Position->X;
+		*y = Position->Y;
+	}
+	else if(NativeWindow)
+	{
+		SDL_GetWindowPosition(NativeWindow, x, y);
+
+		// If we've hit this case, then we're either not in the event
+		// loop, or suddenly have a new window to keep track of.
+		// Record the initial window position.
+		SavedWindowLocationsForEventLoop.FindOrAdd(NativeWindow) = FVector2D(*x, *y);
+	}
+	else
+	{
+		UE_LOG(LogLinuxWindowEvent, Error, TEXT("Tried to get the location of a non-existing window\n"));
+		*x = 0;
+		*y = 0;
+	}
+}
diff --git a/Engine/Source/Runtime/Core/Private/Linux/LinuxCursor.cpp b/Engine/Source/Runtime/Core/Private/Linux/LinuxCursor.cpp
--- a/Engine/Source/Runtime/Core/Private/Linux/LinuxCursor.cpp
+++ b/Engine/Source/Runtime/Core/Private/Linux/LinuxCursor.cpp
@@ -202,7 +202,17 @@ void FLinuxCursor::SetPosition( const int32 X, const int32 Y )
 	if (WndFocus)
 	{
 		int WndX, WndY;
-		SDL_GetWindowPosition(WndFocus, &WndX, &WndY);	//	get top left
+
+		// Get top-left.
+		if(LinuxApplication)
+		{
+			LinuxApplication->GetWindowPositionInEventLoop(WndFocus, &WndX, &WndY);
+		}
+		else
+		{
+			SDL_GetWindowPosition(WndFocus, &WndX, &WndY);
+		}
+
 		SDL_WarpMouseInWindow(WndFocus, X - WndX, Y - WndY);
 	}
 	else
diff --git a/Engine/Source/Runtime/Core/Private/Linux/LinuxPlatformMisc.cpp b/Engine/Source/Runtime/Core/Private/Linux/LinuxPlatformMisc.cpp
--- a/Engine/Source/Runtime/Core/Private/Linux/LinuxPlatformMisc.cpp
+++ b/Engine/Source/Runtime/Core/Private/Linux/LinuxPlatformMisc.cpp
@@ -258,14 +258,27 @@ void FLinuxPlatformMisc::PumpMessages( bool bFromMainLoop )
 {
 	if (GInitializedSDL && bFromMainLoop)
 	{
-		SDL_Event event;
-
-		while (SDL_PollEvent(&event))
+		if( LinuxApplication )
 		{
-			if( LinuxApplication )
+			LinuxApplication->SaveWindowLocationsForEventLoop();
+
+			SDL_Event event;
+
+			while (SDL_PollEvent(&event))
 			{
 				LinuxApplication->AddPendingEvent( event );
 			}
+
+			LinuxApplication->ClearWindowLocationsAfterEventLoop();
+		}
+		else
+		{
+			// No application to send events to. Just flush out the
+			// queue.
+			SDL_Event event;
+			while (SDL_PollEvent(&event))
+			{
+			}
 		}
 	}
 }
diff --git a/Engine/Source/Runtime/Core/Public/Linux/LinuxApplication.h b/Engine/Source/Runtime/Core/Public/Linux/LinuxApplication.h
--- a/Engine/Source/Runtime/Core/Public/Linux/LinuxApplication.h
+++ b/Engine/Source/Runtime/Core/Public/Linux/LinuxApplication.h
@@ -86,6 +86,33 @@ public:
 		return (Windows.Num() > 0) ? bActivateApp : true;
 	}
 
+	/**
+	 * Windows can move during an event loop (like in
+	 * FLinuxPlatformMisc::PumpMessages(), but SDL queues many events
+	 * up before any windows move. This can lead to the screen-space
+	 * position of the mouse cursor being calculated incorrectly with
+	 * the old event data and new window location data. Use this to
+	 * save the window locations for use during the loop.
+	 */
+	void SaveWindowLocationsForEventLoop(void);
+
+	/** Clear out data saved in SaveWindowLocationsForEventLoop(). */
+	void ClearWindowLocationsAfterEventLoop(void);
+
+	/**
+	 * Get a window position inside the event loop. Fall back on
+	 * SDL_GetWindowPosition if not present in the saved window
+	 * locations.
+	 *
+	 * @param NativeWindow The native SDL_HWindow to look up. This is
+	 *					 a drop-in replacement for
+	 *					 SDL_GetWindowPosition, so it uses the
+	 *					 native window type.
+	 * @param x			Pointer to the X component of the position.
+	 * @param y			Pointer to the Y component of the position.
+	 */
+	void GetWindowPositionInEventLoop(SDL_HWindow NativeWindow, int *x, int *y);
+
 private:
 
 	FLinuxApplication();
@@ -173,6 +200,15 @@ private:
 	/** Array of windows to focus when current gets removed. */
 	TArray< TSharedRef< FLinuxWindow > > RevertFocusStack;
 
+	/**
+	 * Saved window locations used for event loop. Note: Using raw
+	 * handle here because weak pointers can change meaning as a
+	 * response to some event in the middle of the loop, potentially
+	 * corrupting our indexing. These keys should not be de-referenced
+	 * (but comparison is okay).
+	 */
+	TMap< SDL_HWindow, FVector2D > SavedWindowLocationsForEventLoop;
+
 	int32 bAllowedToDeferMessageProcessing;
 
 	/** Using high precision mouse input */
