Corrected duplicate loctext keys

diff --git a/Engine/Source/Runtime/Core/Public/Misc/ExpressionParserTypes.inl b/Engine/Source/Runtime/Core/Public/Misc/ExpressionParserTypes.inl
--- a/Engine/Source/Runtime/Core/Public/Misc/ExpressionParserTypes.inl
+++ b/Engine/Source/Runtime/Core/Public/Misc/ExpressionParserTypes.inl
@@ -169,7 +169,7 @@ FExpressionResult TOperatorJumpTable<ContextType>::ExecBinary(const FExpressionT
 	Args.Add(FText::FromString(Operator.Context.GetString()));
 	Args.Add(FText::FromString(L.Context.GetString()));
 	Args.Add(FText::FromString(R.Context.GetString()));
-	return MakeError(FText::Format(LOCTEXT("UnaryExecutionError", "Binary operator {0} cannot operate on {1} and {2}"), Args));
+	return MakeError(FText::Format(LOCTEXT("BinaryExecutionError", "Binary operator {0} cannot operate on {1} and {2}"), Args));
 }
 
 template<typename ContextType>
@@ -184,7 +184,7 @@ FExpressionResult TOperatorJumpTable<ContextType>::ExecPreUnary(const FExpressio
 	FFormatOrderedArguments Args;
 	Args.Add(FText::FromString(Operator.Context.GetString()));
 	Args.Add(FText::FromString(R.Context.GetString()));
-	return MakeError(FText::Format(LOCTEXT("UnaryExecutionError", "Pre-unary operator {0} cannot operate on {1}"), Args));
+	return MakeError(FText::Format(LOCTEXT("PreUnaryExecutionError", "Pre-unary operator {0} cannot operate on {1}"), Args));
 }
 
 template<typename ContextType>
@@ -199,7 +199,7 @@ FExpressionResult TOperatorJumpTable<ContextType>::ExecPostUnary(const FExpressi
 	FFormatOrderedArguments Args;
 	Args.Add(FText::FromString(Operator.Context.GetString()));
 	Args.Add(FText::FromString(L.Context.GetString()));
-	return MakeError(FText::Format(LOCTEXT("UnaryExecutionError", "Post-unary operator {0} cannot operate on {1}"), Args));
+	return MakeError(FText::Format(LOCTEXT("PostUnaryExecutionError", "Post-unary operator {0} cannot operate on {1}"), Args));
 }
 
 template<typename ContextType>
