Log memory stats to output log on OOM.

diff --git a/Engine/Source/Runtime/Core/Private/GenericPlatform/GenericPlatformMemory.cpp b/Engine/Source/Runtime/Core/Private/GenericPlatform/GenericPlatformMemory.cpp
--- a/Engine/Source/Runtime/Core/Private/GenericPlatform/GenericPlatformMemory.cpp
+++ b/Engine/Source/Runtime/Core/Private/GenericPlatform/GenericPlatformMemory.cpp
@@ -74,6 +74,20 @@ void FGenericPlatformMemory::OnOutOfMemory(uint64 Size, uint32 Alignment)
 		FPlatformMemory::BinnedFreeToOS(BackupOOMMemoryPool);
 		UE_LOG(LogMemory, Warning, TEXT("Freeing %d bytes from backup pool to handle out of memory."), BackupOOMMemoryPoolSize);
 	}
+	FPlatformMemoryStats PlatformMemoryStats = FPlatformMemory::GetStats();
+	UE_LOG(LogMemory, Warning, TEXT("MemoryStats:")\
+		TEXT("\n\tAvailablePhysical %llu")\
+		TEXT("\n\tAvailableVirtual %llu")\
+		TEXT("\n\tUsedPhysical %llu")\
+		TEXT("\n\tPeakUsedPhysical %llu")\
+		TEXT("\n\tUsedVirtual %llu")\
+		TEXT("\n\tPeakUsedVirtual %llu"),
+		PlatformMemoryStats.AvailablePhysical,
+		PlatformMemoryStats.AvailableVirtual,
+		PlatformMemoryStats.UsedPhysical,
+		PlatformMemoryStats.PeakUsedPhysical,
+		PlatformMemoryStats.UsedVirtual,
+		PlatformMemoryStats.PeakUsedVirtual);
 	UE_LOG(LogMemory, Fatal, TEXT("Ran out of memory allocating %llu bytes with alignment %u"), Size, Alignment);
 }
 
