Fixed a crash on typing " on U.S. International - PC keyboard on Mac

diff --git a/Engine/Source/Runtime/Core/Private/Mac/MacApplication.cpp b/Engine/Source/Runtime/Core/Private/Mac/MacApplication.cpp
--- a/Engine/Source/Runtime/Core/Private/Mac/MacApplication.cpp
+++ b/Engine/Source/Runtime/Core/Private/Mac/MacApplication.cpp
@@ -919,7 +919,7 @@ void FMacApplication::ProcessKeyDownEvent(const FDeferredMacEvent& Event, TShare
 void FMacApplication::ProcessKeyUpEvent(const FDeferredMacEvent& Event)
 {
 	bool bHandled = false;
-	if (!bSystemModalMode)
+	if (!bSystemModalMode && [Event.Characters length] > 0 && [Event.CharactersIgnoringModifiers length] > 0)
 	{
 		const TCHAR Character = ConvertChar([Event.Characters characterAtIndex:0]);
 		const TCHAR CharCode = [Event.CharactersIgnoringModifiers characterAtIndex:0];
