Added facility to create a slow task dialog automatically after a certain amount of time has passed.

diff --git a/Engine/Source/Runtime/Core/Public/Misc/FeedbackContext.h b/Engine/Source/Runtime/Core/Public/Misc/FeedbackContext.h
--- a/Engine/Source/Runtime/Core/Public/Misc/FeedbackContext.h
+++ b/Engine/Source/Runtime/Core/Public/Misc/FeedbackContext.h
@@ -144,6 +144,9 @@ struct CORE_API FSlowTask
 	/** The time that this scope was created */
 	double StartTime;
 
+	/** Threshold before dialog is opened */
+	TOptional<float> OpenDialogThreshold;
+
 private:
 
 	/** Boolean flag to control whether this scope actually does anything (unset for quiet operations etc) */
@@ -155,6 +158,12 @@ private:
 	/** The feedback context that we belong to */
 	FFeedbackContext& Context;
 
+	/** Specify whether the delayed dialog should show a cancel button */
+	bool bDelayedDialogShowCancelButton : 1;
+
+	/** Specify whether the delayed dialog is allowed in PIE */
+	bool bDelayedDialogAllowInPIE : 1;
+
 	/** Prevent copying */
 	FSlowTask(const FSlowTask&);
 
@@ -229,6 +238,19 @@ public:
 		}
 	}
 
+	/**
+	 * Creates a new dialog for this slow task after the given time threshold. If the task completes before this time, no dialog will be shown.
+	 * @param		Threshold				Time in seconds before dialog will be shown.
+	 * @param		bShowCancelButton		Whether to show a cancel button on the dialog or not
+	 * @param		bAllowInPIE				Whether to allow this dialog in PIE. If false, this dialog will not appear during PIE sessions.
+	 */
+	FORCEINLINE void MakeDialogDelayed(float Threshold, bool bShowCancelButton = false, bool bAllowInPIE = false)
+	{
+		OpenDialogThreshold = Threshold;
+		bDelayedDialogShowCancelButton = bShowCancelButton;
+		bDelayedDialogAllowInPIE = bAllowInPIE;
+	}
+
 	/**
 	 * Creates a new dialog for this slow task, if there is currently not one open
 	 * @param		bShowCancelButton		Whether to show a cancel button on the dialog or not
@@ -250,6 +272,11 @@ public:
 		verifySlow(ExpectedWorkThisFrame <= WorkRemaining);
 		CurrentFrameScope = FMath::Min(WorkRemaining, ExpectedWorkThisFrame);
 
+		if (!bCreatedDialog && OpenDialogThreshold.IsSet() && static_cast<float>(FPlatformTime::Seconds() - StartTime) > OpenDialogThreshold.GetValue())
+		{
+			MakeDialog(bDelayedDialogShowCancelButton, bDelayedDialogAllowInPIE);
+		}
+
 		if (bEnabled)
 		{
 			Context.RequestUpdateUI(bCreatedDialog || (*Context.ScopeStack)[0] == this);
