[UE-16708] Fixing array-entry duplication in INI files by ensuring that we don't double-write lines that were in the array from the start. Might want to reorder some of this code for the more complex array modifications such as subtracting elements, but appears to work at present.

diff --git a/Engine/Source/Runtime/Core/Private/Misc/ConfigCacheIni.cpp b/Engine/Source/Runtime/Core/Private/Misc/ConfigCacheIni.cpp
--- a/Engine/Source/Runtime/Core/Private/Misc/ConfigCacheIni.cpp
+++ b/Engine/Source/Runtime/Core/Private/Misc/ConfigCacheIni.cpp
@@ -1316,7 +1316,11 @@ void FConfigFile::ProcessPropertyAndWriteForDefaults( const TArray< FString >& I
 			for (const FString& NextElement : ArrayProperties)
 			{
 				FString PropertyNameWithRemoveOp = PropertyName.Replace(TEXT("+"), TEXT("-"));
-				OutText += FString::Printf(TEXT("%s=%s") LINE_TERMINATOR, *PropertyNameWithRemoveOp, *NextElement);
+				if (PropertyName.Compare(PropertyNameWithRemoveOp))
+				{
+					// Only output the property if we've made a change. This avoids duplicating array entries.
+					OutText += FString::Printf(TEXT("%s=%s") LINE_TERMINATOR, *PropertyNameWithRemoveOp, *NextElement);
+				}
 			}
 		}
 	}
