UE-16708

Fixed an issue where configs written to the game Agnostic saved config folder were not treated as saved configs, but defaults.
Editor Settings are now processed correctly when output to the engine saved config directory.

Note: Reverting CL# 2660460 as this was not a fix and introduced a tonne of other issues with saving arrays to default ini files.

diff --git a/Engine/Source/Runtime/Core/Private/Misc/ConfigCacheIni.cpp b/Engine/Source/Runtime/Core/Private/Misc/ConfigCacheIni.cpp
--- a/Engine/Source/Runtime/Core/Private/Misc/ConfigCacheIni.cpp
+++ b/Engine/Source/Runtime/Core/Private/Misc/ConfigCacheIni.cpp
@@ -948,8 +948,9 @@ bool FConfigFile::Write( const FString& Filename, bool bDoRemoteWrite/* = true*/
 
 				// If we are writing to a default config file and this property is an array, we need to be careful to remove those from higher up the hierarchy
 				const FString AbsoluteFilename = FPaths::ConvertRelativePathToFull(Filename);
-				const FString AbsoluteGeneratedConfigDir = FPaths::ConvertRelativePathToFull(FPaths::GeneratedConfigDir());
-				const bool bIsADefaultIniWrite = !AbsoluteFilename.Contains(AbsoluteGeneratedConfigDir);
+				const FString AbsoluteGameGeneratedConfigDir = FPaths::ConvertRelativePathToFull(FPaths::GeneratedConfigDir());
+				const FString AbsoluteGameAgnosticGeneratedConfigDir = FPaths::ConvertRelativePathToFull(FPaths::Combine(*FPaths::GameAgnosticSavedDir(), TEXT("Config")) + TEXT("/"));
+				const bool bIsADefaultIniWrite = !AbsoluteFilename.Contains(AbsoluteGameGeneratedConfigDir) && !AbsoluteFilename.Contains(AbsoluteGameAgnosticGeneratedConfigDir);
 
 				// Check if the property matches the source configs. We do not wanna write it out if so.
 				if ((bIsADefaultIniWrite || bDifferentNumberOfElements || !DoesConfigPropertyValueMatch(SourceConfigFile, SectionName, PropertyName, PropertyValue)) && !bOptionIsFromCommandline)
@@ -1316,11 +1317,7 @@ void FConfigFile::ProcessPropertyAndWriteForDefaults( const TArray< FString >& I
 			for (const FString& NextElement : ArrayProperties)
 			{
 				FString PropertyNameWithRemoveOp = PropertyName.Replace(TEXT("+"), TEXT("-"));
-				if (PropertyName.Compare(PropertyNameWithRemoveOp))
-				{
-					// Only output the property if we've made a change. This avoids duplicating array entries.
-					OutText += FString::Printf(TEXT("%s=%s") LINE_TERMINATOR, *PropertyNameWithRemoveOp, *NextElement);
-				}
+				OutText += FString::Printf(TEXT("%s=%s") LINE_TERMINATOR, *PropertyNameWithRemoveOp, *NextElement);
 			}
 		}
 	}
