Added variadic template support to FText formatting

diff --git a/Engine/Source/Runtime/Core/Public/Internationalization/Text.h b/Engine/Source/Runtime/Core/Public/Internationalization/Text.h
--- a/Engine/Source/Runtime/Core/Public/Internationalization/Text.h
+++ b/Engine/Source/Runtime/Core/Public/Internationalization/Text.h
@@ -316,6 +316,71 @@ public:
 	static FText Format(const FText& Fmt,const FText& v1,const FText& v2,const FText& v3);
 	static FText Format(const FText& Fmt,const FText& v1,const FText& v2,const FText& v3,const FText& v4);
 
+#if PLATFORM_COMPILER_HAS_VARIADIC_TEMPLATES
+	/**
+	 * FormatNamed allows you to pass name <-> value pairs to the function to format automatically
+	 *
+	 * @usage FText::FormatNamed( FText::FromString( TEXT( "{PlayerName} is really cool" ) ), TEXT( "PlayerName" ), FText::FromString( TEXT( "Awesomegirl" ) ) );
+	 *
+	 * @param Fmt the format to create from
+	 * @param Args a variadic list of FString to Value (must be even numbered)
+	 * @return a formatted FText
+	 */
+	template < typename... TArguments >
+	static FText FormatNamed( const FText& Fmt, TArguments&&... Args )
+	{
+		static_assert( sizeof...( TArguments ) % 2 == 0, "FormatNamed requires an even number of Name <-> Value pairs" );
+
+		FFormatNamedArguments FormatArguments;
+		FormatNamed_Util( FormatArguments, Forward< TArguments >( Args )... );
+		return FormatInternal( Fmt, FormatArguments, false, false );
+	}
+
+	/**
+	 * FormatOrdered allows you to pass a variadic list of types to use for formatting in order desired
+	 *
+	 * @param Fmt the format to create from
+	 * @param Args a variadic list of values in order of desired formatting
+	 * @return a formatted FText
+	 */
+	template < typename... TArguments >
+	static FText FormatOrdered( const FText& Fmt, TArguments&&... Args )
+	{
+		FFormatOrderedArguments FormatArguments;
+		FormatOrdered_Util( FormatArguments, Forward< TArguments >( Args )... );
+		return FormatInternal( Fmt, FormatArguments, false, false );
+	}
+
+private:
+	template < typename TName, typename TValue, typename... TArguments >
+	static void FormatNamed_Util( OUT FFormatNamedArguments& Result, TName&& Name, TValue&& Value, TArguments&&... Args )
+	{
+		FormatNamed_Util( Result, Name, Value );
+		FormatNamed_Util( Result, Forward< TArguments >( Args )... );
+	}
+
+	template < typename TName, typename TValue >
+	static void FormatNamed_Util( OUT FFormatNamedArguments& Result, TName&& Name, TValue&& Value )
+	{
+		Result.Emplace( Name, Value );
+	}
+
+	template < typename TValue, typename... TArguments >
+	static void FormatOrdered_Util( OUT FFormatOrderedArguments& Result, TValue&& Value, TArguments&&... Args )
+	{
+		FormatOrdered_Util( Result, Value );
+		FormatOrdered_Util( Result, Forward< TArguments >( Args )... );
+	}
+
+	template < typename TValue >
+	static void FormatOrdered_Util( OUT FFormatOrderedArguments& Result, TValue&& Value )
+	{
+		Result.Emplace( Value );
+	}
+
+public:
+#endif
+
 	static void SetEnableErrorCheckingResults(bool bEnable){bEnableErrorCheckingResults=bEnable;}
 	static bool GetEnableErrorCheckingResults(){return bEnableErrorCheckingResults;}
 
