Fix linux / ps4 compile error with incomplete type.

diff --git a/Engine/Source/Runtime/Core/Public/Internationalization/Text.h b/Engine/Source/Runtime/Core/Public/Internationalization/Text.h
--- a/Engine/Source/Runtime/Core/Public/Internationalization/Text.h
+++ b/Engine/Source/Runtime/Core/Public/Internationalization/Text.h
@@ -67,7 +67,32 @@ namespace EFormatArgumentType
 	};
 }
 
-struct FFormatArgumentValue;
+struct CORE_API FFormatArgumentValue
+{
+	EFormatArgumentType::Type Type;
+	union
+	{
+		int64 IntValue;
+		uint64 UIntValue;
+		float FloatValue;
+		double DoubleValue;
+		FText* TextValue;
+	};
+
+	FFormatArgumentValue();
+
+	FFormatArgumentValue( const int Value );
+	FFormatArgumentValue( const unsigned int Value );
+	FFormatArgumentValue( const int64 Value );
+	FFormatArgumentValue( const uint64 Value );
+	FFormatArgumentValue( const float Value );
+	FFormatArgumentValue( const double Value );
+	FFormatArgumentValue( const FText& Value );
+	FFormatArgumentValue( const FFormatArgumentValue& Source );
+	~FFormatArgumentValue();
+
+	friend FArchive& operator<<( FArchive& Ar, FFormatArgumentValue& Value );
+};
 
 typedef TMap<FString, FFormatArgumentValue> FFormatNamedArguments;
 typedef TArray<FFormatArgumentValue> FFormatOrderedArguments;
@@ -519,33 +544,6 @@ public:
 	static int32 GetFlags(const FText& Text);
 };
 
-struct CORE_API FFormatArgumentValue
-{
-	EFormatArgumentType::Type Type;
-	union
-	{
-		int64 IntValue;
-		uint64 UIntValue;
-		float FloatValue;
-		double DoubleValue;
-		FText* TextValue;
-	};
-
-	FFormatArgumentValue();
-
-	FFormatArgumentValue(const int Value);
-	FFormatArgumentValue(const unsigned int Value);
-	FFormatArgumentValue(const int64 Value);
-	FFormatArgumentValue(const uint64 Value);
-	FFormatArgumentValue(const float Value);
-	FFormatArgumentValue(const double Value);
-	FFormatArgumentValue(const FText& Value);
-	FFormatArgumentValue(const FFormatArgumentValue& Source);
-	~FFormatArgumentValue();
-
-	friend FArchive& operator<<( FArchive& Ar, FFormatArgumentValue& Value );
-};
-
 struct FFormatArgumentData
 {
 	FText ArgumentName;
