Implement OS X support for automated, unattended reporting including ensures using CrashReportClient.

diff --git a/Engine/Source/Developer/CrashDebugHelper/Private/Mac/CrashDebugHelperMac.cpp b/Engine/Source/Developer/CrashDebugHelper/Private/Mac/CrashDebugHelperMac.cpp
--- a/Engine/Source/Developer/CrashDebugHelper/Private/Mac/CrashDebugHelperMac.cpp
+++ b/Engine/Source/Developer/CrashDebugHelper/Private/Mac/CrashDebugHelperMac.cpp
@@ -247,6 +247,10 @@ static int32 ParseExceptionCode(TCHAR const* CrashLog, uint32& OutExceptionCode)
 			{
 				OutExceptionCode = SIGABRT;
 			}
+			else if(FCStringWide::Strcmp(Buffer, TEXT("SIGTRAP")) == 0)
+			{
+				OutExceptionCode = SIGTRAP;
+			}
 			else if(FString(Buffer).IsNumeric())
 			{
 				Found = swscanf(Buffer, TEXT("%u"), &OutExceptionCode);
diff --git a/Engine/Source/Runtime/Core/Private/Apple/ApplePlatformStackWalk.cpp b/Engine/Source/Runtime/Core/Private/Apple/ApplePlatformStackWalk.cpp
--- a/Engine/Source/Runtime/Core/Private/Apple/ApplePlatformStackWalk.cpp
+++ b/Engine/Source/Runtime/Core/Private/Apple/ApplePlatformStackWalk.cpp
@@ -308,10 +308,6 @@ void CreateExceptionInfoString(int32 Signal, struct __siginfo* Info)
 #undef HANDLE_CASE
 }
 
-void NewReportEnsure( const TCHAR* ErrorMessage )
-{
-}
-
 int32 ReportCrash(ucontext_t *Context, int32 Signal, struct __siginfo* Info)
 {
 	static bool GAlreadyCreatedMinidump = false;
diff --git a/Engine/Source/Runtime/Core/Private/Mac/MacPlatformMisc.cpp b/Engine/Source/Runtime/Core/Private/Mac/MacPlatformMisc.cpp
--- a/Engine/Source/Runtime/Core/Private/Mac/MacPlatformMisc.cpp
+++ b/Engine/Source/Runtime/Core/Private/Mac/MacPlatformMisc.cpp
@@ -51,6 +51,8 @@ struct FMacApplicationInfo
 		AppPath = FString([[NSBundle mainBundle] executablePath]);
 
 		AppBundleID = FString([[NSBundle mainBundle] bundleIdentifier]);
+		
+		bIsUnattended = FApp::IsUnattended();
 
 		bIsSandboxed = FPlatformProcess::IsSandboxedApplication();
 		
@@ -246,6 +248,7 @@ struct FMacApplicationInfo
 		return PLCrashReportFileName;
 	}
 	
+	bool bIsUnattended;
 	bool bIsSandboxed;
 	bool RunningOnBattery;
 	bool RunningOnMavericks;
@@ -1616,25 +1619,13 @@ void FMacCrashContext::CopyMinidump(char const* OutputPath, char const* InputPat
 	}
 }
 
-void FMacCrashContext::GenerateCrashInfoAndLaunchReporter() const
+void FMacCrashContext::GenerateInfoInFolder(char const* const InfoFolder) const
 {
 	// create a crash-specific directory
 	char CrashInfoFolder[PATH_MAX] = {};
-	FCStringAnsi::Strncpy(CrashInfoFolder, GMacAppInfo.CrashReportPath, PATH_MAX);
-	FCStringAnsi::Strcat(CrashInfoFolder, PATH_MAX, "/CrashReport-UE4-");
-	FCStringAnsi::Strcat(CrashInfoFolder, PATH_MAX, GMacAppInfo.AppNameUTF8);
-	FCStringAnsi::Strcat(CrashInfoFolder, PATH_MAX, "-pid-");
-	FCStringAnsi::Strcat(CrashInfoFolder, PATH_MAX, ItoANSI(getpid(), 10));
-	FCStringAnsi::Strcat(CrashInfoFolder, PATH_MAX, "-");
-	FCStringAnsi::Strcat(CrashInfoFolder, PATH_MAX, ItoANSI(GMacAppInfo.RunUUID.A, 16));
-	FCStringAnsi::Strcat(CrashInfoFolder, PATH_MAX, ItoANSI(GMacAppInfo.RunUUID.B, 16));
-	FCStringAnsi::Strcat(CrashInfoFolder, PATH_MAX, ItoANSI(GMacAppInfo.RunUUID.C, 16));
-	FCStringAnsi::Strcat(CrashInfoFolder, PATH_MAX, ItoANSI(GMacAppInfo.RunUUID.D, 16));
-
-	// Prevent CrashReportClient from spawning another CrashReportClient.
-	const bool bCanRunCrashReportClient = FCString::Stristr( *(GMacAppInfo.ExecutableName), TEXT( "CrashReportClient" ) ) == nullptr;
-
-	if(!mkdir(CrashInfoFolder, 0766) && bCanRunCrashReportClient)
+	FCStringAnsi::Strncpy(CrashInfoFolder, InfoFolder, PATH_MAX);
+	
+	if(!mkdir(CrashInfoFolder, 0766))
 	{
 		char FilePath[PATH_MAX] = {};
 		FCStringAnsi::Strncpy(FilePath, CrashInfoFolder, PATH_MAX);
@@ -1729,17 +1720,47 @@ void FMacCrashContext::GenerateCrashInfoAndLaunchReporter() const
 			close(VideoDst);
 			close(VideoSrc);
 		}
-		
+	}
+}
+
+void FMacCrashContext::GenerateCrashInfoAndLaunchReporter() const
+{
+	// Prevent CrashReportClient from spawning another CrashReportClient.
+	const bool bCanRunCrashReportClient = FCString::Stristr( *(GMacAppInfo.ExecutableName), TEXT( "CrashReportClient" ) ) == nullptr;
+
+	if(bCanRunCrashReportClient)
+	{
+		// create a crash-specific directory
+		char CrashInfoFolder[PATH_MAX] = {};
+		FCStringAnsi::Strncpy(CrashInfoFolder, GMacAppInfo.CrashReportPath, PATH_MAX);
+		FCStringAnsi::Strcat(CrashInfoFolder, PATH_MAX, "/CrashReport-UE4-");
+		FCStringAnsi::Strcat(CrashInfoFolder, PATH_MAX, GMacAppInfo.AppNameUTF8);
+		FCStringAnsi::Strcat(CrashInfoFolder, PATH_MAX, "-pid-");
+		FCStringAnsi::Strcat(CrashInfoFolder, PATH_MAX, ItoANSI(getpid(), 10));
+		FCStringAnsi::Strcat(CrashInfoFolder, PATH_MAX, "-");
+		FCStringAnsi::Strcat(CrashInfoFolder, PATH_MAX, ItoANSI(GMacAppInfo.RunUUID.A, 16));
+		FCStringAnsi::Strcat(CrashInfoFolder, PATH_MAX, ItoANSI(GMacAppInfo.RunUUID.B, 16));
+		FCStringAnsi::Strcat(CrashInfoFolder, PATH_MAX, ItoANSI(GMacAppInfo.RunUUID.C, 16));
+		FCStringAnsi::Strcat(CrashInfoFolder, PATH_MAX, ItoANSI(GMacAppInfo.RunUUID.D, 16));
+		
+		GenerateInfoInFolder(CrashInfoFolder);
+
 		// try launching the tool and wait for its exit, if at all
 		// Use vfork() & execl() as they are async-signal safe, CreateProc can fail in Cocoa
 		int32 ReturnCode = 0;
-		FCStringAnsi::Strncpy(FilePath, CrashInfoFolder, PATH_MAX);
-		FCStringAnsi::Strcat(FilePath, PATH_MAX, "/");
+		FCStringAnsi::Strcat(CrashInfoFolder, PATH_MAX, "/");
 		pid_t ForkPID = vfork();
 		if(ForkPID == 0)
 		{
 			// Child
-			execl(GMacAppInfo.CrashReportClient, "CrashReportClient", FilePath, NULL);
+			if(GMacAppInfo.bIsUnattended)
+			{
+				execl(GMacAppInfo.CrashReportClient, "CrashReportClient", CrashInfoFolder, "-Unattended", NULL);
+			}
+			else
+			{
+				execl(GMacAppInfo.CrashReportClient, "CrashReportClient", CrashInfoFolder, NULL);
+			}
 		}
 		// We no longer wait here because on return the OS will scribble & crash again due to the behaviour of the XPC function
 		// OS X uses internally to launch/wait on the CrashReportClient. It is simpler, easier & safer to just die here like a good little Mac.app.
@@ -1767,3 +1788,61 @@ void FMacCrashContext::GenerateCrashInfoAndLaunchReporter() const
 	
 	_Exit(0);
 }
+
+void FMacCrashContext::GenerateEnsureInfoAndLaunchReporter() const
+{
+	// Prevent CrashReportClient from spawning another CrashReportClient.
+	const bool bCanRunCrashReportClient = FCString::Stristr( *(GMacAppInfo.ExecutableName), TEXT( "CrashReportClient" ) ) == nullptr;
+	
+	if(bCanRunCrashReportClient)
+	{
+		SCOPED_AUTORELEASE_POOL;
+		
+		// Write the PLCrashReporter report to the expected location
+		NSData* CrashReport = [FMacApplicationInfo::CrashReporter generateLiveReport];
+		[CrashReport writeToFile:[NSString stringWithUTF8String:GMacAppInfo.PLCrashReportPath] atomically:YES];
+
+		// Use a slightly different output folder name to not conflict with a subequent crash
+		const FGuid Guid = FGuid::NewGuid();
+		FString GameName = FApp::GetGameName();
+		FString EnsureLogFolder = FString(GMacAppInfo.CrashReportPath) / FString::Printf(TEXT("EnsureReport-%s-%s"), *GameName, *Guid.ToString(EGuidFormats::Digits));
+		
+		GenerateInfoInFolder(TCHAR_TO_UTF8(*EnsureLogFolder));
+		
+		FString Arguments = EnsureLogFolder + TEXT("/ -Unattended");
+		FString ReportClient = FPaths::ConvertRelativePathToFull(FPlatformProcess::GenerateApplicationPath(TEXT("CrashReportClient"), EBuildConfigurations::Development));
+		FPlatformProcess::ExecProcess(*ReportClient, *Arguments, nullptr, nullptr, nullptr);
+	}
+}
+
+static FCriticalSection EnsureLock;
+static bool bReentranceGuard = false;
+
+void NewReportEnsure( const TCHAR* ErrorMessage )
+{
+	// Simple re-entrance guard.
+	EnsureLock.Lock();
+	
+	if( bReentranceGuard )
+	{
+		EnsureLock.Unlock();
+		return;
+	}
+	
+	bReentranceGuard = true;
+	
+	if(FMacApplicationInfo::CrashReporter != nil)
+	{
+		siginfo_t Signal;
+		Signal.si_signo = SIGTRAP;
+		Signal.si_code = TRAP_TRACE;
+		Signal.si_addr = __builtin_return_address(0);
+		
+		FMacCrashContext EnsureContext;
+		EnsureContext.InitFromSignal(SIGTRAP, &Signal, nullptr);
+		EnsureContext.GenerateEnsureInfoAndLaunchReporter();
+	}
+	
+	bReentranceGuard = false;
+	EnsureLock.Unlock();
+}
diff --git a/Engine/Source/Runtime/Core/Public/Mac/MacPlatformCrashContext.h b/Engine/Source/Runtime/Core/Public/Mac/MacPlatformCrashContext.h
--- a/Engine/Source/Runtime/Core/Public/Mac/MacPlatformCrashContext.h
+++ b/Engine/Source/Runtime/Core/Public/Mac/MacPlatformCrashContext.h
@@ -12,8 +12,14 @@ struct CORE_API FMacCrashContext : public FApplePlatformCrashContext
 	/** Copies the PLCrashReporter minidump */
 	void CopyMinidump(char const* OutputPath, char const* InputPath) const;
 
+	/** Generates the ensure/crash info into the given folder */
+	void GenerateInfoInFolder(char const* const InfoFolder) const;
+	
 	/** Generates information for crash reporter */
 	void GenerateCrashInfoAndLaunchReporter() const;
+	
+	/** Generates information for ensures sent via the CrashReporter */
+	void GenerateEnsureInfoAndLaunchReporter() const;
 };
 
 typedef FMacCrashContext FPlatformCrashContext;
