Inline versions of ToLower/ToUpper so that I can avoid some unneeded copying (or at least moving)

diff --git a/Engine/Source/Runtime/Core/Private/Containers/String.cpp b/Engine/Source/Runtime/Core/Private/Containers/String.cpp
--- a/Engine/Source/Runtime/Core/Private/Containers/String.cpp
+++ b/Engine/Source/Runtime/Core/Private/Containers/String.cpp
@@ -79,24 +79,37 @@ int32 FString::Find(const TCHAR* SubStr, ESearchCase::Type SearchCase, ESearchDi
 
 FString FString::ToUpper() const
 {
-	FString New( **this );
+	FString New(**this);
+	New.ToUpperInline();
+	return New;
+}
+
+void FString::ToUpperInline()
+{
 	const int32 StringLength = Len();
-	for( int32 i=0; i < StringLength; ++i )
+	TCHAR* RawData = Data.GetData();
+	for (int32 i = 0; i < StringLength; ++i)
 	{
-		New[i] = FChar::ToUpper(New[i]);
+		RawData[i] = FChar::ToUpper(RawData[i]);
 	}
-	return New;
 }
 
+
 FString FString::ToLower() const
 {
-	FString New( **this );
+	FString New(**this);
+	New.ToLowerInline();
+	return New;
+}
+
+void FString::ToLowerInline()
+{
 	const int32 StringLength = Len();
-	for( int32 i=0; i < StringLength; ++i )
+	TCHAR* RawData = Data.GetData();
+	for (int32 i = 0; i < StringLength; ++i)
 	{
-		New[i] = FChar::ToLower(New[i]);
+		RawData[i] = FChar::ToLower(RawData[i]);
 	}
-	return New;
 }
 
 bool FString::StartsWith(const FString& InPrefix, ESearchCase::Type SearchCase ) const
diff --git a/Engine/Source/Runtime/Core/Public/Containers/UnrealString.h b/Engine/Source/Runtime/Core/Public/Containers/UnrealString.h
--- a/Engine/Source/Runtime/Core/Public/Containers/UnrealString.h
+++ b/Engine/Source/Runtime/Core/Public/Containers/UnrealString.h
@@ -1286,9 +1286,15 @@ public:
 	/** @return a new string with the characters of this converted to uppercase */
 	FString ToUpper() const;
 
+	/** Converts all characters in this string to uppercase */
+	void ToUpperInline();
+
 	/** @return a new string with the characters of this converted to lowercase */
 	FString ToLower() const;
 
+	/** Converts all characters in this string to lowercase */
+	void ToLowerInline();
+
 	/** Pad the left of this string for ChCount characters */
 	FString LeftPad( int32 ChCount ) const;
 
