Fix bogus memory logs on OOM on 32bit platforms.

diff --git a/Engine/Source/Runtime/Core/Private/GenericPlatform/GenericPlatformMemory.cpp b/Engine/Source/Runtime/Core/Private/GenericPlatform/GenericPlatformMemory.cpp
--- a/Engine/Source/Runtime/Core/Private/GenericPlatform/GenericPlatformMemory.cpp
+++ b/Engine/Source/Runtime/Core/Private/GenericPlatform/GenericPlatformMemory.cpp
@@ -69,12 +69,12 @@ void FGenericPlatformMemory::OnOutOfMemory(uint64 Size, uint32 Alignment)
 	OOMAllocationAlignment = Alignment;
 
 	bIsOOM = true;
+	FPlatformMemoryStats PlatformMemoryStats = FPlatformMemory::GetStats();
 	if (BackupOOMMemoryPool)
 	{
 		FPlatformMemory::BinnedFreeToOS(BackupOOMMemoryPool);
 		UE_LOG(LogMemory, Warning, TEXT("Freeing %d bytes from backup pool to handle out of memory."), BackupOOMMemoryPoolSize);
 	}
-	FPlatformMemoryStats PlatformMemoryStats = FPlatformMemory::GetStats();
 	UE_LOG(LogMemory, Warning, TEXT("MemoryStats:")\
 		TEXT("\n\tAvailablePhysical %llu")\
 		TEXT("\n\tAvailableVirtual %llu")\
@@ -82,12 +82,12 @@ void FGenericPlatformMemory::OnOutOfMemory(uint64 Size, uint32 Alignment)
 		TEXT("\n\tPeakUsedPhysical %llu")\
 		TEXT("\n\tUsedVirtual %llu")\
 		TEXT("\n\tPeakUsedVirtual %llu"),
-		PlatformMemoryStats.AvailablePhysical,
-		PlatformMemoryStats.AvailableVirtual,
-		PlatformMemoryStats.UsedPhysical,
-		PlatformMemoryStats.PeakUsedPhysical,
-		PlatformMemoryStats.UsedVirtual,
-		PlatformMemoryStats.PeakUsedVirtual);
+		(uint64)PlatformMemoryStats.AvailablePhysical,
+		(uint64)PlatformMemoryStats.AvailableVirtual,
+		(uint64)PlatformMemoryStats.UsedPhysical,
+		(uint64)PlatformMemoryStats.PeakUsedPhysical,
+		(uint64)PlatformMemoryStats.UsedVirtual,
+		(uint64)PlatformMemoryStats.PeakUsedVirtual);
 	UE_LOG(LogMemory, Fatal, TEXT("Ran out of memory allocating %llu bytes with alignment %u"), Size, Alignment);
 }
 
