Renamed and moved CalculateSegmentIntersection2D to FMath::SegmentIntersection2D #UE4

github PR#1475, by slonopotamus

diff --git a/Engine/Source/Runtime/Core/Private/Math/UnrealMath.cpp b/Engine/Source/Runtime/Core/Private/Math/UnrealMath.cpp
--- a/Engine/Source/Runtime/Core/Private/Math/UnrealMath.cpp
+++ b/Engine/Source/Runtime/Core/Private/Math/UnrealMath.cpp
@@ -1617,18 +1617,37 @@ float FMath::GetTForSegmentPlaneIntersect(const FVector& StartPoint, const FVect
 	return ( Plane.W - (StartPoint|Plane) ) / ( (EndPoint - StartPoint)|Plane);	
 }
 
-bool FMath::SegmentPlaneIntersection(const FVector& StartPoint, const FVector& EndPoint, const FPlane& Plane, FVector& out_IntersectPoint)
+bool FMath::SegmentPlaneIntersection(const FVector& StartPoint, const FVector& EndPoint, const FPlane& Plane, FVector& out_IntersectionPoint)
 {
-	float T = FMath::GetTForSegmentPlaneIntersect(StartPoint,EndPoint,Plane);
+	float T = FMath::GetTForSegmentPlaneIntersect(StartPoint, EndPoint, Plane);
 	// If the parameter value is not between 0 and 1, there is no intersection
-	if( T > -KINDA_SMALL_NUMBER && T < 1.f+KINDA_SMALL_NUMBER )
+	if (T > -KINDA_SMALL_NUMBER && T < 1.f + KINDA_SMALL_NUMBER)
 	{
-		out_IntersectPoint = StartPoint + T * (EndPoint - StartPoint);
+		out_IntersectionPoint = StartPoint + T * (EndPoint - StartPoint);
 		return true;
 	}
 	return false;
 }
 
+bool FMath::SegmentIntersection2D(const FVector& SegmentStartA, const FVector& SegmentEndA, const FVector& SegmentStartB, const FVector& SegmentEndB, FVector& out_IntersectionPoint)
+{
+	const FVector VectorA = SegmentEndA - SegmentStartA;
+	const FVector VectorB = SegmentEndB - SegmentStartB;
+
+	const float S = (-VectorA.Y * (SegmentStartA.X - SegmentStartB.X) + VectorA.X * (SegmentStartA.Y - SegmentStartB.Y)) / (-VectorB.X * VectorA.Y + VectorA.X * VectorB.Y);
+	const float T = (VectorB.X * (SegmentStartA.Y - SegmentStartB.Y) - VectorB.Y * (SegmentStartA.X - SegmentStartB.X)) / (-VectorB.X * VectorA.Y + VectorA.X * VectorB.Y);
+
+	const bool bIntersects = (S >= 0 && S <= 1 && T >= 0 && T <= 1);
+
+	if (bIntersects)
+	{
+		out_IntersectionPoint.X = SegmentStartA.X + (T * VectorA.X);
+		out_IntersectionPoint.Y = SegmentStartA.Y + (T * VectorA.Y);
+		out_IntersectionPoint.Z = SegmentStartA.Z + (T * VectorA.Z);
+	}
+
+	return bIntersects;
+}
 
 /**
  * Compute the screen bounds of a point light along one axis.
diff --git a/Engine/Source/Runtime/Core/Public/Math/UnrealMathUtility.h b/Engine/Source/Runtime/Core/Public/Math/UnrealMathUtility.h
--- a/Engine/Source/Runtime/Core/Public/Math/UnrealMathUtility.h
+++ b/Engine/Source/Runtime/Core/Public/Math/UnrealMathUtility.h
@@ -1014,14 +1014,27 @@ struct FMath : public FPlatformMath
 
 	/**
 	 * Returns true if there is an intersection between the segment specified by StartPoint and Endpoint, and
-	 * the plane on which polygon Plane lies. If there is an intersection, the point is placed in IntersectionPoint
+	 * the plane on which polygon Plane lies. If there is an intersection, the point is placed in out_IntersectionPoint
 	 * @param StartPoint - start point of segment
 	 * @param EndPoint   - end point of segment
 	 * @param Plane		- plane to intersect with
-	 * @param out_InterSectPoint - out var for the point on the segment that intersects the mesh (if any)
+	 * @param out_IntersectionPoint - out var for the point on the segment that intersects the mesh (if any)
 	 * @return true if intersection occured
 	 */
-	static CORE_API bool SegmentPlaneIntersection(const FVector& StartPoint, const FVector& EndPoint, const FPlane& Plane, FVector& out_IntersectPoint);
+	static CORE_API bool SegmentPlaneIntersection(const FVector& StartPoint, const FVector& EndPoint, const FPlane& Plane, FVector& out_IntersectionPoint);
+
+	/**
+	 * Returns true if there is an intersection between the segment specified by SegmentStartA and SegmentEndA, and
+	 * the segment specified by SegmentEndA and SegmentStartB, in 2D space. If there is an intersection, the point is placed in out_IntersectionPoint
+	 * @param SegmentStartA - start point of first segment
+	 * @param SegmentEndA   - end point of first segment
+	 * @param SegmentStartB - start point of second segment
+	 * @param SegmentEndB   - end point of second segment
+	 * @param out_IntersectionPoint - out var for the intersection point (if any)
+	 * @return true if intersection occured
+	 */
+	static CORE_API bool SegmentIntersection2D(const FVector& SegmentStartA, const FVector& SegmentEndA, const FVector& SegmentStartB, const FVector& SegmentEndB, FVector& out_IntersectionPoint);
+
 
 	/**
 	 * Returns closest point on a triangle to a point.
diff --git a/Engine/Source/Runtime/Engine/Private/AI/Navigation/NavigationPath.cpp b/Engine/Source/Runtime/Engine/Private/AI/Navigation/NavigationPath.cpp
--- a/Engine/Source/Runtime/Engine/Private/AI/Navigation/NavigationPath.cpp
+++ b/Engine/Source/Runtime/Engine/Private/AI/Navigation/NavigationPath.cpp
@@ -467,26 +467,6 @@ void FNavMeshPath::PerformStringPulling(const FVector& StartLoc, const FVector&
 
 namespace
 {
-	FORCEINLINE bool CalculateSegmentIntersection2D(const FVector& SegmentStartA, const FVector& SegmentEndA, const FVector& SegmentStartB, const FVector& SegmentEndB, FVector& IntersectionPoint)
-	{
-		const FVector VectorA = SegmentEndA - SegmentStartA;
-		const FVector VectorB = SegmentEndB - SegmentStartB;
-
-		const float S = (-VectorA.Y * (SegmentStartA.X - SegmentStartB.X) + VectorA.X * (SegmentStartA.Y - SegmentStartB.Y)) / (-VectorB.X * VectorA.Y + VectorA.X * VectorB.Y);
-		const float T = ( VectorB.X * (SegmentStartA.Y - SegmentStartB.Y) - VectorB.Y * (SegmentStartA.X - SegmentStartB.X)) / (-VectorB.X * VectorA.Y + VectorA.X * VectorB.Y);
-
-		const bool bIntersects = S >= 0 && S <= 1 && T >= 0 && T <= 1;
-
-		if (bIntersects)
-		{
-			IntersectionPoint.X = SegmentStartA.X + (T * VectorA.X);
-			IntersectionPoint.Y = SegmentStartA.Y + (T * VectorA.Y);
-			IntersectionPoint.Z = SegmentStartA.Z + (T * VectorA.Z);
-		}
-
-		return bIntersects;
-	}
-
 	struct FPathPointInfo
 	{
 		FPathPointInfo() 
@@ -580,7 +560,7 @@ namespace
 #if DEBUG_DRAW_OFFSET
 			DrawDebugLine( GInternalDebugWorld_, Left, Right, FColor::White, true );
 #endif
-			bool bIntersected = CalculateSegmentIntersection2D( Left, Right, StartTrace, EndTrace, IntersectionPoint );
+			bool bIntersected = FMath::SegmentIntersection2D(Left, Right, StartTrace, EndTrace, IntersectionPoint);
 			if ( !bIntersected)
 			{
 				const float EdgeHalfLength = (CurrentEdge->Left - CurrentEdge->Right).Size() * 0.5;
