Added FVector2D::ToDirectionAndLength().

--------
Integrated using branch UE4-To-UE4-Orion (reversed) of change#2663780 by [REDACTED_PERSON] on 2015/08/20 20:44:30.

diff --git a/Engine/Source/Runtime/Core/Public/Math/Vector2D.h b/Engine/Source/Runtime/Core/Public/Math/Vector2D.h
--- a/Engine/Source/Runtime/Core/Public/Math/Vector2D.h
+++ b/Engine/Source/Runtime/Core/Public/Math/Vector2D.h
@@ -404,6 +404,14 @@ public:
 	 */
 	bool IsNearlyZero( float Tolerance=KINDA_SMALL_NUMBER ) const;
 
+	/**
+	 * Util to convert this vector into a unit direction vector and its original length.
+	 *
+	 * @param OutDir Reference passed in to store unit direction vector.
+	 * @param OutLength Reference passed in to store length of the vector.
+	 */
+	void ToDirectionAndLength(FVector2D &OutDir, float &OutLength) const;
+
 	/**
 	 * Checks whether all components of the vector are exactly zero.
 	 *
@@ -801,6 +809,21 @@ FORCEINLINE void FVector2D::Normalize(float Tolerance)
 }
 
 
+FORCEINLINE void FVector2D::ToDirectionAndLength(FVector2D &OutDir, float &OutLength) const
+{
+	OutLength = Size();
+	if (OutLength > SMALL_NUMBER)
+	{
+		float OneOverLength = 1.0f / OutLength;
+		OutDir = FVector2D(X*OneOverLength, Y*OneOverLength);
+	}
+	else
+	{
+		OutDir = FVector2D::ZeroVector;
+	}
+}
+
+
 FORCEINLINE bool FVector2D::IsNearlyZero(float Tolerance) const
 {
 	return	FMath::Abs(X)<=Tolerance
