PR #1365: FVector::DistSquaredXY ? Rama (Contributed by [REDACTED_PERSON])

diff --git a/Engine/Source/Runtime/Core/Public/Math/Vector.h b/Engine/Source/Runtime/Core/Public/Math/Vector.h
--- a/Engine/Source/Runtime/Core/Public/Math/Vector.h
+++ b/Engine/Source/Runtime/Core/Public/Math/Vector.h
@@ -818,6 +818,15 @@ public:
 	 */
 	static FORCEINLINE float DistSquared( const FVector &V1, const FVector &V2 );
 
+	/**
+	 * Squared distance between two points in the XY plane only.
+	 *	
+	 * @param V1 The first point.
+	 * @param V2 The second point.
+	 * @return The squared distance between two points in the XY plane
+	 */
+	static FORCEINLINE float DistSquaredXY( const FVector &V1, const FVector &V2 );
+	
 	/**
 	 * Compute pushout of a box from a plane.
 	 *
@@ -1911,6 +1920,11 @@ FORCEINLINE float FVector::DistSquared( const FVector &V1, const FVector &V2 )
 	return FMath::Square(V2.X-V1.X) + FMath::Square(V2.Y-V1.Y) + FMath::Square(V2.Z-V1.Z);
 }
 
+FORCEINLINE float FVector::DistSquaredXY( const FVector &V1, const FVector &V2 )
+{
+	return FMath::Square(V2.X-V1.X) + FMath::Square(V2.Y-V1.Y);
+}
+
 FORCEINLINE float FVector::BoxPushOut( const FVector& Normal, const FVector& Size )
 {
 	return FMath::Abs(Normal.X*Size.X) + FMath::Abs(Normal.Y*Size.Y) + FMath::Abs(Normal.Z*Size.Z);
