PR-1285 - added CanLaunchURL blueprint function
courtesy of derekvanvliet

diff --git a/Engine/Source/Runtime/Core/Private/GenericPlatform/GenericPlatformProcess.cpp b/Engine/Source/Runtime/Core/Private/GenericPlatform/GenericPlatformProcess.cpp
--- a/Engine/Source/Runtime/Core/Private/GenericPlatform/GenericPlatformProcess.cpp
+++ b/Engine/Source/Runtime/Core/Private/GenericPlatform/GenericPlatformProcess.cpp
@@ -177,6 +177,12 @@ void FGenericPlatformProcess::LaunchURL( const TCHAR* URL, const TCHAR* Parms, F
 	UE_LOG(LogHAL, Fatal, TEXT("FGenericPlatformProcess::LaunchURL not implemented on this platform"));
 }
 
+bool FGenericPlatformProcess::CanLaunchURL(const TCHAR* URL)
+{
+	UE_LOG(LogHAL, Warning, TEXT("FGenericPlatformProcess::CanLaunchURL not implemented on this platform"));
+	return false;
+}
+
 FProcHandle FGenericPlatformProcess::CreateProc( const TCHAR* URL, const TCHAR* Parms, bool bLaunchDetached, bool bLaunchHidden, bool bLaunchReallyHidden, uint32* OutProcessID, int32 PriorityModifier, const TCHAR* OptionalWorkingDirectory, void* PipeWrite )
 {
 	UE_LOG(LogHAL, Fatal, TEXT("FGenericPlatformProcess::CreateProc not implemented on this platform"));
diff --git a/Engine/Source/Runtime/Core/Private/IOS/IOSPlatformProcess.cpp b/Engine/Source/Runtime/Core/Private/IOS/IOSPlatformProcess.cpp
--- a/Engine/Source/Runtime/Core/Private/IOS/IOSPlatformProcess.cpp
+++ b/Engine/Source/Runtime/Core/Private/IOS/IOSPlatformProcess.cpp
@@ -51,6 +51,15 @@ void FIOSPlatformProcess::LaunchURL( const TCHAR* URL, const TCHAR* Parms, FStri
 	}
 }
 
+bool FIOSPlatformProcess::CanLaunchURL(const TCHAR* URL)
+{
+	CFStringRef CFUrl = FPlatformString::TCHARToCFString(URL);
+	CFUrl = CFURLCreateStringByAddingPercentEscapes(NULL, CFUrl, CFSTR("#?+"), NULL, kCFStringEncodingUTF8);
+	bool Result = [[UIApplication sharedApplication] canOpenURL:[NSURL URLWithString : (NSString *)CFUrl]];
+	CFRelease(CFUrl);
+
+	return Result;
+}
 void FIOSPlatformProcess::SetRealTimeMode()
 {
 	if ([IOSAppDelegate GetDelegate].OSVersion < 7 && FPlatformMisc::NumberOfCores() > 1)
diff --git a/Engine/Source/Runtime/Core/Public/GenericPlatform/GenericPlatformProcess.h b/Engine/Source/Runtime/Core/Public/GenericPlatform/GenericPlatformProcess.h
--- a/Engine/Source/Runtime/Core/Public/GenericPlatform/GenericPlatformProcess.h
+++ b/Engine/Source/Runtime/Core/Public/GenericPlatform/GenericPlatformProcess.h
@@ -315,6 +315,11 @@ struct CORE_API FGenericPlatformProcess
 	 */
 	static void LaunchURL( const TCHAR* URL, const TCHAR* Parms, FString* Error );
 
+	/**
+	 * Checks if the platform can launch a uniform resource locator (i.e. http://www.epicgames.com/unreal).
+	 **/
+	static bool CanLaunchURL(const TCHAR* URL);
+
 	/**
 	 * Creates a new process and its primary thread. The new process runs the
 	 * specified executable file in the security context of the calling process.
diff --git a/Engine/Source/Runtime/Core/Public/IOS/IOSPlatformProcess.h b/Engine/Source/Runtime/Core/Public/IOS/IOSPlatformProcess.h
--- a/Engine/Source/Runtime/Core/Public/IOS/IOSPlatformProcess.h
+++ b/Engine/Source/Runtime/Core/Public/IOS/IOSPlatformProcess.h
@@ -31,6 +31,7 @@ struct CORE_API FIOSPlatformProcess : public FGenericPlatformProcess
 	static const TCHAR* BaseDir();
 	static FRunnableThread* CreateRunnableThread();
 	static void LaunchURL( const TCHAR* URL, const TCHAR* Parms, FString* Error );
+	static bool CanLaunchURL(const TCHAR* URL);
 	static void SetRealTimeMode();
 	static void SetupGameOrRenderThread(bool bIsRenderThread);
 	static void SetThreadAffinityMask(uint64 AffinityMask);
diff --git a/Engine/Source/Runtime/Engine/Classes/Kismet/KismetSystemLibrary.h b/Engine/Source/Runtime/Engine/Classes/Kismet/KismetSystemLibrary.h
--- a/Engine/Source/Runtime/Engine/Classes/Kismet/KismetSystemLibrary.h
+++ b/Engine/Source/Runtime/Engine/Classes/Kismet/KismetSystemLibrary.h
@@ -1357,6 +1357,9 @@ class ENGINE_API UKismetSystemLibrary : public UBlueprintFunctionLibrary
 	UFUNCTION(BlueprintCallable, Category = "Utilities|Platform")
 	static void LaunchURL(const FString& URL);
 
+	UFUNCTION(BlueprintCallable, Category = "Utilities|Platform")
+	static bool CanLaunchURL(const FString& URL);
+
 	// Deletes all unreferenced objects, keeping only referenced objects (this command will be queued and happen at the end of the frame)
 	// Note: This can be a slow operation, and should only be performed where a hitch would be acceptable
 	UFUNCTION(BlueprintCallable, Category = "Utilities|Platform")
diff --git a/Engine/Source/Runtime/Engine/Private/KismetSystemLibrary.cpp b/Engine/Source/Runtime/Engine/Private/KismetSystemLibrary.cpp
--- a/Engine/Source/Runtime/Engine/Private/KismetSystemLibrary.cpp
+++ b/Engine/Source/Runtime/Engine/Private/KismetSystemLibrary.cpp
@@ -3030,6 +3030,15 @@ void UKismetSystemLibrary::LaunchURL(const FString& URL)
 	}
 }
 
+bool UKismetSystemLibrary::CanLaunchURL(const FString& URL)
+{
+	if (!URL.IsEmpty())
+	{
+		return FPlatformProcess::CanLaunchURL(*URL);
+	}
+
+	return false;
+}
 void UKismetSystemLibrary::CollectGarbage()
 {
 	GEngine->DeferredCommands.Add(TEXT("obj gc"));
