Stats2: fix trying to access TLS value with invalid TLS key.

- Fixes Linux editor crashing on start.
- Stats messages can be sent during global constructor initialization (e.g. creating global SNullWidget instance in SNullWidget.cpp sends one).
- Accesing FPlatformTLS::GetTlsValue() with key (slot id) that was not allocated using FPlatformTLS:AllocTlsSlot() results is bad for pthread-based platforms.

diff --git a/Engine/Source/Runtime/Core/Public/Stats/Stats2.h b/Engine/Source/Runtime/Core/Public/Stats/Stats2.h
--- a/Engine/Source/Runtime/Core/Public/Stats/Stats2.h
+++ b/Engine/Source/Runtime/Core/Public/Stats/Stats2.h
@@ -1366,7 +1366,7 @@ public:
 	template<typename TValue>
 	static FORCEINLINE_STATS void AddMessage(FName InStatName, EStatOperation::Type InStatOperation, TValue Value, bool bIsCycle = false)
 	{
-		if (!InStatName.IsNone() && WillEverCollectData())
+		if (!InStatName.IsNone() && WillEverCollectData() && IsThreadingReady())
 		{
 			FThreadStats* ThreadStats = GetThreadStats();
 			ThreadStats->AddStatMessage(FStatMessage(InStatName, InStatOperation, Value, bIsCycle));
