Fix to FMath::SegmentIntersection2D comment #UE4

diff --git a/Engine/Source/Runtime/Core/Public/Math/UnrealMathUtility.h b/Engine/Source/Runtime/Core/Public/Math/UnrealMathUtility.h
--- a/Engine/Source/Runtime/Core/Public/Math/UnrealMathUtility.h
+++ b/Engine/Source/Runtime/Core/Public/Math/UnrealMathUtility.h
@@ -1019,19 +1019,19 @@ struct FMath : public FPlatformMath
 	 * @param EndPoint   - end point of segment
 	 * @param Plane		- plane to intersect with
 	 * @param out_IntersectionPoint - out var for the point on the segment that intersects the mesh (if any)
-	 * @return true if intersection occured
+	 * @return true if intersection occurred
 	 */
 	static CORE_API bool SegmentPlaneIntersection(const FVector& StartPoint, const FVector& EndPoint, const FPlane& Plane, FVector& out_IntersectionPoint);
 
 	/**
 	 * Returns true if there is an intersection between the segment specified by SegmentStartA and SegmentEndA, and
-	 * the segment specified by SegmentEndA and SegmentStartB, in 2D space. If there is an intersection, the point is placed in out_IntersectionPoint
+	 * the segment specified by SegmentStartB and SegmentEndB, in 2D space. If there is an intersection, the point is placed in out_IntersectionPoint
 	 * @param SegmentStartA - start point of first segment
 	 * @param SegmentEndA   - end point of first segment
 	 * @param SegmentStartB - start point of second segment
 	 * @param SegmentEndB   - end point of second segment
 	 * @param out_IntersectionPoint - out var for the intersection point (if any)
-	 * @return true if intersection occured
+	 * @return true if intersection occurred
 	 */
 	static CORE_API bool SegmentIntersection2D(const FVector& SegmentStartA, const FVector& SegmentEndA, const FVector& SegmentStartB, const FVector& SegmentEndB, FVector& out_IntersectionPoint);
 
