FCachedReadPlatformFile: default to off on Windows (UEPLAT-1017).

- The buffering in WindowsPlatformFile seems to be enough.
- Added an explicit switch to turn it back on in case broader testing proves the above wrong.

(Merging CL 2666741 to main)

diff --git a/Engine/Source/Runtime/Core/Public/HAL/IPlatformFileCachedWrapper.h b/Engine/Source/Runtime/Core/Public/HAL/IPlatformFileCachedWrapper.h
--- a/Engine/Source/Runtime/Core/Public/HAL/IPlatformFileCachedWrapper.h
+++ b/Engine/Source/Runtime/Core/Public/HAL/IPlatformFileCachedWrapper.h
@@ -229,7 +229,8 @@ public:
 	}
 	virtual bool ShouldBeUsed(IPlatformFile* Inner, const TCHAR* CmdLine) const override
 	{
-		bool bResult = FPlatformProperties::RequiresCookedData();
+		// default to false on Windows since FAsyncBufferedFileReaderWindows already buffers the data
+		bool bResult = !PLATFORM_WINDOWS && FPlatformProperties::RequiresCookedData();
 
 		// Allow a choice between shorter load times or less memory on desktop platforms.
 		// Note: this cannot be in config since they aren't read at that point.
@@ -239,6 +240,10 @@ public:
 			{
 				bResult = false;
 			}
+			else if (FParse::Param(CmdLine, TEXT("CachedReadFile")))
+			{
+				bResult = true;
+			}
 
 			UE_LOG(LogPlatformFile, Log, TEXT("%s cached read wrapper"), bResult ? TEXT("Using") : TEXT("Not using"));
 		}
