VS2015 dissable __is_pod update warning

diff --git a/Engine/Source/Runtime/Core/Public/Templates/UnrealTypeTraits.h b/Engine/Source/Runtime/Core/Public/Templates/UnrealTypeTraits.h
--- a/Engine/Source/Runtime/Core/Public/Templates/UnrealTypeTraits.h
+++ b/Engine/Source/Runtime/Core/Public/Templates/UnrealTypeTraits.h
@@ -200,10 +200,18 @@ template<> struct TIsVoidType<void const volatile> { enum { Value = true }; };
  * TIsPODType
  * @todo - POD array and member pointer detection
  */
+// __is_pod changed in VS2015, however the results are still correct for all usages I've been able to locate.
+#if _MSC_VER == 1900
+#pragma warning(push)
+// #pragma warning(disable:4647)
+#endif // _MSC_VER == 1900
 template<typename T> struct TIsPODType 
 { 
 	enum { Value = IS_POD(T) || IS_ENUM(T) || TIsArithmeticType<T>::Value || TIsPointerType<T>::Value }; 
 };
+#if _MSC_VER == 1900
+#pragma warning(pop)
+#endif // _MSC_VER == 1900
 
 /**
  * TIsFundamentalType
