Fix accidentally commented line

diff --git a/Engine/Source/Runtime/Core/Public/Templates/UnrealTypeTraits.h b/Engine/Source/Runtime/Core/Public/Templates/UnrealTypeTraits.h
--- a/Engine/Source/Runtime/Core/Public/Templates/UnrealTypeTraits.h
+++ b/Engine/Source/Runtime/Core/Public/Templates/UnrealTypeTraits.h
@@ -203,7 +203,7 @@ template<> struct TIsVoidType<void const volatile> { enum { Value = true }; };
 // __is_pod changed in VS2015, however the results are still correct for all usages I've been able to locate.
 #if _MSC_VER == 1900
 #pragma warning(push)
-// #pragma warning(disable:4647)
+#pragma warning(disable:4647)
 #endif // _MSC_VER == 1900
 template<typename T> struct TIsPODType 
 { 
