Core: Avoiding allocations in DoesConfigPropertyValueMatch
[INTEGRATE] Change 2668488 by [REDACTED_PERSON][REDACTED_HANDLE] on 2015/08/25 17:39:10

diff --git a/Engine/Source/Runtime/Core/Private/Misc/ConfigCacheIni.cpp b/Engine/Source/Runtime/Core/Private/Misc/ConfigCacheIni.cpp
--- a/Engine/Source/Runtime/Core/Private/Misc/ConfigCacheIni.cpp
+++ b/Engine/Source/Runtime/Core/Private/Misc/ConfigCacheIni.cpp
@@ -809,12 +809,9 @@ bool DoesConfigPropertyValueMatch( FConfigFile* InConfigFile, const FString& InS
 		if( Section )
 		{
 			// Start Array check, if the property is in an array, we need to iterate over all properties.
-			TArray< FString > MatchingProperties;
-			Section->MultiFind( InPropertyName, MatchingProperties );
-
-			for( int32 PropertyIndex = 0; PropertyIndex < MatchingProperties.Num() && !bFoundAMatch; PropertyIndex++ )
+			for (FConfigSection::TConstKeyIterator It(*Section, InPropertyName); It && !bFoundAMatch; ++It)
 			{
-				const FString& PropertyValue = MatchingProperties[ PropertyIndex ];
+				const FString& PropertyValue = It.Value();
 				bFoundAMatch = PropertyValue == InPropertyValue;
 
 				// if our properties don't match, run further checks
