Compile fixes for TCircularQueue::Count.

diff --git a/Engine/Source/Runtime/Core/Public/Containers/CircularQueue.h b/Engine/Source/Runtime/Core/Public/Containers/CircularQueue.h
--- a/Engine/Source/Runtime/Core/Public/Containers/CircularQueue.h
+++ b/Engine/Source/Runtime/Core/Public/Containers/CircularQueue.h
@@ -43,11 +43,11 @@ public:
 	 */
 	uint32 Count() const
 	{
-		int32 Count = Head - Tail;
+		int32 Count = Tail - Head;
 
 		if (Count < 0)
 		{
-			Count += Buffer.GetSize();
+			Count += Buffer.Capacity();
 		}
 
 		return Count;
