Fix for FString::Compare comment.

diff --git a/Engine/Source/Runtime/Core/Public/Containers/UnrealString.h b/Engine/Source/Runtime/Core/Public/Containers/UnrealString.h
--- a/Engine/Source/Runtime/Core/Public/Containers/UnrealString.h
+++ b/Engine/Source/Runtime/Core/Public/Containers/UnrealString.h
@@ -1246,7 +1246,7 @@ public:
 	 * 
 	 * @param Other 	The string test against
 	 * @param SearchCase 	Whether or not the comparison should ignore case
-	 * @return 0 if equal, -1 if less than, 1 if greater than
+	 * @return 0 if equal, negative if less than, positive if greater than
 	 */
 	FORCEINLINE int32 Compare( const FString& Other, ESearchCase::Type SearchCase = ESearchCase::CaseSensitive ) const
 	{
