Linux dedicated servers should launch crash reporter with -unattended

Merging 2671555

//depot/UE4-UT/...

to //depot/UE4/...

diff --git a/Engine/Source/Runtime/Core/Private/Linux/LinuxPlatformCrashContext.cpp b/Engine/Source/Runtime/Core/Private/Linux/LinuxPlatformCrashContext.cpp
--- a/Engine/Source/Runtime/Core/Private/Linux/LinuxPlatformCrashContext.cpp
+++ b/Engine/Source/Runtime/Core/Private/Linux/LinuxPlatformCrashContext.cpp
@@ -711,9 +711,20 @@ void DLLEXPORT GenerateCrashInfoAndLaunchReporter(const FLinuxCrashContext & Con
 			RelativePathToCrashReporter = TEXT("../../../engine/binaries/linux/crashreportclient");	// FIXME: even more painfully hard-coded
 		}
 
+		FString CrashReportClientArguments;
+
+		// Suppress the user input dialog if we're running in unattended mode
+		bool bNoDialog = FApp::IsUnattended() || IsRunningDedicatedServer();
+		if (bNoDialog)
+		{
+			CrashReportClientArguments += TEXT(" -Unattended ");
+		}
+
+		CrashReportClientArguments += CrashInfoAbsolute + TEXT("/");
+
 		// show on the console
 		printf("Starting %s\n", StringCast<ANSICHAR>(RelativePathToCrashReporter).Get());
-		FProcHandle RunningProc = FPlatformProcess::CreateProc(RelativePathToCrashReporter, *(CrashInfoAbsolute + TEXT("/")), true, false, false, NULL, 0, NULL, NULL);
+		FProcHandle RunningProc = FPlatformProcess::CreateProc(RelativePathToCrashReporter, *CrashReportClientArguments, true, false, false, NULL, 0, NULL, NULL);
 		if (FPlatformProcess::IsProcRunning(RunningProc))
 		{
 			// do not wait indefinitely
