Added an interface to bubble up console commands.

diff --git a/Engine/Source/Runtime/Core/Public/GenericPlatform/GenericApplication.h b/Engine/Source/Runtime/Core/Public/GenericPlatform/GenericApplication.h
--- a/Engine/Source/Runtime/Core/Public/GenericPlatform/GenericApplication.h
+++ b/Engine/Source/Runtime/Core/Public/GenericPlatform/GenericApplication.h
@@ -379,6 +379,9 @@ class GenericApplication
 {
 public:
 
+	DECLARE_MULTICAST_DELEGATE_OneParam( FOnConsoleCommandAdded, const FString& /*Command*/ );
+	typedef FOnConsoleCommandAdded::FDelegate FOnConsoleCommandListener;
+
 	GenericApplication( const TSharedPtr< ICursor >& InCursor )
 		: Cursor( InCursor )
 		, MessageHandler( MakeShareable( new FGenericApplicationMessageHandler() ) )
@@ -421,6 +424,10 @@ public:
 
 	virtual bool IsMouseAttached() const { return true; }
 
+	virtual void RegisterConsoleCommandListener(const FOnConsoleCommandListener& InListener) {}
+
+	virtual void AddPendingConsoleCommand(const FString& InCommand) {}
+
 	virtual FPlatformRect GetWorkArea( const FPlatformRect& CurrentWindow ) const
 	{
 		FPlatformRect OutRect;
