Linux: added prefetch implementations.

diff --git a/Engine/Source/Runtime/Core/Private/Linux/LinuxPlatformMisc.cpp b/Engine/Source/Runtime/Core/Private/Linux/LinuxPlatformMisc.cpp
--- a/Engine/Source/Runtime/Core/Private/Linux/LinuxPlatformMisc.cpp
+++ b/Engine/Source/Runtime/Core/Private/Linux/LinuxPlatformMisc.cpp
@@ -119,6 +119,25 @@ namespace
 	bool GInitializedSDL = false;
 }
 
+size_t GCacheLineSize = CACHE_LINE_SIZE;
+
+void LinuxPlatform_UpdateCacheLineSize()
+{
+	// sysfs "API", as usual ;/
+	FILE * SysFsFile = fopen("/sys/devices/system/cpu/cpu0/cache/index0/coherency_line_size", "r");
+	if (SysFsFile)
+	{
+		int SystemLineSize = 0;
+		fscanf(SysFsFile, "%d", &SystemLineSize);
+		fclose(SysFsFile);
+
+		if (SystemLineSize > 0)
+		{
+			GCacheLineSize = SystemLineSize;
+		}
+	}
+}
+
 void FLinuxPlatformMisc::PlatformInit()
 {
 	// install a platform-specific signal handler
@@ -132,6 +151,8 @@ void FLinuxPlatformMisc::PlatformInit()
 	UE_LOG(LogInit, Log, TEXT(" - we're logged in %s"), FPlatformMisc::HasBeenStartedRemotely() ? TEXT("remotely") : TEXT("locally"));
 	UE_LOG(LogInit, Log, TEXT(" - Number of physical cores available for the process: %d"), FPlatformMisc::NumberOfCores());
 	UE_LOG(LogInit, Log, TEXT(" - Number of logical cores available for the process: %d"), FPlatformMisc::NumberOfCoresIncludingHyperthreads());
+	LinuxPlatform_UpdateCacheLineSize();
+	UE_LOG(LogInit, Log, TEXT(" - Cache line size: %Zu"), GCacheLineSize);
 	UE_LOG(LogInit, Log, TEXT(" - Memory allocator used: %s"), GMalloc->GetDescriptiveName());
 
 	// programs don't need it by default
diff --git a/Engine/Source/Runtime/Core/Public/Linux/LinuxPlatformMisc.h b/Engine/Source/Runtime/Core/Public/Linux/LinuxPlatformMisc.h
--- a/Engine/Source/Runtime/Core/Public/Linux/LinuxPlatformMisc.h
+++ b/Engine/Source/Runtime/Core/Public/Linux/LinuxPlatformMisc.h
@@ -86,6 +86,24 @@ struct CORE_API FLinuxPlatformMisc : public FGenericPlatformMisc
 		__sync_synchronize();
 	}
 
+	FORCEINLINE static void PrefetchBlock(const void* InPtr, int32 NumBytes = 1)
+	{
+		extern size_t GCacheLineSize;
+
+		const char* Ptr = static_cast<const char*>(InPtr);
+		const size_t CacheLineSize = GCacheLineSize;
+		for (size_t BytesPrefetched = 0; BytesPrefetched < NumBytes; BytesPrefetched += CacheLineSize)
+		{
+			__builtin_prefetch(Ptr);
+			Ptr += CacheLineSize;
+		}
+	}
+
+	FORCEINLINE static void Prefetch(void const* Ptr, int32 Offset = 0)
+	{
+		__builtin_prefetch(static_cast<char const*>(Ptr) + Offset);
+	}
+
 	static int32 NumberOfCores();
 	static int32 NumberOfCoresIncludingHyperthreads();
 	static void LoadPreInitModules();
