TLS: harden slot allocation code to enforce assumptions (Windows, Linux).

diff --git a/Engine/Source/Runtime/Core/Public/Linux/LinuxPlatformTLS.h b/Engine/Source/Runtime/Core/Public/Linux/LinuxPlatformTLS.h
--- a/Engine/Source/Runtime/Core/Public/Linux/LinuxPlatformTLS.h
+++ b/Engine/Source/Runtime/Core/Public/Linux/LinuxPlatformTLS.h
@@ -45,14 +45,36 @@ struct CORE_API FLinuxTLS : public FGenericPlatformTLS
 	/**
 	 * Allocates a thread local store slot
 	 */
-	static FORCEINLINE uint32 AllocTlsSlot(void)
+	static uint32 AllocTlsSlot(void)
 	{
 		// allocate a per-thread mem slot
 		pthread_key_t Key = 0;
-		if (pthread_key_create(&Key, NULL) != 0)
+		if (pthread_key_create(&Key, nullptr) != 0)
 		{
-			Key = 0xFFFFFFFF;  // matches the Windows TlsAlloc() retval //@todo android: should probably check for this below, or assert out instead
+			return static_cast<uint32>(INDEX_NONE); // matches the Windows TlsAlloc() retval.
 		}
+
+		// pthreads can return an arbitrary key, yet we reserve INDEX_NONE as an invalid one. Handle this very unlikely case
+		// by allocating another one first (so we get another value) and releasing existing key.
+		if (static_cast<uint32>(Key) == static_cast<uint32>(INDEX_NONE))
+		{
+			pthread_key_t NewKey = 0;
+			int SecondKeyAllocResult = pthread_key_create(&NewKey, nullptr);
+			// discard the previous one
+			pthread_key_delete((pthread_key_t)Key);
+
+			if (SecondKeyAllocResult != 0)
+			{
+				// could not alloc the second key, treat this as an error
+				return static_cast<uint32>(INDEX_NONE); // matches the Windows TlsAlloc() retval.
+			}
+
+			// check that we indeed got something different
+			checkf(NewKey != static_cast<uint32>(INDEX_NONE), TEXT("Could not allocate a usable TLS slot id."));
+
+			Key = NewKey;
+		}
+
 		return Key;
 	}
 
diff --git a/Engine/Source/Runtime/Core/Public/Windows/WindowsPlatformTLS.h b/Engine/Source/Runtime/Core/Public/Windows/WindowsPlatformTLS.h
--- a/Engine/Source/Runtime/Core/Public/Windows/WindowsPlatformTLS.h
+++ b/Engine/Source/Runtime/Core/Public/Windows/WindowsPlatformTLS.h
@@ -30,6 +30,7 @@ struct CORE_API FWindowsPlatformTLS
 	 */
 	static FORCEINLINE uint32 AllocTlsSlot(void)
 	{
+		static_assert(static_cast<uint32>(INDEX_NONE) == TLS_OUT_OF_INDEXES, "TLS_OUT_OF_INDEXES is different from INDEX_NONE, change FWindowsPlatformTLS::AllocTlsSlot() implementation.");
 		return ::TlsAlloc();
 	}
 
