Updated disabled test flags.
Removed outdated test case.

diff --git a/Engine/Plugins/Messaging/UdpMessaging/Source/UdpMessaging/Private/Transport/Tests/UdpMessageSegmenterTest.cpp b/Engine/Plugins/Messaging/UdpMessaging/Source/UdpMessaging/Private/Transport/Tests/UdpMessageSegmenterTest.cpp
deleted file mode 100644
diff --git a/Engine/Source/Runtime/Core/Private/Tests/Math/RangeSetTest.cpp b/Engine/Source/Runtime/Core/Private/Tests/Math/RangeSetTest.cpp
--- a/Engine/Source/Runtime/Core/Private/Tests/Math/RangeSetTest.cpp
+++ b/Engine/Source/Runtime/Core/Private/Tests/Math/RangeSetTest.cpp
@@ -4,7 +4,7 @@
 #include "AutomationTest.h"
 
 
-IMPLEMENT_SIMPLE_AUTOMATION_TEST(FRangeSetTest, "System.Core.Math.RangeSet", EAutomationTestFlags::EditorContext | EAutomationTestFlags::ClientContext | EAutomationTestFlags::SmokeFilter)
+IMPLEMENT_SIMPLE_AUTOMATION_TEST(FRangeSetTest, "System.Core.Math.RangeSet", EAutomationTestFlags::Disabled | EAutomationTestFlags::EditorContext | EAutomationTestFlags::ClientContext | EAutomationTestFlags::SmokeFilter)
 
 bool FRangeSetTest::RunTest( const FString& Parameters )
 {
diff --git a/Engine/Source/Runtime/Core/Private/Tests/Misc/TypeContainerTest.cpp b/Engine/Source/Runtime/Core/Private/Tests/Misc/TypeContainerTest.cpp
--- a/Engine/Source/Runtime/Core/Private/Tests/Misc/TypeContainerTest.cpp
+++ b/Engine/Source/Runtime/Core/Private/Tests/Misc/TypeContainerTest.cpp
@@ -11,7 +11,7 @@
 #include "AutomationTest.h"
 
 
-IMPLEMENT_SIMPLE_AUTOMATION_TEST(FTypeContainerTest, "System.Core.Misc.TypeContainer", EAutomationTestFlags::ApplicationContextMask | EAutomationTestFlags::SmokeFilter)
+IMPLEMENT_SIMPLE_AUTOMATION_TEST(FTypeContainerTest, "System.Core.Misc.TypeContainer", EAutomationTestFlags::Disabled | EAutomationTestFlags::ApplicationContextMask | EAutomationTestFlags::SmokeFilter)
 
 
 /* Helpers
