[INTEGRATE] merge ue4<-ue4-launcherDev cl 2676696

diff --git a/Engine/Source/Runtime/Core/Private/HAL/FileManagerGeneric.cpp b/Engine/Source/Runtime/Core/Private/HAL/FileManagerGeneric.cpp
--- a/Engine/Source/Runtime/Core/Private/HAL/FileManagerGeneric.cpp
+++ b/Engine/Source/Runtime/Core/Private/HAL/FileManagerGeneric.cpp
@@ -587,12 +587,12 @@ void FFileManagerGeneric::FindFilesRecursiveInternal( TArray<FString>& FileNames
 }
 
 FArchiveFileReaderGeneric::FArchiveFileReaderGeneric( IFileHandle* InHandle, const TCHAR* InFilename, int64 InSize )
-	:	Filename		( InFilename )
-	,   Size           ( InSize )
-	,   Pos            ( 0 )
-	,   BufferBase     ( 0 )
-	,   BufferCount    ( 0 )
-	,		Handle( InHandle )
+	: Filename( InFilename )
+	, Size( InSize )
+	, Pos( 0 )
+	, BufferBase( 0 )
+	, BufferCount( 0 )
+	, Handle( InHandle )
 {
 	ArIsLoading = ArIsPersistent = true;
 }
@@ -718,10 +718,11 @@ void FArchiveFileReaderGeneric::Serialize( void* V, int64 Length )
 }
 
 FArchiveFileWriterGeneric::FArchiveFileWriterGeneric( IFileHandle* InHandle, const TCHAR* InFilename, int64 InPos )
-	:	Filename	( InFilename )
-	,   Pos			( InPos )
-	,   BufferCount	( 0 )
-	,	Handle		( InHandle )
+	: Filename( InFilename )
+	, Pos( InPos )
+	, BufferCount( 0 )
+	, Handle( InHandle )
+	, bLoggingError( false )
 {
 	ArIsSaving = ArIsPersistent = true;
 }
@@ -825,8 +826,14 @@ void FArchiveFileWriterGeneric::Flush()
 
 void FArchiveFileWriterGeneric::LogWriteError(const TCHAR* Message)
 {
-	TCHAR ErrorBuffer[1024];
-	UE_LOG(LogFileManager, Error, TEXT("%s: %s (%s)"), Message, *Filename, FPlatformMisc::GetSystemErrorMessage(ErrorBuffer, 1024, 0));
+	// Prevent re-entry if logging causes another log error leading to a stack overflow
+	if (!bLoggingError)
+	{
+		bLoggingError = true;
+		TCHAR ErrorBuffer[1024];
+		UE_LOG(LogFileManager, Error, TEXT("%s: %s (%s)"), Message, *Filename, FPlatformMisc::GetSystemErrorMessage(ErrorBuffer, 1024, 0));
+		bLoggingError = false;
+	}
 }
 //---
 
diff --git a/Engine/Source/Runtime/Core/Public/HAL/FileManagerGeneric.h b/Engine/Source/Runtime/Core/Public/HAL/FileManagerGeneric.h
--- a/Engine/Source/Runtime/Core/Public/HAL/FileManagerGeneric.h
+++ b/Engine/Source/Runtime/Core/Public/HAL/FileManagerGeneric.h
@@ -228,13 +228,13 @@ protected:
 	virtual void ReadLowLevel(uint8* Dest, int64 CountToRead, int64& OutBytesRead);
 
 	/** Filename for debugging purposes. */
-	FString			Filename;
-	int64            Size;
-	int64            Pos;
-	int64            BufferBase;
-	int64            BufferCount;
-	TAutoPtr<IFileHandle>	Handle;
-	uint8            Buffer[1024];	
+	FString Filename;
+	int64 Size;
+	int64 Pos;
+	int64 BufferBase;
+	int64 BufferCount;
+	TAutoPtr<IFileHandle> Handle;
+	uint8 Buffer[1024];
 };
 
 
@@ -292,9 +292,10 @@ protected:
 	void LogWriteError(const TCHAR* Message);
 
 	/** Filename for debugging purposes */
-	FString			Filename;
-	int64            Pos;
-	int64            BufferCount;
-	TAutoPtr<IFileHandle>		 Handle;
-	uint8            Buffer[4096];
+	FString Filename;
+	int64 Pos;
+	int64 BufferCount;
+	TAutoPtr<IFileHandle> Handle;
+	uint8 Buffer[4096];
+	bool bLoggingError;
 };
