Fixed a double FText history allocation in AsNumber, AsPercent, AsCurrency, AsDate, AsTime, AsDateTime, and Format

The FText constructor taking a single FString would always allocate a default history object, which would later be replaced with the correct history object for the functions listed above.

This constructor has been changed to take an optional history object. If one is provided, the FText will just use that, otherwise it will allocate a default history object.

diff --git a/Engine/Source/Runtime/Core/Private/Internationalization/ICUText.cpp b/Engine/Source/Runtime/Core/Private/Internationalization/ICUText.cpp
--- a/Engine/Source/Runtime/Core/Private/Internationalization/ICUText.cpp
+++ b/Engine/Source/Runtime/Core/Private/Internationalization/ICUText.cpp
@@ -44,10 +44,7 @@ FText FText::AsDate(const FDateTime& DateTime, const EDateTimeStyle::Type DateSt
 	FString NativeString;
 	ICUUtilities::ConvertString(FormattedString, NativeString);
 
-	FText ResultText = FText::CreateChronologicalText(MoveTemp(NativeString));
-	ResultText.History = MakeShareable(new FTextHistory_AsDate(DateTime, DateStyle, TimeZone, TargetCulture));
-
-	return ResultText;
+	return FText::CreateChronologicalText(MoveTemp(NativeString), MakeShareable(new FTextHistory_AsDate(DateTime, DateStyle, TimeZone, TargetCulture)));
 }
 
 FText FText::AsTime(const FDateTime& DateTime, const EDateTimeStyle::Type TimeStyle, const FString& TimeZone, const FCulturePtr& TargetCulture)
@@ -67,10 +64,7 @@ FText FText::AsTime(const FDateTime& DateTime, const EDateTimeStyle::Type TimeSt
 	FString NativeString;
 	ICUUtilities::ConvertString(FormattedString, NativeString);
 
-	FText ResultText = FText::CreateChronologicalText(MoveTemp(NativeString));
-	ResultText.History = MakeShareable(new FTextHistory_AsTime(DateTime, TimeStyle, TimeZone, TargetCulture));
-
-	return ResultText;
+	return FText::CreateChronologicalText(MoveTemp(NativeString), MakeShareable(new FTextHistory_AsTime(DateTime, TimeStyle, TimeZone, TargetCulture)));
 }
 
 FText FText::AsTimespan(const FTimespan& Timespan, const FCulturePtr& TargetCulture)
@@ -114,10 +108,7 @@ FText FText::AsDateTime(const FDateTime& DateTime, const EDateTimeStyle::Type Da
 	FString NativeString;
 	ICUUtilities::ConvertString(FormattedString, NativeString);
 
-	FText ResultText = FText::CreateChronologicalText(MoveTemp(NativeString));
-	ResultText.History = MakeShareable(new FTextHistory_AsDateTime(DateTime, DateStyle, TimeStyle, TimeZone, TargetCulture));
-
-	return ResultText;
+	return FText::CreateChronologicalText(MoveTemp(NativeString), MakeShareable(new FTextHistory_AsDateTime(DateTime, DateStyle, TimeStyle, TimeZone, TargetCulture)));
 }
 
 FText FText::AsMemory(SIZE_T NumBytes, const FNumberFormattingOptions* const Options, const FCulturePtr& TargetCulture)
diff --git a/Engine/Source/Runtime/Core/Private/Internationalization/ICUText.h b/Engine/Source/Runtime/Core/Private/Internationalization/ICUText.h
--- a/Engine/Source/Runtime/Core/Private/Internationalization/ICUText.h
+++ b/Engine/Source/Runtime/Core/Private/Internationalization/ICUText.h
@@ -25,9 +25,7 @@ FText FText::AsNumberTemplate(T1 Val, const FNumberFormattingOptions* const Opti
 	FString NativeString;
 	ICUUtilities::ConvertString(FormattedString, NativeString);
 
-	FText ReturnText = FText::CreateNumericalText(MoveTemp(NativeString));
-	ReturnText.History = MakeShareable(new FTextHistory_AsNumber(Val, Options, TargetCulture));
-	return ReturnText;
+	return FText::CreateNumericalText(MoveTemp(NativeString), MakeShareable(new FTextHistory_AsNumber(Val, Options, TargetCulture)));
 }
 
 template<typename T1, typename T2>
@@ -45,9 +43,7 @@ FText FText::AsCurrencyTemplate(T1 Val, const FString& CurrencyCode, const FNumb
 	FString NativeString;
 	ICUUtilities::ConvertString(FormattedString, NativeString);
 
-	FText ReturnText = FText::CreateNumericalText(MoveTemp(NativeString));
-	ReturnText.History = MakeShareable(new FTextHistory_AsCurrency(Val, CurrencyCode, Options, TargetCulture));
-	return ReturnText;
+	return FText::CreateNumericalText(MoveTemp(NativeString), MakeShareable(new FTextHistory_AsCurrency(Val, CurrencyCode, Options, TargetCulture)));
 }
 
 template<typename T1, typename T2>
@@ -65,7 +61,5 @@ FText FText::AsPercentTemplate(T1 Val, const FNumberFormattingOptions* const Opt
 	FString NativeString;
 	ICUUtilities::ConvertString(FormattedString, NativeString);
 
-	FText ReturnText = FText::CreateNumericalText(MoveTemp(NativeString));
-	ReturnText.History = MakeShareable(new FTextHistory_AsPercent(Val, Options, TargetCulture));
-	return ReturnText;
+	return FText::CreateNumericalText(MoveTemp(NativeString), MakeShareable(new FTextHistory_AsPercent(Val, Options, TargetCulture)));
 }
diff --git a/Engine/Source/Runtime/Core/Private/Internationalization/LegacyText.h b/Engine/Source/Runtime/Core/Private/Internationalization/LegacyText.h
--- a/Engine/Source/Runtime/Core/Private/Internationalization/LegacyText.h
+++ b/Engine/Source/Runtime/Core/Private/Internationalization/LegacyText.h
@@ -8,7 +8,7 @@ FText FText::AsNumberTemplate(T1 Val, const FNumberFormattingOptions* const Opti
 	FInternationalization& I18N = FInternationalization::Get();
 	checkf(I18N.IsInitialized() == true, TEXT("FInternationalization is not initialized. An FText formatting method was likely used in static object initialization - this is not supported."));
 	const FCulture& Culture = TargetCulture.IsValid() ? *TargetCulture : *I18N.GetCurrentCulture();
-	return FText::CreateNumericalText( Culture.NumberFormattingRule.AsNumber(Val) );
+	return FText::CreateNumericalText( Culture.NumberFormattingRule.AsNumber(Val), nullptr );
 }
 
 template<typename T1, typename T2>
@@ -17,7 +17,7 @@ FText FText::AsCurrencyTemplate(T1 Val, const FString& CurrencyCode, const FNumb
 	FInternationalization& I18N = FInternationalization::Get();
 	checkf(I18N.IsInitialized() == true, TEXT("FInternationalization is not initialized. An FText formatting method was likely used in static object initialization - this is not supported."));
 	const FCulture& Culture = TargetCulture.IsValid() ? *TargetCulture : *I18N.GetCurrentCulture();
-	return FText::CreateNumericalText( Culture.NumberFormattingRule.AsCurrency(Val) );
+	return FText::CreateNumericalText( Culture.NumberFormattingRule.AsCurrency(Val), nullptr );
 }
 
 template<typename T1, typename T2>
@@ -26,5 +26,5 @@ FText FText::AsPercentTemplate(T1 Val, const FNumberFormattingOptions* const Opt
 	FInternationalization& I18N = FInternationalization::Get();
 	checkf(I18N.IsInitialized() == true, TEXT("FInternationalization is not initialized. An FText formatting method was likely used in static object initialization - this is not supported."));
 	const FCulture& Culture = TargetCulture.IsValid() ? *TargetCulture : *I18N.GetCurrentCulture();
-	return FText::CreateNumericalText( Culture.NumberFormattingRule.AsPercent(Val) );
+	return FText::CreateNumericalText( Culture.NumberFormattingRule.AsPercent(Val), nullptr );
 }
\ No newline at end of file
diff --git a/Engine/Source/Runtime/Core/Private/Internationalization/Text.cpp b/Engine/Source/Runtime/Core/Private/Internationalization/Text.cpp
--- a/Engine/Source/Runtime/Core/Private/Internationalization/Text.cpp
+++ b/Engine/Source/Runtime/Core/Private/Internationalization/Text.cpp
@@ -322,11 +322,15 @@ FText& FText::operator=(FText&& Source)
 }
 #endif
 
-FText::FText( FString InSourceString )
+FText::FText( FString InSourceString, TSharedPtr<class FTextHistory, ESPMode::ThreadSafe> InHistory )
 	: DisplayString( new FString( MoveTemp(InSourceString) ))
+	, History( MoveTemp(InHistory) )
 	, Flags(0)
 {
-	History = MakeShareable(new FTextHistory_Base(DisplayString));
+	if (!History.IsValid())
+	{
+		History = MakeShareable(new FTextHistory_Base(DisplayString));
+	}
 }
 
 FText::FText( FString InSourceString, FString InNamespace, FString InKey, int32 InFlags )
@@ -373,7 +377,7 @@ FText FText::TrimPreceding( const FText& InText )
 		TrimmedString = TrimmedString.Right( TrimmedString.Len() - StartPos );
 	}
 
-	FText NewText = FText( MoveTemp( TrimmedString ) );
+	FText NewText = FText( MoveTemp( TrimmedString ), nullptr );
 
 	if (!GIsEditor)
 	{
@@ -408,7 +412,7 @@ FText FText::TrimTrailing( const FText& InText )
 		TrimmedString = TrimmedString.Left( EndPos + 1 );
 	}
 
-	FText NewText = FText( MoveTemp ( TrimmedString ) );
+	FText NewText = FText( MoveTemp ( TrimmedString ), nullptr );
 
 	if (!GIsEditor)
 	{
@@ -455,7 +459,7 @@ FText FText::TrimPrecedingAndTrailing( const FText& InText )
 		TrimmedString = TrimmedString.Mid( StartPos, Len );
 	}
 
-	FText NewText = FText( MoveTemp( TrimmedString ) );
+	FText NewText = FText( MoveTemp( TrimmedString ), nullptr );
 
 	if (!GIsEditor)
 	{
@@ -836,8 +840,7 @@ public:
 			}
 		}
 
-		FText Result = FText(MoveTemp(ResultString));
-		Result.History = InHistory;
+		FText Result = FText(MoveTemp(ResultString), InHistory);
 		if (!GIsEditor)
 		{
 			Result.Flags = Result.Flags | ETextFlag::Transient;
@@ -1147,9 +1150,9 @@ FText FText::ChangeKey( FString Namespace, FString Key, const FText& Text )
 }
 #endif
 
-FText FText::CreateNumericalText(FString InSourceString)
+FText FText::CreateNumericalText(FString InSourceString, TSharedPtr<class FTextHistory, ESPMode::ThreadSafe> InHistory)
 {
-	FText NewText = FText( MoveTemp( InSourceString ) );
+	FText NewText = FText( MoveTemp( InSourceString ), MoveTemp( InHistory ) );
 	if (!GIsEditor)
 	{
 		NewText.Flags |= ETextFlag::Transient;
@@ -1157,9 +1160,9 @@ FText FText::CreateNumericalText(FString InSourceString)
 	return NewText;
 }
 
-FText FText::CreateChronologicalText(FString InSourceString)
+FText FText::CreateChronologicalText(FString InSourceString, TSharedPtr<class FTextHistory, ESPMode::ThreadSafe> InHistory)
 {
-	FText NewText = FText( MoveTemp( InSourceString ) );
+	FText NewText = FText( MoveTemp( InSourceString ), MoveTemp( InHistory ) );
 	if (!GIsEditor)
 	{
 		NewText.Flags |= ETextFlag::Transient;
@@ -1169,7 +1172,7 @@ FText FText::CreateChronologicalText(FString InSourceString)
 
 FText FText::FromName( const FName& Val) 
 {
-	FText NewText = FText( Val.ToString() );
+	FText NewText = FText( Val.ToString(), nullptr );
 
 	if (!GIsEditor)
 	{
@@ -1181,7 +1184,7 @@ FText FText::FromName( const FName& Val)
 
 FText FText::FromString( FString String )
 {
-	FText NewText = FText( MoveTemp(String) );
+	FText NewText = FText( MoveTemp(String), nullptr );
 
 	if (!GIsEditor)
 	{
@@ -1193,7 +1196,7 @@ FText FText::FromString( FString String )
 
 FText FText::AsCultureInvariant( FString String )
 {
-	FText NewText = FText( String );
+	FText NewText = FText( String, nullptr );
 	NewText.Flags |= ETextFlag::CultureInvariant;
 
 	return NewText;
diff --git a/Engine/Source/Runtime/Core/Public/Internationalization/Text.h b/Engine/Source/Runtime/Core/Public/Internationalization/Text.h
--- a/Engine/Source/Runtime/Core/Public/Internationalization/Text.h
+++ b/Engine/Source/Runtime/Core/Public/Internationalization/Text.h
@@ -423,9 +423,9 @@ private:
 	enum class EInitToEmptyString : uint8 { Value };
 	explicit FText( EInitToEmptyString );
 
-	explicit FText( FString InSourceString );
+	FText( FString InSourceString, TSharedPtr<class FTextHistory, ESPMode::ThreadSafe> InHistory );
 
-	explicit FText( FString InSourceString, FString InNamespace, FString InKey, int32 InFlags=0 );
+	FText( FString InSourceString, FString InNamespace, FString InKey, int32 InFlags=0 );
 
 	friend CORE_API FArchive& operator<<( FArchive& Ar, FText& Value );
 
@@ -439,12 +439,12 @@ private:
 	/**
 	 * Generate an FText for a string formatted numerically.
 	 */
-	static FText CreateNumericalText(FString InSourceString);
+	static FText CreateNumericalText(FString InSourceString, TSharedPtr<class FTextHistory, ESPMode::ThreadSafe> InHistory);
 
 	/**
 	 * Generate an FText for a string formatted from a date/time.
 	 */
-	static FText CreateChronologicalText(FString InSourceString);
+	static FText CreateChronologicalText(FString InSourceString, TSharedPtr<class FTextHistory, ESPMode::ThreadSafe> InHistory);
 
 	/** Returns the source string of the FText */
 	TSharedPtr< FString, ESPMode::ThreadSafe > GetSourceString() const;
