Added specific rounding implementations

Added some specific float and double rounding implementations (to be used by FastDecimalFormat):
 * FMath::RoundHalfToEven
 * FMath::RoundHalfFromZero
 * FMath::RoundHalfToZero
 * FMath::RoundFromZero
 * FMath::RoundToZero
 * FMath::RoundToNegativeInfinity
 * FMath::RoundToPositiveInfinity

Added a generic platform wrapper for the float and double modf function (FMath::Modf).

Added some missing double implementations (FMath::RoundToDouble, FMath::CeilToDouble).

diff --git a/Engine/Source/Runtime/Core/Private/Math/UnrealMath.cpp b/Engine/Source/Runtime/Core/Private/Math/UnrealMath.cpp
--- a/Engine/Source/Runtime/Core/Private/Math/UnrealMath.cpp
+++ b/Engine/Source/Runtime/Core/Private/Math/UnrealMath.cpp
@@ -2613,6 +2613,99 @@ void FVector::GenerateClusterCenters(TArray<FVector>& Clusters, const TArray<FVe
 	}
 }
 
+namespace MathRoundingUtil
+{
+
+float TruncateToHalfIfClose(float F)
+{
+	float ValueToFudgeIntegralPart = 0.0f;
+	float ValueToFudgeFractionalPart = FMath::Modf(F, &ValueToFudgeIntegralPart);
+	if (F < 0.0f)
+	{
+		return ValueToFudgeIntegralPart + ((FMath::IsNearlyEqual(ValueToFudgeFractionalPart, -0.5f)) ? -0.5f : ValueToFudgeFractionalPart);
+	}
+	else
+	{
+		return ValueToFudgeIntegralPart + ((FMath::IsNearlyEqual(ValueToFudgeFractionalPart, 0.5f)) ? 0.5f : ValueToFudgeFractionalPart);
+	}
+}
+
+double TruncateToHalfIfClose(double F)
+{
+	double ValueToFudgeIntegralPart = 0.0;
+	double ValueToFudgeFractionalPart = FMath::Modf(F, &ValueToFudgeIntegralPart);
+	if (F < 0.0)
+	{
+		return ValueToFudgeIntegralPart + ((FMath::IsNearlyEqual(ValueToFudgeFractionalPart, -0.5)) ? -0.5 : ValueToFudgeFractionalPart);
+	}
+	else
+	{
+		return ValueToFudgeIntegralPart + ((FMath::IsNearlyEqual(ValueToFudgeFractionalPart, 0.5)) ? 0.5 : ValueToFudgeFractionalPart);
+	}
+}
+
+} // namespace GenericPlatformMathInternal
+
+float FMath::RoundHalfToEven(float F)
+{
+	F = MathRoundingUtil::TruncateToHalfIfClose(F);
+
+	const bool bIsNegative = F < 0.0f;
+	const bool bValueIsEven = static_cast<uint32>(FloorToFloat(((bIsNegative) ? -F : F))) % 2 == 0;
+	if (bValueIsEven)
+	{
+		// Round towards value (eg, value is -2.5 or 2.5, and should become -2 or 2)
+		return (bIsNegative) ? FloorToFloat(F + 0.5f) : CeilToFloat(F - 0.5f);
+	}
+	else
+	{
+		// Round away from value (eg, value is -3.5 or 3.5, and should become -4 or 4)
+		return (bIsNegative) ? CeilToFloat(F - 0.5f) : FloorToFloat(F + 0.5f);
+	}
+}
+
+double FMath::RoundHalfToEven(double F)
+{
+	F = MathRoundingUtil::TruncateToHalfIfClose(F);
+
+	const bool bIsNegative = F < 0.0;
+	const bool bValueIsEven = static_cast<uint64>(FMath::FloorToDouble(((bIsNegative) ? -F : F))) % 2 == 0;
+	if (bValueIsEven)
+	{
+		// Round towards value (eg, value is -2.5 or 2.5, and should become -2 or 2)
+		return (bIsNegative) ? FloorToDouble(F + 0.5) : CeilToDouble(F - 0.5);
+	}
+	else
+	{
+		// Round away from value (eg, value is -3.5 or 3.5, and should become -4 or 4)
+		return (bIsNegative) ? CeilToDouble(F - 0.5) : FloorToDouble(F + 0.5);
+	}
+}
+
+float FMath::RoundHalfFromZero(float F)
+{
+	F = MathRoundingUtil::TruncateToHalfIfClose(F);
+	return (F < 0.0f) ? CeilToFloat(F - 0.5f) : FloorToFloat(F + 0.5f);
+}
+
+double FMath::RoundHalfFromZero(double F)
+{
+	F = MathRoundingUtil::TruncateToHalfIfClose(F);
+	return (F < 0.0) ? CeilToDouble(F - 0.5) : FloorToDouble(F + 0.5);
+}
+
+float FMath::RoundHalfToZero(float F)
+{
+	F = MathRoundingUtil::TruncateToHalfIfClose(F);
+	return (F < 0.0f) ? FloorToFloat(F + 0.5f) : CeilToFloat(F - 0.5f);
+}
+
+double FMath::RoundHalfToZero(double F)
+{
+	F = MathRoundingUtil::TruncateToHalfIfClose(F);
+	return (F < 0.0) ? FloorToDouble(F + 0.5) : CeilToDouble(F - 0.5);
+}
+
 bool FMath::MemoryTest( void* BaseAddress, uint32 NumBytes )
 {
 	volatile uint32* Ptr;
diff --git a/Engine/Source/Runtime/Core/Public/GenericPlatform/GenericPlatformMath.h b/Engine/Source/Runtime/Core/Public/GenericPlatform/GenericPlatformMath.h
--- a/Engine/Source/Runtime/Core/Public/GenericPlatform/GenericPlatformMath.h
+++ b/Engine/Source/Runtime/Core/Public/GenericPlatform/GenericPlatformMath.h
@@ -117,6 +117,16 @@ struct FGenericPlatformMath
 		return FloorToFloat(F + 0.5f);
 	}
 
+	/**
+	* Converts a double to the nearest integer. Rounds up when the fraction is .5
+	* @param F		Floating point value to convert
+	* @return		The nearest integer to 'F'.
+	*/
+	static FORCEINLINE double RoundToDouble(double F)
+	{
+		return FloorToDouble(F + 0.5);
+	}
+
 	/**
 	* Converts a float to the nearest integer. Rounds up when the fraction is .5
 	* @param F		Floating point value to convert
@@ -148,6 +158,16 @@ struct FGenericPlatformMath
 		return ceilf(F);
 	}
 
+	/**
+	* Converts a double to the nearest greater or equal integer.
+	* @param F		Floating point value to convert
+	* @return		An integer greater or equal to 'F'.
+	*/
+	static FORCEINLINE double CeilToDouble(double F)
+	{
+		return ceil(F);
+	}
+
 	/**
 	* Converts a float to a greater or equal integer.
 	* @param F		Floating point value to convert
@@ -179,6 +199,28 @@ struct FGenericPlatformMath
 		return Value - FloorToFloat(Value);
 	}
 
+	/**
+	* Breaks the given value into an integral and a fractional part.
+	* @param InValue	Floating point value to convert
+	* @param OutIntPart Floating point value that receives the integral part of the number.
+	* @return			The fractional part of the number.
+	*/
+	static FORCEINLINE float Modf(const float InValue, float* OutIntPart)
+	{
+		return modff(InValue, OutIntPart);
+	}
+
+	/**
+	* Breaks the given value into an integral and a fractional part.
+	* @param InValue	Floating point value to convert
+	* @param OutIntPart Floating point value that receives the integral part of the number.
+	* @return			The fractional part of the number.
+	*/
+	static FORCEINLINE double Modf(const double InValue, double* OutIntPart)
+	{
+		return modf(InValue, OutIntPart);
+	}
+
 	static FORCEINLINE float Exp( float Value ) { return expf(Value); }
 	static FORCEINLINE float Loge( float Value ) {	return logf(Value); }
 	static FORCEINLINE float LogX( float Base, float Value ) { return Loge(Value) / Loge(Base); }
diff --git a/Engine/Source/Runtime/Core/Public/Math/UnrealMath.h b/Engine/Source/Runtime/Core/Public/Math/UnrealMath.h
--- a/Engine/Source/Runtime/Core/Public/Math/UnrealMath.h
+++ b/Engine/Source/Runtime/Core/Public/Math/UnrealMath.h
@@ -404,6 +404,46 @@ inline FVector FMath::VRand()
 	return Result * (1.0f / Sqrt(L));
 }
 
+FORCEINLINE float FMath::RoundFromZero(float F)
+{
+	return (F < 0.0f) ? FloorToFloat(F) : CeilToFloat(F);
+}
+
+FORCEINLINE double FMath::RoundFromZero(double F)
+{
+	return (F < 0.0) ? FloorToDouble(F) : CeilToDouble(F);
+}
+
+FORCEINLINE float FMath::RoundToZero(float F)
+{
+	return (F < 0.0f) ? CeilToFloat(F) : FloorToFloat(F);
+}
+
+FORCEINLINE double FMath::RoundToZero(double F)
+{
+	return (F < 0.0) ? CeilToDouble(F) : FloorToDouble(F);
+}
+
+FORCEINLINE float FMath::RoundToNegativeInfinity(float F)
+{
+	return FloorToFloat(F);
+}
+
+FORCEINLINE double FMath::RoundToNegativeInfinity(double F)
+{
+	return FloorToDouble(F);
+}
+
+FORCEINLINE float FMath::RoundToPositiveInfinity(float F)
+{
+	return CeilToFloat(F);
+}
+
+FORCEINLINE double FMath::RoundToPositiveInfinity(double F)
+{
+	return CeilToDouble(F);
+}
+
 inline FString FMath::FormatIntToHumanReadable(int32 Val)
 {
 	FString Src = *FString::Printf( TEXT("%i"), Val );
diff --git a/Engine/Source/Runtime/Core/Public/Math/UnrealMathUtility.h b/Engine/Source/Runtime/Core/Public/Math/UnrealMathUtility.h
--- a/Engine/Source/Runtime/Core/Public/Math/UnrealMathUtility.h
+++ b/Engine/Source/Runtime/Core/Public/Math/UnrealMathUtility.h
@@ -1095,6 +1095,62 @@ struct FMath : public FPlatformMath
 	 */
 	static CORE_API bool PointsAreCoplanar(const TArray<FVector>& Points, const float Tolerance = 0.1f);
 
+	/**
+	* Converts a floating point number to the nearest integer, equidistant ties go to the value which is closest to an even value: 1.5 becomes 2, 0.5 becomes 0
+	* @param F		Floating point value to convert
+	* @return		The rounded integer
+	*/
+	static CORE_API float RoundHalfToEven(float F);
+	static CORE_API double RoundHalfToEven(double F);
+
+	/**
+	* Converts a floating point number to the nearest integer, equidistant ties go to the value which is further from zero: -0.5 becomes -1.0, 0.5 becomes 1.0
+	* @param F		Floating point value to convert
+	* @return		The rounded integer
+	*/
+	static CORE_API float RoundHalfFromZero(float F);
+	static CORE_API double RoundHalfFromZero(double F);
+
+	/**
+	* Converts a floating point number to the nearest integer, equidistant ties go to the value which is closer to zero: -0.5 becomes 0, 0.5 becomes 0
+	* @param F		Floating point value to convert
+	* @return		The rounded integer
+	*/
+	static CORE_API float RoundHalfToZero(float F);
+	static CORE_API double RoundHalfToZero(double F);
+
+	/**
+	* Converts a floating point number to an integer which is further from zero, "larger" in absolute value: 0.1 becomes 1, -0.1 becomes -1
+	* @param F		Floating point value to convert
+	* @return		The rounded integer
+	*/
+	static float RoundFromZero(float F);
+	static double RoundFromZero(double F);
+
+	/**
+	* Converts a floating point number to an integer which is closer to zero, "smaller" in absolute value: 0.1 becomes 0, -0.1 becomes 0
+	* @param F		Floating point value to convert
+	* @return		The rounded integer
+	*/
+	static float RoundToZero(float F);
+	static double RoundToZero(double F);
+
+	/**
+	* Converts a floating point number to an integer which is more negative: 0.1 becomes 0, -0.1 becomes -1
+	* @param F		Floating point value to convert
+	* @return		The rounded integer
+	*/
+	static float RoundToNegativeInfinity(float F);
+	static double RoundToNegativeInfinity(double F);
+
+	/**
+	* Converts a floating point number to an integer which is more positive: 0.1 becomes 1, -0.1 becomes 0
+	* @param F		Floating point value to convert
+	* @return		The rounded integer
+	*/
+	static float RoundToPositiveInfinity(float F);
+	static double RoundToPositiveInfinity(double F);
+
 	// Formatting functions
 
 	/**
