Don't open control channel when detecting malicious packets. Fixes issue of client staying infinitely connected because close packets keep the connection alive.

diff --git a/Engine/Source/Runtime/Core/Public/Misc/AssertionMacros.h b/Engine/Source/Runtime/Core/Public/Misc/AssertionMacros.h
--- a/Engine/Source/Runtime/Core/Public/Misc/AssertionMacros.h
+++ b/Engine/Source/Runtime/Core/Public/Misc/AssertionMacros.h
@@ -467,13 +467,8 @@ struct FTCharArrayTester
 	static_assert(IS_TCHAR_ARRAY(Format), "Formatting string must be a TCHAR array."); \
 	check(NetConnection != nullptr); \
 	FString SecurityPrint = FString::Printf(Format, ##__VA_ARGS__); \
-	if(Channels[0] == nullptr) \
-	{ \
-		Channels[0] = CreateChannel(EChannelType::CHTYPE_Control, false, 0); \
-	} \
 	UE_SECURITY_LOG(NetConnection, SecurityEventType, Format, ##__VA_ARGS__); \
 	UE_SECURITY_LOG(NetConnection, ESecurityEvent::Closed, TEXT("Connection closed")); \
-	NOTIFY_CLIENT_OF_SECURITY_EVENT_IF_NOT_SHIPPING(NetConnection, SecurityPrint); \
 	NetConnection->Close(); \
 	PerfCountersIncrement(TEXT("ClosedConnectionsDueToSecurityViolations")); \
 }
