Linux: fix calculation of physical cores.

diff --git a/Engine/Source/Runtime/Core/Private/Linux/LinuxPlatformMisc.cpp b/Engine/Source/Runtime/Core/Private/Linux/LinuxPlatformMisc.cpp
--- a/Engine/Source/Runtime/Core/Private/Linux/LinuxPlatformMisc.cpp
+++ b/Engine/Source/Runtime/Core/Private/Linux/LinuxPlatformMisc.cpp
@@ -612,7 +612,16 @@ int32 FLinuxPlatformMisc::NumberOfCores()
 		else
 		{
 			char FileNameBuffer[1024];
-			unsigned char PossibleCores[CPU_SETSIZE] = { 0 };
+			struct CpuInfo
+			{
+				int Core;
+				int Package;
+			}
+			CpuInfos[CPU_SETSIZE];
+
+			FMemory::Memzero(CpuInfos);
+			int MaxCoreId = 0;
+			int MaxPackageId = 0;
 
 			for(int32 CpuIdx = 0; CpuIdx < CPU_SETSIZE; ++CpuIdx)
 			{
@@ -620,29 +629,49 @@ int32 FLinuxPlatformMisc::NumberOfCores()
 				{
 					sprintf(FileNameBuffer, "/sys/devices/system/cpu/cpu%d/topology/core_id", CpuIdx);
 					
-					FILE* CoreIdFile = fopen(FileNameBuffer, "r");
-					unsigned int CoreId = 0;
-					if (CoreIdFile)
+					if (FILE* CoreIdFile = fopen(FileNameBuffer, "r"))
 					{
-						if (1 != fscanf(CoreIdFile, "%d", &CoreId))
+						if (1 != fscanf(CoreIdFile, "%d", &CpuInfos[CpuIdx].Core))
 						{
-							CoreId = 0;
+							CpuInfos[CpuIdx].Core = 0;
 						}
 						fclose(CoreIdFile);
 					}
 
-					if (CoreId >= ARRAY_COUNT(PossibleCores))
+					sprintf(FileNameBuffer, "/sys/devices/system/cpu/cpu%d/topology/physical_package_id", CpuIdx);
+
+					unsigned int PackageId = 0;
+					if (FILE* PackageIdFile = fopen(FileNameBuffer, "r"))
 					{
-						CoreId = 0;
+						if (1 != fscanf(PackageIdFile, "%d", &CpuInfos[CpuIdx].Package))
+						{
+							CpuInfos[CpuIdx].Package = 0;
+						}
+						fclose(PackageIdFile);
 					}
-					
-					PossibleCores[ CoreId ] = 1;
+
+					MaxCoreId = FMath::Max(MaxCoreId, CpuInfos[CpuIdx].Core);
+					MaxPackageId = FMath::Max(MaxPackageId, CpuInfos[CpuIdx].Package);
+				}
+			}
+
+			int NumCores = MaxCoreId + 1;
+			int NumPackages = MaxPackageId + 1;
+			int NumPairs = NumPackages * NumCores;
+			unsigned char * Pairs = reinterpret_cast<unsigned char *>(FMemory_Alloca(NumPairs * sizeof(unsigned char)));
+			FMemory::Memzero(Pairs, NumPairs * sizeof(unsigned char));
+
+			for (int32 CpuIdx = 0; CpuIdx < CPU_SETSIZE; ++CpuIdx)
+			{
+				if (CPU_ISSET(CpuIdx, &AvailableCpusMask))
+				{
+					Pairs[CpuInfos[CpuIdx].Package * NumCores + CpuInfos[CpuIdx].Core] = 1;
 				}
 			}
 
-			for(int32 Idx = 0; Idx < ARRAY_COUNT(PossibleCores); ++Idx)
+			for (int32 Idx = 0; Idx < NumPairs; ++Idx)
 			{
-				NumCoreIds += PossibleCores[Idx];
+				NumCoreIds += Pairs[Idx];
 			}
 		}
 	}
