Added FastDecimalFormat

This provides efficient and culture aware number formatting, and will be used by our FText implementation.

The ICU culture now has functions to extract culture correct number formatting rules from the ICU DecimalFormat objects. These rules can then be used with FastDecimalFormat::NumberToString to convert a number to the correct decimal representation for the active culture.

diff --git a/Engine/Source/Runtime/Core/Private/Internationalization/FastDecimalFormat.cpp b/Engine/Source/Runtime/Core/Private/Internationalization/FastDecimalFormat.cpp
new file mode 100644
--- /dev/null
+++ b/Engine/Source/Runtime/Core/Private/Internationalization/FastDecimalFormat.cpp
@@ -0,0 +1,293 @@
+// Copyright 1998-2015 Epic Games, Inc. All Rights Reserved.
+
+#include "CorePrivatePCH.h"
+#include "FastDecimalFormat.h"
+
+namespace FastDecimalFormat
+{
+
+namespace Internal
+{
+
+static const int32 MaxIntegralPrintLength = 20;
+static const int32 MaxFractionalPrintPrecision = 18;
+static const int32 MinRequiredIntegralBufferSize = (MaxIntegralPrintLength * 2) + 1; // *2 for an absolute worst case group separator scenario, +1 for null terminator
+
+static const uint64 Pow10Table[] = {
+	1,						// 10^0
+	10,						// 10^1
+	100,					// 10^2
+	1000,					// 10^3
+	10000,					// 10^4
+	100000,					// 10^5
+	1000000,				// 10^6
+	10000000,				// 10^7
+	100000000,				// 10^8
+	1000000000,				// 10^9
+	10000000000,			// 10^10
+	100000000000,			// 10^11
+	1000000000000,			// 10^12
+	10000000000000,			// 10^13
+	100000000000000,		// 10^14
+	1000000000000000,		// 10^15
+	10000000000000000,		// 10^16
+	100000000000000000,		// 10^17
+	1000000000000000000,	// 10^18
+};
+
+static_assert(ARRAY_COUNT(Pow10Table) - 1 >= MaxFractionalPrintPrecision, "Pow10Table must at big enough to index any value up-to MaxFractionalPrintPrecision");
+
+int32 IntegralToString_UInt64ToString(
+	const uint64 InVal, 
+	const bool InUseGrouping, const uint8 InPrimaryGroupingSize, const uint8 InSecondaryGroupingSize, const TCHAR InGroupingSeparatorCharacter, 
+	const int32 InMinDigitsToPrint, const int32 InMaxDigitsToPrint, 
+	TCHAR* InBufferToFill, const int32 InBufferToFillSize
+	)
+{
+	check(InBufferToFillSize >= MinRequiredIntegralBufferSize);
+
+	TCHAR TmpBuffer[MinRequiredIntegralBufferSize];
+	int32 StringLen = 0;
+
+	int32 DigitsPrinted = 0;
+	uint8 NumUntilNextGroup = InPrimaryGroupingSize;
+
+	if (InVal > 0)
+	{
+		// Perform the initial number -> string conversion
+		uint64 TmpNum = InVal;
+		while (StringLen < InMaxDigitsToPrint && TmpNum != 0)
+		{
+			if (InUseGrouping && NumUntilNextGroup-- == 0)
+			{
+				TmpBuffer[StringLen++] = InGroupingSeparatorCharacter;
+				NumUntilNextGroup = InSecondaryGroupingSize;
+			}
+
+			TmpBuffer[StringLen++] = '0' + (TmpNum % 10);
+			TmpNum /= 10;
+
+			++DigitsPrinted;
+		}
+	}
+
+	// Pad the string to the min digits requested
+	{
+		const int32 PaddingToApply = FMath::Min(InMinDigitsToPrint - DigitsPrinted, MaxIntegralPrintLength - DigitsPrinted);
+		for (int32 PaddingIndex = 0; PaddingIndex < PaddingToApply; ++PaddingIndex)
+		{
+			if (InUseGrouping && NumUntilNextGroup-- == 0)
+			{
+				TmpBuffer[StringLen++] = InGroupingSeparatorCharacter;
+				NumUntilNextGroup = InSecondaryGroupingSize;
+			}
+
+			TmpBuffer[StringLen++] = '0';
+		}
+	}
+
+	// TmpBuffer is backwards, flip it into the final output buffer
+	for (int32 FinalBufferIndex = 0; FinalBufferIndex < StringLen; ++FinalBufferIndex)
+	{
+		InBufferToFill[FinalBufferIndex] = TmpBuffer[StringLen - FinalBufferIndex - 1];
+	}
+	InBufferToFill[StringLen] = 0;
+
+	return StringLen;
+}
+
+FORCEINLINE int32 IntegralToString_Common(const uint64 InVal, const FDecimalNumberFormattingRules& InFormattingRules, const FNumberFormattingOptions& InFormattingOptions, TCHAR* InBufferToFill, const int32 InBufferToFillSize)
+{
+	// Perform the initial format to a decimal string
+	return IntegralToString_UInt64ToString(
+		InVal, 
+		InFormattingOptions.UseGrouping, 
+		InFormattingRules.PrimaryGroupingSize, 
+		InFormattingRules.SecondaryGroupingSize, 
+		InFormattingRules.GroupingSeparatorCharacter, 
+		InFormattingOptions.MinimumIntegralDigits, 
+		InFormattingOptions.MaximumIntegralDigits, 
+		InBufferToFill, 
+		InBufferToFillSize
+		);
+}
+
+void FractionalToString_SplitAndRoundNumber(const double InValue, const int32 InNumDecimalPlaces, ERoundingMode InRoundingMode, double& OutIntegralPart, double& OutFractionalPart)
+{
+	const int32 DecimalPlacesToRoundTo = FMath::Min(InNumDecimalPlaces, MaxFractionalPrintPrecision);
+
+	const bool bIsNegative = InValue < 0.0;
+	const bool bIsRoundingEntireNumber = DecimalPlacesToRoundTo == 0;
+
+	// We split the value before performing the rounding to avoid losing precision during the rounding calculations
+	// If we're rounding to zero decimal places, then we just apply rounding to the number as a whole
+	double IntegralPart = InValue;
+	double FractionalPart = (bIsRoundingEntireNumber) ? 0.0 : FMath::Modf(InValue, &IntegralPart);
+
+	// Multiply the value to round by 10^DecimalPlacesToRoundTo - this will allow us to perform rounding calculations 
+	// that correctly trim any remaining fractional parts that are outside of our rounding range
+	double& ValueToRound = ((bIsRoundingEntireNumber) ? IntegralPart : FractionalPart);
+	ValueToRound *= Pow10Table[DecimalPlacesToRoundTo];
+
+	// The rounding modes here mimic those of ICU. See http://userguide.icu-project.org/formatparse/numbers/rounding-modes
+	switch (InRoundingMode)
+	{
+	case ERoundingMode::HalfToEven:
+		// Rounds to the nearest place, equidistant ties go to the value which is closest to an even value: 1.5 becomes 2, 0.5 becomes 0
+		ValueToRound = FMath::RoundHalfToEven(ValueToRound);
+		break;
+
+	case ERoundingMode::HalfFromZero:
+		// Rounds to nearest place, equidistant ties go to the value which is further from zero: -0.5 becomes -1.0, 0.5 becomes 1.0
+		ValueToRound = FMath::RoundHalfFromZero(ValueToRound);
+		break;
+	
+	case ERoundingMode::HalfToZero:
+		// Rounds to nearest place, equidistant ties go to the value which is closer to zero: -0.5 becomes 0, 0.5 becomes 0
+		ValueToRound = FMath::RoundHalfToZero(ValueToRound);
+		break;
+
+	case ERoundingMode::FromZero:
+		// Rounds to the value which is further from zero, "larger" in absolute value: 0.1 becomes 1, -0.1 becomes -1
+		ValueToRound = FMath::RoundFromZero(ValueToRound);
+		break;
+	
+	case ERoundingMode::ToZero:
+		// Rounds to the value which is closer to zero, "smaller" in absolute value: 0.1 becomes 0, -0.1 becomes 0
+		ValueToRound = FMath::RoundToZero(ValueToRound);
+		break;
+	
+	case ERoundingMode::ToNegativeInfinity:
+		// Rounds to the value which is more negative: 0.1 becomes 0, -0.1 becomes -1
+		ValueToRound = FMath::RoundToNegativeInfinity(ValueToRound);
+		break;
+	
+	case ERoundingMode::ToPositiveInfinity:
+		// Rounds to the value which is more positive: 0.1 becomes 1, -0.1 becomes 0
+		ValueToRound = FMath::RoundToPositiveInfinity(ValueToRound);
+		break;
+
+	default:
+		break;
+	}
+
+	// Copy to the correct output param depending on whether we were rounding to the number as a whole
+	if (bIsRoundingEntireNumber)
+	{
+		OutIntegralPart = ValueToRound;
+		OutFractionalPart = 0.0;
+	}
+	else
+	{
+		OutIntegralPart = IntegralPart;
+		OutFractionalPart = ValueToRound;
+	}
+}
+
+FString BuildFinalString(const bool bIsNegative, const FDecimalNumberFormattingRules& InFormattingRules, const TCHAR* InIntegralBuffer, const int32 InIntegralLen, const TCHAR* InFractionalBuffer, const int32 InFractionalLen)
+{
+	const FString& FinalPrefixStr = (bIsNegative) ? InFormattingRules.NegativePrefixString : InFormattingRules.PositivePrefixString;
+	const FString& FinalSuffixStr = (bIsNegative) ? InFormattingRules.NegativeSuffixString : InFormattingRules.PositiveSuffixString;
+
+	FString FinalStr;
+	FinalStr.Reserve(FinalPrefixStr.Len() + InIntegralLen + 1 + InFractionalLen + FinalSuffixStr.Len());
+
+	FinalStr.Append(FinalPrefixStr);
+	FinalStr.AppendChars(InIntegralBuffer, InIntegralLen);
+	if (InFractionalLen > 0)
+	{
+		FinalStr.AppendChar(InFormattingRules.DecimalSeparatorCharacter);
+		FinalStr.AppendChars(InFractionalBuffer, InFractionalLen);
+	}
+	FinalStr.Append(FinalSuffixStr);
+
+	return FinalStr;
+}
+
+FString IntegralToString(const bool bIsNegative, const uint64 InVal, const FDecimalNumberFormattingRules& InFormattingRules, const FNumberFormattingOptions& InFormattingOptions)
+{
+	// Deal with the integral part (produces a string of the integral part, inserting group separators if requested and required, and padding as needed)
+	TCHAR IntegralPartBuffer[MinRequiredIntegralBufferSize];
+	const int32 IntegralPartLen = IntegralToString_Common(InVal, InFormattingRules, InFormattingOptions, IntegralPartBuffer, ARRAY_COUNT(IntegralPartBuffer));
+
+	// Deal with any forced fractional part (produces a string zeros up to the required minimum length)
+	TCHAR FractionalPartBuffer[MinRequiredIntegralBufferSize];
+	int32 FractionalPartLen = 0;
+	if (InFormattingOptions.MinimumFractionalDigits > 0)
+	{
+		const int32 PaddingToApply = FMath::Min(InFormattingOptions.MinimumFractionalDigits, MaxFractionalPrintPrecision);
+		for (int32 PaddingIndex = 0; PaddingIndex < PaddingToApply; ++PaddingIndex)
+		{
+			FractionalPartBuffer[FractionalPartLen++] = '0';
+		}
+	}
+	FractionalPartBuffer[FractionalPartLen] = 0;
+
+	return BuildFinalString(bIsNegative, InFormattingRules, IntegralPartBuffer, IntegralPartLen, FractionalPartBuffer, FractionalPartLen);
+}
+
+FString FractionalToString(const double InVal, const FDecimalNumberFormattingRules& InFormattingRules, const FNumberFormattingOptions& InFormattingOptions)
+{
+	if (FMath::IsNaN(InVal))
+	{
+		return InFormattingRules.NaNString;
+	}
+
+	double IntegralPart = 0.0;
+	double FractionalPart = 0.0;
+	FractionalToString_SplitAndRoundNumber(InVal, InFormattingOptions.MaximumFractionalDigits, InFormattingOptions.RoundingMode, IntegralPart, FractionalPart);
+
+	const bool bIsNegative = IntegralPart < 0.0;
+	if (bIsNegative)
+	{
+		IntegralPart = -IntegralPart;
+		FractionalPart = -FractionalPart;
+	}
+
+	// Deal with the integral part (produces a string of the integral part, inserting group separators if requested and required, and padding as needed)
+	TCHAR IntegralPartBuffer[MinRequiredIntegralBufferSize];
+	const int32 IntegralPartLen = IntegralToString_Common(static_cast<uint64>(IntegralPart), InFormattingRules, InFormattingOptions, IntegralPartBuffer, ARRAY_COUNT(IntegralPartBuffer));
+
+	// Deal with the fractional part (produces a string of the fractional part, potentially padding with zeros up to InFormattingOptions.MaximumFractionalDigits)
+	TCHAR FractionalPartBuffer[MinRequiredIntegralBufferSize];
+	int32 FractionalPartLen = 0;
+	if (FractionalPart != 0.0)
+	{
+		FractionalPartLen = IntegralToString_UInt64ToString(static_cast<uint64>(FractionalPart), false, 0, 0, ' ', 0, InFormattingOptions.MaximumFractionalDigits, FractionalPartBuffer, ARRAY_COUNT(FractionalPartBuffer));
+	
+		// Trim any trailing zeros back down to InFormattingOptions.MinimumFractionalDigits
+		while (FractionalPartLen > InFormattingOptions.MinimumFractionalDigits && FractionalPartBuffer[FractionalPartLen - 1] == '0')
+		{
+			--FractionalPartLen;
+		}
+	}
+	FractionalPartBuffer[FractionalPartLen] = 0;
+
+	// Pad the fractional part with any zeros that may have been missed so far
+	{
+		const int32 PaddingToApply = FMath::Min(InFormattingOptions.MinimumFractionalDigits - FractionalPartLen, MaxFractionalPrintPrecision - FractionalPartLen);
+		for (int32 PaddingIndex = 0; PaddingIndex < PaddingToApply; ++PaddingIndex)
+		{
+			FractionalPartBuffer[FractionalPartLen++] = '0';
+		}
+		FractionalPartBuffer[FractionalPartLen] = 0;
+	}
+
+	return BuildFinalString(bIsNegative, InFormattingRules, IntegralPartBuffer, IntegralPartLen, FractionalPartBuffer, FractionalPartLen);
+}
+
+} // namespace Internal
+
+FDecimalNumberFormattingRules GetCultureAgnosticFormattingRules()
+{
+	FDecimalNumberFormattingRules CultureAgnosticFormattingRules;
+	CultureAgnosticFormattingRules.NaNString = TEXT("NaN");
+	CultureAgnosticFormattingRules.NegativePrefixString = TEXT("-");
+	CultureAgnosticFormattingRules.GroupingSeparatorCharacter = ',';
+	CultureAgnosticFormattingRules.DecimalSeparatorCharacter = '.';
+	CultureAgnosticFormattingRules.PrimaryGroupingSize = 3;
+	CultureAgnosticFormattingRules.SecondaryGroupingSize = 3;
+	return CultureAgnosticFormattingRules;
+}
+
+} // namespace FastDecimalFormat
diff --git a/Engine/Source/Runtime/Core/Private/Internationalization/ICUCulture.cpp b/Engine/Source/Runtime/Core/Private/Internationalization/ICUCulture.cpp
--- a/Engine/Source/Runtime/Core/Private/Internationalization/ICUCulture.cpp
+++ b/Engine/Source/Runtime/Core/Private/Internationalization/ICUCulture.cpp
@@ -1,6 +1,7 @@
 // Copyright 1998-2015 Epic Games, Inc. All Rights Reserved.
 
 #include "CorePrivatePCH.h"
+#include "FastDecimalFormat.h"
 
 #if UE_ENABLE_ICU
 #include "ICUUtilities.h"
@@ -565,4 +566,186 @@ TSharedRef<const icu::DateFormat> FCulture::FICUCultureImplementation::GetDateTi
 	}
 }
 
+namespace
+{
+
+FDecimalNumberFormattingRules ExtractNumberFormattingRulesFromICUDecimalFormatter(icu::DecimalFormat& InICUDecimalFormat)
+{
+	FDecimalNumberFormattingRules NewUEDecimalNumberFormattingRules;
+
+	// Extract the default formatting options before we mess around with the formatter object settings
+	NewUEDecimalNumberFormattingRules.CultureDefaultFormattingOptions
+		.SetUseGrouping(InICUDecimalFormat.isGroupingUsed() != 0)
+		.SetRoundingMode(ICUToUE(InICUDecimalFormat.getRoundingMode()))
+		.SetMinimumIntegralDigits(InICUDecimalFormat.getMinimumIntegerDigits())
+		.SetMaximumIntegralDigits(InICUDecimalFormat.getMaximumIntegerDigits())
+		.SetMinimumFractionalDigits(InICUDecimalFormat.getMinimumFractionDigits())
+		.SetMaximumFractionalDigits(InICUDecimalFormat.getMaximumFractionDigits());
+
+	// We force grouping to be on, even if a culture doesn't use it by default, so that we can extract meaningful grouping information
+	// This allows us to use the correct groupings if we should ever force grouping for a number, rather than use the culture default
+	InICUDecimalFormat.setGroupingUsed(true);
+
+	auto ICUStringToTCHAR = [](const icu::UnicodeString& InICUString) -> TCHAR
+	{
+		check(InICUString.length() == 1);
+		return static_cast<TCHAR>(InICUString.charAt(0));
+	};
+
+	auto ExtractFormattingSymbolAsCharacter = [&](icu::DecimalFormatSymbols::ENumberFormatSymbol InSymbolToExtract) -> TCHAR
+	{
+		const icu::UnicodeString& ICUSymbolString = InICUDecimalFormat.getDecimalFormatSymbols()->getConstSymbol(InSymbolToExtract);
+		return ICUStringToTCHAR(ICUSymbolString); // For efficiency we assume that these symbols are always a single character
+	};
+
+	icu::UnicodeString ScratchICUString;
+
+	// Extract the rules from the decimal formatter
+	NewUEDecimalNumberFormattingRules.NaNString						= ICUUtilities::ConvertString(InICUDecimalFormat.getDecimalFormatSymbols()->getConstSymbol(icu::DecimalFormatSymbols::kNaNSymbol));
+	NewUEDecimalNumberFormattingRules.NegativePrefixString			= ICUUtilities::ConvertString(InICUDecimalFormat.getNegativePrefix(ScratchICUString));
+	NewUEDecimalNumberFormattingRules.NegativeSuffixString			= ICUUtilities::ConvertString(InICUDecimalFormat.getNegativeSuffix(ScratchICUString));
+	NewUEDecimalNumberFormattingRules.PositivePrefixString			= ICUUtilities::ConvertString(InICUDecimalFormat.getPositivePrefix(ScratchICUString));
+	NewUEDecimalNumberFormattingRules.PositiveSuffixString			= ICUUtilities::ConvertString(InICUDecimalFormat.getPositiveSuffix(ScratchICUString));
+	NewUEDecimalNumberFormattingRules.GroupingSeparatorCharacter	= ExtractFormattingSymbolAsCharacter(icu::DecimalFormatSymbols::kGroupingSeparatorSymbol);
+	NewUEDecimalNumberFormattingRules.DecimalSeparatorCharacter		= ExtractFormattingSymbolAsCharacter(icu::DecimalFormatSymbols::kDecimalSeparatorSymbol);
+	NewUEDecimalNumberFormattingRules.PrimaryGroupingSize			= static_cast<uint8>(InICUDecimalFormat.getGroupingSize());
+	NewUEDecimalNumberFormattingRules.SecondaryGroupingSize			= (InICUDecimalFormat.getSecondaryGroupingSize() < 1) 
+																		? NewUEDecimalNumberFormattingRules.PrimaryGroupingSize 
+																		: static_cast<uint8>(InICUDecimalFormat.getSecondaryGroupingSize());
+
+	return NewUEDecimalNumberFormattingRules;
+}
+
+} // anonymous namespace
+
+const FDecimalNumberFormattingRules& FCulture::FICUCultureImplementation::GetDecimalNumberFormattingRules()
+{
+	if (UEDecimalNumberFormattingRules.IsValid())
+	{
+		return *UEDecimalNumberFormattingRules;
+	}
+
+	// Create a culture decimal formatter (doesn't call CreateDecimalFormat as we need a mutable instance)
+	TSharedPtr<icu::DecimalFormat> DecimalFormatterForCulture;
+	{
+		UErrorCode ICUStatus = U_ZERO_ERROR;
+		DecimalFormatterForCulture = MakeShareable(static_cast<icu::DecimalFormat*>(icu::NumberFormat::createInstance(ICULocale, ICUStatus)));
+		checkf(DecimalFormatterForCulture.IsValid(), TEXT("Creating a decimal format object failed using locale %s. Perhaps this locale has no data."), StringCast<TCHAR>(ICULocale.getName()).Get());
+	}
+
+	const FDecimalNumberFormattingRules NewUEDecimalNumberFormattingRules = ExtractNumberFormattingRulesFromICUDecimalFormatter(*DecimalFormatterForCulture);
+
+	// Check the pointer again in case another thread beat us to it
+	{
+		FScopeLock PtrLock(&UEDecimalNumberFormattingRulesCS);
+
+		if (!UEDecimalNumberFormattingRules.IsValid())
+		{
+			UEDecimalNumberFormattingRules = MakeShareable(new FDecimalNumberFormattingRules(NewUEDecimalNumberFormattingRules));
+		}
+	}
+
+	return *UEDecimalNumberFormattingRules;
+}
+
+const FDecimalNumberFormattingRules& FCulture::FICUCultureImplementation::GetPercentFormattingRules()
+{
+	if (UEPercentFormattingRules.IsValid())
+	{
+		return *UEPercentFormattingRules;
+	}
+
+	// Create a culture percent formatter (doesn't call CreatePercentFormat as we need a mutable instance)
+	TSharedPtr<icu::DecimalFormat> PercentFormatterForCulture;
+	{
+		UErrorCode ICUStatus = U_ZERO_ERROR;
+		PercentFormatterForCulture = MakeShareable(static_cast<icu::DecimalFormat*>(icu::NumberFormat::createPercentInstance(ICULocale, ICUStatus)));
+		checkf(PercentFormatterForCulture.IsValid(), TEXT("Creating a percent format object failed using locale %s. Perhaps this locale has no data."), StringCast<TCHAR>(ICULocale.getName()).Get());
+	}
+
+	const FDecimalNumberFormattingRules NewUEPercentFormattingRules = ExtractNumberFormattingRulesFromICUDecimalFormatter(*PercentFormatterForCulture);
+
+	// Check the pointer again in case another thread beat us to it
+	{
+		FScopeLock PtrLock(&UEPercentFormattingRulesCS);
+
+		if (!UEPercentFormattingRules.IsValid())
+		{
+			UEPercentFormattingRules = MakeShareable(new FDecimalNumberFormattingRules(NewUEPercentFormattingRules));
+		}
+	}
+
+	return *UEPercentFormattingRules;
+}
+
+const FDecimalNumberFormattingRules& FCulture::FICUCultureImplementation::GetCurrencyFormattingRules(const FString& InCurrencyCode)
+{
+	const bool bUseDefaultFormattingRules = InCurrencyCode.IsEmpty();
+
+	if (bUseDefaultFormattingRules)
+	{
+		if (UECurrencyFormattingRules.IsValid())
+		{
+			return *UECurrencyFormattingRules;
+		}
+	}
+	else
+	{
+		FScopeLock MapLock(&UEAlternateCurrencyFormattingRulesCS);
+
+		auto FoundUEAlternateCurrencyFormattingRules = UEAlternateCurrencyFormattingRules.FindRef(InCurrencyCode);
+		if (FoundUEAlternateCurrencyFormattingRules.IsValid())
+		{
+			return *FoundUEAlternateCurrencyFormattingRules;
+		}
+	}
+
+	// Create a culture percent formatter (doesn't call CreateCurrencyFormat as we need a mutable instance)
+	TSharedPtr<icu::DecimalFormat> CurrencyFormatterForCulture;
+	{
+		UErrorCode ICUStatus = U_ZERO_ERROR;
+		CurrencyFormatterForCulture = MakeShareable(static_cast<icu::DecimalFormat*>(icu::NumberFormat::createCurrencyInstance(ICULocale, ICUStatus)));
+		checkf(CurrencyFormatterForCulture.IsValid(), TEXT("Creating a currency format object failed using locale %s. Perhaps this locale has no data."), StringCast<TCHAR>(ICULocale.getName()).Get());
+	}
+
+	if (!bUseDefaultFormattingRules)
+	{
+		// Set the custom currency before we extract the data from the formatter
+		icu::UnicodeString ICUCurrencyCode = ICUUtilities::ConvertString(InCurrencyCode);
+		CurrencyFormatterForCulture->setCurrency(ICUCurrencyCode.getBuffer());
+	}
+
+	const FDecimalNumberFormattingRules NewUECurrencyFormattingRules = ExtractNumberFormattingRulesFromICUDecimalFormatter(*CurrencyFormatterForCulture);
+
+	if (bUseDefaultFormattingRules)
+	{
+		// Check the pointer again in case another thread beat us to it
+		{
+			FScopeLock PtrLock(&UECurrencyFormattingRulesCS);
+
+			if (!UECurrencyFormattingRules.IsValid())
+			{
+				UECurrencyFormattingRules = MakeShareable(new FDecimalNumberFormattingRules(NewUECurrencyFormattingRules));
+			}
+		}
+
+		return *UECurrencyFormattingRules;
+	}
+	else
+	{
+		FScopeLock MapLock(&UEAlternateCurrencyFormattingRulesCS);
+
+		// Find again in case another thread beat us to it
+		auto FoundUEAlternateCurrencyFormattingRules = UEAlternateCurrencyFormattingRules.FindRef(InCurrencyCode);
+		if (FoundUEAlternateCurrencyFormattingRules.IsValid())
+		{
+			return *FoundUEAlternateCurrencyFormattingRules;
+		}
+
+		FoundUEAlternateCurrencyFormattingRules = MakeShareable(new FDecimalNumberFormattingRules(NewUECurrencyFormattingRules));
+		UEAlternateCurrencyFormattingRules.Add(InCurrencyCode, FoundUEAlternateCurrencyFormattingRules);
+		return *FoundUEAlternateCurrencyFormattingRules;
+	}
+}
+
 #endif
diff --git a/Engine/Source/Runtime/Core/Private/Internationalization/ICUCulture.h b/Engine/Source/Runtime/Core/Private/Internationalization/ICUCulture.h
--- a/Engine/Source/Runtime/Core/Private/Internationalization/ICUCulture.h
+++ b/Engine/Source/Runtime/Core/Private/Internationalization/ICUCulture.h
@@ -10,6 +10,8 @@
 #include <unicode/decimfmt.h>
 #include <unicode/datefmt.h>
 
+struct FDecimalNumberFormattingRules;
+
 inline UColAttributeValue UEToICU(const ETextComparisonLevel::Type ComparisonLevel)
 {
 	UColAttributeValue Value;
@@ -100,6 +102,39 @@ inline icu::DecimalFormat::ERoundingMode UEToICU(const ERoundingMode RoundingMod
 	return Value;
 }
 
+inline ERoundingMode ICUToUE(const icu::DecimalFormat::ERoundingMode RoundingMode)
+{
+	ERoundingMode Value;
+	switch(RoundingMode)
+	{
+	case icu::DecimalFormat::ERoundingMode::kRoundHalfEven:
+		Value = ERoundingMode::HalfToEven;
+		break;
+	case icu::DecimalFormat::ERoundingMode::kRoundHalfUp:
+		Value = ERoundingMode::HalfFromZero;
+		break;
+	case icu::DecimalFormat::ERoundingMode::kRoundHalfDown:
+		Value = ERoundingMode::HalfToZero;
+		break;
+	case icu::DecimalFormat::ERoundingMode::kRoundUp:
+		Value = ERoundingMode::FromZero;
+		break;
+	case icu::DecimalFormat::ERoundingMode::kRoundDown:
+		Value = ERoundingMode::ToZero;
+		break;
+	case icu::DecimalFormat::ERoundingMode::kRoundFloor:
+		Value = ERoundingMode::ToNegativeInfinity;
+		break;
+	case icu::DecimalFormat::ERoundingMode::kRoundCeiling:
+		Value = ERoundingMode::ToPositiveInfinity;
+		break;
+	default:
+		Value = ERoundingMode::HalfToEven;
+		break;
+	}
+	return Value;
+}
+
 enum class EBreakIteratorType
 {
 	Grapheme,
@@ -365,6 +400,10 @@ class FCulture::FICUCultureImplementation
 	TSharedRef<const icu::DateFormat> GetTimeFormatter(const EDateTimeStyle::Type TimeStyle, const FString& TimeZone);
 	TSharedRef<const icu::DateFormat> GetDateTimeFormatter(const EDateTimeStyle::Type DateStyle, const EDateTimeStyle::Type TimeStyle, const FString& TimeZone);
 
+	const FDecimalNumberFormattingRules& GetDecimalNumberFormattingRules();
+	const FDecimalNumberFormattingRules& GetPercentFormattingRules();
+	const FDecimalNumberFormattingRules& GetCurrencyFormattingRules(const FString& InCurrencyCode);
+
 	icu::Locale ICULocale;
 	TSharedPtr<const icu::BreakIterator> ICUGraphemeBreakIterator;
 	TSharedPtr<const icu::BreakIterator> ICUWordBreakIterator;
@@ -384,5 +423,17 @@ class FCulture::FICUCultureImplementation
 	TSharedPtr<const icu::DateFormat> ICUDateFormat;
 	TSharedPtr<const icu::DateFormat> ICUTimeFormat;
 	TSharedPtr<const icu::DateFormat> ICUDateTimeFormat;
+
+	TSharedPtr<const FDecimalNumberFormattingRules, ESPMode::ThreadSafe> UEDecimalNumberFormattingRules;
+	FCriticalSection UEDecimalNumberFormattingRulesCS;
+
+	TSharedPtr<const FDecimalNumberFormattingRules, ESPMode::ThreadSafe> UEPercentFormattingRules;
+	FCriticalSection UEPercentFormattingRulesCS;
+
+	TSharedPtr<const FDecimalNumberFormattingRules, ESPMode::ThreadSafe> UECurrencyFormattingRules;
+	FCriticalSection UECurrencyFormattingRulesCS;
+
+	TMap<FString, TSharedPtr<const FDecimalNumberFormattingRules>> UEAlternateCurrencyFormattingRules;
+	FCriticalSection UEAlternateCurrencyFormattingRulesCS;
 };
 #endif
diff --git a/Engine/Source/Runtime/Core/Public/Internationalization/FastDecimalFormat.h b/Engine/Source/Runtime/Core/Public/Internationalization/FastDecimalFormat.h
new file mode 100644
--- /dev/null
+++ b/Engine/Source/Runtime/Core/Public/Internationalization/FastDecimalFormat.h
@@ -0,0 +1,88 @@
+// Copyright 1998-2015 Epic Games, Inc. All Rights Reserved.
+
+#pragma once
+
+/** Rules used to format a decimal number */
+struct FDecimalNumberFormattingRules
+{
+	FDecimalNumberFormattingRules()
+		: GroupingSeparatorCharacter(0)
+		, DecimalSeparatorCharacter(0)
+		, PrimaryGroupingSize(0)
+		, SecondaryGroupingSize(0)
+	{
+	}
+
+	/** Number formatting rules, typically extracted from the ICU decimal formatter for a given culture */
+	FString NaNString;
+	FString NegativePrefixString;
+	FString NegativeSuffixString;
+	FString PositivePrefixString;
+	FString PositiveSuffixString;
+	TCHAR GroupingSeparatorCharacter;
+	TCHAR DecimalSeparatorCharacter;
+	uint8 PrimaryGroupingSize;
+	uint8 SecondaryGroupingSize;
+
+	/** Default number formatting options for a given culture */
+	FNumberFormattingOptions CultureDefaultFormattingOptions;
+};
+
+/**
+ * Provides efficient and culture aware number formatting.
+ * You would call FastDecimalFormat::NumberToString to convert a number to the correct decimal representation based on the given formatting rules and options.
+ * The primary consumer of this is FText, however you can use it for other things. GetCultureAgnosticFormattingRules can provide formatting rules for cases where you don't care about culture.
+ */
+namespace FastDecimalFormat
+{
+
+namespace Internal
+{
+
+CORE_API FString IntegralToString(const bool bIsNegative, const uint64 InVal, const FDecimalNumberFormattingRules& InFormattingRules, const FNumberFormattingOptions& InFormattingOptions);
+CORE_API FString FractionalToString(const double InVal, const FDecimalNumberFormattingRules& InFormattingRules, const FNumberFormattingOptions& InFormattingOptions);
+
+} // namespace Internal
+
+#define FAST_DECIMAL_FORMAT_SIGNED_IMPL(NUMBER_TYPE)																															\
+	FORCEINLINE FString NumberToString(const NUMBER_TYPE InVal, const FDecimalNumberFormattingRules& InFormattingRules, const FNumberFormattingOptions& InFormattingOptions)	\
+	{																																											\
+		const bool bIsNegative = InVal < 0;																																		\
+		return Internal::IntegralToString(bIsNegative, static_cast<uint64>((bIsNegative) ? -InVal : InVal), InFormattingRules, InFormattingOptions);							\
+	}
+
+#define FAST_DECIMAL_FORMAT_UNSIGNED_IMPL(NUMBER_TYPE)																															\
+	FORCEINLINE FString NumberToString(const NUMBER_TYPE InVal, const FDecimalNumberFormattingRules& InFormattingRules, const FNumberFormattingOptions& InFormattingOptions)	\
+	{																																											\
+		return Internal::IntegralToString(false, static_cast<uint64>(InVal), InFormattingRules, InFormattingOptions);															\
+	}
+
+#define FAST_DECIMAL_FORMAT_FRACTIONAL_IMPL(NUMBER_TYPE)																														\
+	FORCEINLINE FString NumberToString(const NUMBER_TYPE InVal, const FDecimalNumberFormattingRules& InFormattingRules, const FNumberFormattingOptions& InFormattingOptions)	\
+	{																																											\
+		return Internal::FractionalToString(static_cast<double>(InVal), InFormattingRules, InFormattingOptions);																\
+	}
+
+FAST_DECIMAL_FORMAT_SIGNED_IMPL(int8)
+FAST_DECIMAL_FORMAT_SIGNED_IMPL(int16)
+FAST_DECIMAL_FORMAT_SIGNED_IMPL(int32)
+FAST_DECIMAL_FORMAT_SIGNED_IMPL(int64)
+
+FAST_DECIMAL_FORMAT_UNSIGNED_IMPL(uint8)
+FAST_DECIMAL_FORMAT_UNSIGNED_IMPL(uint16)
+FAST_DECIMAL_FORMAT_UNSIGNED_IMPL(uint32)
+FAST_DECIMAL_FORMAT_UNSIGNED_IMPL(uint64)
+
+FAST_DECIMAL_FORMAT_FRACTIONAL_IMPL(float)
+FAST_DECIMAL_FORMAT_FRACTIONAL_IMPL(double)
+
+#undef FAST_DECIMAL_FORMAT_SIGNED_IMPL
+#undef FAST_DECIMAL_FORMAT_UNSIGNED_IMPL
+#undef FAST_DECIMAL_FORMAT_FRACTIONAL_IMPL
+
+/**
+ * Get the formatting rules to use when you don't care about culture.
+ */
+CORE_API FDecimalNumberFormattingRules GetCultureAgnosticFormattingRules();
+
+} // namespace FastDecimalFormat
