Added TextRounding and TextPadding tests

These are used to verify that FText is producing the expected output, and verify that FastDecimalFormat is producing equivalent results compared to the ICU DecimalFormat.

diff --git a/Engine/Source/Runtime/Core/Private/Tests/Internationalization/TextTest.cpp b/Engine/Source/Runtime/Core/Private/Tests/Internationalization/TextTest.cpp
--- a/Engine/Source/Runtime/Core/Private/Tests/Internationalization/TextTest.cpp
+++ b/Engine/Source/Runtime/Core/Private/Tests/Internationalization/TextTest.cpp
@@ -636,6 +636,200 @@ bool FTextTest::RunTest (const FString& Parameters)
 }
 
 
+IMPLEMENT_SIMPLE_AUTOMATION_TEST(FTextRoundingTest, "System.Core.Misc.TextRounding", EAutomationTestFlags::EditorContext | EAutomationTestFlags::ClientContext | EAutomationTestFlags::SmokeFilter)
+
+bool FTextRoundingTest::RunTest (const FString& Parameters)
+{
+	static const TCHAR* RoundingModeNames[] = {
+		TEXT("HalfToEven"),
+		TEXT("HalfFromZero"),
+		TEXT("HalfToZero"),
+		TEXT("FromZero"),
+		TEXT("ToZero"),
+		TEXT("ToNegativeInfinity"),
+		TEXT("ToPositiveInfinity"),
+	};
+
+	static_assert(ERoundingMode::ToPositiveInfinity == ARRAY_COUNT(RoundingModeNames) - 1, "RoundingModeNames array needs updating");
+
+	static const double InputValues[] = {
+		1000.1224,
+		1000.1225,
+		1000.1226,
+		1000.1234,
+		1000.1235,
+		1000.1236,
+		
+		1000.1244,
+		1000.1245,
+		1000.1246,
+		1000.1254,
+		1000.1255,
+		1000.1256,
+
+		-1000.1224,
+		-1000.1225,
+		-1000.1226,
+		-1000.1234,
+		-1000.1235,
+		-1000.1236,
+		
+		-1000.1244,
+		-1000.1245,
+		-1000.1246,
+		-1000.1254,
+		-1000.1255,
+		-1000.1256,
+	};
+
+	static const TCHAR* OutputValues[][ARRAY_COUNT(RoundingModeNames)] = 
+	{
+		// HalfToEven        | HalfFromZero      | HalfToZero        | FromZero          | ToZero            | ToNegativeInfinity | ToPositiveInfinity
+		{  TEXT("1000.122"),   TEXT("1000.122"),   TEXT("1000.122"),   TEXT("1000.123"),   TEXT("1000.122"),   TEXT("1000.122"),    TEXT("1000.123") },
+		{  TEXT("1000.122"),   TEXT("1000.123"),   TEXT("1000.122"),   TEXT("1000.123"),   TEXT("1000.122"),   TEXT("1000.122"),    TEXT("1000.123") },
+		{  TEXT("1000.123"),   TEXT("1000.123"),   TEXT("1000.123"),   TEXT("1000.123"),   TEXT("1000.122"),   TEXT("1000.122"),    TEXT("1000.123") },
+		{  TEXT("1000.123"),   TEXT("1000.123"),   TEXT("1000.123"),   TEXT("1000.124"),   TEXT("1000.123"),   TEXT("1000.123"),    TEXT("1000.124") },
+		{  TEXT("1000.124"),   TEXT("1000.124"),   TEXT("1000.123"),   TEXT("1000.124"),   TEXT("1000.123"),   TEXT("1000.123"),    TEXT("1000.124") },
+		{  TEXT("1000.124"),   TEXT("1000.124"),   TEXT("1000.124"),   TEXT("1000.124"),   TEXT("1000.123"),   TEXT("1000.123"),    TEXT("1000.124") },
+
+		{  TEXT("1000.124"),   TEXT("1000.124"),   TEXT("1000.124"),   TEXT("1000.125"),   TEXT("1000.124"),   TEXT("1000.124"),    TEXT("1000.125") },
+		{  TEXT("1000.124"),   TEXT("1000.125"),   TEXT("1000.124"),   TEXT("1000.125"),   TEXT("1000.124"),   TEXT("1000.124"),    TEXT("1000.125") },
+		{  TEXT("1000.125"),   TEXT("1000.125"),   TEXT("1000.125"),   TEXT("1000.125"),   TEXT("1000.124"),   TEXT("1000.124"),    TEXT("1000.125") },
+		{  TEXT("1000.125"),   TEXT("1000.125"),   TEXT("1000.125"),   TEXT("1000.126"),   TEXT("1000.125"),   TEXT("1000.125"),    TEXT("1000.126") },
+		{  TEXT("1000.126"),   TEXT("1000.126"),   TEXT("1000.125"),   TEXT("1000.126"),   TEXT("1000.125"),   TEXT("1000.125"),    TEXT("1000.126") },
+		{  TEXT("1000.126"),   TEXT("1000.126"),   TEXT("1000.126"),   TEXT("1000.126"),   TEXT("1000.125"),   TEXT("1000.125"),    TEXT("1000.126") },
+
+		{ TEXT("-1000.122"),  TEXT("-1000.122"),  TEXT("-1000.122"),  TEXT("-1000.123"),  TEXT("-1000.122"),  TEXT("-1000.123"),   TEXT("-1000.122") },
+		{ TEXT("-1000.122"),  TEXT("-1000.123"),  TEXT("-1000.122"),  TEXT("-1000.123"),  TEXT("-1000.122"),  TEXT("-1000.123"),   TEXT("-1000.122") },
+		{ TEXT("-1000.123"),  TEXT("-1000.123"),  TEXT("-1000.123"),  TEXT("-1000.123"),  TEXT("-1000.122"),  TEXT("-1000.123"),   TEXT("-1000.122") },
+		{ TEXT("-1000.123"),  TEXT("-1000.123"),  TEXT("-1000.123"),  TEXT("-1000.124"),  TEXT("-1000.123"),  TEXT("-1000.124"),   TEXT("-1000.123") },
+		{ TEXT("-1000.124"),  TEXT("-1000.124"),  TEXT("-1000.123"),  TEXT("-1000.124"),  TEXT("-1000.123"),  TEXT("-1000.124"),   TEXT("-1000.123") },
+		{ TEXT("-1000.124"),  TEXT("-1000.124"),  TEXT("-1000.124"),  TEXT("-1000.124"),  TEXT("-1000.123"),  TEXT("-1000.124"),   TEXT("-1000.123") },
+
+		{ TEXT("-1000.124"),  TEXT("-1000.124"),  TEXT("-1000.124"),  TEXT("-1000.125"),  TEXT("-1000.124"),  TEXT("-1000.125"),   TEXT("-1000.124") },
+		{ TEXT("-1000.124"),  TEXT("-1000.125"),  TEXT("-1000.124"),  TEXT("-1000.125"),  TEXT("-1000.124"),  TEXT("-1000.125"),   TEXT("-1000.124") },
+		{ TEXT("-1000.125"),  TEXT("-1000.125"),  TEXT("-1000.125"),  TEXT("-1000.125"),  TEXT("-1000.124"),  TEXT("-1000.125"),   TEXT("-1000.124") },
+		{ TEXT("-1000.125"),  TEXT("-1000.125"),  TEXT("-1000.125"),  TEXT("-1000.126"),  TEXT("-1000.125"),  TEXT("-1000.126"),   TEXT("-1000.125") },
+		{ TEXT("-1000.126"),  TEXT("-1000.126"),  TEXT("-1000.125"),  TEXT("-1000.126"),  TEXT("-1000.125"),  TEXT("-1000.126"),   TEXT("-1000.125") },
+		{ TEXT("-1000.126"),  TEXT("-1000.126"),  TEXT("-1000.126"),  TEXT("-1000.126"),  TEXT("-1000.125"),  TEXT("-1000.126"),   TEXT("-1000.125") },
+	};
+
+	static_assert(ARRAY_COUNT(InputValues) == ARRAY_COUNT(OutputValues), "The size of InputValues does not match OutputValues");
+
+	// This test needs to be run using an English culture
+	FInternationalization& I18N = FInternationalization::Get();
+	const FString OriginalCulture = I18N.GetCurrentCulture()->GetName();
+	I18N.SetCurrentCulture(TEXT("en"));
+
+	// Test to make sure that the decimal formatter is rounding fractional numbers correctly (to 3 decimal places)
+	FNumberFormattingOptions FormattingOptions = FNumberFormattingOptions()
+		.SetUseGrouping(false)
+		.SetMaximumFractionalDigits(3);
+
+	auto DoSingleTest = [&](const double InNumber, const FString& InExpectedString, const FString& InDescription)
+	{
+		const FText ResultText = FText::AsNumber(InNumber, &FormattingOptions);
+		if(ResultText.ToString() != InExpectedString)
+		{
+			AddError(FString::Printf(TEXT("Text rounding failure: source '%f' - expected '%s' - result '%s'. %s."), InNumber, *InExpectedString, *ResultText.ToString(), *InDescription));
+		}
+	};
+
+	auto DoAllTests = [&](const ERoundingMode InRoundingMode)
+	{
+		FormattingOptions.SetRoundingMode(InRoundingMode);
+
+		for (int32 TestValueIndex = 0; TestValueIndex < ARRAY_COUNT(InputValues); ++TestValueIndex)
+		{
+			DoSingleTest(InputValues[TestValueIndex], OutputValues[TestValueIndex][InRoundingMode], RoundingModeNames[InRoundingMode]);
+		}
+	};
+
+	DoAllTests(ERoundingMode::HalfToEven);
+	DoAllTests(ERoundingMode::HalfFromZero);
+	DoAllTests(ERoundingMode::HalfToZero);
+	DoAllTests(ERoundingMode::FromZero);
+	DoAllTests(ERoundingMode::ToZero);
+	DoAllTests(ERoundingMode::ToNegativeInfinity);
+	DoAllTests(ERoundingMode::ToPositiveInfinity);
+
+	// HalfToEven - Rounds to the nearest place, equidistant ties go to the value which is closest to an even value: 1.5 becomes 2, 0.5 becomes 0
+	{
+		FormattingOptions.SetRoundingMode(ERoundingMode::HalfToEven);
+
+		DoSingleTest(1000.12459, TEXT("1000.125"), TEXT("HalfToEven"));
+		DoSingleTest(1000.124549, TEXT("1000.125"), TEXT("HalfToEven"));
+		DoSingleTest(1000.124551, TEXT("1000.125"), TEXT("HalfToEven"));
+		DoSingleTest(1000.12451, TEXT("1000.125"), TEXT("HalfToEven"));
+		DoSingleTest(1000.1245000001, TEXT("1000.125"), TEXT("HalfToEven"));
+		DoSingleTest(1000.12450000000001, TEXT("1000.124"), TEXT("HalfToEven"));
+	}
+
+	// Restore original culture
+	I18N.SetCurrentCulture(OriginalCulture);
+
+	return true;
+}
+
+
+IMPLEMENT_SIMPLE_AUTOMATION_TEST(FTextPaddingTest, "System.Core.Misc.TextPadding", EAutomationTestFlags::EditorContext | EAutomationTestFlags::ClientContext | EAutomationTestFlags::SmokeFilter)
+
+bool FTextPaddingTest::RunTest (const FString& Parameters)
+{
+	// This test needs to be run using an English culture
+	FInternationalization& I18N = FInternationalization::Get();
+	const FString OriginalCulture = I18N.GetCurrentCulture()->GetName();
+	I18N.SetCurrentCulture(TEXT("en"));
+
+	// Test to make sure that the decimal formatter is padding integral numbers correctly
+	FNumberFormattingOptions FormattingOptions;
+
+	auto DoSingleTest = [&](const int32 InNumber, const FString& InExpectedString, const FString& InDescription)
+	{
+		const FText ResultText = FText::AsNumber(InNumber, &FormattingOptions);
+		if(ResultText.ToString() != InExpectedString)
+		{
+			AddError(FString::Printf(TEXT("Text padding failure: source '%d' - expected '%s' - result '%s'. %s."), InNumber, *InExpectedString, *ResultText.ToString(), *InDescription));
+		}
+	};
+
+	// Test with a max limit of 3
+	{
+		FormattingOptions = FNumberFormattingOptions()
+			.SetUseGrouping(false)
+			.SetMaximumIntegralDigits(3);
+
+		DoSingleTest(123456,  TEXT("456"),  TEXT("Truncating '123456' to a max of 3 integral digits"));
+		DoSingleTest(-123456, TEXT("-456"), TEXT("Truncating '-123456' to a max of 3 integral digits"));
+	}
+
+	// Test with a min limit of 6
+	{
+		FormattingOptions = FNumberFormattingOptions()
+			.SetUseGrouping(false)
+			.SetMinimumIntegralDigits(6);
+
+		DoSingleTest(123,  TEXT("000123"),  TEXT("Padding '123' to a min of 6 integral digits"));
+		DoSingleTest(-123, TEXT("-000123"), TEXT("Padding '-123' to a min of 6 integral digits"));
+	}
+
+	// Test with forced fractional digits
+	{
+		FormattingOptions = FNumberFormattingOptions()
+			.SetUseGrouping(false)
+			.SetMinimumFractionalDigits(3);
+
+		DoSingleTest(123,  TEXT("123.000"),  TEXT("Padding '123' to a min of 3 fractional digits"));
+		DoSingleTest(-123, TEXT("-123.000"), TEXT("Padding '-123' to a min of 3 fractional digits"));
+	}
+
+	// Restore original culture
+	I18N.SetCurrentCulture(OriginalCulture);
+
+	return true;
+}
+
+
 #if UE_ENABLE_ICU
 
 IMPLEMENT_SIMPLE_AUTOMATION_TEST(FICUTextTest, "System.Core.Misc.ICUText", EAutomationTestFlags::EditorContext | EAutomationTestFlags::ClientContext | EAutomationTestFlags::SmokeFilter)
