Swapped FText over to use FastDecimalFormat by default

This is controlled by the EXPERIMENTAL_TEXT_FAST_DECIMAL_FORMAT define in ICUText.h, and is turned on by default. Setting this to 0 will go through the slower ICU DecimalFormat path instead.

diff --git a/Engine/Source/Runtime/Core/Private/Internationalization/ICUText.h b/Engine/Source/Runtime/Core/Private/Internationalization/ICUText.h
--- a/Engine/Source/Runtime/Core/Private/Internationalization/ICUText.h
+++ b/Engine/Source/Runtime/Core/Private/Internationalization/ICUText.h
@@ -5,25 +5,36 @@
 #include "ICUUtilities.h"
 #include "ICUCulture.h"
 #include "TextHistory.h"
+#include "FastDecimalFormat.h"
 #include "unicode/utypes.h"
 #include <unicode/numfmt.h>
 #include "unicode/fmtable.h"
 #include "unicode/unistr.h"
 
+#define EXPERIMENTAL_TEXT_FAST_DECIMAL_FORMAT 1
+
 template<typename T1, typename T2>
 FText FText::AsNumberTemplate(T1 Val, const FNumberFormattingOptions* const Options, const FCulturePtr& TargetCulture)
 {
 	FInternationalization& I18N = FInternationalization::Get();
 	checkf(I18N.IsInitialized() == true, TEXT("FInternationalization is not initialized. An FText formatting method was likely used in static object initialization - this is not supported."));
 	const FCulture& Culture = TargetCulture.IsValid() ? *TargetCulture : *I18N.GetCurrentCulture();
+
+	FString NativeString;
+
+#if EXPERIMENTAL_TEXT_FAST_DECIMAL_FORMAT
+	const FDecimalNumberFormattingRules& FormattingRules = Culture.Implementation->GetDecimalNumberFormattingRules();
+	const FNumberFormattingOptions& FormattingOptions = (Options) ? *Options : FormattingRules.CultureDefaultFormattingOptions;
+	NativeString = FastDecimalFormat::NumberToString(Val, FormattingRules, FormattingOptions);
+#else // EXPERIMENTAL_TEXT_FAST_DECIMAL_FORMAT
 	UErrorCode ICUStatus = U_ZERO_ERROR;
 	const TSharedRef<const icu::DecimalFormat, ESPMode::ThreadSafe>& ICUDecimalFormat( Culture.Implementation->GetDecimalFormatter(Options) );
 	icu::Formattable FormattableVal(static_cast<T2>(Val));
 	icu::UnicodeString FormattedString;
 	ICUDecimalFormat->format(FormattableVal, FormattedString, ICUStatus);
 
-	FString NativeString;
 	ICUUtilities::ConvertString(FormattedString, NativeString);
+#endif // EXPERIMENTAL_TEXT_FAST_DECIMAL_FORMAT
 
 	return FText::CreateNumericalText(MoveTemp(NativeString), MakeShareable(new FTextHistory_AsNumber(Val, Options, TargetCulture)));
 }
@@ -34,14 +45,22 @@ FText FText::AsCurrencyTemplate(T1 Val, const FString& CurrencyCode, const FNumb
 	FInternationalization& I18N = FInternationalization::Get();
 	checkf(I18N.IsInitialized() == true, TEXT("FInternationalization is not initialized. An FText formatting method was likely used in static object initialization - this is not supported."));
 	const FCulture& Culture = TargetCulture.IsValid() ? *TargetCulture : *I18N.GetCurrentCulture();
+
+	FString NativeString;
+
+#if EXPERIMENTAL_TEXT_FAST_DECIMAL_FORMAT
+	const FDecimalNumberFormattingRules& FormattingRules = Culture.Implementation->GetCurrencyFormattingRules(CurrencyCode);
+	const FNumberFormattingOptions& FormattingOptions = (Options) ? *Options : FormattingRules.CultureDefaultFormattingOptions;
+	NativeString = FastDecimalFormat::NumberToString(Val, FormattingRules, FormattingOptions);
+#else // EXPERIMENTAL_TEXT_FAST_DECIMAL_FORMAT
 	UErrorCode ICUStatus = U_ZERO_ERROR;
 	const TSharedRef<const icu::DecimalFormat>& ICUDecimalFormat(Culture.Implementation->GetCurrencyFormatter(CurrencyCode, Options));
 	icu::Formattable FormattableVal(static_cast<T2>(Val));
 	icu::UnicodeString FormattedString;
 	ICUDecimalFormat->format(FormattableVal, FormattedString, ICUStatus);
 
-	FString NativeString;
 	ICUUtilities::ConvertString(FormattedString, NativeString);
+#endif // EXPERIMENTAL_TEXT_FAST_DECIMAL_FORMAT
 
 	return FText::CreateNumericalText(MoveTemp(NativeString), MakeShareable(new FTextHistory_AsCurrency(Val, CurrencyCode, Options, TargetCulture)));
 }
@@ -52,14 +71,24 @@ FText FText::AsPercentTemplate(T1 Val, const FNumberFormattingOptions* const Opt
 	FInternationalization& I18N = FInternationalization::Get();
 	checkf(I18N.IsInitialized() == true, TEXT("FInternationalization is not initialized. An FText formatting method was likely used in static object initialization - this is not supported."));
 	const FCulture& Culture = TargetCulture.IsValid() ? *TargetCulture : *I18N.GetCurrentCulture();
+
+	FString NativeString;
+
+#if EXPERIMENTAL_TEXT_FAST_DECIMAL_FORMAT
+	const FDecimalNumberFormattingRules& FormattingRules = Culture.Implementation->GetPercentFormattingRules();
+	const FNumberFormattingOptions& FormattingOptions = (Options) ? *Options : FormattingRules.CultureDefaultFormattingOptions;
+	NativeString = FastDecimalFormat::NumberToString(Val * static_cast<T1>(100), FormattingRules, FormattingOptions);
+#else // EXPERIMENTAL_TEXT_FAST_DECIMAL_FORMAT
 	UErrorCode ICUStatus = U_ZERO_ERROR;
 	const TSharedRef<const icu::DecimalFormat>& ICUDecimalFormat( Culture.Implementation->GetPercentFormatter(Options) );
 	icu::Formattable FormattableVal(static_cast<T2>(Val));
 	icu::UnicodeString FormattedString;
 	ICUDecimalFormat->format(FormattableVal, FormattedString, ICUStatus);
 
-	FString NativeString;
 	ICUUtilities::ConvertString(FormattedString, NativeString);
+#endif // EXPERIMENTAL_TEXT_FAST_DECIMAL_FORMAT
 
 	return FText::CreateNumericalText(MoveTemp(NativeString), MakeShareable(new FTextHistory_AsPercent(Val, Options, TargetCulture)));
 }
+
+#undef EXPERIMENTAL_TEXT_FAST_DECIMAL_FORMAT
