Split Async test into individual tests.
Added validation to Void Task test.

diff --git a/Engine/Source/Runtime/Core/Private/Tests/Async/AsyncTest.cpp b/Engine/Source/Runtime/Core/Private/Tests/Async/AsyncTest.cpp
--- a/Engine/Source/Runtime/Core/Private/Tests/Async/AsyncTest.cpp
+++ b/Engine/Source/Runtime/Core/Private/Tests/Async/AsyncTest.cpp
@@ -7,45 +7,78 @@
 #include "Future.h"
 
 
-IMPLEMENT_SIMPLE_AUTOMATION_TEST(FAsyncTest, "System.Core.Async.Async", EAutomationTestFlags::EditorContext | EAutomationTestFlags::EngineFilter)
+IMPLEMENT_SIMPLE_AUTOMATION_TEST(FAsyncGraphTest, "System.Core.Async.Task Graph", EAutomationTestFlags::EditorContext | EAutomationTestFlags::EngineFilter)
+IMPLEMENT_SIMPLE_AUTOMATION_TEST(FAsyncThreadedTaskTest, "System.Core.Async.Threaded Task", EAutomationTestFlags::EditorContext | EAutomationTestFlags::EngineFilter)
+IMPLEMENT_SIMPLE_AUTOMATION_TEST(FAsyncThreadedPoolTest, "System.Core.Async.Threaded Pool", EAutomationTestFlags::EditorContext | EAutomationTestFlags::EngineFilter)
+IMPLEMENT_SIMPLE_AUTOMATION_TEST(FAsyncVoidTaskTest, "System.Core.Async.Void Task", EAutomationTestFlags::EditorContext | EAutomationTestFlags::EngineFilter)
 
-
-bool FAsyncTest::RunTest(const FString& Parameters)
+/**
+* Helper methods used in the test cases
+*/
+namespace AsyncTestUtils
 {
 	TFunction<int()> Task = []() {
 		return 123;
 	};
 
-	// task graph task
-	{
-		auto Future = Async(EAsyncExecution::TaskGraph, Task);
-		int Result = Future.Get();
+	bool bHasVoidTaskFinished = false;
+
+	TFunction<void()> VoidTask = []() {
+		bHasVoidTaskFinished = true;
+	};
+}
+
+/**
+* Test that task graph tasks return correctly.
+*/
+bool FAsyncGraphTest::RunTest(const FString& Parameters)
+{
+	auto Future = Async(EAsyncExecution::TaskGraph, AsyncTestUtils::Task);
+	int Result = Future.Get();
 
-		TestEqual(TEXT("Task graph task must return expected value"), Result, 123);
-	}
+	TestEqual(TEXT("Task graph task must return expected value"), Result, 123);
 
-	// thread task
-	{
-		auto Future = Async(EAsyncExecution::Thread, Task);
-		int Result = Future.Get();
+	return true;
+}
 
-		TestEqual(TEXT("Threaded task must return expected value"), Result, 123);
-	}
+/**
+* Test that threaded tasks return correctly.
+*/
+bool FAsyncThreadedTaskTest::RunTest(const FString& Parameters)
+{
+	auto Future = Async(EAsyncExecution::Thread, AsyncTestUtils::Task);
+	int Result = Future.Get();
 
-	// thread pool task
-	{
-		auto Future = Async(EAsyncExecution::ThreadPool, Task);
-		int Result = Future.Get();
+	TestEqual(TEXT("Threaded task must return expected value"), Result, 123);
 
-		TestEqual(TEXT("Thread pool task must return expected value"), Result, 123);
-	}
+	return true;
+}
+
+/**
+* Test that threaded pool tasks return correctly.
+*/
+bool FAsyncThreadedPoolTest::RunTest(const FString& Parameters)
+{
+	auto Future = Async(EAsyncExecution::ThreadPool, AsyncTestUtils::Task);
+	int Result = Future.Get();
+
+	TestEqual(TEXT("Thread pool task must return expected value"), Result, 123);
+
+	return true;
+}
+
+/**
+* Test that void tasks run without errors or warnings.
+*/
+bool FAsyncVoidTaskTest::RunTest(const FString& Parameters)
+{
+	// Reset test variable before running
+	AsyncTestUtils::bHasVoidTaskFinished = false;
+	auto Future = Async(EAsyncExecution::TaskGraph, AsyncTestUtils::VoidTask);
+	Future.Get();
 
-	TFunction<void()> VoidTask = []() { };
-	
-	// TAsync<void>
-	{
-		auto Future = Async(EAsyncExecution::TaskGraph, VoidTask);
-	}
+	// Check that the variable state was updated by task
+	TestTrue(TEXT("Void tasks should run"), AsyncTestUtils::bHasVoidTaskFinished);
 
 	return true;
 }
