Remove dead code from Containers test

diff --git a/Engine/Source/Runtime/Core/Private/Containers/ContainersTest.cpp b/Engine/Source/Runtime/Core/Private/Containers/ContainersTest.cpp
--- a/Engine/Source/Runtime/Core/Private/Containers/ContainersTest.cpp
+++ b/Engine/Source/Runtime/Core/Private/Containers/ContainersTest.cpp
@@ -272,15 +272,6 @@ namespace
 	}
 }
 
-static bool doodle = false;
-void checkDoodle(bool b)
-{
-	if (doodle)
-	{
-		check(b);
-	}
-}
-
 bool FContainersTest::RunTest( const FString& Parameters )
 {
 	typedef TMap<int32, FContainerTestValueType> TestMapType;
