Added FString::Format.

Both named and ordered arguments are supported, using patterns of the form "Some text with an {Arg}" and "Some text with an {0}" respectively.
Escaped { characters can be achieved using the ` character (for exaple, "Escaped char `{}") No other escape characters are supported or necessary.
Format behaviour is similar to that of FText::Format. Numeric formatting rules may be added in future.

diff --git a/Engine/Source/Runtime/Core/Private/Misc/ExpressionParserExamples.cpp b/Engine/Source/Runtime/Core/Private/Misc/ExpressionParserExamples.cpp
--- a/Engine/Source/Runtime/Core/Private/Misc/ExpressionParserExamples.cpp
+++ b/Engine/Source/Runtime/Core/Private/Misc/ExpressionParserExamples.cpp
@@ -73,191 +73,6 @@ private:
 };
 DEFINE_EXPRESSION_NODE_TYPE(FExampleMathEvaluator::FExampleAdd, 0x4BF093DC, 0xA92247B5, 0xAE53A9B3, 0x64A84D2C)
 
-
-/** A simple string formatter */
-class FExampleStringFormatter
-{
-public:
-	/** Token representing a literal string inside the string */
-	struct FStringLiteral
-	{
-		FStringLiteral(FStringToken InString) : String(InString), Len(InString.GetTokenEndPos() - InString.GetTokenStartPos()) {}
-		FStringToken String;
-		int32 Len;
-	};
-	/** Token representing a user-defined token, such as {Argument} */
-	struct FFormatToken
-	{
-		FFormatToken(FStringToken InIdentifier) : Identifier(InIdentifier), Len(Identifier.GetTokenEndPos() - Identifier.GetTokenStartPos()) {}
-		FStringToken Identifier;
-		int32 Len;
-	};
-
-	FExampleStringFormatter()
-	{
-		using namespace ExpressionParser;
-
-		// Define a tokenizer that will parse user supplied tokens (FFormatToken) in the expression
-		TokenDefinitions.DefineToken([](FExpressionTokenConsumer& Consumer) -> TOptional<FExpressionError> {
-			auto& Stream = Consumer.GetStream();
-
-			TOptional<FStringToken> OpeningChar = Stream.ParseSymbol('{');
-			if (!OpeningChar.IsSet())
-			{
-				return TOptional<FExpressionError>();
-			}
-
-			FStringToken& EntireToken = OpeningChar.GetValue();
-
-			// Optional whitespace
-			Stream.ParseToken([](TCHAR InC) { return FChar::IsWhitespace(InC) ? EParseState::Continue : EParseState::StopBefore; }, &EntireToken);
-
-			// The identifier itself
-			TOptional<FStringToken> Identifier = Stream.ParseToken([](TCHAR InC) { return FChar::IsWhitespace(InC) || InC == '}' ? EParseState::StopBefore : EParseState::Continue; }, &EntireToken);
-
-			if (!Identifier.IsSet())
-			{
-				FFormatOrderedArguments Args;
-				Args.Add(FText::FromString(FString(EntireToken.GetTokenEndPos()).Left(10) + TEXT("...")));
-				return FExpressionError(FText::Format(LOCTEXT("UnrecognizedIdentifier", "Unrecognized identifier symbol: '{0}'"), Args));
-			}
-
-			// Optional whitespace
-			Stream.ParseToken([](TCHAR InC) { return FChar::IsWhitespace(InC) ? EParseState::Continue : EParseState::StopBefore; }, &EntireToken);
-			
-			if (!Stream.ParseSymbol('}', &EntireToken).IsSet())
-			{
-				FFormatOrderedArguments Args;
-				Args.Add(FText::FromString(FString(EntireToken.GetTokenEndPos() - EntireToken.GetTokenStartPos(), EntireToken.GetTokenStartPos())));
-				return FExpressionError(FText::Format(LOCTEXT("InvalidTokenDefinition", "Invalid token definition: '{0}'"), Args));
-			}
-
-			// Add the token to the consumer. This moves the read position in the stream to the end of the token.
-			Consumer.Add(EntireToken, FFormatToken(Identifier.GetValue()));
-
-			return TOptional<FExpressionError>();
-		});
-
-		// Define a tokenizer for anything else
-		TokenDefinitions.DefineToken([](FExpressionTokenConsumer& Consumer){
-			
-			// Called at the start of every new token - see if it's a { character
-			TOptional<FStringToken> Token = Consumer.GetStream().ParseToken([](TCHAR InC){
-				return InC == '{' ? EParseState::StopBefore : EParseState::Continue;
-			});
-
-			if (Token.IsSet())
-			{
-				// Add the token to the consumer. This moves the read position in the stream to the end of the token.
-				Consumer.Add(Token.GetValue(), FStringLiteral(Token.GetValue()));
-			}
-
-			return TOptional<FExpressionError>();
-		});
-	}
-
-	/** Format the specified text using the specified arguments. Replaces instances of { Argument } with keys in the map matching 'Argument' */
-	TValueOrError<FText, FExpressionError> Format(FText InExpression, const TMap<FString, FText>& Args) const
-	{
-		const FString& String = InExpression.ToString();
-
-		TValueOrError<TArray<FExpressionToken>, FExpressionError> Result = ExpressionParser::Lex(*String, TokenDefinitions);
-		if (!Result.IsValid())
-		{
-			return MakeError(Result.GetError());
-		}
-		
-		// This code deliberately tries to reallocate as little as possible
-		FString Formatted;
-		Formatted.Reserve(1024);
-		for (auto& Token : Result.GetValue())
-		{
-			if (const auto* Literal = Token.Node.Cast<FStringLiteral>())
-			{
-				Formatted.AppendChars(Literal->String.GetTokenStartPos(), Literal->Len);
-			}
-			else if (const auto* FormatToken = Token.Node.Cast<FFormatToken>())
-			{
-				const FText* Arg = nullptr;
-				for (auto& Pair : Args)
-				{
-					if (Pair.Key.Len() == FormatToken->Len && FCString::Strnicmp(FormatToken->Identifier.GetTokenStartPos(), *Pair.Key, FormatToken->Len) == 0)
-					{
-						Arg = &Pair.Value;
-						break;
-					}
-				}
-
-				if (Arg)
-				{
-					Formatted.AppendChars(FormatToken->Identifier.GetTokenStartPos(), FormatToken->Len);
-				}
-				else
-				{
-					return MakeError(FText::Format(LOCTEXT("UnspecifiedToken", "Unspecified token: {0}"), FText::FromString(FormatToken->Identifier.GetString())));
-				}
-			}
-		}
-
-		return MakeValue(FText::FromString(Formatted));
-	}
-
-private:
-	FTokenDefinitions TokenDefinitions;
-};
-
-DEFINE_EXPRESSION_NODE_TYPE(FExampleStringFormatter::FStringLiteral, 0x03ED3A25, 0x85D94664, 0x8A8001A1, 0xDCC637F7)
-DEFINE_EXPRESSION_NODE_TYPE(FExampleStringFormatter::FFormatToken, 0xAAB48E5B, 0xEDA94853, 0xA951ED2D, 0x0A8E795D)
-
-/** Unit test for the string formatter */
-/*
-DEFINE_LOG_CATEGORY_STATIC(LogExampleStringTests, Log, All);
-
-#include "TextHistory.h"
-
-IMPLEMENT_SIMPLE_AUTOMATION_TEST(FExampleStringFormatterTest, "System.Core.Expression Parser", EAutomationTestFlags::Editor)
-bool FExampleStringFormatterTest::RunTest( const FString& Parameters )
-{
-	FExampleStringFormatter Formatter;
-
-	double Time = FPlatformTime::Seconds();
-
-	const FText Pattern = LOCTEXT("Pattern", "This is some text containing two arguments, { Argument1 } and {Argument2}.");
-
-
-	FFormatNamedArguments TextArgs;
-	TextArgs.Add(TEXT("Argument1"), LOCTEXT("Arg1", "Argument1"));
-	TextArgs.Add(TEXT("Argument2"), LOCTEXT("Arg2", "Argument2"));
-
-	{
-
-		TMap<FString, FText> Args;
-		Args.Add(TEXT("Argument1"), LOCTEXT("Arg1", "Argument1"));
-		Args.Add(TEXT("Argument2"), LOCTEXT("Arg2", "Argument2"));
-
-		for (int32 Bla = 0; Bla < 1000000; ++Bla)
-		{
-			MakeShareable(new FTextHistory_NamedFormat(Pattern, TextArgs));
-			Formatter.Format(Pattern, Args);
-		}
-	}
-
-	UE_LOG(LogExampleStringTests, Log, TEXT("Expression parser took %fs"), FPlatformTime::Seconds() - Time);
-
-	Time = FPlatformTime::Seconds();
-	{
-		for (int32 Bla = 0; Bla < 1000000; ++Bla)
-		{
-			FText::Format(Pattern, TextArgs);
-		}
-	}
-	UE_LOG(LogExampleStringTests, Log, TEXT("FText::Format took %fs"), FPlatformTime::Seconds() - Time);
-
-	return true;
-}
-
-*/
-
 #undef LOCTEXT_NAMESPACE
 
 #endif //!(UE_BUILD_SHIPPING || UE_BUILD_TEST)
\ No newline at end of file
diff --git a/Engine/Source/Runtime/Core/Private/Misc/StringFormatter.cpp b/Engine/Source/Runtime/Core/Private/Misc/StringFormatter.cpp
new file mode 100644
--- /dev/null
+++ b/Engine/Source/Runtime/Core/Private/Misc/StringFormatter.cpp
@@ -0,0 +1,650 @@
+// Copyright 1998-2015 Epic Games, Inc. All Rights Reserved.
+
+#include "CorePrivatePCH.h"
+#include "ExpressionParser.h"
+#include "StringFormatter.h"
+
+#define LOCTEXT_NAMESPACE "StringFormatter"
+
+FStringFormatArg::FStringFormatArg( const int32 Value ) : Type(Int), IntValue(Value) {}
+FStringFormatArg::FStringFormatArg( const uint32 Value ) : Type(UInt), UIntValue(Value) {}
+FStringFormatArg::FStringFormatArg( const int64 Value ) : Type(Int), IntValue(Value) {}
+FStringFormatArg::FStringFormatArg( const uint64 Value ) : Type(UInt), UIntValue(Value) {}
+FStringFormatArg::FStringFormatArg( const float Value ) : Type(Double), DoubleValue(Value) {}
+FStringFormatArg::FStringFormatArg( const double Value ) : Type(Double), DoubleValue(Value) {}
+FStringFormatArg::FStringFormatArg( const FString& Value ) : Type(String), StringValue(&Value) {}
+FStringFormatArg::FStringFormatArg( const TCHAR* Value ) : Type(StringLiteral), StringLiteralValue(Value) {}
+FStringFormatArg::FStringFormatArg( const FStringFormatArg& RHS )
+{
+	Type = RHS.Type;
+	switch (Type)
+	{
+		case Int: 				IntValue = RHS.IntValue; break;
+		case UInt: 				UIntValue = RHS.UIntValue; break;
+		case Double: 			IntValue = RHS.IntValue; break;
+		case String: 			StringValue = RHS.StringValue; break;
+		case StringLiteral: 	StringLiteralValue = RHS.StringLiteralValue; break;
+	}
+}
+
+void AppendToString(const FStringFormatArg& Arg, FString& StringToAppendTo)
+{
+	switch(Arg.Type)
+	{
+		case FStringFormatArg::Int: 			StringToAppendTo.Append(LexicalConversion::ToString(Arg.IntValue)); break;
+		case FStringFormatArg::UInt: 			StringToAppendTo.Append(LexicalConversion::ToString(Arg.UIntValue)); break;
+		case FStringFormatArg::Double: 			StringToAppendTo.Append(LexicalConversion::ToString(Arg.DoubleValue)); break;
+		case FStringFormatArg::String: 			StringToAppendTo.AppendChars(**Arg.StringValue, Arg.StringValue->Len()); break;
+		case FStringFormatArg::StringLiteral: 	StringToAppendTo += Arg.StringLiteralValue; break;
+	}
+}
+
+/** Token representing a literal string inside the string */
+struct FStringLiteral
+{
+	FStringLiteral(const FStringToken& InString) : String(InString), Len(InString.GetTokenEndPos() - InString.GetTokenStartPos()) {}
+	/** The string literal token */
+	FStringToken String;
+	/** Cached length of the string */
+	int32 Len;
+};
+
+/** Token representing a user-defined token, such as {Argument} */
+struct FFormatSpecifier
+{
+	FFormatSpecifier(const FStringToken& InIdentifier, const FStringToken& InEntireToken) : Identifier(InIdentifier), EntireToken(InEntireToken), Len(Identifier.GetTokenEndPos() - Identifier.GetTokenStartPos()) {}
+
+	/** The identifier part of the token */
+	FStringToken Identifier;
+	/** The entire token */
+	FStringToken EntireToken;
+	/** Cached length of the identifier */
+	int32 Len;
+};
+
+/** Token representing a user-defined index token, such as {0} */
+struct FIndexSpecifier
+{
+	FIndexSpecifier(int32 InIndex, const FStringToken& InEntireToken) : Index(InIndex), EntireToken(InEntireToken) {}
+
+	/** The index of the parsed token */
+	int32 Index;
+	/** The entire token */
+	FStringToken EntireToken;
+};
+
+/** Token representing an escaped character */
+struct FEscapedCharacter
+{
+	FEscapedCharacter(TCHAR InChar) : Character(InChar) {}
+
+	/** The character that was escaped */
+	TCHAR Character;
+};
+
+DEFINE_EXPRESSION_NODE_TYPE(FStringLiteral, 0x03ED3A25, 0x85D94664, 0x8A8001A1, 0xDCC637F7)
+DEFINE_EXPRESSION_NODE_TYPE(FFormatSpecifier, 0xAAB48E5B, 0xEDA94853, 0xA951ED2D, 0x0A8E795D)
+DEFINE_EXPRESSION_NODE_TYPE(FIndexSpecifier, 0xE11F9937, 0xAF714AC5, 0x88A4E04E, 0x723A753C)
+DEFINE_EXPRESSION_NODE_TYPE(FEscapedCharacter, 0x48FF0754, 0x508941BB, 0x9D5447FF, 0xCAC61362)
+
+FExpressionError GenerateErrorMsg(const FStringToken& Token)
+{
+	FFormatOrderedArguments Args;
+	Args.Add(FText::FromString(FString(Token.GetTokenEndPos()).Left(10) + TEXT("...")));
+	return FExpressionError(FText::Format(LOCTEXT("InvalidTokenDefinition", "Invalid token definition at '{0}'"), Args));
+}
+
+TOptional<FExpressionError> ParseIndex(FExpressionTokenConsumer& Consumer, bool bEmitErrors)
+{
+	auto& Stream = Consumer.GetStream();
+
+	TOptional<FStringToken> OpeningChar = Stream.ParseSymbol('{');
+	if (!OpeningChar.IsSet())
+	{
+		return TOptional<FExpressionError>();
+	}
+
+	FStringToken& EntireToken = OpeningChar.GetValue();
+
+	// Optional whitespace
+	Stream.ParseToken([](TCHAR InC) { return FChar::IsWhitespace(InC) ? EParseState::Continue : EParseState::StopBefore; }, &EntireToken);
+
+	// The identifier itself
+	TOptional<int32> Index;
+	Stream.ParseToken([&](TCHAR InC) {
+		if (FChar::IsDigit(InC))
+		{
+			if (!Index.IsSet())
+			{
+				Index = 0;
+			}
+			Index.GetValue() *= 10;
+			Index.GetValue() += InC - '0';
+			return EParseState::Continue;
+		}
+		return EParseState::StopBefore;
+	}, &EntireToken);
+
+	if (!Index.IsSet())
+	{
+		// Not a valid token
+		if (bEmitErrors)
+		{
+			return GenerateErrorMsg(EntireToken);
+		}
+		else
+		{
+			return TOptional<FExpressionError>();
+		}
+	}
+
+	// Optional whitespace
+	Stream.ParseToken([](TCHAR InC) { return FChar::IsWhitespace(InC) ? EParseState::Continue : EParseState::StopBefore; }, &EntireToken);
+	
+	if (!Stream.ParseSymbol('}', &EntireToken).IsSet())
+	{
+		// Not a valid token
+		if (bEmitErrors)
+		{
+			return GenerateErrorMsg(EntireToken);
+		}
+		else
+		{
+			return TOptional<FExpressionError>();
+		}
+	}
+
+	// Add the token to the consumer. This moves the read position in the stream to the end of the token.
+	Consumer.Add(EntireToken, FIndexSpecifier(Index.GetValue(), EntireToken));
+	return TOptional<FExpressionError>();
+}
+
+TOptional<FExpressionError> ParseSpecifier(FExpressionTokenConsumer& Consumer, bool bEmitErrors)
+{
+	auto& Stream = Consumer.GetStream();
+
+	TOptional<FStringToken> OpeningChar = Stream.ParseSymbol('{');
+	if (!OpeningChar.IsSet())
+	{
+		return TOptional<FExpressionError>();
+	}
+
+	FStringToken& EntireToken = OpeningChar.GetValue();
+
+	// Optional whitespace
+	Stream.ParseToken([](TCHAR InC) { return FChar::IsWhitespace(InC) ? EParseState::Continue : EParseState::StopBefore; }, &EntireToken);
+
+	// The identifier itself
+	TOptional<FStringToken> Identifier = Stream.ParseToken([](TCHAR InC) {
+		if (FChar::IsWhitespace(InC) || InC == '}')
+		{
+			return EParseState::StopBefore;
+		}
+		else if (FChar::IsIdentifier(InC))
+		{
+			return EParseState::Continue;
+		}
+		else
+		{
+			return EParseState::Cancel;
+		}
+
+	}, &EntireToken);
+
+	if (!Identifier.IsSet())
+	{
+		// Not a valid token
+		// Not a valid token
+		if (bEmitErrors)
+		{
+			return GenerateErrorMsg(EntireToken);
+		}
+		else
+		{
+			return TOptional<FExpressionError>();
+		}
+	}
+
+	// Optional whitespace
+	Stream.ParseToken([](TCHAR InC) { return FChar::IsWhitespace(InC) ? EParseState::Continue : EParseState::StopBefore; }, &EntireToken);
+
+	if (!Stream.ParseSymbol('}', &EntireToken).IsSet())
+	{		
+		// Not a valid token
+		if (bEmitErrors)
+		{
+			return GenerateErrorMsg(EntireToken);
+		}
+		else
+		{
+			return TOptional<FExpressionError>();
+		}
+	}
+
+	// Add the token to the consumer. This moves the read position in the stream to the end of the token.
+	Consumer.Add(EntireToken, FFormatSpecifier(Identifier.GetValue(), EntireToken));
+	return TOptional<FExpressionError>();
+}
+
+static const TCHAR EscapeChar = '`';
+
+/** Parse an escaped character */
+TOptional<FExpressionError> ParseEscapedChar(FExpressionTokenConsumer& Consumer, bool bEmitErrors)
+{
+	static const TCHAR* ValidEscapeChars = TEXT("{`");
+
+	TOptional<FStringToken> Token = Consumer.GetStream().ParseSymbol(EscapeChar);
+	if (!Token.IsSet())
+	{
+		return TOptional<FExpressionError>();
+	}
+
+	// Accumulate the next character into the token
+	TOptional<FStringToken> EscapedChar = Consumer.GetStream().ParseSymbol(&Token.GetValue());
+	if (!EscapedChar.IsSet())
+	{
+		return TOptional<FExpressionError>();
+	}
+
+	// Check for a valid escape character
+	const TCHAR Character = *EscapedChar->GetTokenStartPos();
+	if (FCString::Strchr(ValidEscapeChars, Character))
+	{
+		// Add the token to the consumer. This moves the read position in the stream to the end of the token.
+		Consumer.Add(Token.GetValue(), FEscapedCharacter(Character));
+		return TOptional<FExpressionError>();
+	}
+	else if (bEmitErrors)
+	{
+		FString CharStr;
+		CharStr += Character;
+		FFormatOrderedArguments Args;
+		Args.Add(FText::FromString(CharStr));
+		return FExpressionError(FText::Format(LOCTEXT("InvalidEscapeCharacter", "Invalid escape character '{0}'"), Args));
+	}
+	else
+	{
+		return TOptional<FExpressionError>();
+	}
+}
+
+/** Parse anything until we find an unescaped { */
+TOptional<FExpressionError> ParseLiteral(FExpressionTokenConsumer& Consumer, bool bEmitErrors)
+{
+	// Include a leading { character - if it was a valid argument token it would have been picked up by a previous token definition
+	bool bFirstChar = true;
+	TOptional<FStringToken> Token = Consumer.GetStream().ParseToken([&](TCHAR C){
+		if (C == '{' && !bFirstChar)
+		{
+			return EParseState::StopBefore;
+		}
+		else if (C == EscapeChar)
+		{
+			return EParseState::StopBefore;
+		}
+		else
+		{
+			bFirstChar = false;
+			// Keep consuming
+			return EParseState::Continue;
+		}
+	});
+
+	if (Token.IsSet())
+	{
+		// Add the token to the consumer. This moves the read position in the stream to the end of the token.
+		Consumer.Add(Token.GetValue(), FStringLiteral(Token.GetValue()));
+	}
+	return TOptional<FExpressionError>();
+}
+
+
+FStringFormatter::FStringFormatter()
+{
+	using namespace ExpressionParser;
+
+	// Token definition logic for named tokens
+	NamedDefinitions.DefineToken([](FExpressionTokenConsumer& Consumer)			{ return ParseSpecifier(Consumer, false); });
+	NamedDefinitions.DefineToken([](FExpressionTokenConsumer& Consumer)			{ return ParseEscapedChar(Consumer, false); });
+	NamedDefinitions.DefineToken([](FExpressionTokenConsumer& Consumer)			{ return ParseLiteral(Consumer, false); });
+
+	// Token definition logic for strict named tokens - will emit errors for any syntax errors
+	StrictNamedDefinitions.DefineToken([](FExpressionTokenConsumer& Consumer)	{ return ParseSpecifier(Consumer, true); });
+	StrictNamedDefinitions.DefineToken([](FExpressionTokenConsumer& Consumer)	{ return ParseEscapedChar(Consumer, true); });
+	StrictNamedDefinitions.DefineToken([](FExpressionTokenConsumer& Consumer)	{ return ParseLiteral(Consumer, true); });
+
+	// Token definition logic for ordered tokens
+	OrderedDefinitions.DefineToken([](FExpressionTokenConsumer& Consumer)		{ return ParseIndex(Consumer, false); });
+	OrderedDefinitions.DefineToken([](FExpressionTokenConsumer& Consumer)		{ return ParseEscapedChar(Consumer, false); });
+	OrderedDefinitions.DefineToken([](FExpressionTokenConsumer& Consumer)		{ return ParseLiteral(Consumer, false); });
+
+	// Token definition logic for strict ordered tokens - will emit errors for any syntax errors
+	StrictOrderedDefinitions.DefineToken([](FExpressionTokenConsumer& Consumer)	{ return ParseIndex(Consumer, true); });
+	OrderedDefinitions.DefineToken([](FExpressionTokenConsumer& Consumer)		{ return ParseEscapedChar(Consumer, true); });
+	StrictOrderedDefinitions.DefineToken([](FExpressionTokenConsumer& Consumer)	{ return ParseLiteral(Consumer, true); });
+}
+
+TValueOrError<FString, FExpressionError> FStringFormatter::FormatInternal(const TCHAR* InExpression, const TMap<FString, FStringFormatArg>& Args, bool bStrict) const
+{
+	TValueOrError<TArray<FExpressionToken>, FExpressionError> Result = ExpressionParser::Lex(InExpression, bStrict ? StrictNamedDefinitions : NamedDefinitions);
+	if (!Result.IsValid())
+	{
+		return MakeError(Result.StealError());
+	}
+
+	TArray<FExpressionToken>& Tokens = Result.GetValue();
+	if (Tokens.Num() == 0)
+	{
+		return MakeValue(InExpression);
+	}
+	
+	// This code deliberately tries to reallocate as little as possible
+	FString Formatted;
+	Formatted.Reserve(Tokens.Last().Context.GetTokenEndPos() - InExpression);
+	for (auto& Token : Tokens)
+	{
+		if (const auto* Literal = Token.Node.Cast<FStringLiteral>())
+		{
+			Formatted.AppendChars(Literal->String.GetTokenStartPos(), Literal->Len);
+		}
+		else if (auto* Escaped = Token.Node.Cast<FEscapedCharacter>())
+		{
+			Formatted.AppendChar(Escaped->Character);
+		}
+		else if (const auto* FormatToken = Token.Node.Cast<FFormatSpecifier>())
+		{
+			const FStringFormatArg* Arg = nullptr;
+			for (auto& Pair : Args)
+			{
+				if (Pair.Key.Len() == FormatToken->Len && FCString::Strnicmp(FormatToken->Identifier.GetTokenStartPos(), *Pair.Key, FormatToken->Len) == 0)
+				{
+					Arg = &Pair.Value;
+					break;
+				}
+			}
+
+			if (Arg)
+			{
+				AppendToString(*Arg, Formatted);
+			}
+			else if (bStrict)
+			{
+				return MakeError(FText::Format(LOCTEXT("UndefinedFormatSpecifier", "Undefined format token: {0}"), FText::FromString(FormatToken->Identifier.GetString())));
+			}
+			else
+			{
+				// No replacement found, so just add the original token string
+				const int32 Length = FormatToken->EntireToken.GetTokenEndPos() - FormatToken->EntireToken.GetTokenStartPos();
+				Formatted.AppendChars(FormatToken->EntireToken.GetTokenStartPos(), Length);
+			}
+		}
+	}
+
+	return MakeValue(Formatted);
+}
+
+TValueOrError<FString, FExpressionError> FStringFormatter::FormatInternal(const TCHAR* InExpression, const TArray<FStringFormatArg>& Args, bool bStrict) const
+{
+	TValueOrError<TArray<FExpressionToken>, FExpressionError> Result = ExpressionParser::Lex(InExpression, bStrict ? StrictOrderedDefinitions : OrderedDefinitions);
+	if (!Result.IsValid())
+	{
+		return MakeError(Result.StealError());
+	}
+
+	TArray<FExpressionToken>& Tokens = Result.GetValue();
+	if (Tokens.Num() == 0)
+	{
+		return MakeValue(InExpression);
+	}
+	
+	// This code deliberately tries to reallocate as little as possible
+	FString Formatted;
+	Formatted.Reserve(Tokens.Last().Context.GetTokenEndPos() - InExpression);
+	for (auto& Token : Tokens)
+	{
+		if (const auto* Literal = Token.Node.Cast<FStringLiteral>())
+		{
+			Formatted.AppendChars(Literal->String.GetTokenStartPos(), Literal->Len);
+		}
+		else if (auto* Escaped = Token.Node.Cast<FEscapedCharacter>())
+		{
+			Formatted.AppendChar(Escaped->Character);
+		}
+		else if (const auto* IndexToken = Token.Node.Cast<FIndexSpecifier>())
+		{
+			if (Args.IsValidIndex(IndexToken->Index))
+			{
+				AppendToString(Args[IndexToken->Index], Formatted);
+			}
+			else if (bStrict)
+			{
+				return MakeError(FText::Format(LOCTEXT("InvalidArgumentIndex", "Invalid argument index: {0}"), FText::AsNumber(IndexToken->Index)));
+			}
+			else
+			{
+				// No replacement found, so just add the original token string
+				const int32 Length = IndexToken->EntireToken.GetTokenEndPos() - IndexToken->EntireToken.GetTokenStartPos();
+				Formatted.AppendChars(IndexToken->EntireToken.GetTokenStartPos(), Length);
+			}
+		}
+	}
+
+	return MakeValue(Formatted);
+}
+
+/** Default formatter for string formatting - thread safe since all formatting is const */
+FStringFormatter DefaultFormatter;
+
+FString FString::Format(const TCHAR* InFormatString, const TMap<FString, FStringFormatArg>& InNamedArguments)
+{
+	return DefaultFormatter.Format(InFormatString, InNamedArguments);
+}
+
+FString FString::Format(const TCHAR* InFormatString, const TArray<FStringFormatArg>& InOrderedArguments)
+{
+	return DefaultFormatter.Format(InFormatString, InOrderedArguments);
+}
+
+// Don't include this code in shipping or perf test builds
+#if !(UE_BUILD_SHIPPING || UE_BUILD_TEST)
+
+bool TestStrings(FAutomationTestBase& Test, const TCHAR* Format, const TMap<FString, FStringFormatArg>& Arguments, const TCHAR* ExpectedResult)
+{
+	FString Result = FString::Format(Format, Arguments);
+	if (FCString::Strcmp(*Result, ExpectedResult) != 0)
+	{
+		Test.AddError(FText::Format(LOCTEXT("TestError", "FString::Format failed. \"{0}\" != \"{1}\""), FText::FromString(Result), FText::FromString(ExpectedResult)).ToString());
+		return false;
+	}
+	return true;
+}
+
+bool TestStrings(FAutomationTestBase& Test, const TCHAR* Format, const TArray<FStringFormatArg>& Arguments, const TCHAR* ExpectedResult)
+{
+	FString Result = FString::Format(Format, Arguments);
+	if (FCString::Strcmp(*Result, ExpectedResult) != 0)
+	{
+		Test.AddError(FText::Format(LOCTEXT("TestError", "FString::Format failed. \"{0}\" != \"{1}\""), FText::FromString(Result), FText::FromString(ExpectedResult)).ToString());
+		return false;
+	}
+	return true;
+}
+
+IMPLEMENT_SIMPLE_AUTOMATION_TEST(FStringFormattingTestNamedSimple, "System.Core.String Formatting.Simple (Named)", EAutomationTestFlags::EditorContext | EAutomationTestFlags::SmokeFilter)
+bool FStringFormattingTestNamedSimple::RunTest( const FString& Parameters )
+{
+	TMap<FString, FStringFormatArg> Args;
+	Args.Add(TEXT("Argument1"), TEXT("Replacement 1"));
+	Args.Add(TEXT("Argument2"), TEXT("Replacement 2"));
+
+	const TCHAR* Pattern =	TEXT("This is some text containing two arguments, { Argument1 } and {Argument2}.");
+	const TCHAR* Result =	TEXT("This is some text containing two arguments, Replacement 1 and Replacement 2.");
+
+	return TestStrings(*this, Pattern, Args, Result);
+}
+
+IMPLEMENT_SIMPLE_AUTOMATION_TEST(FStringFormattingTestNamedMultiple, "System.Core.String Formatting.Multiple (Named)", EAutomationTestFlags::EditorContext | EAutomationTestFlags::SmokeFilter)
+bool FStringFormattingTestNamedMultiple::RunTest( const FString& Parameters )
+{
+	TMap<FString, FStringFormatArg> Args;
+	Args.Add(TEXT("Argument1"), TEXT("Replacement 1"));
+
+	const TCHAR* Pattern =	TEXT("This is some text containing the same argument, { Argument1 } and {Argument1}.");
+	const TCHAR* Result =	TEXT("This is some text containing the same argument, Replacement 1 and Replacement 1.");
+
+	return TestStrings(*this, Pattern, Args, Result);
+}
+
+IMPLEMENT_SIMPLE_AUTOMATION_TEST(FStringFormattingTestNamedEscaped, "System.Core.String Formatting.Escaped (Named)", EAutomationTestFlags::EditorContext | EAutomationTestFlags::SmokeFilter)
+bool FStringFormattingTestNamedEscaped::RunTest( const FString& Parameters )
+{
+	TMap<FString, FStringFormatArg> Args;
+	Args.Add(TEXT("Argument1"), TEXT("Replacement 1"));
+
+	const TCHAR* Pattern =	TEXT("This is some text containing brace and an arg, `{ Argument1 } and {Argument1}.");
+	const TCHAR* Result =	TEXT("This is some text containing brace and an arg, { Argument1 } and Replacement 1.");
+
+	return TestStrings(*this, Pattern, Args, Result);
+}
+
+IMPLEMENT_SIMPLE_AUTOMATION_TEST(FStringFormattingTestNamedUnbounded, "System.Core.String Formatting.Unbounded (Named)", EAutomationTestFlags::EditorContext | EAutomationTestFlags::SmokeFilter)
+bool FStringFormattingTestNamedUnbounded::RunTest( const FString& Parameters )
+{
+	TMap<FString, FStringFormatArg> Args;
+	Args.Add(TEXT("Argument1"), TEXT("Replacement 1"));
+
+	const TCHAR* Pattern =	TEXT("This is some text containing an unbounded arg, { Argument1");
+	const TCHAR* Result =	TEXT("This is some text containing an unbounded arg, { Argument1");
+
+	return TestStrings(*this, Pattern, Args, Result);
+}
+
+IMPLEMENT_SIMPLE_AUTOMATION_TEST(FStringFormattingTestNamedNonExistent, "System.Core.String Formatting.Non Existent (Named)", EAutomationTestFlags::EditorContext | EAutomationTestFlags::SmokeFilter)
+bool FStringFormattingTestNamedNonExistent::RunTest( const FString& Parameters )
+{
+	TMap<FString, FStringFormatArg> Args;
+	Args.Add(TEXT("Argument1"), TEXT("Replacement 1"));
+
+	const TCHAR* Pattern =	TEXT("This is some text containing a non-existent arg { Argument2 }");
+	const TCHAR* Result =	TEXT("This is some text containing a non-existent arg { Argument2 }");
+
+	return TestStrings(*this, Pattern, Args, Result);
+}
+
+IMPLEMENT_SIMPLE_AUTOMATION_TEST(FStringFormattingTestNamedInvalid, "System.Core.String Formatting.Invalid (Named)", EAutomationTestFlags::EditorContext | EAutomationTestFlags::SmokeFilter)
+bool FStringFormattingTestNamedInvalid::RunTest( const FString& Parameters )
+{
+	TMap<FString, FStringFormatArg> Args;
+	Args.Add(TEXT("Argument1"), TEXT("Replacement 1"));
+
+	const TCHAR* Pattern =	TEXT("This is some text containing an invalid arg {Argument1 1}  { a8f7690-23\\ {} }");
+	const TCHAR* Result =	TEXT("This is some text containing an invalid arg {Argument1 1}  { a8f7690-23\\ {} }");
+
+	return TestStrings(*this, Pattern, Args, Result);
+}
+
+IMPLEMENT_SIMPLE_AUTOMATION_TEST(FStringFormattingTestNamedEmpty, "System.Core.String Formatting.Empty (Named)", EAutomationTestFlags::EditorContext | EAutomationTestFlags::SmokeFilter)
+bool FStringFormattingTestNamedEmpty::RunTest( const FString& Parameters )
+{
+	TMap<FString, FStringFormatArg> Args;
+	Args.Add(TEXT("Argument1"), TEXT("Replacement 1"));
+	Args.Add(TEXT("Argument2"), TEXT("Replacement 2"));
+
+	const TCHAR* Pattern = TEXT("");
+	const TCHAR* Result = TEXT("");
+
+	return TestStrings(*this, Pattern, Args, Result);
+}
+
+IMPLEMENT_SIMPLE_AUTOMATION_TEST(FStringFormattingTestOrderedSimple, "System.Core.String Formatting.Simple (Ordered)", EAutomationTestFlags::EditorContext | EAutomationTestFlags::SmokeFilter)
+bool FStringFormattingTestOrderedSimple::RunTest( const FString& Parameters )
+{
+	TArray<FStringFormatArg> Args;
+	Args.Add(TEXT("Replacement 1"));
+	Args.Add(TEXT("Replacement 2"));
+
+	const TCHAR* Pattern =	TEXT("This is some text containing two arguments, { 0 } and {1}.");
+	const TCHAR* Result =	TEXT("This is some text containing two arguments, Replacement 1 and Replacement 2.");
+
+	return TestStrings(*this, Pattern, Args, Result);
+}
+
+IMPLEMENT_SIMPLE_AUTOMATION_TEST(FStringFormattingTestOrderedMultiple, "System.Core.String Formatting.Multiple (Ordered)", EAutomationTestFlags::EditorContext | EAutomationTestFlags::SmokeFilter)
+bool FStringFormattingTestOrderedMultiple::RunTest( const FString& Parameters )
+{
+	TArray<FStringFormatArg> Args;
+	Args.Add(TEXT("Replacement 1"));
+
+	const TCHAR* Pattern =	TEXT("This is some text containing the same argument, { 0 } and {0}.");
+	const TCHAR* Result =	TEXT("This is some text containing the same argument, Replacement 1 and Replacement 1.");
+
+	return TestStrings(*this, Pattern, Args, Result);
+}
+
+IMPLEMENT_SIMPLE_AUTOMATION_TEST(FStringFormattingTestOrderedEscaped, "System.Core.String Formatting.Escaped (Ordered)", EAutomationTestFlags::EditorContext | EAutomationTestFlags::SmokeFilter)
+bool FStringFormattingTestOrderedEscaped::RunTest( const FString& Parameters )
+{
+	TArray<FStringFormatArg> Args;
+	Args.Add(TEXT("Replacement 1"));
+
+	const TCHAR* Pattern =	TEXT("This is some text containing brace and an arg, `{ 0 } and {0}.");
+	const TCHAR* Result =	TEXT("This is some text containing brace and an arg, { 0 } and Replacement 1.");
+
+	return TestStrings(*this, Pattern, Args, Result);
+}
+
+IMPLEMENT_SIMPLE_AUTOMATION_TEST(FStringFormattingTestOrderedUnbounded, "System.Core.String Formatting.Unbounded (Ordered)", EAutomationTestFlags::EditorContext | EAutomationTestFlags::SmokeFilter)
+bool FStringFormattingTestOrderedUnbounded::RunTest( const FString& Parameters )
+{
+	TArray<FStringFormatArg> Args;
+	Args.Add(TEXT("Replacement 1"));
+
+	const TCHAR* Pattern =	TEXT("This is some text containing an unbounded arg, { 0");
+	const TCHAR* Result =	TEXT("This is some text containing an unbounded arg, { 0");
+
+	return TestStrings(*this, Pattern, Args, Result);
+}
+
+IMPLEMENT_SIMPLE_AUTOMATION_TEST(FStringFormattingTestOrderedNonExistent, "System.Core.String Formatting.Non Existent (Ordered)", EAutomationTestFlags::EditorContext | EAutomationTestFlags::SmokeFilter)
+bool FStringFormattingTestOrderedNonExistent::RunTest( const FString& Parameters )
+{
+	TArray<FStringFormatArg> Args;
+	Args.Add(TEXT("Replacement 1"));
+
+	const TCHAR* Pattern =	TEXT("This is some text containing a non-existent arg { 5 }");
+	const TCHAR* Result =	TEXT("This is some text containing a non-existent arg { 5 }");
+
+	return TestStrings(*this, Pattern, Args, Result);
+}
+
+IMPLEMENT_SIMPLE_AUTOMATION_TEST(FStringFormattingTestOrderedAllTypes, "System.Core.String Formatting.All Types (Ordered)", EAutomationTestFlags::EditorContext | EAutomationTestFlags::SmokeFilter)
+bool FStringFormattingTestOrderedAllTypes::RunTest( const FString& Parameters )
+{
+	const int32 Value1 = 100;
+	const uint32 Value2 = 200;
+	const int64 Value3 = 300;
+	const uint64 Value4 = 400;
+	const float Value5 = 500.0;
+	const double Value6 = 600.0;
+	const FString Value7 = TEXT("Text");
+	const TCHAR* Value8 = TEXT("Text");
+
+	TArray<FStringFormatArg> Args;
+	Args.Add(Value1);
+	Args.Add(Value2);
+	Args.Add(Value3);
+	Args.Add(Value4);
+	Args.Add(Value5);
+	Args.Add(Value6);
+	Args.Add(Value7);
+	Args.Add(Value8);
+
+	const TCHAR* Pattern =	TEXT("{0}, {1}, {2}, {3}, {4}, {5}, {6}, {7}");
+	const TCHAR* Result =	TEXT("100, 200, 300, 400, 500.000000, 600.000000, Text, Text");
+
+	return TestStrings(*this, Pattern, Args, Result);
+}
+
+#endif //!(UE_BUILD_SHIPPING || UE_BUILD_TEST)
+
+#undef LOCTEXT_NAMESPACE
\ No newline at end of file
diff --git a/Engine/Source/Runtime/Core/Public/Containers/UnrealString.h b/Engine/Source/Runtime/Core/Public/Containers/UnrealString.h
--- a/Engine/Source/Runtime/Core/Public/Containers/UnrealString.h
+++ b/Engine/Source/Runtime/Core/Public/Containers/UnrealString.h
@@ -10,6 +10,7 @@
 #include "Math/UnrealMathUtility.h"
 #include "Misc/Crc.h"
 #include "Misc/CString.h"
+#include "Misc/StringFormatArg.h"
 #include "Templates/MemoryOps.h"
 
 
@@ -1314,6 +1315,22 @@ public:
 	 */
 	VARARG_DECL( static FString, static FString, return, Printf, VARARG_NONE, const TCHAR*, VARARG_NONE, VARARG_NONE );
 
+	/**
+	 * Format the specified string using the specified arguments. Replaces instances of { Argument } with keys in the map matching 'Argument'
+	 * @param InFormatString		A string representing the format expression
+	 * @param InNamedArguments		A map of named arguments that match the tokens specified in InExpression
+	 * @return A string containing the formatted text
+	 */
+	static FString Format(const TCHAR* InFormatString, const TMap<FString, FStringFormatArg>& InNamedArguments);
+
+	/**
+	 * Format the specified string using the specified arguments. Replaces instances of {0} with indices from the given array matching the index specified in the token
+	 * @param InFormatString		A string representing the format expression
+	 * @param InOrderedArguments	An array of ordered arguments that match the tokens specified in InExpression
+	 * @return A string containing the formatted text
+	 */
+	static FString Format(const TCHAR* InFormatString, const TArray<FStringFormatArg>& InOrderedArguments);
+
 	// @return string with Ch character
 	static FString Chr( TCHAR Ch );
 
diff --git a/Engine/Source/Runtime/Core/Public/Misc/StringFormatArg.h b/Engine/Source/Runtime/Core/Public/Misc/StringFormatArg.h
new file mode 100644
--- /dev/null
+++ b/Engine/Source/Runtime/Core/Public/Misc/StringFormatArg.h
@@ -0,0 +1,44 @@
+// Copyright 1998-2015 Epic Games, Inc. All Rights Reserved.
+
+#pragma once
+
+/** An argument supplied to FString::Format */
+struct CORE_API FStringFormatArg
+{
+	enum EType { Int, UInt, Double, String, StringLiteral };
+
+	/** The type of this arg */
+	EType Type;
+
+	/* todo: convert this to a TVariant */
+	union
+	{
+		/** Value as integer */
+		int64 IntValue;
+		/** Value as uint */
+		uint64 UIntValue;
+		/** Value as double */
+		double DoubleValue;
+		/** Value as a pointer to an FString */
+		const FString* StringValue;
+		/** Value as a string literal */
+		const TCHAR* StringLiteralValue;
+	};
+
+	FStringFormatArg( const int32 Value );
+	FStringFormatArg( const uint32 Value );
+	FStringFormatArg( const int64 Value );
+	FStringFormatArg( const uint64 Value );
+	FStringFormatArg( const float Value );
+	FStringFormatArg( const double Value );
+	FStringFormatArg( const FString& Value );
+	FStringFormatArg( const TCHAR* Value );
+
+	/** Copyable */
+	FStringFormatArg( const FStringFormatArg& RHS );
+	
+private:
+
+	/** Not default constructible */
+	FStringFormatArg();
+};
\ No newline at end of file
diff --git a/Engine/Source/Runtime/Core/Public/Misc/StringFormatter.h b/Engine/Source/Runtime/Core/Public/Misc/StringFormatter.h
new file mode 100644
--- /dev/null
+++ b/Engine/Source/Runtime/Core/Public/Misc/StringFormatter.h
@@ -0,0 +1,81 @@
+// Copyright 1998-2015 Epic Games, Inc. All Rights Reserved.
+
+#pragma once
+
+/** A string formatter is responsible for formatting string patterns using a set of named, or ordered arguments */
+class CORE_API FStringFormatter
+{
+public:
+
+	FStringFormatter();
+
+	/**
+	 * Format the specified string using the specified arguments. Replaces instances of { Argument } with keys in the map matching 'Argument'
+	 * @param InExpression		A string representing the format expression
+	 * @param InArgs			A map of named arguments that match the tokens specified in InExpression
+	 * @return A string containing the formatted text
+	 */
+	FString Format(const TCHAR* InExpression, const TMap<FString, FStringFormatArg>& InArgs) const
+	{
+		auto Result = FormatInternal(InExpression, InArgs, false);
+		if (ensure(Result.IsValid()))
+		{
+			return MoveTemp(Result.GetValue());
+		}
+
+		return InExpression;
+	}
+
+	/**
+	 * Format the specified string using the specified arguments. Replaces instances of {0} with indices from the given array matching the index specified in the token
+	 * @param InExpression		A string representing the format expression
+	 * @param InArgs			An array of ordered arguments that match the tokens specified in InExpression
+	 * @return A string containing the formatted text
+	 */
+	FString Format(const TCHAR* InExpression, const TArray<FStringFormatArg>& InArgs) const
+	{
+		auto Result = FormatInternal(InExpression, InArgs, false);
+		if (ensure(Result.IsValid()))
+		{
+			return MoveTemp(Result.GetValue());
+		}
+		
+		return InExpression;
+	}
+
+	/**
+	 * Format the specified string using the specified arguments. Replaces instances of { Argument } with keys in the map matching 'Argument'
+	 * @param InExpression		A string representing the format expression
+	 * @param InArgs			A map of named arguments that match the tokens specified in InExpression
+	 * @return A string containing the formatted text, or an error where InExpression is ill-formed, or contains undefined arguments
+	*/
+	TValueOrError<FString, FExpressionError> FormatStrict(const TCHAR* InExpression, const TMap<FString, FStringFormatArg>& InArgs) const
+	{
+		return FormatInternal(InExpression, InArgs, true);
+	}
+
+	/**
+	 * Format the specified string using the specified arguments. Replaces instances of {0} with indices from the given array matching the index specified in the token
+	 * @param InExpression		A string representing the format expression
+	 * @param InArgs			An array of ordered arguments that match the tokens specified in InExpression
+	 * @return A string containing the formatted text, or an error where InExpression is ill-formed, or contains undefined arguments
+	 */
+	TValueOrError<FString, FExpressionError> FormatStrict(const TCHAR* InExpression, const TArray<FStringFormatArg>& InArgs) const
+	{
+		return FormatInternal(InExpression, InArgs, true);
+	}
+
+private:
+
+	/** Internal formatting logic */
+	TValueOrError<FString, FExpressionError> FormatInternal(const TCHAR* InExpression, const TMap<FString, FStringFormatArg>& InArgs, bool bStrict) const;
+	TValueOrError<FString, FExpressionError> FormatInternal(const TCHAR* InExpression, const TArray<FStringFormatArg>& InArgs, bool bStrict) const;
+
+	/** Token definitions for lenient lexers */
+	FTokenDefinitions NamedDefinitions;
+	FTokenDefinitions OrderedDefinitions;
+
+	/** Token definitions for strict lexers */
+	FTokenDefinitions StrictNamedDefinitions;
+	FTokenDefinitions StrictOrderedDefinitions;
+};
\ No newline at end of file
