UE4 - Lots of changes to the task graph include new experimental scheduler, current used on PS4, but expected to be used on all platforms after it shakes out.

diff --git a/Engine/Plugins/Experimental/CodeEditor/Source/CodeEditor/Private/DirectoryScanner.cpp b/Engine/Plugins/Experimental/CodeEditor/Source/CodeEditor/Private/DirectoryScanner.cpp
--- a/Engine/Plugins/Experimental/CodeEditor/Source/CodeEditor/Private/DirectoryScanner.cpp
+++ b/Engine/Plugins/Experimental/CodeEditor/Source/CodeEditor/Private/DirectoryScanner.cpp
@@ -39,7 +39,7 @@ struct FDirectoryScannerCommand : public IQueuedWork
 		class FDirectoryEnumerator : public IPlatformFile::FDirectoryVisitor
 		{
 		public:
-			FDirectoryEnumerator(TLockFreePointerList<FDirectoryResult>& InFoundFiles)
+			FDirectoryEnumerator(TLockFreePointerListUnordered<FDirectoryResult>& InFoundFiles)
 				: FoundFiles(InFoundFiles)
 			{
 			}
@@ -58,7 +58,7 @@ struct FDirectoryScannerCommand : public IQueuedWork
 				return true;
 			}
 
-			TLockFreePointerList<FDirectoryResult>& FoundFiles;
+			TLockFreePointerListUnordered<FDirectoryResult>& FoundFiles;
 		};
 
 		FDirectoryEnumerator DirectoryEnumerator(FoundFiles);
@@ -73,7 +73,7 @@ struct FDirectoryScannerCommand : public IQueuedWork
 
 	FOnDirectoryScanned OnDirectoryScanned;
 
-	TLockFreePointerList<FDirectoryResult> FoundFiles;
+	TLockFreePointerListUnordered<FDirectoryResult> FoundFiles;
 
 	volatile int32 bExecuted;
 };
diff --git a/Engine/Source/Developer/Profiler/Private/Widgets/SProfilerMiniView.h b/Engine/Source/Developer/Profiler/Private/Widgets/SProfilerMiniView.h
--- a/Engine/Source/Developer/Profiler/Private/Widgets/SProfilerMiniView.h
+++ b/Engine/Source/Developer/Profiler/Private/Widgets/SProfilerMiniView.h
@@ -344,7 +344,6 @@ protected:
 
 	/** Recently added thread times, not processed yet. */
 	TArray<FFrameThreadTimes> RecentlyAddedFrames;
-	TLockFreePointerList<FFrameThreadTimes> RecentlyAddedFramesLockFree;
 
 	/** Maximum total mini-view frame time seen so far. */
 	float MaxFrameTime;
diff --git a/Engine/Source/Programs/UnrealLightmass/Private/Lighting/LightingSystem.h b/Engine/Source/Programs/UnrealLightmass/Private/Lighting/LightingSystem.h
--- a/Engine/Source/Programs/UnrealLightmass/Private/Lighting/LightingSystem.h
+++ b/Engine/Source/Programs/UnrealLightmass/Private/Lighting/LightingSystem.h
@@ -2589,10 +2589,12 @@ private:
 	volatile int32 MappingTasksInProgressThatWillNeedHelp;
 
 	/** List of tasks to cache indirect lighting, used by all mapping threads. */
-	TLockFreePointerList<FCacheIndirectTaskDescription> CacheIndirectLightingTasks;
+	// consider changing this from FIFO to Unordered, which may be faster
+	TLockFreePointerListLIFO<FCacheIndirectTaskDescription> CacheIndirectLightingTasks;
 
 	/** List of tasks to interpolate indirect lighting, used by all mapping threads. */
-	TLockFreePointerList<FInterpolateIndirectTaskDescription> InterpolateIndirectLightingTasks;
+	// consider changing this from FIFO to Unordered, which may be faster
+	TLockFreePointerListLIFO<FInterpolateIndirectTaskDescription> InterpolateIndirectLightingTasks;
 
 	TArray<FVolumeSamplesTaskDescription> VolumeSampleTasks;
 
diff --git a/Engine/Source/Runtime/Core/Private/Async/TaskGraph.cpp b/Engine/Source/Runtime/Core/Private/Async/TaskGraph.cpp
--- a/Engine/Source/Runtime/Core/Private/Async/TaskGraph.cpp
+++ b/Engine/Source/Runtime/Core/Private/Async/TaskGraph.cpp
@@ -91,7 +91,10 @@ namespace ENamedThreads
  *	Because of the multithreaded nature of this system an ordinary singleton cannot be used.
  *	FTaskGraphImplementation::Startup() creates the singleton and the constructor actually sets this value.
 **/
-static class FTaskGraphImplementation* TaskGraphImplementationSingleton = NULL;
+class FTaskGraphImplementation;
+struct FWorkerThread;
+
+static FTaskGraphImplementation* TaskGraphImplementationSingleton = NULL;
 
 #if !UE_BUILD_SHIPPING && !UE_BUILD_TEST
 
@@ -124,21 +127,21 @@ static struct FChaosMode
 		// ave = 1/8
 		for (int32 Index = 0; Index < NumSamples; Index++)
 		{
-			DelayTimes[Index] *= 0.0001f;
+			DelayTimes[Index] *= 0.00001f;
 		}
-		// ave = 0.0000125s
-		for (int32 Zeros = 0; Zeros < NumSamples / 2; Zeros++)
+		// ave = 0.00000125s
+		for (int32 Zeros = 0; Zeros < NumSamples / 20; Zeros++)
 		{
 			int32 Index = Stream.RandHelper(NumSamples);
 			DelayTimes[Index] = 0.0f;
 		}
-		// half the samples are now zero
+		// 95% the samples are now zero
 		for (int32 Zeros = 0; Zeros < NumSamples / 100; Zeros++)
 		{
 			int32 Index = Stream.RandHelper(NumSamples);
-			DelayTimes[Index] = .005f;
+			DelayTimes[Index] = .00005f;
 		}
-		// 1% of the samples are 5ms
+		// .001% of the samples are 5ms
 	}
 	FORCEINLINE void Delay()
 	{
@@ -195,6 +198,49 @@ static FAutoConsoleVariableRef CVarNumWorkerThreadsToIgnore(
 	ECVF_Cheat
 	);
 
+static int32 GFastScheduler = 0;
+#if USE_NEW_LOCK_FREE_LISTS
+static int32 GFastSchedulerLatched = 1;
+#else
+static int32 GFastSchedulerLatched = 0;
+#endif
+static FAutoConsoleVariableRef CVarFastScheduler(
+	TEXT("TaskGraph.FastScheduler"),
+	GFastScheduler,
+	TEXT("If > 0, then use the new experimental anythread task scheduler."),
+	ECVF_Cheat
+	);
+
+// move to platform abstraction
+#if PLATFORM_XBOXONE || PLATFORM_PS4
+	#define PLATFORM_OK_TO_BURN_CPU (1)
+#else
+	#define PLATFORM_OK_TO_BURN_CPU (0)
+#endif
+
+#define USE_INTRUSIVE_TASKQUEUES (USE_NEW_LOCK_FREE_LISTS)
+
+static int32 GConsoleSpinMode = PLATFORM_OK_TO_BURN_CPU * 2;
+static int32 GConsoleSpinModeLatched = 0;
+static FAutoConsoleVariableRef CVarConsoleSpinMode(
+	TEXT("TaskGraph.ConsoleSpinMode"),
+	GConsoleSpinMode,
+	TEXT("If > 0, then we never allow all task threads to idle; one will always be spinning looking for work (if mode == 2, then we sleep0 when spinning); this relieve the calling thread from the buden of starting threads. Only active with TaskGraph.FastScheduler 1."),
+	ECVF_Cheat
+	);
+
+#if USE_NEW_LOCK_FREE_LISTS
+static int32 GMaxTasksToStartOnDequeue = 1;
+#else
+static int32 GMaxTasksToStartOnDequeue = 2;
+#endif
+static FAutoConsoleVariableRef CVarMaxTasksToStartOnDequeue(
+	TEXT("TaskGraph.MaxTasksToStartOnDequeue"),
+	GMaxTasksToStartOnDequeue,
+	TEXT("Performance tweak, controls how many additional task threads a task thread starts when it grabs a list of new tasks. This only applies in TaskGraph.ConsoleSpinMode"),
+	ECVF_Cheat
+	);
+
 /** 
  *	FTaskQueue
  *	High performance, SINGLE threaded, FIFO task queue for the private queue on named threads.
@@ -209,7 +255,7 @@ public:
 	{
 	}
 	/** Returns the number of non-NULL items in the queue. **/
-	int32 Num() const
+	FORCEINLINE int32 Num() const
 	{
 		CheckInvariants();
 		return EndIndex - StartIndex;
@@ -218,7 +264,7 @@ public:
 	 *	Adds a task to the queue.
 	 *	@param Task; the task to add to the queue
 	**/
-	void Enqueue(FBaseGraphTask* Task)
+	FORCEINLINE void Enqueue(FBaseGraphTask* Task)
 	{
 		CheckInvariants();
 		if (EndIndex >= Tasks.Num())
@@ -244,7 +290,7 @@ public:
 	 *	Pops a task off the queue.
 	 *	@return The oldest task in the queue or NULL if the queue is empty
 	**/
-	FBaseGraphTask* Dequeue()
+	FORCEINLINE FBaseGraphTask* Dequeue()
 	{
 		if (!Num())
 		{
@@ -286,29 +332,21 @@ private:
 
 };
 
-// this is used to signify a task that is just a call to wake up
-// It is generally bad to reuse pointers for non-pointer data, but efficiency is important here
-static FBaseGraphTask* WakeUpBaseGraphTask = (FBaseGraphTask*)0x3;
-
 /** 
- *	FTaskThread
- *	A class for managing a worker or named thread. 
+ *	FTaskThreadBase
+ *	Base class for a thread that executes tasks
  *	This class implements the FRunnable API, but external threads don't use that because those threads are created elsewhere.
 **/
-class FTaskThread : public FRunnable, FSingleThreadRunnable
+class FTaskThreadBase : public FRunnable, FSingleThreadRunnable
 {
 public:
-	/** Array of tasks for this task thread. */
-	TArray<FBaseGraphTask*> NewTasks;
-
 	// Calls meant to be called from a "main" or supervisor thread.
 
 	/** Constructor, initializes everything to unusable values. Meant to be called from a "main" thread. **/
-	FTaskThread()
+	FTaskThreadBase()
 		: ThreadId(ENamedThreads::AnyThread)
 		, PerThreadIDTLSSlot(0xffffffff)
-		, bAllowsStealsFromMe(false)
-		, bStealsFromOthers(false)
+		, OwnerWorker(nullptr)
 	{
 		NewTasks.Reset(128);
 	}
@@ -320,13 +358,12 @@ public:
 	 *	@param bInAllowsStealsFromMe; If true, this is a worker thread and any other thread can steal tasks from my incoming queue.
 	 *	@param bInStealsFromOthers If true, this is a worker thread and I will attempt to steal tasks when I run out of work.
 	**/
-	void Setup(ENamedThreads::Type InThreadId, uint32 InPerThreadIDTLSSlot, bool bInAllowsStealsFromMe, bool bInStealsFromOthers)
+	void Setup(ENamedThreads::Type InThreadId, uint32 InPerThreadIDTLSSlot, FWorkerThread* InOwnerWorker)
 	{
 		ThreadId = InThreadId;
 		check(ThreadId >= 0);
 		PerThreadIDTLSSlot = InPerThreadIDTLSSlot;
-		bAllowsStealsFromMe = bInAllowsStealsFromMe;
-		bStealsFromOthers = bInStealsFromOthers;
+		OwnerWorker = InOwnerWorker;
 	}
 
 	// Calls meant to be called from "this thread".
@@ -334,40 +371,293 @@ public:
 	/** A one-time call to set the TLS entry for this thread. **/
 	void InitializeForCurrentThread()
 	{
-		FPlatformTLS::SetTlsValue(PerThreadIDTLSSlot,this);
+		FPlatformTLS::SetTlsValue(PerThreadIDTLSSlot,OwnerWorker);
+	}
+
+	/** Return the index of this thread. **/
+	ENamedThreads::Type GetThreadId() const
+	{
+		checkThreadGraph(OwnerWorker); // make sure we are started up
+		return ThreadId;
 	}
 
 	/** Used for named threads to start processing tasks until the thread is idle and RequestQuit has been called. **/
-	void ProcessTasksUntilQuit(int32 QueueIndex)
+	virtual void ProcessTasksUntilQuit(int32 QueueIndex) = 0;
+
+	/** Used for named threads to start processing tasks until the thread is idle and RequestQuit has been called. **/
+	virtual void ProcessTasksUntilIdle(int32 QueueIndex)
+	{
+		check(0);
+	}
+
+	/** 
+	 *	Queue a task, assuming that this thread is the same as the current thread.
+	 *	For named threads, these go directly into the private queue.
+	 *	@param QueueIndex, Queue to enqueue for
+	 *	@param Task Task to queue.
+	 **/
+	virtual void EnqueueFromThisThread(int32 QueueIndex, FBaseGraphTask* Task)
+	{
+		check(0);
+	}
+
+	// Calls meant to be called from any thread.
+
+	/** 
+	 *	Will cause the thread to return to the caller when it becomes idle. Used to return from ProcessTasksUntilQuit for named threads or to shut down unnamed threads. 
+	 *	CAUTION: This will not work under arbitrary circumstances. For example you should not attempt to stop unnamed threads unless they are known to be idle.
+	 *	Return requests for named threads should be submitted from that named thread as FReturnGraphTask does.
+	 *	@param QueueIndex, Queue to request quit from
+	**/
+	virtual void RequestQuit(int32 QueueIndex) = 0;
+
+	/** 
+	 *	Queue a task, assuming that this thread is not the same as the current thread.
+	 *	@param QueueIndex, Queue to enqueue into
+	 *	@param Task; Task to queue.
+	 **/
+	virtual bool EnqueueFromOtherThread(int32 QueueIndex, FBaseGraphTask* Task)
+	{
+		check(0);
+		return false;
+	}
+
+	virtual void WakeUp()
+	{
+		check(0);
+	}
+	/** 
+	 *Return true if this thread is processing tasks. This is only a "guess" if you ask for a thread other than yourself because that can change before the function returns.
+	 *@param QueueIndex, Queue to request quit from
+	 **/
+	virtual bool IsProcessingTasks(int32 QueueIndex) = 0;
+
+	// SingleThreaded API
+
+	/** Tick single-threaded. */
+	virtual void Tick() override
+	{
+		ProcessTasksUntilIdle(0);
+	}
+
+
+	// FRunnable API
+
+	/**
+	 * Allows per runnable object initialization. NOTE: This is called in the
+	 * context of the thread object that aggregates this, not the thread that
+	 * passes this runnable to a new thread.
+	 *
+	 * @return True if initialization was successful, false otherwise
+	 */
+	virtual bool Init() override
+	{
+		InitializeForCurrentThread();
+		return true;
+	}
+
+	/**
+	 * This is where all per object thread work is done. This is only called
+	 * if the initialization was successful.
+	 *
+	 * @return The exit code of the runnable object
+	 */
+	virtual uint32 Run() override
+	{
+		ProcessTasksUntilQuit(0);
+		return 0;
+	}
+
+	/**
+	 * This is called if a thread is requested to terminate early
+	 */
+	virtual void Stop() override
+	{
+		RequestQuit(-1);
+	}
+
+	/**
+	 * Called in the context of the aggregating thread to perform any cleanup.
+	 */
+	virtual void Exit() override
+	{
+	}
+
+	/**
+	 * Return single threaded interface when multithreading is disabled.
+	 */
+	virtual FSingleThreadRunnable* GetSingleThreadInterface() override
+	{
+		return this;
+	}
+
+protected:
+
+	/** Id / Index of this thread. **/
+	ENamedThreads::Type									ThreadId;
+	/** TLS SLot that we store the FTaskThread* this pointer in. **/
+	uint32												PerThreadIDTLSSlot;
+	/** Used to signal stalling. Not safe for synchronization in most cases. **/
+	FThreadSafeCounter									IsStalled;
+	/** Array of tasks for this task thread. */
+	TArray<FBaseGraphTask*> NewTasks;
+	/** back pointer to the owning FWorkerThread **/
+	FWorkerThread* OwnerWorker;
+};
+
+/** 
+ *	FNamedTaskThread
+ *	A class for managing a named thread. 
+ */
+class FNamedTaskThread : public FTaskThreadBase
+{
+public:
+
+	virtual void ProcessTasksUntilQuit(int32 QueueIndex) override
 	{
 		Queue(QueueIndex).QuitWhenIdle.Reset();
+		verify(++Queue(QueueIndex).RecursionGuard == 1);
+		//checkThreadGraph(!Queue(QueueIndex).IncomingQueue.IsClosed()); // make sure we are started up
 		while (Queue(QueueIndex).QuitWhenIdle.GetValue() == 0)
 		{
-			ProcessTasks(QueueIndex, true);
+			ProcessTasksNamedThread(QueueIndex, true);
 			// @Hack - quit now when running with only one thread.
 			if (!FPlatformProcess::SupportsMultithreading())
 			{
 				break;
 			}
 		}
+		//checkThreadGraph(!Queue(QueueIndex).IncomingQueue.IsClosed()); // make sure we are started up
+		verify(!--Queue(QueueIndex).RecursionGuard);
 	}
 
-	/** Tick single-threaded. */
-	virtual void Tick() override
+	virtual void ProcessTasksUntilIdle(int32 QueueIndex) override
 	{
-		if (Queue(0).QuitWhenIdle.GetValue() == 0)
+		verify(++Queue(QueueIndex).RecursionGuard == 1);
+		//checkThreadGraph(!Queue(QueueIndex).IncomingQueue.IsClosed()); // make sure we are started up
+		ProcessTasksNamedThread(QueueIndex, false);
+		//checkThreadGraph(!Queue(QueueIndex).IncomingQueue.IsClosed()); // make sure we are started up
+		verify(!--Queue(QueueIndex).RecursionGuard);
+	}
+#if USE_NEW_LOCK_FREE_LISTS
+	FORCEINLINE FBaseGraphTask* GetNextTask(int32 QueueIndex)
+	{
+		FBaseGraphTask* Task = NULL;
+		Task = Queue(QueueIndex).PrivateQueueHiPri.Dequeue();
+		if (!Task)
+		{
+			if (Queue(QueueIndex).OutstandingHiPriTasks.GetValue())
+			{
+				while (true)
+				{
+					FBaseGraphTask* NewTask = Queue(QueueIndex).IncomingQueue.Pop();
+					if (!NewTask)
+					{
+						break;
+					}
+					if (ENamedThreads::GetPriority(NewTask->ThreadToExecuteOn))
+					{
+						Queue(QueueIndex).OutstandingHiPriTasks.Decrement();
+						Task = NewTask;
+						break;
+					}
+					else
+					{
+						Queue(QueueIndex).PrivateQueue.Enqueue(NewTask);
+					}
+				}
+			}
+			if (!Task)
+			{
+				Task = Queue(QueueIndex).PrivateQueue.Dequeue();
+			}
+		}
+		if (!Task)
 		{
-			ProcessTasks(0, true);
+			Task = Queue(QueueIndex).IncomingQueue.Pop();
 		}
+		return Task;
 	}
-
-	/** 
-	 *	Process tasks until idle. May block if bAllowStall is true
-	 *	@param QueueIndex, Queue to process tasks from
-	 *	@param bAllowStall,  if true, the thread will block on the stall event when it runs out of tasks.
-	 **/
-	void ProcessTasks(int32 QueueIndex, bool bAllowStall)
+	void ProcessTasksNamedThread(int32 QueueIndex, bool bAllowStall)
+	{
+		TStatId StallStatId;
+		bool bCountAsStall = false;
+#if STATS
+		TStatId StatName;
+		FCycleCounter ProcessingTasks;
+		if (ThreadId == ENamedThreads::GameThread)
+		{
+			StatName = GET_STATID(STAT_TaskGraph_GameTasks);
+			StallStatId = GET_STATID(STAT_TaskGraph_GameStalls);
+			bCountAsStall = true;
+		}
+		else if (ThreadId == ENamedThreads::RenderThread)
+		{
+			if (QueueIndex > 0)
+			{
+				StallStatId = GET_STATID(STAT_TaskGraph_RenderStalls);
+				bCountAsStall = true;
+			}
+			// else StatName = none, we need to let the scope empty so that the render thread submits tasks in a timely manner. 
+		}
+		else if (ThreadId != ENamedThreads::StatsThread)
+		{
+			StatName = GET_STATID(STAT_TaskGraph_OtherTasks);
+			StallStatId = GET_STATID(STAT_TaskGraph_OtherStalls);
+			bCountAsStall = true;
+		}
+		bool bTasksOpen = false;
+		if (FThreadStats::IsCollectingData(StatName))
+		{
+			bTasksOpen = true;
+			ProcessingTasks.Start(StatName);
+		}
+#endif
+		while (1)
+		{
+			FBaseGraphTask* Task = GetNextTask(QueueIndex);
+			if (!Task)
+			{
+#if STATS
+				if(bTasksOpen)
+				{
+					ProcessingTasks.Stop();
+					bTasksOpen = false;
+				}
+#endif
+				if (bAllowStall)
+				{
+					bool bNewTaskReady = Stall(QueueIndex, StallStatId, bCountAsStall);
+					if (!bNewTaskReady)
+					{
+						break; // we didn't really stall, rather we were asked to quit
+					}
+#if STATS
+					if (!bTasksOpen && FThreadStats::IsCollectingData(StatName))
+					{
+						bTasksOpen = true;
+						ProcessingTasks.Start(StatName);
+					}
+#endif
+					continue;
+				}
+				else
+				{
+					break; // we were asked to quit
+				}
+			}
+			else
+			{
+				TestRandomizedThreads();
+				Task->Execute(NewTasks, ENamedThreads::Type(ThreadId | (QueueIndex << ENamedThreads::QueueIndexShift)));
+				TestRandomizedThreads();
+			}
+		}
+	}
+#else
+	void ProcessTasksNamedThread(int32 QueueIndex, bool bAllowStall)
 	{
+		//checkThreadGraph(!Queue(QueueIndex).IncomingQueue.IsClosed()); // make sure we are started up
 		TStatId StallStatId;
 		bool bCountAsStall = false;
 #if STATS
@@ -396,15 +686,13 @@ public:
 		}
 		bool bTasksOpen = false;
 #endif
-		check(!Queue(QueueIndex).RecursionGuard.GetValue());
-		Queue(QueueIndex).RecursionGuard.Increment();
 		while (1)
 		{
 			FBaseGraphTask* Task = NULL;
 			Task = Queue(QueueIndex).PrivateQueueHiPri.Dequeue();
 			if (!Task)
 			{
-				if (!bAllowsStealsFromMe && Queue(QueueIndex).OutstandingHiPriTasks.GetValue())
+				if (Queue(QueueIndex).OutstandingHiPriTasks.GetValue())
 				{
 					Queue(QueueIndex).IncomingQueue.PopAll(NewTasks);
 					if (NewTasks.Num())
@@ -433,100 +721,45 @@ public:
 			}
 			if (!Task)
 			{
-				if (!bAllowsStealsFromMe)
+				Queue(QueueIndex).IncomingQueue.PopAll(NewTasks);
+				if (!NewTasks.Num())
 				{
-					// no steals so we will take all of the items, and this also ensures ordering
-					for (int32 Count = SPIN_COUNT + 1; Count && !NewTasks.Num() ; Count--)
-					{
-						Queue(QueueIndex).IncomingQueue.PopAll(NewTasks);
-					}
-					if (FPlatformProcess::SupportsMultithreading())
-					{
-						for (int32 Count = SLEEP_COUNT; Count && !NewTasks.Num() ; Count--)
-						{
-							FPlatformProcess::Sleep(0.0f);
-							Queue(QueueIndex).IncomingQueue.PopAll(NewTasks);
-						}
-					}
-					if (!NewTasks.Num())
+					if (bAllowStall)
 					{
-						if (bAllowStall)
-						{
 #if STATS
-							if(bTasksOpen)
-							{
-								ProcessingTasks.Stop();
-								bTasksOpen = false;
-							}
-#endif
-							if (Stall(QueueIndex, StallStatId, bCountAsStall))
-							{
-								Queue(QueueIndex).IncomingQueue.PopAll(NewTasks);
-							}
-						}
-					}
-					if (NewTasks.Num())
-					{
-						for (int32 Index = NewTasks.Num() - 1; Index >= 0 ; Index--) // reverse the order since PopAll is implicitly backwards
+						if(bTasksOpen)
 						{
-							FBaseGraphTask* NewTask = NewTasks[Index];
-							if (ENamedThreads::GetPriority(NewTask->ThreadToExecuteOn))
-							{
-								Queue(QueueIndex).PrivateQueueHiPri.Enqueue(NewTask);
-								Queue(QueueIndex).OutstandingHiPriTasks.Decrement();
-							}
-							else
-							{
-								Queue(QueueIndex).PrivateQueue.Enqueue(NewTask);
-							}
+							ProcessingTasks.Stop();
+							bTasksOpen = false;
 						}
-						NewTasks.Reset();
-						Task = Queue(QueueIndex).PrivateQueueHiPri.Dequeue();
-						if (!Task)
+#endif
+						if (Stall(QueueIndex, StallStatId, bCountAsStall))
 						{
-							Task = Queue(QueueIndex).PrivateQueue.Dequeue();
+							Queue(QueueIndex).IncomingQueue.PopAll(NewTasks);
 						}
 					}
 				}
-				else
+				if (NewTasks.Num())
 				{
-					// because of stealing, we are only going to take one item
-					for (int32 Count = SPIN_COUNT + 1; !Task && Count ; Count--)
+					for (int32 Index = NewTasks.Num() - 1; Index >= 0 ; Index--) // reverse the order since PopAll is implicitly backwards
 					{
-						Task = Queue(QueueIndex).IncomingQueue.PopIfNotClosed();
-						if (!Task)
+						FBaseGraphTask* NewTask = NewTasks[Index];
+						checkThreadGraph(NewTask);
+						if (ENamedThreads::GetPriority(NewTask->ThreadToExecuteOn))
 						{
-							Task = FindWork();
+							Queue(QueueIndex).PrivateQueueHiPri.Enqueue(NewTask);
+							Queue(QueueIndex).OutstandingHiPriTasks.Decrement();
 						}
-					}
-					if (FPlatformProcess::SupportsMultithreading())
-					{
-						for (int32 Count = SLEEP_COUNT; !Task && Count ; Count--)
+						else
 						{
-							FPlatformProcess::Sleep(0.0f);
-							Task = Queue(QueueIndex).IncomingQueue.PopIfNotClosed();
-							if (!Task)
-							{
-								Task = FindWork();
-							}
+							Queue(QueueIndex).PrivateQueue.Enqueue(NewTask);
 						}
 					}
+					NewTasks.Reset();
+					Task = Queue(QueueIndex).PrivateQueueHiPri.Dequeue();
 					if (!Task)
 					{
-						if (bAllowStall)
-						{
-#if STATS
-							if(bTasksOpen)
-							{
-								ProcessingTasks.Stop();
-								bTasksOpen = false;
-							}
-#endif
-							if (Stall(QueueIndex, StallStatId, bCountAsStall))
-							{
-								Task = Queue(QueueIndex).IncomingQueue.PopIfNotClosed();
-							}
-						}
+						Task = Queue(QueueIndex).PrivateQueue.Dequeue();
 					}
 				}
 			}
@@ -540,10 +773,8 @@ public:
 					ProcessingTasks.Start(StatName);
 				}
 #endif
-				if (Task != WakeUpBaseGraphTask)
-				{
-					Task->Execute(NewTasks, ENamedThreads::Type(ThreadId | (QueueIndex << ENamedThreads::QueueIndexShift)));
-				}
+				Task->Execute(NewTasks, ENamedThreads::Type(ThreadId | (QueueIndex << ENamedThreads::QueueIndexShift)));
+
 				TestRandomizedThreads();
 			}
 			else
@@ -558,84 +789,62 @@ public:
 			bTasksOpen = false;
 		}
 #endif
-		Queue(QueueIndex).RecursionGuard.Decrement();
-		check(!Queue(QueueIndex).RecursionGuard.GetValue());
 	}
-
-	/** 
-	 *	Queue a task, assuming that this thread is the same as the current thread.
-	 *	For named threads, these go directly into the private queue.
-	 *	@param QueueIndex, Queue to enqueue for
-	 *	@param Task Task to queue.
-	 **/
-	void EnqueueFromThisThread(int32 QueueIndex, FBaseGraphTask* Task)
+#endif
+	virtual void EnqueueFromThisThread(int32 QueueIndex, FBaseGraphTask* Task) override
 	{
-		checkThreadGraph(Queue(QueueIndex).StallRestartEvent); // make sure we are started up
-		if (bAllowsStealsFromMe)
+		checkThreadGraph(Task && Queue(QueueIndex).StallRestartEvent); // make sure we are started up
+		if (ENamedThreads::GetPriority(Task->ThreadToExecuteOn))
 		{
-			bool bWasReopenedByMe = Queue(QueueIndex).IncomingQueue.ReopenIfClosedAndPush(Task);
-			checkThreadGraph(!bWasReopenedByMe); // if I am stalled, why am I here?
+			Queue(QueueIndex).PrivateQueueHiPri.Enqueue(Task);
 		}
 		else
 		{
-			checkThreadGraph((FTaskThread*)FPlatformTLS::GetTlsValue(PerThreadIDTLSSlot) == this); // verify that we are the thread they say we are
-			if (Task != WakeUpBaseGraphTask && ENamedThreads::GetPriority(Task->ThreadToExecuteOn))
-			{
-				Queue(QueueIndex).PrivateQueueHiPri.Enqueue(Task);
-			}
-			else
-			{
-				Queue(QueueIndex).PrivateQueue.Enqueue(Task);
-			}
+			Queue(QueueIndex).PrivateQueue.Enqueue(Task);
 		}
 	}
 
-
-	// Calls meant to be called from any thread.
-
-	/** 
-	 *	Will cause the thread to return to the caller when it becomes idle. Used to return from ProcessTasksUntilQuit for named threads or to shut down unnamed threads. 
-	 *	CAUTION: This will not work under arbitrary circumstances. For example you should not attempt to stop unnamed threads unless they are known to be idle.
-	 *	Return requests for named threads should be submitted from that named thread as FReturnGraphTask does.
-	 *	@param QueueIndex, Queue to request quit from
-	**/
-	void RequestQuit(int32 QueueIndex)
+	virtual void RequestQuit(int32 QueueIndex) override
 	{
 		// this will not work under arbitrary circumstances. For example you should not attempt to stop threads unless they are known to be idle.
+		if (!Queue(0).StallRestartEvent)
+		{
+			return;
+		}
+#if USE_NEW_LOCK_FREE_LISTS
+		if (QueueIndex == -1)
+		{
+			// we are shutting down
+			checkThreadGraph(Queue(0).StallRestartEvent); // make sure we are started up
+			checkThreadGraph(Queue(1).StallRestartEvent); // make sure we are started up
+			Queue(0).QuitWhenIdle.Increment();
+			Queue(1).QuitWhenIdle.Increment();
+			Queue(0).StallRestartEvent->Trigger(); 
+			Queue(1).StallRestartEvent->Trigger(); 
+		}
+		else
+		{
+			checkThreadGraph(Queue(QueueIndex).StallRestartEvent); // make sure we are started up
+			Queue(QueueIndex).QuitWhenIdle.Increment();
+		}
+#else
+		if (QueueIndex == -1)
+		{
+			QueueIndex = 0;
+		}
 		checkThreadGraph(Queue(QueueIndex).StallRestartEvent); // make sure we are started up
 		Queue(QueueIndex).QuitWhenIdle.Increment();
 		Queue(QueueIndex).StallRestartEvent->Trigger(); 
+#endif
 	}
 
-	/** 
-	 *	Check the (unsafe) status of a thread.
-	 *	@return true if this thread was idle. 
-	 *	CAUTION the status of the thread can easily change before this routine returns.
-	 **/
-	bool IsProbablyStalled() const
-	{
-		checkThreadGraph(Queue(0).StallRestartEvent); // make sure we are started up
-		return IsStalled.GetValue() != 0;
-	}
-
-	/** Return the index of this thread. **/
-	ENamedThreads::Type GetThreadId() const
-	{
-		checkThreadGraph(Queue(0).StallRestartEvent); // make sure we are started up
-		return ThreadId;
-	}
-
-	/** 
-	 *	Queue a task, assuming that this thread is not the same as the current thread.
-	 *	@param QueueIndex, Queue to enqueue into
-	 *	@param Task; Task to queue.
-	 **/
-	bool EnqueueFromOtherThread(int32 QueueIndex, FBaseGraphTask* Task)
+	virtual bool EnqueueFromOtherThread(int32 QueueIndex, FBaseGraphTask* Task) override
 	{
 		TestRandomizedThreads();
-		checkThreadGraph(Queue(QueueIndex).StallRestartEvent); // make sure we are started up
+		checkThreadGraph(Task && Queue(QueueIndex).StallRestartEvent); // make sure we are started up
+
 		bool bWasReopenedByMe;
-		bool bHiPri = Task != WakeUpBaseGraphTask && ENamedThreads::GetPriority(Task->ThreadToExecuteOn);
+		bool bHiPri = !!ENamedThreads::GetPriority(Task->ThreadToExecuteOn);
 		{
 			TASKGRAPH_SCOPE_CYCLE_COUNTER(0, STAT_TaskGraph_EnqueueFromOtherThread_ReopenIfClosedAndPush);
 			bWasReopenedByMe = Queue(QueueIndex).IncomingQueue.ReopenIfClosedAndPush(Task);
@@ -652,85 +861,13 @@ public:
 		return bWasReopenedByMe;
 	}
 
-	/** 
-	 *	Attempt to give up a task for another thread.
-	 *	@return Task; Stolen task, if one was found, otherwise NULL.
-	 **/
-	FBaseGraphTask* RequestSteal()
-	{
-		checkThreadGraph(bAllowsStealsFromMe); 
-		return Queue(0).IncomingQueue.PopIfNotClosed();
-	}
-
-	/** 
-	 *Return true if this thread is processing tasks. This is only a "guess" if you ask for a thread other than yourself because that can change before the function returns.
-	 *@param QueueIndex, Queue to request quit from
-	 **/
-	bool IsProcessingTasks(int32 QueueIndex)
-	{
-		return !!Queue(QueueIndex).RecursionGuard.GetValue();
-	}
-
-	// FRunnable API
-
-	/**
-	 * Allows per runnable object initialization. NOTE: This is called in the
-	 * context of the thread object that aggregates this, not the thread that
-	 * passes this runnable to a new thread.
-	 *
-	 * @return True if initialization was successful, false otherwise
-	 */
-	virtual bool Init() override
-	{
-		InitializeForCurrentThread();
-		return true;
-	}
-
-	/**
-	 * This is where all per object thread work is done. This is only called
-	 * if the initialization was successful.
-	 *
-	 * @return The exit code of the runnable object
-	 */
-	virtual uint32 Run() override
-	{
-		ProcessTasksUntilQuit(0);
-		return 0;
-	}
-
-	/**
-	 * This is called if a thread is requested to terminate early
-	 */
-	virtual void Stop() override
+	virtual bool IsProcessingTasks(int32 QueueIndex) override
 	{
-		RequestQuit(0);
-	}
-
-	/**
-	 * Called in the context of the aggregating thread to perform any cleanup.
-	 */
-	virtual void Exit() override
-	{
-	}
-
-	/**
-	 * Return single threaded interface when multithreading is disabled.
-	 */
-	virtual FSingleThreadRunnable* GetSingleThreadInterface() override
-	{
-		return this;
+		return !!Queue(QueueIndex).RecursionGuard;
 	}
 
 private:
 
-	enum
-	{
-		/** The number of times to look for work before deciding to block on the stall event **/
-		SPIN_COUNT=0,
-		/** The number of times to call FPlatformProcess::Sleep(0) and look for work before deciding to block on the stall event **/
-		SLEEP_COUNT=0,
-	};
-
 	/** Grouping of the data for an individual queue. **/
 	struct FThreadTaskQueue
 	{
@@ -745,16 +882,23 @@ private:
 		 *	For unnamed thread this is the public queue, subject to stealing.
 		 *	In either case this queue is closely related to the stall event. Other threads that reopen the incoming queue must trigger the stall event to allow the thread to run.
 		**/
-		TReopenableLockFreePointerList<FBaseGraphTask>		IncomingQueue;
+#if USE_NEW_LOCK_FREE_LISTS && USE_INTRUSIVE_TASKQUEUES
+		FCloseableLockFreePointerQueueBaseSingleBaseConsumerIntrusive<FBaseGraphTask>		IncomingQueue;
+#elif USE_NEW_LOCK_FREE_LISTS
+		TReopenableLockFreePointerListFIFOSingleConsumer<FBaseGraphTask>		IncomingQueue;
+#else
+		TReopenableLockFreePointerListLIFO<FBaseGraphTask>		IncomingQueue;
+#endif
 		/** Used to signal the thread to quit when idle. **/
 		FThreadSafeCounter									QuitWhenIdle;
 		/** We need to disallow reentry of the processing loop **/
-		FThreadSafeCounter									RecursionGuard;
+		uint32												RecursionGuard;
 		/** Event that this thread blocks on when it runs out of work. **/
 		FEvent*												StallRestartEvent;
 
 		FThreadTaskQueue()
-			: StallRestartEvent(FPlatformProcess::GetSynchEventFromPool(true))
+			: RecursionGuard(0)
+			, StallRestartEvent(FPlatformProcess::GetSynchEventFromPool(true))
 		{
 
 		}
@@ -782,20 +926,16 @@ private:
 			// @Hack - Only stall when multithreading is enabled.
 			if (FPlatformProcess::SupportsMultithreading())
 			{
-				Queue(QueueIndex).StallRestartEvent->Reset();
 				FPlatformMisc::MemoryBarrier();
+				Queue(QueueIndex).StallRestartEvent->Reset();
+				TestRandomizedThreads();
 				if (Queue(QueueIndex).IncomingQueue.CloseIfEmpty())
 				{
 					FScopeCycleCounter Scope( StallStatId );
 
-					int32 NewValue = IsStalled.Increment();
-					NotifyStalling();
-					checkThreadGraph(NewValue == 1); // there should be no concurrent calls to Stall!
 					TestRandomizedThreads();
 					Queue(QueueIndex).StallRestartEvent->Wait(MAX_uint32, bCountAsStall);
-					TestRandomizedThreads();
-					NewValue = IsStalled.Decrement();
-					checkThreadGraph(NewValue == 0); // there should be no concurrent calls to Stall!
+					//checkThreadGraph(!Queue(QueueIndex).IncomingQueue.IsClosed()); // make sure we are started up
 					return true;
 				}
 			}
@@ -807,152 +947,316 @@ private:
 		return false;
 	}
 
-	/**
-	 *	Internal function to call the system looking for work. Called from this thread.
-	 *	@return New task to process.
-	 */
-	FBaseGraphTask* FindWork();
-
-	/**
-	 *	Internal function to notify the that system that I am stalling. This is a hint to give me a job asap.
-	 */
-	void NotifyStalling();
-
 	FORCEINLINE FThreadTaskQueue& Queue(int32 QueueIndex)
 	{
-		checkThreadGraph(QueueIndex >= 0 && QueueIndex < ENamedThreads::NumQueues && (!bAllowsStealsFromMe || !QueueIndex)); // range check, unnamed threads cannot use an alternate queue
+		checkThreadGraph(QueueIndex >= 0 && QueueIndex < ENamedThreads::NumQueues);
 		return Queues[QueueIndex];
 	}
 	FORCEINLINE const FThreadTaskQueue& Queue(int32 QueueIndex) const
 	{
-		checkThreadGraph(QueueIndex >= 0 && QueueIndex < ENamedThreads::NumQueues && (!bAllowsStealsFromMe || !QueueIndex)); // range check, unnamed threads cannot use an alternate queue
+		checkThreadGraph(QueueIndex >= 0 && QueueIndex < ENamedThreads::NumQueues);
 		return Queues[QueueIndex];
 	}
 
 	/** Array of queues, only the first one is used for unnamed threads. **/
 	FThreadTaskQueue Queues[ENamedThreads::NumQueues];
-
-	/** Id / Index of this thread. **/
-	ENamedThreads::Type									ThreadId;
-	/** TLS SLot that we store the FTaskThread* this pointer in. **/
-	uint32												PerThreadIDTLSSlot;
-	/** Used to signal stalling. Not safe for synchronization in most cases. **/
-	FThreadSafeCounter									IsStalled;
-	/** If true, this is a worker thread and any other thread can steal tasks from my incoming queue. **/
-	bool												bAllowsStealsFromMe;
-	/** If true, this is a worker thread and I will attempt to steal tasks when I run out of work. **/
-	bool												bStealsFromOthers;
-
 };
 
-
-
-
 /** 
- *	FTaskGraphImplementation
- *	Implementation of the centralized part of the task graph system.
- *	These parts of the system have no knowledge of the dependency graph, they exclusively work on tasks.
+ *	FTaskThreadAnyThread
+ *	A class for managing a worker threads. 
 **/
-class FTaskGraphImplementation : public FTaskGraphInterface  
+class FTaskThreadAnyThread : public FTaskThreadBase
 {
-	/** 
-	 *	FWorkerThread
-	 *	Helper structure to aggregate a few items related to the individual threads.
-	**/
-	struct FWorkerThread
+public:
+	/** Used for named threads to start processing tasks until the thread is idle and RequestQuit has been called. **/
+	virtual void ProcessTasksUntilQuit(int32 QueueIndex) override
 	{
-		/** The actual FTaskThread that manager this task **/
-		FTaskThread			TaskGraphWorker;
-		/** For internal threads, the is non-NULL and holds the information about the runable thread that was created. **/
-		FRunnableThread*	RunnableThread;
-		/** For external threads, this determines if they have been "attached" yet. Attachment is mostly setting up TLS for this individual thread. **/
-		bool				bAttached;
-
-		/** Constructor to set reasonable defaults. **/
-		FWorkerThread()
-			: RunnableThread(NULL)
-			, bAttached(false)
+		check(!QueueIndex);
+		Queue.QuitWhenIdle.Reset();
+		while (Queue.QuitWhenIdle.GetValue() == 0)
 		{
+			ProcessTasks();
+			// @Hack - quit now when running with only one thread.
+			if (!FPlatformProcess::SupportsMultithreading())
+			{
+				break;
+			}
 		}
-	};
-
-public:
-
-	// API related to life cycle of the system and singletons
-
-	/** 
-	 *	Singleton returning the one and only FTaskGraphImplementation.
-	 *	Note that unlike most singletons, a manual call to FTaskGraphInterface::Startup is required before the singleton will return a valid reference.
-	**/
-	static FTaskGraphImplementation& Get()
-	{		
-		checkThreadGraph(TaskGraphImplementationSingleton);
-		return *TaskGraphImplementationSingleton;
 	}
+	// Calls meant to be called from any thread.
 
 	/** 
-	 *	Constructor - initializes the data structures, sets the singleton pointer and creates the internal threads.
-	 *	@param InNumThreads; total number of threads in the system, including named threads, unnamed threads, internal threads and external threads. Must be at least 1 + the number of named threads.
+	 *	Will cause the thread to return to the caller when it becomes idle. Used to return from ProcessTasksUntilQuit for named threads or to shut down unnamed threads. 
+	 *	CAUTION: This will not work under arbitrary circumstances. For example you should not attempt to stop unnamed threads unless they are known to be idle.
+	 *	Return requests for named threads should be submitted from that named thread as FReturnGraphTask does.
+	 *	@param QueueIndex, Queue to request quit from
 	**/
-	FTaskGraphImplementation(int32 InNumThreads)
+	virtual void RequestQuit(int32 QueueIndex) override
 	{
-		// if we don't want any performance-based threads, then force the task graph to not create any worker threads, and run in game thread
-		if (!FPlatformProcess::SupportsMultithreading())
-		{
-			// this is the logic that used to be spread over a couple of places, that will make the rest of this function disable a worker thread
-			// @todo: it could probably be made simpler/clearer
-			InNumThreads = 1;
-			// this - 1 tells the below code there is no rendering thread
-			LastExternalThread = (ENamedThreads::Type)(ENamedThreads::ActualRenderingThread - 1);
-		}
-		else
-		{
-			LastExternalThread = ENamedThreads::ActualRenderingThread;
-		}
-		
-		NumNamedThreads = LastExternalThread + 1;
-		NumThreads = FMath::Max<int32>(FMath::Min<int32>(InNumThreads + NumNamedThreads,MAX_THREADS),NumNamedThreads + 1);
-		// Cap number of extra threads to the platform worker thread count
-		NumThreads = FMath::Min(NumThreads, NumNamedThreads + FPlatformMisc::NumberOfWorkerThreadsToSpawn());
-		UE_LOG(LogTaskGraph, Log, TEXT("Started task graph with %d named threads and %d total threads."), NumNamedThreads, NumThreads);
-		check(NumThreads - NumNamedThreads >= 1);  // need at least one pure worker thread
-		check(NumThreads <= MAX_THREADS);
-		check(!NextStealFromThread.GetValue()); // reentrant?
-		NextStealFromThread.Increment(); // just checking for reentrancy
-		PerThreadIDTLSSlot = FPlatformTLS::AllocTlsSlot();
-
-		NextUnnamedThreadMod = NumThreads - NumNamedThreads;
+		check(QueueIndex < 1);
 
-		for (int32 ThreadIndex = 0; ThreadIndex < NumThreads; ThreadIndex++)
-		{
-			check(!WorkerThreads[ThreadIndex].bAttached); // reentrant?
-			bool bAllowsStealsFromMe = ThreadIndex >= NumNamedThreads;
-			bool bStealsFromOthers = ThreadIndex >= NumNamedThreads;
-			WorkerThreads[ThreadIndex].TaskGraphWorker.Setup(ENamedThreads::Type(ThreadIndex), PerThreadIDTLSSlot, bAllowsStealsFromMe, bStealsFromOthers);
-		}
+		// this will not work under arbitrary circumstances. For example you should not attempt to stop threads unless they are known to be idle.
+		checkThreadGraph(Queue.StallRestartEvent); // make sure we are started up
+		Queue.QuitWhenIdle.Increment();
+		Queue.StallRestartEvent->Trigger(); 
+	}
 
-		TaskGraphImplementationSingleton = this; // now reentrancy is ok
+	virtual void WakeUp() final override
+	{
+		TASKGRAPH_SCOPE_CYCLE_COUNTER(1, STAT_TaskGraph_Wakeup_Trigger);
+		Queue.StallRestartEvent->Trigger();
+	}
 
-		for (int32 ThreadIndex = LastExternalThread + 1; ThreadIndex < NumThreads; ThreadIndex++)
-		{
-			FString Name = FString::Printf(TEXT("TaskGraphThread %d"), ThreadIndex - (LastExternalThread + 1));
-			uint32 StackSize = 256 * 1024;
-			WorkerThreads[ThreadIndex].RunnableThread = FRunnableThread::Create(&Thread(ThreadIndex), *Name, StackSize, TPri_Normal, FPlatformAffinity::GetTaskGraphThreadMask()); // these are below normal threads? so that they sleep when the named threads are active
-			WorkerThreads[ThreadIndex].bAttached = true;
-		}
+	/** 
+	 *Return true if this thread is processing tasks. This is only a "guess" if you ask for a thread other than yourself because that can change before the function returns.
+	 *@param QueueIndex, Queue to request quit from
+	 **/
+	virtual bool IsProcessingTasks(int32 QueueIndex) override
+	{
+		check(!QueueIndex);
+		return !!Queue.RecursionGuard;
 	}
 
+private:
+
 	/** 
-	 *	Destructor - probably only works reliably when the system is completely idle. The system has no idea if it is idle or not.
-	**/
-	virtual ~FTaskGraphImplementation()
+	 *	Process tasks until idle. May block if bAllowStall is true
+	 *	@param QueueIndex, Queue to process tasks from
+	 *	@param bAllowStall,  if true, the thread will block on the stall event when it runs out of tasks.
+	 **/
+	void ProcessTasks()
 	{
-		for (int32 ThreadIndex = 0; ThreadIndex < NumThreads; ThreadIndex++)
-		{
-			Thread(ThreadIndex).RequestQuit(0);
-		}
-		for (int32 ThreadIndex = 0; ThreadIndex < NumThreads; ThreadIndex++)
-		{
+		TStatId StallStatId;
+		bool bCountAsStall = true;
+#if STATS
+		TStatId StatName;
+		FCycleCounter ProcessingTasks;
+		StatName = GET_STATID(STAT_TaskGraph_OtherTasks);
+		StallStatId = GET_STATID(STAT_TaskGraph_OtherStalls);
+		bool bTasksOpen = false;
+		if (FThreadStats::IsCollectingData(StatName))
+		{
+			bTasksOpen = true;
+			ProcessingTasks.Start(StatName);
+		}
+#endif
+		verify(++Queue.RecursionGuard == 1);
+		while (1)
+		{
+			FBaseGraphTask* Task = FindWork();
+			if (!Task)
+			{
+#if STATS
+				if(bTasksOpen)
+				{
+					ProcessingTasks.Stop();
+					bTasksOpen = false;
+				}
+#endif
+				Stall(StallStatId, bCountAsStall);
+				if (Queue.QuitWhenIdle.GetValue())
+				{
+					break;
+				}
+#if STATS
+				if (FThreadStats::IsCollectingData(StatName))
+				{
+					bTasksOpen = true;
+					ProcessingTasks.Start(StatName);
+				}
+#endif
+				continue;
+			}
+			TestRandomizedThreads();
+			Task->Execute(NewTasks, ENamedThreads::Type(ThreadId));
+			TestRandomizedThreads();
+		}
+		verify(!--Queue.RecursionGuard);
+	}
+
+	/** Grouping of the data for an individual queue. **/
+	struct FThreadTaskQueue
+	{
+		/** Used to signal the thread to quit when idle. **/
+		FThreadSafeCounter									QuitWhenIdle;
+		/** We need to disallow reentry of the processing loop **/
+		uint32												RecursionGuard;
+		/** Event that this thread blocks on when it runs out of work. **/
+		FEvent*												StallRestartEvent;
+
+		FThreadTaskQueue()
+			: RecursionGuard(0)
+			, StallRestartEvent(FPlatformProcess::GetSynchEventFromPool(true))
+		{
+
+		}
+		~FThreadTaskQueue()
+		{
+			FPlatformProcess::ReturnSynchEventToPool( StallRestartEvent );
+			StallRestartEvent = nullptr;
+		}
+	};
+
+	/**
+	 *	Internal function to block on the stall wait event.
+	 *  @param QueueIndex		- Queue to stall
+	 *  @param StallStatId		- Stall stat id
+	 *  @param bCountAsStall	- true if StallStatId is a valid stat id
+	 */
+	void Stall(TStatId StallStatId, bool bCountAsStall)
+	{
+		TestRandomizedThreads();
+
+		checkThreadGraph(Queue.StallRestartEvent); // make sure we are started up
+		if (Queue.QuitWhenIdle.GetValue() == 0)
+		{
+			// @Hack - Only stall when multithreading is enabled.
+			if (FPlatformProcess::SupportsMultithreading())
+			{
+				FPlatformMisc::MemoryBarrier();
+
+				FScopeCycleCounter Scope( StallStatId );
+
+				NotifyStalling();
+				TestRandomizedThreads();
+				Queue.StallRestartEvent->Wait(MAX_uint32, bCountAsStall);
+				TestRandomizedThreads();
+				Queue.StallRestartEvent->Reset();
+			}
+		}
+	}
+
+	/**
+	 *	Internal function to call the system looking for work. Called from this thread.
+	 *	@return New task to process.
+	 */
+	FBaseGraphTask* FindWork();
+
+	/**
+	 *	Internal function to notify the that system that I am stalling. This is a hint to give me a job asap.
+	 */
+	void NotifyStalling();
+
+	/** Array of queues, only the first one is used for unnamed threads. **/
+	FThreadTaskQueue Queue;
+};
+
+/** 
+ *	FTaskGraphImplementation
+ *	Implementation of the centralized part of the task graph system.
+ *	These parts of the system have no knowledge of the dependency graph, they exclusively work on tasks.
+**/
+#define FAtomicStateBitfield_MAX_THREADS (13)
+
+/** 
+	*	FWorkerThread
+	*	Helper structure to aggregate a few items related to the individual threads.
+**/
+struct FWorkerThread
+{
+	/** The actual FTaskThread that manager this task **/
+	FTaskThreadBase*	TaskGraphWorker;
+	/** For internal threads, the is non-NULL and holds the information about the runable thread that was created. **/
+	FRunnableThread*	RunnableThread;
+	/** For external threads, this determines if they have been "attached" yet. Attachment is mostly setting up TLS for this individual thread. **/
+	bool				bAttached;
+
+	/** Constructor to set reasonable defaults. **/
+	FWorkerThread()
+		: TaskGraphWorker(nullptr)
+		, RunnableThread(nullptr)
+		, bAttached(false)
+	{
+	}
+};
+
+class FTaskGraphImplementation : public FTaskGraphInterface  
+{
+public:
+
+	// API related to life cycle of the system and singletons
+
+	/** 
+	 *	Singleton returning the one and only FTaskGraphImplementation.
+	 *	Note that unlike most singletons, a manual call to FTaskGraphInterface::Startup is required before the singleton will return a valid reference.
+	**/
+	static FTaskGraphImplementation& Get()
+	{		
+		checkThreadGraph(TaskGraphImplementationSingleton);
+		return *TaskGraphImplementationSingleton;
+	}
+
+	/** 
+	 *	Constructor - initializes the data structures, sets the singleton pointer and creates the internal threads.
+	 *	@param InNumThreads; total number of threads in the system, including named threads, unnamed threads, internal threads and external threads. Must be at least 1 + the number of named threads.
+	**/
+	FTaskGraphImplementation(int32 InNumThreads)
+	{
+		// if we don't want any performance-based threads, then force the task graph to not create any worker threads, and run in game thread
+		if (!FPlatformProcess::SupportsMultithreading())
+		{
+			// this is the logic that used to be spread over a couple of places, that will make the rest of this function disable a worker thread
+			// @todo: it could probably be made simpler/clearer
+			InNumThreads = 1;
+			// this - 1 tells the below code there is no rendering thread
+			LastExternalThread = (ENamedThreads::Type)(ENamedThreads::ActualRenderingThread - 1);
+		}
+		else
+		{
+			LastExternalThread = ENamedThreads::ActualRenderingThread;
+		}
+		
+		NumNamedThreads = LastExternalThread + 1;
+		NumThreads = FMath::Max<int32>(FMath::Min<int32>(InNumThreads + NumNamedThreads,MAX_THREADS),NumNamedThreads + 1);
+		// Cap number of extra threads to the platform worker thread count
+		NumThreads = FMath::Min(NumThreads, NumNamedThreads + FPlatformMisc::NumberOfWorkerThreadsToSpawn());
+		UE_LOG(LogTaskGraph, Log, TEXT("Started task graph with %d named threads and %d total threads."), NumNamedThreads, NumThreads);
+		check(NumThreads - NumNamedThreads >= 1);  // need at least one pure worker thread
+		check(NumThreads <= MAX_THREADS);
+		check(!NextStealFromThread.GetValue()); // reentrant?
+		NextStealFromThread.Increment(); // just checking for reentrancy
+		PerThreadIDTLSSlot = FPlatformTLS::AllocTlsSlot();
+
+		NextUnnamedThreadMod = NumThreads - NumNamedThreads;
+
+		for (int32 ThreadIndex = 0; ThreadIndex < NumThreads; ThreadIndex++)
+		{
+			check(!WorkerThreads[ThreadIndex].bAttached); // reentrant?
+			bool bAnyTaskThread = ThreadIndex >= NumNamedThreads;
+			if (bAnyTaskThread)
+			{
+				WorkerThreads[ThreadIndex].TaskGraphWorker = new FTaskThreadAnyThread;
+			}
+			else
+			{
+				WorkerThreads[ThreadIndex].TaskGraphWorker = new FNamedTaskThread;
+			}
+			WorkerThreads[ThreadIndex].TaskGraphWorker->Setup(ENamedThreads::Type(ThreadIndex), PerThreadIDTLSSlot, &WorkerThreads[ThreadIndex]);
+		}
+
+		TaskGraphImplementationSingleton = this; // now reentrancy is ok
+
+		for (int32 ThreadIndex = LastExternalThread + 1; ThreadIndex < NumThreads; ThreadIndex++)
+		{
+			FString Name = FString::Printf(TEXT("TaskGraphThread %d"), ThreadIndex - (LastExternalThread + 1));
+			uint32 StackSize = 256 * 1024;
+			WorkerThreads[ThreadIndex].RunnableThread = FRunnableThread::Create(&Thread(ThreadIndex), *Name, StackSize, TPri_Normal, FPlatformAffinity::GetTaskGraphThreadMask()); // these are below normal threads? so that they sleep when the named threads are active
+			WorkerThreads[ThreadIndex].bAttached = true;
+		}
+	}
+
+	/** 
+	 *	Destructor - probably only works reliably when the system is completely idle. The system has no idea if it is idle or not.
+	**/
+	virtual ~FTaskGraphImplementation()
+	{
+		for (int32 ThreadIndex = 0; ThreadIndex < NumThreads; ThreadIndex++)
+		{
+			Thread(ThreadIndex).RequestQuit(-1);
+		}
+		for (int32 ThreadIndex = 0; ThreadIndex < NumThreads; ThreadIndex++)
+		{
 			if (ThreadIndex > LastExternalThread)
 			{
 				WorkerThreads[ThreadIndex].RunnableThread->WaitForCompletion();
@@ -963,7 +1267,7 @@ public:
 		}
 		TaskGraphImplementationSingleton = NULL;
 		NextUnnamedThreadMod = 0;
-		TArray<FTaskThread*> NotProperlyUnstalled;
+		TArray<FTaskThreadBase*> NotProperlyUnstalled;
 		StalledUnnamedThreads.PopAll(NotProperlyUnstalled);
 		FPlatformTLS::FreeTlsSlot(PerThreadIDTLSSlot);
 	}
@@ -977,21 +1281,52 @@ public:
 	 *	@param	ThreadToExecuteOn; Either a named thread for a threadlocked task or ENamedThreads::AnyThread for a task that is to run on a worker thread
 	 *	@param	CurrentThreadIfKnown; This should be the current thread if it is known, or otherwise use ENamedThreads::AnyThread and the current thread will be determined.
 	**/
-	virtual void QueueTask(FBaseGraphTask* Task, ENamedThreads::Type ThreadToExecuteOn, ENamedThreads::Type CurrentThreadIfKnown = ENamedThreads::AnyThread) override
+	virtual void QueueTask(FBaseGraphTask* Task, ENamedThreads::Type ThreadToExecuteOn, ENamedThreads::Type InCurrentThreadIfKnown = ENamedThreads::AnyThread) final override
 	{
 		TASKGRAPH_SCOPE_CYCLE_COUNTER(2, STAT_TaskGraph_QueueTask);
 
 		TestRandomizedThreads();
 		checkThreadGraph(NextUnnamedThreadMod);
-		if (ENamedThreads::GetThreadIndex(CurrentThreadIfKnown) == ENamedThreads::AnyThread)
+		if (GFastSchedulerLatched != GFastScheduler && IsInGameThread())
 		{
-			 CurrentThreadIfKnown = GetCurrentThread();
+#if USE_NEW_LOCK_FREE_LISTS
+			GFastScheduler = !!FApp::ShouldUseThreadingForPerformance();
+#endif
+			if (GFastSchedulerLatched != GFastScheduler)
+			{
+				if (GFastScheduler)
+				{
+					AtomicForConsoleApproach = FAtomicStateBitfield();
+					AtomicForConsoleApproach.Stalled = (1 << GetNumWorkerThreads()) - 1; // everyone is stalled
+				}
+				// this is all kinda of sketchy, but runtime switching just has to barely work.
+				FPlatformProcess::Sleep(.1f); // hopefully all task threads are idle
+				GFastSchedulerLatched = GFastScheduler;
+
+				// wake up all threads to prevent deadlock on the transition
+				if (GFastSchedulerLatched)
+				{
+					GConsoleSpinModeLatched = 0; // nobody is started now, this will re latch and start someone if we have that mode on
+				}
+				else
+				{
+					for (int32 Index = 0; Index < GetNumWorkerThreads(); Index++)
+					{
+						StartTaskThread(Index);
+					}
+				}
+			}
 		}
-		else
+		if (GFastSchedulerLatched && GConsoleSpinModeLatched != GConsoleSpinMode && IsInGameThread())
 		{
-			CurrentThreadIfKnown = ENamedThreads::GetThreadIndex(CurrentThreadIfKnown);
-			checkThreadGraph(CurrentThreadIfKnown == GetCurrentThread());
+			bool bStartTask = GConsoleSpinMode && ((!!GConsoleSpinModeLatched) != !!GConsoleSpinMode);
+			GConsoleSpinModeLatched = GConsoleSpinMode;
+			if (bStartTask)
+			{
+				StartTaskThreadFastMode();
+			}
 		}
+
 		if (ENamedThreads::GetThreadIndex(ThreadToExecuteOn) == ENamedThreads::AnyThread)
 		{
 			TASKGRAPH_SCOPE_CYCLE_COUNTER(3, STAT_TaskGraph_QueueTask_AnyThread);
@@ -1008,28 +1343,48 @@ public:
 						IncomingAnyThreadTasks.Push(Task);
 					}
 				}
-				FTaskThread* TempTarget;
+				if (GFastSchedulerLatched)
 				{
-					TASKGRAPH_SCOPE_CYCLE_COUNTER(5, STAT_TaskGraph_QueueTask_StalledUnnamedThreads_Pop);
-					TempTarget = StalledUnnamedThreads.Pop(); //@todo it is possible that a thread is in the process of stalling and we just missed it, non-fatal, but we could lose a whole task of potential parallelism.
-					if (TempTarget && GNumWorkerThreadsToIgnore && (TempTarget->GetThreadId() - NumNamedThreads) >= GetNumWorkerThreads())
+					if (!GConsoleSpinModeLatched)
 					{
-						TempTarget = nullptr;
+						StartTaskThreadFastMode();
 					}
 				}
-				if (TempTarget)
-				{
-					ThreadToExecuteOn = TempTarget->GetThreadId();
-				}
 				else
 				{
-					check(NextUnnamedThreadMod - GNumWorkerThreadsToIgnore > 0); // can't tune it to zero task threads
-					ThreadToExecuteOn = ENamedThreads::Type((uint32(NextUnnamedThreadForTaskFromUnknownThread.Increment()) % uint32(NextUnnamedThreadMod - GNumWorkerThreadsToIgnore)) + NumNamedThreads);
-				}
-				FTaskThread* Target = &Thread(ThreadToExecuteOn);
-				if (ThreadToExecuteOn != CurrentThreadIfKnown)
-				{
-					Target->EnqueueFromOtherThread(0, WakeUpBaseGraphTask);
+					FTaskThreadBase* TempTarget;
+					{
+						TASKGRAPH_SCOPE_CYCLE_COUNTER(5, STAT_TaskGraph_QueueTask_StalledUnnamedThreads_Pop);
+						TempTarget = StalledUnnamedThreads.Pop(); //@todo it is possible that a thread is in the process of stalling and we just missed it, non-fatal, but we could lose a whole task of potential parallelism.
+						if (TempTarget && GNumWorkerThreadsToIgnore && (TempTarget->GetThreadId() - NumNamedThreads) >= GetNumWorkerThreads())
+						{
+							TempTarget = nullptr;
+						}
+					}
+					if (TempTarget)
+					{
+						ThreadToExecuteOn = TempTarget->GetThreadId();
+					}
+					else
+					{
+						check(NextUnnamedThreadMod - GNumWorkerThreadsToIgnore > 0); // can't tune it to zero task threads
+						ThreadToExecuteOn = ENamedThreads::Type((uint32(NextUnnamedThreadForTaskFromUnknownThread.Increment()) % uint32(NextUnnamedThreadMod - GNumWorkerThreadsToIgnore)) + NumNamedThreads);
+					}
+					FTaskThreadBase* Target = &Thread(ThreadToExecuteOn);
+					ENamedThreads::Type CurrentThreadIfKnown;
+					if (ENamedThreads::GetThreadIndex(InCurrentThreadIfKnown) == ENamedThreads::AnyThread)
+					{
+						CurrentThreadIfKnown = GetCurrentThread();
+					}
+					else
+					{
+						CurrentThreadIfKnown = ENamedThreads::GetThreadIndex(InCurrentThreadIfKnown);
+						checkThreadGraph(CurrentThreadIfKnown == GetCurrentThread());
+					}
+					if (ThreadToExecuteOn != CurrentThreadIfKnown)
+					{
+						Target->WakeUp();
+					}
 				}
 				return;
 			}
@@ -1038,10 +1393,20 @@ public:
 				ThreadToExecuteOn = ENamedThreads::GameThread;
 			}
 		}
+		ENamedThreads::Type CurrentThreadIfKnown;
+		if (ENamedThreads::GetThreadIndex(InCurrentThreadIfKnown) == ENamedThreads::AnyThread)
+		{
+			CurrentThreadIfKnown = GetCurrentThread();
+		}
+		else
+		{
+			CurrentThreadIfKnown = ENamedThreads::GetThreadIndex(InCurrentThreadIfKnown);
+			checkThreadGraph(CurrentThreadIfKnown == GetCurrentThread());
+		}
 		{
 			int32 QueueToExecuteOn = ENamedThreads::GetQueueIndex(ThreadToExecuteOn);
 			ThreadToExecuteOn = ENamedThreads::GetThreadIndex(ThreadToExecuteOn);
-			FTaskThread* Target = &Thread(ThreadToExecuteOn);
+			FTaskThreadBase* Target = &Thread(ThreadToExecuteOn);
 			if (ThreadToExecuteOn == CurrentThreadIfKnown)
 			{
 				Target->EnqueueFromThisThread(QueueToExecuteOn, Task);
@@ -1054,14 +1419,14 @@ public:
 	}
 
 
-	virtual	int32 GetNumWorkerThreads() override
+	virtual	int32 GetNumWorkerThreads() final override
 	{
 		int32 Result = NumThreads - NumNamedThreads - GNumWorkerThreadsToIgnore;
 		check(Result > 0); // can't tune it to zero task threads
 		return Result;
 	}
 
-	virtual ENamedThreads::Type GetCurrentThreadIfKnown(bool bLocalQueue) override
+	virtual ENamedThreads::Type GetCurrentThreadIfKnown(bool bLocalQueue) final override
 	{
 		ENamedThreads::Type Result = GetCurrentThread();
 		if (bLocalQueue && Result >= 0 && Result < NumNamedThreads)
@@ -1071,7 +1436,7 @@ public:
 		return Result;
 	}
 
-	virtual bool IsThreadProcessingTasks(ENamedThreads::Type ThreadToCheck) override
+	virtual bool IsThreadProcessingTasks(ENamedThreads::Type ThreadToCheck) final override
 	{
 		int32 QueueIndex = ENamedThreads::GetQueueIndex(ThreadToCheck);
 		ThreadToCheck = ENamedThreads::GetThreadIndex(ThreadToCheck);
@@ -1081,7 +1446,7 @@ public:
 
 	// External Thread API
 
-	virtual void AttachToThread(ENamedThreads::Type CurrentThread) override
+	virtual void AttachToThread(ENamedThreads::Type CurrentThread) final override
 	{
 		CurrentThread = ENamedThreads::GetThreadIndex(CurrentThread);
 		check(NextUnnamedThreadMod);
@@ -1090,16 +1455,16 @@ public:
 		Thread(CurrentThread).InitializeForCurrentThread();
 	}
 
-	virtual void ProcessThreadUntilIdle(ENamedThreads::Type CurrentThread) override
+	virtual void ProcessThreadUntilIdle(ENamedThreads::Type CurrentThread) final override
 	{
 		int32 QueueIndex = ENamedThreads::GetQueueIndex(CurrentThread);
 		CurrentThread = ENamedThreads::GetThreadIndex(CurrentThread);
 		check(CurrentThread >= 0 && CurrentThread < NumNamedThreads);
 		check(CurrentThread == GetCurrentThread());
-		Thread(CurrentThread).ProcessTasks(QueueIndex, false);
+		Thread(CurrentThread).ProcessTasksUntilIdle(QueueIndex);
 	}
 
-	virtual void ProcessThreadUntilRequestReturn(ENamedThreads::Type CurrentThread) override
+	virtual void ProcessThreadUntilRequestReturn(ENamedThreads::Type CurrentThread) final override
 	{
 		int32 QueueIndex = ENamedThreads::GetQueueIndex(CurrentThread);
 		CurrentThread = ENamedThreads::GetThreadIndex(CurrentThread);
@@ -1108,7 +1473,7 @@ public:
 		Thread(CurrentThread).ProcessTasksUntilQuit(QueueIndex);
 	}
 
-	virtual void RequestReturn(ENamedThreads::Type CurrentThread) override
+	virtual void RequestReturn(ENamedThreads::Type CurrentThread) final override
 	{
 		int32 QueueIndex = ENamedThreads::GetQueueIndex(CurrentThread);
 		CurrentThread = ENamedThreads::GetThreadIndex(CurrentThread);
@@ -1116,7 +1481,7 @@ public:
 		Thread(CurrentThread).RequestQuit(QueueIndex);
 	}
 
-	virtual void WaitUntilTasksComplete(const FGraphEventArray& Tasks, ENamedThreads::Type CurrentThreadIfKnown = ENamedThreads::AnyThread) override
+	virtual void WaitUntilTasksComplete(const FGraphEventArray& Tasks, ENamedThreads::Type CurrentThreadIfKnown = ENamedThreads::AnyThread) final override
 	{
 		ENamedThreads::Type CurrentThread = CurrentThreadIfKnown;
 		if (ENamedThreads::GetThreadIndex(CurrentThreadIfKnown) == ENamedThreads::AnyThread)
@@ -1160,7 +1525,7 @@ public:
 		}
 	}
 
-	virtual void TriggerEventWhenTasksComplete(FEvent* InEvent, const FGraphEventArray& Tasks, ENamedThreads::Type CurrentThreadIfKnown = ENamedThreads::AnyThread) override
+	virtual void TriggerEventWhenTasksComplete(FEvent* InEvent, const FGraphEventArray& Tasks, ENamedThreads::Type CurrentThreadIfKnown = ENamedThreads::AnyThread) final override
 	{
 		check(InEvent);
 		bool bAnyPending = false;
@@ -1181,28 +1546,355 @@ public:
 		TGraphTask<FTriggerEventGraphTask>::CreateTask(&Tasks, CurrentThreadIfKnown).ConstructAndDispatchWhenReady(InEvent);
 	}
 
-	// API used by FWorkerThread's
+	// Scheduling utilities
 
-	/** 
-	 *	Attempt to steal some work from another thread.
-	 *	@param	ThreadInNeed; Id of the thread requesting work.
-	 *	@return Task that was stolen if any was found.
-	**/
-	FBaseGraphTask* FindWork(ENamedThreads::Type ThreadInNeed)
+	void StartTaskThread(int32 IndexToStart)
 	{
-		TestRandomizedThreads();
+		ENamedThreads::Type ThreadToWake = ENamedThreads::Type(IndexToStart + NumNamedThreads);
+		((FTaskThreadAnyThread&)Thread(ThreadToWake)).FTaskThreadAnyThread::WakeUp();
+	}
+
+	bool StartTaskThreadFastMode(int32 MyIndex = -1)
+	{
+		while (true)
 		{
-			FBaseGraphTask* Task = SortedAnyThreadTasksHiPri.Pop();
-			if (Task)
+			TestRandomizedThreads();
+			FAtomicStateBitfield LocalAtomicForConsoleApproachInner = AtomicForConsoleApproach;
+			check(MyIndex < 0 || !(LocalAtomicForConsoleApproachInner.Stalled & (1 << MyIndex))); // if I was stalled, I would not be here
+
+			if (!LocalAtomicForConsoleApproachInner.Stalled)
 			{
-				return Task;
+				break;
+			}
+
+			uint32 IndexToStart =  31 - FMath::CountLeadingZeros((uint32)LocalAtomicForConsoleApproachInner.Stalled);
+			check(IndexToStart != MyIndex);
+			check(((uint32)LocalAtomicForConsoleApproachInner.Stalled) & (1 << IndexToStart));
+
+			FAtomicStateBitfield NewAtomicForConsoleApproachInner = LocalAtomicForConsoleApproachInner;
+			NewAtomicForConsoleApproachInner.Stalled = LocalAtomicForConsoleApproachInner.Stalled & ~(1 << IndexToStart);
+			check(!(((uint32)NewAtomicForConsoleApproachInner.Stalled) & (1 << IndexToStart)));
+			check(LocalAtomicForConsoleApproachInner != NewAtomicForConsoleApproachInner);
+			if (FAtomicStateBitfield::InterlockedCompareExchange(&AtomicForConsoleApproach, NewAtomicForConsoleApproachInner, LocalAtomicForConsoleApproachInner) == LocalAtomicForConsoleApproachInner)
+			{
+				TestRandomizedThreads();
+				StartTaskThread(IndexToStart);
+				return true;
 			}
 		}
-		if (!IncomingAnyThreadTasksHiPri.IsEmpty())
+		return false;
+	}
+
+	FORCEINLINE void SetWorking(uint32 MyIndex, int32 LocalNumWorkingThread)
+	{
+		if (!GConsoleSpinModeLatched)
 		{
-			do
+			// we don't need to track the number of threads working or start threads when we take a job unless we are in spin mode
+			return;
+		}
+		while (true)
+		{
+			TestRandomizedThreads();
+			FPlatformMisc::MemoryBarrier();
+			FAtomicStateBitfield LocalAtomicForConsoleApproachInner = AtomicForConsoleApproach;
+			check(!(LocalAtomicForConsoleApproachInner.Stalled & (1 << MyIndex))); // if I was stalled, I would not be here
+			bool CurrentlyWorking = !!((LocalAtomicForConsoleApproachInner.Working) & (1 << MyIndex));
+			if (CurrentlyWorking)
+			{
+				break;
+			}
+
+			FAtomicStateBitfield NewAtomicForConsoleApproachInner = LocalAtomicForConsoleApproachInner;
+			NewAtomicForConsoleApproachInner.Working = LocalAtomicForConsoleApproachInner.Working | (1 << MyIndex);
+			int32 NumStalled = LocalAtomicForConsoleApproachInner.NumberOfStalledThreads();
+
+			// we are going to do work, so if everyone else is either working or stalled, we need to start a thread
+			if (NumStalled && LocalAtomicForConsoleApproachInner.NumberOfStalledThreads() + LocalAtomicForConsoleApproachInner.NumberOfWorkingThreads() + 1 == LocalNumWorkingThread)
+			{
+				// start one while we flip ourselves to working.
+				check(LocalAtomicForConsoleApproachInner.Stalled);
+
+				uint32 IndexToStart =  31 - FMath::CountLeadingZeros((uint32)LocalAtomicForConsoleApproachInner.Stalled);
+				check(IndexToStart != MyIndex);
+				check(((uint32)LocalAtomicForConsoleApproachInner.Stalled) & (1 << IndexToStart));
+
+				NewAtomicForConsoleApproachInner.Stalled = LocalAtomicForConsoleApproachInner.Stalled & ~(1 << IndexToStart);
+
+				check(LocalAtomicForConsoleApproachInner != NewAtomicForConsoleApproachInner);
+				if (FAtomicStateBitfield::InterlockedCompareExchange(&AtomicForConsoleApproach, NewAtomicForConsoleApproachInner, LocalAtomicForConsoleApproachInner) == LocalAtomicForConsoleApproachInner)
+				{
+					TestRandomizedThreads();
+					StartTaskThread(IndexToStart);
+					// if we have another stalled thread and it looks like there are more tasks, then start an extra thread to get fan-out on the thread starts
+#if USE_NEW_LOCK_FREE_LISTS
+					if (GMaxTasksToStartOnDequeue > 1 && NumStalled > 1 && (!IncomingAnyThreadTasks.IsEmptyFast() || !IncomingAnyThreadTasksHiPri.IsEmptyFast()))
+					{
+						StartTaskThreadFastMode(MyIndex);
+					}
+#endif
+					return;
+				}
+			}
+			else
+			{
+				// we don't need to wake anyone either because there are other unstalled idle threads or everyone is busy
+				check(LocalAtomicForConsoleApproachInner != NewAtomicForConsoleApproachInner);
+				if (FAtomicStateBitfield::InterlockedCompareExchange(&AtomicForConsoleApproach, NewAtomicForConsoleApproachInner, LocalAtomicForConsoleApproachInner) == LocalAtomicForConsoleApproachInner)
+				{
+					TestRandomizedThreads();
+					return;
+				}
+			}
+		}
+	}
+#if USE_NEW_LOCK_FREE_LISTS
+	FBaseGraphTask* FindWork(ENamedThreads::Type ThreadInNeed)
+	{
+		int32 LocalNumWorkingThread = GetNumWorkerThreads();
+		uint32 MyIndex = uint32(ThreadInNeed) - NumNamedThreads;
+		check(MyIndex >= 0 && (int32)MyIndex < LocalNumWorkingThread && 
+			LocalNumWorkingThread <= FAtomicStateBitfield_MAX_THREADS); // this is an atomic on 32 bits currently; we use two fields
+		while (true)
+		{
+			if (!GFastSchedulerLatched)
+			{
+				return nullptr;
+			}
+			bool bOkWithFalseEmpties = !!GConsoleSpinModeLatched;
+			{
+				FBaseGraphTask* Task = IncomingAnyThreadTasksHiPri.Pop(bOkWithFalseEmpties);
+				if (!Task)
+				{
+					Task = IncomingAnyThreadTasks.Pop(bOkWithFalseEmpties);
+				}
+				if (Task)
+				{
+					SetWorking(MyIndex, LocalNumWorkingThread);
+					return Task;
+				}
+			}
+
+			TestRandomizedThreads();
+			FPlatformMisc::MemoryBarrier();
+			FAtomicStateBitfield LocalAtomicForConsoleApproachInner = AtomicForConsoleApproach;
+			check(!(LocalAtomicForConsoleApproachInner.Stalled & (1 << MyIndex))); // if I was stalled, I would not be here
+
+
+			FAtomicStateBitfield NewAtomicForConsoleApproachInner = LocalAtomicForConsoleApproachInner;
+			bool CurrentlyWorking = !!((LocalAtomicForConsoleApproachInner.Working) & (1 << MyIndex));
+			NewAtomicForConsoleApproachInner.Working = NewAtomicForConsoleApproachInner.Working & ~(1 << MyIndex);
+			int32 NumWorking = LocalAtomicForConsoleApproachInner.NumberOfWorkingThreads() - !!CurrentlyWorking;
+			bool bAttemptStall = (!GConsoleSpinModeLatched || LocalAtomicForConsoleApproachInner.NumberOfStalledThreads() + NumWorking < LocalNumWorkingThread - 1);
+			if (bAttemptStall)
+			{
+				NewAtomicForConsoleApproachInner.Stalled = LocalAtomicForConsoleApproachInner.Stalled | (1 << MyIndex);
+				check((((uint32)NewAtomicForConsoleApproachInner.Stalled) & (1 << MyIndex)));
+
+				// if we are going to attempt to stall, we better be sure there really are not tasks...unless we are already sure
+				if (bOkWithFalseEmpties)
+				{
+					FBaseGraphTask* Task = IncomingAnyThreadTasksHiPri.Pop();
+					if (!Task)
+					{
+						Task = IncomingAnyThreadTasks.Pop();
+					}
+					if (Task)
+					{
+						SetWorking(MyIndex, LocalNumWorkingThread);
+						return Task;
+					}
+				}
+			}
+			if (bAttemptStall || CurrentlyWorking) // otherwise we have no reason to change the state
+			{
+				check(LocalAtomicForConsoleApproachInner != NewAtomicForConsoleApproachInner);
+				if (FAtomicStateBitfield::InterlockedCompareExchange(&AtomicForConsoleApproach, NewAtomicForConsoleApproachInner, LocalAtomicForConsoleApproachInner) == LocalAtomicForConsoleApproachInner)
+				{
+					if (bAttemptStall)
+					{
+						TestRandomizedThreads();
+						return nullptr; // we did stall
+					}
+				}
+			}
+			else if (GConsoleSpinModeLatched == 2)
+			{
+				FPlatformProcess::Sleep(0.0f);
+			}
+		}
+	}
+#else
+	FBaseGraphTask* FindWorkConsole(ENamedThreads::Type ThreadInNeed)
+	{
+		int32 LocalNumWorkingThread = GetNumWorkerThreads();
+		uint32 MyIndex = uint32(ThreadInNeed) - NumNamedThreads;
+		check(MyIndex >= 0 && (int32)MyIndex < LocalNumWorkingThread && 
+			LocalNumWorkingThread <= FAtomicStateBitfield_MAX_THREADS); // this is an atomic on 32 bits currently; we use two fields
+		while (true)
+		{
+			if (!GFastSchedulerLatched)
+			{
+				return nullptr;
+			}
+			{
+				FBaseGraphTask* Task = SortedAnyThreadTasksHiPri.Pop();
+				if (Task)
+				{
+					SetWorking(MyIndex, LocalNumWorkingThread);
+					return Task;
+				}
+			}
+			TestRandomizedThreads();
+			FPlatformMisc::MemoryBarrier();
+			FAtomicStateBitfield LocalAtomicForConsoleApproach = AtomicForConsoleApproach;
+			check(!(LocalAtomicForConsoleApproach.Stalled & (1 << MyIndex))); // if I was stalled, I would not be here
+			if (!LocalAtomicForConsoleApproach.QueueOwned)
+			{
+				if (IncomingAnyThreadTasksHiPri.IsEmpty())
+				{
+					FBaseGraphTask* Task = SortedAnyThreadTasks.Pop();
+					if (Task)
+					{
+						SetWorking(MyIndex, LocalNumWorkingThread);
+						return Task;
+					}
+				}
+				FAtomicStateBitfield NewAtomicForConsoleApproach = LocalAtomicForConsoleApproach;
+				// Queue is unowned, lets grab it and try to empty them
+				NewAtomicForConsoleApproach.QueueOwned = 1;
+				NewAtomicForConsoleApproach.QueueOwnerIndex = MyIndex;
+				NewAtomicForConsoleApproach.Working = LocalAtomicForConsoleApproach.Working & ~(1 << MyIndex);
+
+				check(LocalAtomicForConsoleApproach != NewAtomicForConsoleApproach);
+
+				if (FAtomicStateBitfield::InterlockedCompareExchange(&AtomicForConsoleApproach, NewAtomicForConsoleApproach, LocalAtomicForConsoleApproach) == LocalAtomicForConsoleApproach)
+				{
+					TestRandomizedThreads();
+					// we win
+					// dequeue any pending things
+					static TArray<FBaseGraphTask*> NewTasks;
+					int32 TotalTasks = 0;
+					bool bWeWereBeat = false;
+					if (SortedAnyThreadTasksHiPri.IsEmpty()) // otherwise someone beat us to it, so we need to drain the queue first
+					{
+						NewTasks.Reset();
+						IncomingAnyThreadTasksHiPri.PopAll(NewTasks);
+						if (NewTasks.Num())
+						{
+							TotalTasks += NewTasks.Num();
+							TLockFreePointerListLIFO<FBaseGraphTask> TempSortedAnyThreadTasks;
+							for (int32 Index = 0 ; Index < NewTasks.Num(); Index++) // we are going to take the last one for ourselves
+							{
+								TempSortedAnyThreadTasks.Push(NewTasks[Index]);
+							}
+							verify(SortedAnyThreadTasksHiPri.ReplaceListIfEmpty(TempSortedAnyThreadTasks));
+						}
+					}
+					else
+					{
+						bWeWereBeat = true;
+					}
+					if (SortedAnyThreadTasks.IsEmpty()) // otherwise someone beat us to it, so we need to drain the queue first
+					{
+						NewTasks.Reset();
+						IncomingAnyThreadTasks.PopAll(NewTasks);
+						if (NewTasks.Num())
+						{
+							TotalTasks += NewTasks.Num();
+							TLockFreePointerListLIFO<FBaseGraphTask> TempSortedAnyThreadTasks;
+							for (int32 Index = 0 ; Index < NewTasks.Num(); Index++) // we are going to take the last one for ourselves
+							{
+								TempSortedAnyThreadTasks.Push(NewTasks[Index]);
+							}
+							verify(SortedAnyThreadTasks.ReplaceListIfEmpty(TempSortedAnyThreadTasks));
+						}
+					}
+					else
+					{
+						bWeWereBeat = true;
+					}
+					// release the ownership
+					while (true)
+					{
+						TestRandomizedThreads();
+						FAtomicStateBitfield LocalAtomicForConsoleApproachInner = AtomicForConsoleApproach;
+						check(!(LocalAtomicForConsoleApproachInner.Stalled & (1 << MyIndex))); // if I was stalled, I would not be here
+						check(!GConsoleSpinModeLatched || !(LocalAtomicForConsoleApproachInner.Working & (1 << MyIndex))); // if I was working, I would not be here
+
+
+						FAtomicStateBitfield NewAtomicForConsoleApproachInner = LocalAtomicForConsoleApproachInner;
+						check(LocalAtomicForConsoleApproachInner.QueueOwned && LocalAtomicForConsoleApproachInner.QueueOwnerIndex == MyIndex);
+						NewAtomicForConsoleApproachInner.QueueOwned = false;
+						NewAtomicForConsoleApproachInner.QueueOwnerIndex = 0;
+						bool bAttemptStall = !TotalTasks && !bWeWereBeat &&
+							(!GConsoleSpinModeLatched || LocalAtomicForConsoleApproachInner.NumberOfStalledThreads() + LocalAtomicForConsoleApproachInner.NumberOfWorkingThreads() < LocalNumWorkingThread - 1);
+						if (bAttemptStall)
+						{
+							NewAtomicForConsoleApproachInner.Stalled = LocalAtomicForConsoleApproachInner.Stalled | (1 << MyIndex);
+							NewAtomicForConsoleApproachInner.Working = LocalAtomicForConsoleApproachInner.Working & ~(1 << MyIndex); // this is redundant in most cases, but might as well track the correct state
+							check((((uint32)NewAtomicForConsoleApproachInner.Stalled) & (1 << MyIndex)));
+						}
+
+						check(LocalAtomicForConsoleApproachInner != NewAtomicForConsoleApproachInner);
+						if (FAtomicStateBitfield::InterlockedCompareExchange(&AtomicForConsoleApproach, NewAtomicForConsoleApproachInner, LocalAtomicForConsoleApproachInner) == LocalAtomicForConsoleApproachInner)
+						{
+							TestRandomizedThreads();
+							if (bAttemptStall)
+							{
+								return nullptr; // we did stall
+							}
+							break;
+						}
+					}
+					if (GConsoleSpinModeLatched)
+					{
+						// -- here because if we actually do one of these tasks, we will start another thread by transitioning into working
+						TotalTasks--;
+						int32 ThreadsToStart = FMath::Min<int32>(FMath::Min<int32>(TotalTasks, GMaxTasksToStartOnDequeue), LocalNumWorkingThread);
+						while (ThreadsToStart > 0)
+						{
+							TestRandomizedThreads();
+							if (!StartTaskThreadFastMode(MyIndex))
+							{
+								break;
+							}
+							ThreadsToStart--;
+						}
+					}
+				}
+			}
+		}
+	}
+	// API used by FWorkerThread's
+
+	/** 
+	 *	Attempt to steal some work from another thread.
+	 *	@param	ThreadInNeed; Id of the thread requesting work.
+	 *	@return Task that was stolen if any was found.
+	**/
+	FBaseGraphTask* FindWork(ENamedThreads::Type ThreadInNeed)
+	{
+		TestRandomizedThreads();
+		if (GFastSchedulerLatched)
+		{
+			return FindWorkConsole(ThreadInNeed);
+		}
+		{
+			FBaseGraphTask* Task = SortedAnyThreadTasksHiPri.Pop();
+			if (Task)
+			{
+				return Task;
+			}
+		}
+		if (!IncomingAnyThreadTasksHiPri.IsEmpty())
+		{
+			do
 			{
 				FScopeLock ScopeLock(&CriticalSectionForSortingIncomingAnyThreadTasksHiPri);
+				if (GFastSchedulerLatched)
+				{
+					return nullptr;
+				}
 				if (!IncomingAnyThreadTasksHiPri.IsEmpty() && SortedAnyThreadTasksHiPri.IsEmpty())
 				{
 					static TArray<FBaseGraphTask*> NewTasks;
@@ -1212,7 +1904,7 @@ public:
 
 					if (NewTasks.Num() > 1)
 					{
-						TLockFreePointerList<FBaseGraphTask> TempSortedAnyThreadTasks;
+						TLockFreePointerListLIFO<FBaseGraphTask> TempSortedAnyThreadTasks;
 						for (int32 Index = 0 ; Index < NewTasks.Num() - 1; Index++) // we are going to take the last one for ourselves
 						{
 							TempSortedAnyThreadTasks.Push(NewTasks[Index]);
@@ -1240,6 +1932,10 @@ public:
 		do
 		{
 			FScopeLock ScopeLock(&CriticalSectionForSortingIncomingAnyThreadTasks);
+			if (GFastSchedulerLatched)
+			{
+				return nullptr;
+			}
 			if (!IncomingAnyThreadTasks.IsEmpty() && SortedAnyThreadTasks.IsEmpty())
 			{
 				static TArray<FBaseGraphTask*> NewTasks;
@@ -1249,7 +1945,7 @@ public:
 
 				if (NewTasks.Num() > 1)
 				{
-					TLockFreePointerList<FBaseGraphTask> TempSortedAnyThreadTasks;
+					TLockFreePointerListLIFO<FBaseGraphTask> TempSortedAnyThreadTasks;
 					for (int32 Index = 0 ; Index < NewTasks.Num() - 1; Index++) // we are going to take the last one for ourselves
 					{
 						TempSortedAnyThreadTasks.Push(NewTasks[Index]);
@@ -1266,43 +1962,18 @@ public:
 				}
 			}
 		} while (!IncomingAnyThreadTasks.IsEmpty() || !SortedAnyThreadTasks.IsEmpty());
-
-		// this can be called before my constructor is finished
-		for (int32 Pass = 0; Pass < 2; Pass++)
-		{
-			for (int32 Test = ThreadInNeed - 1; Test >= NumNamedThreads; Test--)
-			{
-				if (Pass || !Thread(Test).IsProbablyStalled())
-				{
-					FBaseGraphTask* Task = Thread(Test).RequestSteal();
-					if (Task)
-					{
-						return Task;
-					}
-				}
-			}
-			for (int32 Test = NumThreads - 1; Test > ThreadInNeed; Test--)
-			{
-				if (Pass || !Thread(Test).IsProbablyStalled())
-				{
-					FBaseGraphTask* Task = Thread(Test).RequestSteal();
-					if (Task)
-					{
-						return Task;
-					}
-				}
-			}
-		}
-		return NULL;
+		return nullptr;
 	}
 
+#endif
+
 	/** 
 	 *	Hint from a worker thread that it is stalling.
 	 *	@param	StallingThread; Id of the thread that is stalling.
 	**/
 	void NotifyStalling(ENamedThreads::Type StallingThread)
 	{
-		if (StallingThread >= NumNamedThreads)
+		if (StallingThread >= NumNamedThreads && !GFastSchedulerLatched)
 		{
 			StalledUnnamedThreads.Push(&Thread(StallingThread));
 		}
@@ -1317,11 +1988,11 @@ private:
 	 *	@param	Index; Id of the thread to retrieve.
 	 *	@return	Reference to the corresponding thread.
 	**/
-	FTaskThread& Thread(int32 Index)
+	FTaskThreadBase& Thread(int32 Index)
 	{
 		checkThreadGraph(Index >= 0 && Index < NumThreads);
-		checkThreadGraph(WorkerThreads[Index].TaskGraphWorker.GetThreadId() == Index);
-		return WorkerThreads[Index].TaskGraphWorker;
+		checkThreadGraph(WorkerThreads[Index].TaskGraphWorker->GetThreadId() == Index);
+		return *WorkerThreads[Index].TaskGraphWorker;
 	}
 
 	/** 
@@ -1330,7 +2001,6 @@ private:
 	**/
 	ENamedThreads::Type GetCurrentThread()
 	{
-		static_assert(offsetof(FWorkerThread,TaskGraphWorker)==0, "Task graph worker must be the first member.");
 		ENamedThreads::Type CurrentThreadIfKnown = ENamedThreads::AnyThread;
 		FWorkerThread* TLSPointer = (FWorkerThread*)FPlatformTLS::GetTlsValue(PerThreadIDTLSSlot);
 		if (TLSPointer)
@@ -1370,25 +2040,133 @@ private:
 	/** Index of TLS slot for FWorkerThread* pointer. **/
 	uint32				PerThreadIDTLSSlot;
 	/** Thread safe list of stalled thread "Hints". **/
-	TLockFreePointerList<FTaskThread>		StalledUnnamedThreads; 
-
-	TLockFreePointerList<FBaseGraphTask>		IncomingAnyThreadTasks;
-	TLockFreePointerList<FBaseGraphTask>		IncomingAnyThreadTasksHiPri;
-	TLockFreePointerList<FBaseGraphTask>		SortedAnyThreadTasks;
-	TLockFreePointerList<FBaseGraphTask>		SortedAnyThreadTasksHiPri;
+	TLockFreePointerListUnordered<FTaskThreadBase>		StalledUnnamedThreads; 
+
+#if USE_NEW_LOCK_FREE_LISTS
+	#if USE_INTRUSIVE_TASKQUEUES
+		FLockFreePointerListFIFOIntrusive<FBaseGraphTask>		IncomingAnyThreadTasks;
+		FLockFreePointerListFIFOIntrusive<FBaseGraphTask>		IncomingAnyThreadTasksHiPri;
+	#else
+		TLockFreePointerListFIFO<FBaseGraphTask>		IncomingAnyThreadTasks;
+		TLockFreePointerListFIFO<FBaseGraphTask>		IncomingAnyThreadTasksHiPri;
+	#endif
+#else
+	TLockFreePointerListLIFO<FBaseGraphTask>		IncomingAnyThreadTasks;
+	TLockFreePointerListLIFO<FBaseGraphTask>		IncomingAnyThreadTasksHiPri;
+	TLockFreePointerListLIFO<FBaseGraphTask>		SortedAnyThreadTasks;
+	TLockFreePointerListLIFO<FBaseGraphTask>		SortedAnyThreadTasksHiPri;
+#endif
 	FCriticalSection CriticalSectionForSortingIncomingAnyThreadTasks;
 	FCriticalSection CriticalSectionForSortingIncomingAnyThreadTasksHiPri;
+
+	// we use a single atomic to control the state of the anythread queues. This limits the maximum number of threads.
+	struct FAtomicStateBitfield
+	{
+#if !USE_NEW_LOCK_FREE_LISTS
+		uint32 QueueOwned: 1;
+		uint32 QueueOwnerIndex: 4;
+#endif
+		uint32 Stalled: FAtomicStateBitfield_MAX_THREADS;
+		uint32 Working: FAtomicStateBitfield_MAX_THREADS;
+
+		FAtomicStateBitfield()
+			: 
+#if !USE_NEW_LOCK_FREE_LISTS
+			QueueOwned(false)
+			, QueueOwnerIndex(0) ,
+#endif
+			 Stalled(0)
+			, Working(0)
+		{
+		}
+
+		static FORCEINLINE FAtomicStateBitfield InterlockedCompareExchange(FAtomicStateBitfield* Dest, FAtomicStateBitfield Exchange, FAtomicStateBitfield Comparand)
+		{
+			FAtomicStateBitfield Result;
+			*(int32*)&Result = FPlatformAtomics::InterlockedCompareExchange((volatile int32*)Dest, *(int32*)&Exchange, *(int32*)&Comparand);
+			return Result;			
+		}
+
+		static int32 CountBits(uint32 Bits)
+		{
+			MS_ALIGN(64) static uint8 Table[64] GCC_ALIGN(64) = {0};
+			if (!Table[63])
+			{
+				for (uint32 Index = 0; Index < 63; Index++)
+				{
+					uint32 LocalIndex = Index;
+					uint8 Result = 0;
+					while (LocalIndex)
+					{
+						if (LocalIndex & 1)
+						{
+							Result++;
+						}
+						LocalIndex >>= 1;
+					}
+					Table[Index] = Result;
+				}
+				FPlatformMisc::MemoryBarrier();
+				Table[63] = 6;
+			}
+			int32 Count = 0;
+			while (true)
+			{
+				Count += Table[Bits & 63];
+				if (!(Bits & ~63))
+				{
+					break;
+				}
+				Bits >>= 6;
+			}
+			return Count;
+		}
+
+		FORCEINLINE int32 NumberOfStalledThreads()
+		{
+			return CountBits(Stalled);
+		}
+
+		FORCEINLINE int32 NumberOfWorkingThreads()
+		{
+			return CountBits(Working);
+		}
+
+		bool operator==(FAtomicStateBitfield Other) const
+		{
+			return 
+#if !USE_NEW_LOCK_FREE_LISTS
+				QueueOwned == Other.QueueOwned &&
+				QueueOwnerIndex == Other.QueueOwnerIndex &&
+#endif
+				Stalled == Other.Stalled &&
+				Working == Other.Working;
+		}
+
+		bool operator!=(FAtomicStateBitfield Other) const
+		{
+			return 
+#if !USE_NEW_LOCK_FREE_LISTS
+				QueueOwned != Other.QueueOwned ||
+				QueueOwnerIndex != Other.QueueOwnerIndex ||
+#endif
+				Stalled != Other.Stalled ||
+				Working != Other.Working;
+		}
+	};
+
+	FAtomicStateBitfield AtomicForConsoleApproach;
 };
 
 
 // Implementations of FTaskThread function that require knowledge of FTaskGraphImplementation
 
-FBaseGraphTask* FTaskThread::FindWork()
+FBaseGraphTask* FTaskThreadAnyThread::FindWork()
 {
 	return FTaskGraphImplementation::Get().FindWork(ThreadId);
 }
 
-void FTaskThread::NotifyStalling()
+void FTaskThreadAnyThread::NotifyStalling()
 {
 	return FTaskGraphImplementation::Get().NotifyStalling(ThreadId);
 }
@@ -1436,6 +2214,133 @@ void FBaseGraphTask::LogPossiblyInvalidSubsequentsTask(const TCHAR* TaskName)
 }
 #endif
 
+#if USE_NEW_LOCK_FREE_LISTS
+
+#define USE_TLS_FOR_EVENTS (1)
+
+#define MONITOR_TASK_ALLOCATION (0)
+
+#if	MONITOR_TASK_ALLOCATION==1
+	typedef FThreadSafeCounter	FTaskLockFreeListCounter;
+#else
+	typedef FNoopCounter		FTaskLockFreeListCounter;
+#endif
+
+static FTaskLockFreeListCounter GraphEventsInUse;
+static FTaskLockFreeListCounter GraphEventsWithInlineStorageInUse;
+
+#if USE_TLS_FOR_EVENTS
+
+static TPointerSet_TLSCacheBase<FGraphEvent, TLockFreePointerListUnordered<FGraphEvent>, FTaskLockFreeListCounter> TheGraphEventAllocator;
+static TPointerSet_TLSCacheBase<FGraphEventAndSmallTaskStorage, TLockFreePointerListUnordered<FGraphEventAndSmallTaskStorage>, FTaskLockFreeListCounter> TheGraphEventWithInlineStorageAllocator;
+
+static FThreadSafeCounter GraphEventsAllocated;
+static FThreadSafeCounter GraphEventsWithInlineStorageAllocated;
+
+FGraphEvent* FGraphEvent::GetBundle(int32 NumBundle)
+{
+	GraphEventsAllocated.Add(NumBundle);
+	FGraphEvent* Root = nullptr;
+	uint8* Block = (uint8*)FMemory::Malloc(NumBundle * sizeof(FGraphEvent));
+	for (int32 Index = 0; Index < NumBundle; Index++)
+	{
+		FGraphEvent* Event = new ((void*)Block) FGraphEvent(false);
+		Event->LockFreePointerQueueNext = Root;
+		Root = Event;
+		Block += sizeof(FGraphEvent);
+	}
+	return Root;
+}
+
+FGraphEventAndSmallTaskStorage* FGraphEventAndSmallTaskStorage::GetBundle(int32 NumBundle)
+{
+	GraphEventsWithInlineStorageAllocated.Add(NumBundle);
+	FGraphEventAndSmallTaskStorage* Root = nullptr;
+	uint8* Block = (uint8*)FMemory::Malloc(NumBundle * sizeof(FGraphEventAndSmallTaskStorage));
+	for (int32 Index = 0; Index < NumBundle; Index++)
+	{
+		FGraphEventAndSmallTaskStorage* Event = new ((void*)Block) FGraphEventAndSmallTaskStorage();
+		Event->LockFreePointerQueueNext = Root;
+		Root = Event;
+		Block += sizeof(FGraphEventAndSmallTaskStorage);
+	}
+	return Root;
+}
+
+#else
+static FLockFreePointerListFIFOIntrusive<FGraphEvent> TheGraphEventAllocator;
+static FLockFreePointerListFIFOIntrusive<FGraphEventAndSmallTaskStorage> TheGraphEventWithInlineStorageAllocator;
+
+#endif
+
+FGraphEventRef FGraphEvent::CreateGraphEvent()
+{
+	GraphEventsInUse.Increment();
+	FGraphEvent* Used =  TheGraphEventAllocator.Pop();
+	if (Used)
+	{
+		FPlatformMisc::MemoryBarrier();
+		checkThreadGraph(!Used->SubsequentList.IsClosed());
+		checkThreadGraph(Used->ReferenceCount.GetValue() == 0);
+		checkThreadGraph(!Used->EventsToWaitFor.Num());
+		checkThreadGraph(!Used->bComplete);
+		checkThreadGraph(!Used->bInline);
+//		Used->ReferenceCount.Reset();  // temp
+//		Used->LockFreePointerQueueNext = nullptr; // temp
+		return FGraphEventRef(Used);
+	}
+	return FGraphEventRef(new FGraphEvent(false));
+}
+
+FGraphEvent* FGraphEvent::CreateGraphEventWithInlineStorage()
+{
+	GraphEventsWithInlineStorageInUse.Increment();
+	FGraphEventAndSmallTaskStorage* Used =  TheGraphEventWithInlineStorageAllocator.Pop();
+	if (Used)
+	{
+		FPlatformMisc::MemoryBarrier();
+		checkThreadGraph(!Used->SubsequentList.IsClosed());
+		checkThreadGraph(Used->ReferenceCount.GetValue() == 0);
+		checkThreadGraph(!Used->EventsToWaitFor.Num());
+		checkThreadGraph(!Used->bComplete);
+		checkThreadGraph(Used->bInline);
+	}
+	else
+	{
+		Used = new FGraphEventAndSmallTaskStorage();
+	}
+	checkThreadGraph(Used->bInline);
+	return Used;
+}
+
+
+void FGraphEvent::Recycle(FGraphEvent* ToRecycle)
+{
+	ToRecycle->Reset();
+	if (ToRecycle->bInline)
+	{
+		GraphEventsWithInlineStorageInUse.Decrement();
+		TheGraphEventWithInlineStorageAllocator.Push((FGraphEventAndSmallTaskStorage*)ToRecycle);
+	}
+	else
+	{
+		GraphEventsInUse.Decrement();
+		TheGraphEventAllocator.Push(ToRecycle);
+	}
+}
+
+
+void FGraphEvent::Reset()
+{
+//	ReferenceCount.Set(-5); // temp
+	SubsequentList.Reset();
+	checkThreadGraph(!SubsequentList.Pop());
+	EventsToWaitFor.Empty();  // we will let the memory block free here if there is one; these are not common
+	// checkThreadGraph(!LockFreePointerQueueNext); // temp
+	bComplete= false;
+}
+
+#else
 static TLockFreeClassAllocator_TLSCache<FGraphEvent> TheGraphEventAllocator;
 
 FGraphEventRef FGraphEvent::CreateGraphEvent()
@@ -1443,12 +2348,12 @@ FGraphEventRef FGraphEvent::CreateGraphEvent()
 	return TheGraphEventAllocator.New();
 }
 
-void FGraphEvent::DispatchSubsequents(ENamedThreads::Type CurrentThreadIfKnown)
+void FGraphEvent::Recycle(FGraphEvent* ToRecycle)
 {
-	TArray<FBaseGraphTask*> NewTasks;
-	NewTasks.Reset( 128 );
-	DispatchSubsequents( NewTasks, CurrentThreadIfKnown );
+	TheGraphEventAllocator.Free(ToRecycle);
 }
+#endif
+
 
 void FGraphEvent::DispatchSubsequents(TArray<FBaseGraphTask*>& NewTasks, ENamedThreads::Type CurrentThreadIfKnown)
 {
@@ -1466,6 +2371,24 @@ void FGraphEvent::DispatchSubsequents(TArray<FBaseGraphTask*>& NewTasks, ENamedT
 		return;
 	}
 
+#if USE_NEW_LOCK_FREE_LISTS
+	bComplete = true;
+	int32 SpinCount = 0;
+	while (true)
+	{
+		FBaseGraphTask* NewTask = SubsequentList.Pop();
+		if (!NewTask)
+		{
+			if (SubsequentList.CloseIfEmpty())
+			{
+				break;
+			}
+			LockFreeCriticalSpin(SpinCount);
+			continue;
+		}
+		NewTask->ConditionalQueueTask(CurrentThreadIfKnown);
+	}
+#else
 	SubsequentList.PopAllAndClose(NewTasks);
 	for (int32 Index = NewTasks.Num() - 1; Index >= 0 ; Index--) // reverse the order since PopAll is implicitly backwards
 	{
@@ -1474,11 +2397,7 @@ void FGraphEvent::DispatchSubsequents(TArray<FBaseGraphTask*>& NewTasks, ENamedT
 		NewTask->ConditionalQueueTask(CurrentThreadIfKnown);
 	}
 	NewTasks.Reset();
-}
-
-void FGraphEvent::Recycle(FGraphEvent* ToRecycle)
-{
-	TheGraphEventAllocator.Free(ToRecycle);
+#endif
 }
 
 FGraphEvent::~FGraphEvent()
@@ -1486,12 +2405,373 @@ FGraphEvent::~FGraphEvent()
 #if DO_CHECK
 	if (!IsComplete())
 	{
+#if USE_NEW_LOCK_FREE_LISTS
+		checkThreadGraph(!SubsequentList.Pop());
+#else
 		// Verifies that the event is completed. We do not allow events to die before completion.
 		TArray<FBaseGraphTask*> NewTasks;
 		SubsequentList.PopAllAndClose(NewTasks);
 		checkThreadGraph(!NewTasks.Num());
+#endif
 	}
 #endif
 	CheckDontCompleteUntilIsEmpty(); // We should not have any wait untils outstanding
 }
 
+
+
+// Benchmark
+
+#include "ParallelFor.h"
+
+static FORCEINLINE void DoWork(void* Hash, FThreadSafeCounter& Counter, FThreadSafeCounter& Cycles, int32 Work)
+{
+	if (Work > 0)
+	{
+		uint32 CyclesStart = FPlatformTime::Cycles();
+		Counter.Increment();
+		int32 Sum = 0;
+		for (int32 Index = 0; Index < Work; Index++)
+		{
+			Sum += PointerHash(((uint64*)Hash) + Index);
+		}
+		Cycles.Add(FPlatformTime::Cycles() - CyclesStart + (Sum & 1));
+	}
+	else if (Work == 0)
+	{
+		Counter.Increment();
+	}
+}
+
+class FIncGraphTask : public FCustomStatIDGraphTaskBase
+{
+public:
+	FORCEINLINE FIncGraphTask(FThreadSafeCounter& InCounter, FThreadSafeCounter& InCycles, int32 InWork)
+		: FCustomStatIDGraphTaskBase(TStatId())
+		, Counter(InCounter)
+		, Cycles(InCycles)
+		, Work(InWork)
+	{
+	}
+	static FORCEINLINE ENamedThreads::Type GetDesiredThread()
+	{
+		return ENamedThreads::AnyThread;
+	}
+
+	static FORCEINLINE ESubsequentsMode::Type GetSubsequentsMode() { return ESubsequentsMode::FireAndForget; }
+	void FORCEINLINE DoTask(ENamedThreads::Type CurrentThread, const FGraphEventRef& MyCompletionGraphEvent)
+	{
+		DoWork(this, Counter, Cycles, Work);
+	}
+private:
+	FThreadSafeCounter& Counter;
+	FThreadSafeCounter& Cycles;
+	int32 Work;
+};
+
+class FIncGraphTaskGT : public FIncGraphTask
+{
+public:
+	FORCEINLINE FIncGraphTaskGT(FThreadSafeCounter& InCounter, FThreadSafeCounter& InCycles, int32 InWork)
+		: FIncGraphTask(InCounter, InCycles, InWork)
+	{
+	}
+	static FORCEINLINE ENamedThreads::Type GetDesiredThread()
+	{
+		return ENamedThreads::GameThread_Local;
+	}
+};
+
+class FBoolGraphTask : public FCustomStatIDGraphTaskBase
+{
+public:
+	FORCEINLINE FBoolGraphTask(bool *InOut)
+		: FCustomStatIDGraphTaskBase(TStatId())
+		, Out(InOut)
+	{
+	}
+	static FORCEINLINE ENamedThreads::Type GetDesiredThread()
+	{
+		return ENamedThreads::AnyThread;
+	}
+	static FORCEINLINE ESubsequentsMode::Type GetSubsequentsMode() { return ESubsequentsMode::FireAndForget; }
+	void FORCEINLINE DoTask(ENamedThreads::Type CurrentThread, const FGraphEventRef& MyCompletionGraphEvent)
+	{
+		*Out = true;
+	}
+private:
+	bool *Out;
+};
+
+
+void PrintResult(double& StartTime, double& QueueTime, double& EndTime, double& JoinTime, FThreadSafeCounter& Counter, FThreadSafeCounter& Cycles, const TCHAR* Message)
+{
+	UE_LOG(LogConsoleResponse, Display, TEXT("Total %6.3fms   %6.3fms queue   %6.3fms join   %6.3fms wait   %6.3fms work   : %s")
+		, float(1000.0 * (EndTime-StartTime)), float(1000.0 * (QueueTime-StartTime)), float(1000.0 * (JoinTime-QueueTime)), float(1000.0 * (EndTime-JoinTime))
+		, float(FPlatformTime::GetSecondsPerCycle() * double(Cycles.GetValue()) * 1000.0)
+		, Message
+		);
+
+	Counter.Reset();
+	Cycles.Reset();
+	StartTime = 0.0;
+	QueueTime = 0.0;
+	EndTime = 0.0;
+	JoinTime = 0.0;
+}
+
+static void TaskGraphBenchmark(const TArray<FString>& Args)
+{
+	double StartTime, QueueTime, EndTime, JoinTime;
+	FThreadSafeCounter Counter;
+	FThreadSafeCounter Cycles;
+	
+	if (Args.Num() == 1 && Args[0] == TEXT("infinite"))
+	{
+		while (true)
+		{
+			{
+				StartTime = FPlatformTime::Seconds();
+
+				ParallelFor(1000, 
+					[&Counter, &Cycles](int32 Index)
+				{
+					TGraphTask<FIncGraphTaskGT>::CreateTask().ConstructAndDispatchWhenReady(Counter, Cycles, -1);
+				}
+				);
+				QueueTime = FPlatformTime::Seconds();
+				JoinTime = QueueTime;
+				FTaskGraphInterface::Get().ProcessThreadUntilIdle(ENamedThreads::GameThread_Local);
+				EndTime = FPlatformTime::Seconds();
+			}
+		}
+	}
+	{
+		StartTime = FPlatformTime::Seconds();
+		FGraphEventArray Tasks;
+		Tasks.Reserve(1000);
+		for (int32 Index = 0; Index < 1000; Index++)
+		{
+			Tasks.Emplace(TGraphTask<FNullGraphTask>::CreateTask(nullptr, ENamedThreads::GameThread).ConstructAndDispatchWhenReady(TStatId()));
+		}
+		QueueTime = FPlatformTime::Seconds();
+		FGraphEventRef Join = TGraphTask<FNullGraphTask>::CreateTask(&Tasks, ENamedThreads::GameThread).ConstructAndDispatchWhenReady(TStatId());
+		JoinTime = FPlatformTime::Seconds();
+		FTaskGraphInterface::Get().WaitUntilTaskCompletes(Join, ENamedThreads::GameThread_Local);
+		EndTime = FPlatformTime::Seconds();
+	}
+	PrintResult(StartTime, QueueTime, EndTime, JoinTime, Counter, Cycles, TEXT("1000 tasks, ordinary GT start"));
+	{
+		StartTime = FPlatformTime::Seconds();
+		FGraphEventArray Tasks;
+		Tasks.AddZeroed(1000);
+
+		ParallelFor(1000, 
+			[&Tasks](int32 Index)
+			{
+				Tasks[Index] = TGraphTask<FNullGraphTask>::CreateTask().ConstructAndDispatchWhenReady(TStatId());
+			}
+		);
+		QueueTime = FPlatformTime::Seconds();
+		FGraphEventRef Join = TGraphTask<FNullGraphTask>::CreateTask(&Tasks, ENamedThreads::GameThread).ConstructAndDispatchWhenReady(TStatId());
+		JoinTime = FPlatformTime::Seconds();
+		FTaskGraphInterface::Get().WaitUntilTaskCompletes(Join, ENamedThreads::GameThread_Local);
+		EndTime = FPlatformTime::Seconds();
+	}
+	PrintResult(StartTime, QueueTime, EndTime, JoinTime, Counter, Cycles, TEXT("1000 tasks, ParallelFor start"));
+	{
+		StartTime = FPlatformTime::Seconds();
+		FGraphEventArray Tasks;
+		Tasks.AddZeroed(10);
+
+		ParallelFor(10, 
+			[&Tasks](int32 Index)
+		{
+			FGraphEventArray InnerTasks;
+			InnerTasks.AddZeroed(100);
+			ENamedThreads::Type CurrentThread = FTaskGraphInterface::Get().GetCurrentThreadIfKnown();
+			for (int32 InnerIndex = 0; InnerIndex < 100; InnerIndex++)
+			{
+				InnerTasks[InnerIndex] = TGraphTask<FNullGraphTask>::CreateTask(nullptr, CurrentThread).ConstructAndDispatchWhenReady(TStatId());
+			}
+			// join the above tasks
+			Tasks[Index] = TGraphTask<FNullGraphTask>::CreateTask(&InnerTasks, CurrentThread).ConstructAndDispatchWhenReady(TStatId());
+		}
+		);
+		QueueTime = FPlatformTime::Seconds();
+		FGraphEventRef Join = TGraphTask<FNullGraphTask>::CreateTask(&Tasks, ENamedThreads::GameThread).ConstructAndDispatchWhenReady(TStatId());
+		JoinTime = FPlatformTime::Seconds();
+		FTaskGraphInterface::Get().WaitUntilTaskCompletes(Join, ENamedThreads::GameThread_Local);
+		EndTime = FPlatformTime::Seconds();
+	}
+	PrintResult(StartTime, QueueTime, EndTime, JoinTime, Counter, Cycles, TEXT("1000 tasks, ParallelFor start, batched completion 10x100"));
+	{
+		StartTime = FPlatformTime::Seconds();
+		FGraphEventArray Tasks;
+		Tasks.AddZeroed(100);
+
+		ParallelFor(100, 
+			[&Tasks](int32 Index)
+		{
+			FGraphEventArray InnerTasks;
+			InnerTasks.AddZeroed(10);
+			ENamedThreads::Type CurrentThread = FTaskGraphInterface::Get().GetCurrentThreadIfKnown();
+			for (int32 InnerIndex = 0; InnerIndex < 10; InnerIndex++)
+			{
+				InnerTasks[InnerIndex] = TGraphTask<FNullGraphTask>::CreateTask(nullptr, CurrentThread).ConstructAndDispatchWhenReady(TStatId());
+			}
+			// join the above tasks
+			Tasks[Index] = TGraphTask<FNullGraphTask>::CreateTask(&InnerTasks, CurrentThread).ConstructAndDispatchWhenReady(TStatId());
+		}
+		);
+		QueueTime = FPlatformTime::Seconds();
+		FGraphEventRef Join = TGraphTask<FNullGraphTask>::CreateTask(&Tasks, ENamedThreads::GameThread).ConstructAndDispatchWhenReady(TStatId());
+		JoinTime = FPlatformTime::Seconds();
+		FTaskGraphInterface::Get().WaitUntilTaskCompletes(Join, ENamedThreads::GameThread_Local);
+		EndTime = FPlatformTime::Seconds();
+	}
+	PrintResult(StartTime, QueueTime, EndTime, JoinTime, Counter, Cycles, TEXT("1000 tasks, ParallelFor start, batched completion 100x10"));
+
+	{
+		StartTime = FPlatformTime::Seconds();
+
+		ParallelFor(1000, 
+			[&Counter, &Cycles](int32 Index)
+			{
+				TGraphTask<FIncGraphTask>::CreateTask().ConstructAndDispatchWhenReady(Counter, Cycles, 0);
+			}
+		);
+		QueueTime = FPlatformTime::Seconds();
+		JoinTime = QueueTime;
+		while (Counter.GetValue() < 1000)
+		{
+			FPlatformMisc::MemoryBarrier();
+		}
+		EndTime = FPlatformTime::Seconds();
+	}
+	PrintResult(StartTime, QueueTime, EndTime, JoinTime, Counter, Cycles, TEXT("1000 tasks, ParallelFor, counter tracking"));
+
+	{
+		StartTime = FPlatformTime::Seconds();
+
+		static bool Output[1000];
+		FPlatformMemory::Memzero(Output, 1000);
+
+		ParallelFor(1000, 
+			[](int32 Index)
+		{
+			TGraphTask<FBoolGraphTask>::CreateTask().ConstructAndDispatchWhenReady(Output + Index);
+		}
+		);
+		QueueTime = FPlatformTime::Seconds();
+		JoinTime = QueueTime;
+		for (int32 Index = 0; Index < 1000; Index++)
+		{
+			while (!Output[Index])
+			{
+				FPlatformProcess::Sleep(0.0f);
+			}
+		}
+		EndTime = FPlatformTime::Seconds();
+	}
+	PrintResult(StartTime, QueueTime, EndTime, JoinTime, Counter, Cycles, TEXT("1000 tasks, ParallelFor, bool* tracking"));
+
+	{
+		StartTime = FPlatformTime::Seconds();
+
+		ParallelFor(1000, 
+			[&Counter, &Cycles](int32 Index)
+		{
+			TGraphTask<FIncGraphTask>::CreateTask().ConstructAndDispatchWhenReady(Counter, Cycles, 1000);
+		}
+		);
+		QueueTime = FPlatformTime::Seconds();
+		JoinTime = QueueTime;
+		while (Counter.GetValue() < 1000)
+		{
+			FPlatformProcess::Sleep(0.0f);
+		}
+		EndTime = FPlatformTime::Seconds();
+	}
+	PrintResult(StartTime, QueueTime, EndTime, JoinTime, Counter, Cycles, TEXT("1000 tasks, ParallelFor, counter tracking, with work"));
+	{
+		StartTime = FPlatformTime::Seconds();
+		for (int32 Index = 0; Index < 1000; Index++)
+		{
+			TGraphTask<FIncGraphTask>::CreateTask(nullptr, ENamedThreads::GameThread).ConstructAndDispatchWhenReady(Counter, Cycles, 1000);
+		}
+		QueueTime = FPlatformTime::Seconds();
+		JoinTime = QueueTime;
+		while (Counter.GetValue() < 1000)
+		{
+			FPlatformProcess::Sleep(0.0f);
+		}
+		EndTime = FPlatformTime::Seconds();
+	}
+	PrintResult(StartTime, QueueTime, EndTime, JoinTime, Counter, Cycles, TEXT("1000 tasks, GT submit, counter tracking, with work"));
+
+	{
+		StartTime = FPlatformTime::Seconds();
+
+		ParallelFor(1000, 
+			[&Counter, &Cycles](int32 Index)
+		{
+			TGraphTask<FIncGraphTaskGT>::CreateTask().ConstructAndDispatchWhenReady(Counter, Cycles, -1);
+		}
+		);
+		QueueTime = FPlatformTime::Seconds();
+		JoinTime = QueueTime;
+		FTaskGraphInterface::Get().ProcessThreadUntilIdle(ENamedThreads::GameThread_Local);
+		EndTime = FPlatformTime::Seconds();
+	}
+	PrintResult(StartTime, QueueTime, EndTime, JoinTime, Counter, Cycles, TEXT("1000 GT tasks, ParallelFor, no tracking (none needed)"));
+
+
+	{
+		StartTime = FPlatformTime::Seconds();
+		QueueTime = StartTime;
+		JoinTime = QueueTime;
+		ParallelFor(1000, 
+			[&Counter, &Cycles](int32 Index)
+			{
+				DoWork(&Counter, Counter, Cycles, -1);
+			}
+		);
+		EndTime = FPlatformTime::Seconds();
+	}
+	PrintResult(StartTime, QueueTime, EndTime, JoinTime, Counter, Cycles, TEXT("1000 element do-nothing ParallelFor"));
+	{
+		StartTime = FPlatformTime::Seconds();
+		QueueTime = StartTime;
+		JoinTime = QueueTime;
+		ParallelFor(1000, 
+			[&Counter, &Cycles](int32 Index)
+			{
+				DoWork(&Counter, Counter, Cycles, 1000);
+			}
+		);
+		EndTime = FPlatformTime::Seconds();
+	}
+	PrintResult(StartTime, QueueTime, EndTime, JoinTime, Counter, Cycles, TEXT("1000 element ParallelFor, with work"));
+
+	{
+		StartTime = FPlatformTime::Seconds();
+		QueueTime = StartTime;
+		JoinTime = QueueTime;
+		ParallelFor(1000, 
+			[&Counter, &Cycles](int32 Index)
+			{
+				DoWork(&Counter, Counter, Cycles, 1000);
+			},
+			true
+		);
+		EndTime = FPlatformTime::Seconds();
+	}
+	PrintResult(StartTime, QueueTime, EndTime, JoinTime, Counter, Cycles, TEXT("1000 element ParallelFor, single threaded, with work"));
+}
+
+static FAutoConsoleCommand TaskGraphBenchmarkCmd(
+	TEXT("TaskGraph.Benchmark"),
+	TEXT("Prints the time to run 1000 no-op tasks."),
+	FConsoleCommandWithArgsDelegate::CreateStatic(&TaskGraphBenchmark)
+	);
diff --git a/Engine/Source/Runtime/Core/Private/Containers/LockFreeList.cpp b/Engine/Source/Runtime/Core/Private/Containers/LockFreeList.cpp
--- a/Engine/Source/Runtime/Core/Private/Containers/LockFreeList.cpp
+++ b/Engine/Source/Runtime/Core/Private/Containers/LockFreeList.cpp
@@ -6,9 +6,6 @@
 
 DEFINE_LOG_CATEGORY(LogLockFreeList);
 
-FThreadSafeCounter FLockFreeListStats::NumOperations;
-FThreadSafeCounter FLockFreeListStats::Cycles;
-
 
 #if USE_LOCKFREELIST_128
 
@@ -88,5 +85,58 @@ FLockFreeVoidPointerListBase::FLinkAllocator::~FLinkAllocator()
 #endif // 0
 }
 
+FLockFreePointerQueueBaseLinkAllocator& FLockFreePointerQueueBaseLinkAllocator::Get()
+{
+	static FLockFreePointerQueueBaseLinkAllocator Singleton;
+	return Singleton;
+}
+
+static FThreadSafeCounter Sleeps;
+static FThreadSafeCounter BigSleeps;
+
+void LockFreeCriticalSpin(int32& SpinCount)
+{
+	SpinCount++;
+	if (SpinCount > 256)
+	{
+		FPlatformProcess::Sleep(0.00001f);
+		BigSleeps.Increment();
+	}
+	else if (SpinCount > 32)
+	{
+		FPlatformProcess::Sleep(0.0f);
+		if (Sleeps.Increment() % 10000 == 9999)
+		{
+			UE_LOG(LogTemp, Warning, TEXT("sleep0 %d sleep 10us %d"), Sleeps.GetValue() + 1, BigSleeps.GetValue());
+		}
+	}
+	else if (SpinCount > 8)
+	{
+		FPlatformMisc::MemoryBarrier();
+	}
+}
+
+#if !UE_BUILD_SHIPPING && !UE_BUILD_TEST
+
+static FThreadSafeCounter TestStall;
+
+void DoTestCriticalStall()
+{
+	if (TestStall.Increment() % 9713 == 9712)
+	{
+		FPlatformProcess::Sleep(0.000001f);
+	}
+}
+
+int32 GTestCriticalStalls = 0;
+static FAutoConsoleVariableRef CVarTestCriticalLockFree(
+	TEXT("TaskGraph.TestCriticalLockFree"),
+	GTestCriticalStalls,
+	TEXT("If > 0, then we sleep periodically at critical points in the lock free lists. Threads must not starve...this will encourage them to starve at the right place to find livelocks."),
+	ECVF_Cheat
+	);
+
+#endif
+
 #endif //USE_LOCKFREELIST_128
 
diff --git a/Engine/Source/Runtime/Core/Private/Misc/EventPool.h b/Engine/Source/Runtime/Core/Private/Misc/EventPool.h
--- a/Engine/Source/Runtime/Core/Private/Misc/EventPool.h
+++ b/Engine/Source/Runtime/Core/Private/Misc/EventPool.h
@@ -88,5 +88,5 @@ public:
 private:
 
 	/** Holds the collection of recycled events. */
-	TLockFreePointerList<FEvent> Pool;
+	TLockFreePointerListUnordered<FEvent> Pool;
 };
diff --git a/Engine/Source/Runtime/Core/Public/Async/ParallelFor.h b/Engine/Source/Runtime/Core/Public/Async/ParallelFor.h
--- a/Engine/Source/Runtime/Core/Public/Async/ParallelFor.h
+++ b/Engine/Source/Runtime/Core/Public/Async/ParallelFor.h
@@ -13,29 +13,49 @@
 struct FParallelForData
 {
 	int32 Num;
+	int32 BlockSize;
+	int32 LastBlockExtraNum;
 	TFunctionRef<void(int32)> Body;
 	FEvent* Event;
 	FThreadSafeCounter IndexToDo;
 	FThreadSafeCounter NumCompleted;
 	bool bExited;
 	bool bTriggered;
-	FParallelForData(int32 InNum, TFunctionRef<void(int32)> InBody)
-		: Num(InNum)
-		, Body(InBody)
+	bool bSaveLastBlockForMaster;
+	FParallelForData(int32 InTotalNum, int32 InNumThreads, bool bInSaveLastBlockForMaster, TFunctionRef<void(int32)> InBody)
+		: Body(InBody)
 		, Event(FPlatformProcess::GetSynchEventFromPool(false))
 		, bExited(false)
 		, bTriggered(false)
+		, bSaveLastBlockForMaster(bInSaveLastBlockForMaster)
 	{
+		check(InTotalNum >= InNumThreads);
+		BlockSize = 0;
+		Num = 0;
+		for (int32 Div = 3; Div; Div--)
+		{
+			BlockSize = InTotalNum / (InNumThreads * Div);
+			if (BlockSize)
+			{
+				Num = InTotalNum / BlockSize;
+				if (Num >= InNumThreads + !!bSaveLastBlockForMaster)
+				{
+					break;
+				}
+			}
+		}
+		check(BlockSize && Num);
+		LastBlockExtraNum = InTotalNum - Num * BlockSize;
+		check(LastBlockExtraNum >= 0);
 	}
 	~FParallelForData()
 	{
-		FPlatformMisc::MemoryBarrier();
 		check(IndexToDo.GetValue() >= Num);
 		check(NumCompleted.GetValue() == Num);
 		check(bExited);
 		FPlatformProcess::ReturnSynchEventToPool(Event);
 	}
-	bool Process(int32 TasksToSpawn, TSharedRef<FParallelForData, ESPMode::ThreadSafe>& Data);
+	bool Process(int32 TasksToSpawn, TSharedRef<FParallelForData, ESPMode::ThreadSafe>& Data, bool bMaster);
 };
 
 class FParallelForTask
@@ -62,16 +82,16 @@ public:
 	}
 	void DoTask(ENamedThreads::Type CurrentThread, const FGraphEventRef& MyCompletionGraphEvent)
 	{
-		if (Data->Process(TasksToSpawn, Data))
+		if (Data->Process(TasksToSpawn, Data, false))
 		{
-			check(!Data->bTriggered);
+			checkSlow(!Data->bTriggered);
 			Data->bTriggered = true;
 			Data->Event->Trigger();
 		}
 	}
 };
 
-inline bool FParallelForData::Process(int32 TasksToSpawn, TSharedRef<FParallelForData, ESPMode::ThreadSafe>& Data)
+inline bool FParallelForData::Process(int32 TasksToSpawn, TSharedRef<FParallelForData, ESPMode::ThreadSafe>& Data, bool bMaster)
 {
 	int32 MaybeTasksLeft = Num - IndexToDo.GetValue();
 	if (TasksToSpawn && MaybeTasksLeft > 0)
@@ -79,23 +99,44 @@ inline bool FParallelForData::Process(int32 TasksToSpawn, TSharedRef<FParallelFo
 		TasksToSpawn = FMath::Min<int32>(TasksToSpawn, MaybeTasksLeft);
 		TGraphTask<FParallelForTask>::CreateTask().ConstructAndDispatchWhenReady(Data, TasksToSpawn - 1);		
 	}
+	int32 LocalBlockSize = BlockSize;
+	int32 LocalNum = Num;
+	bool bLocalSaveLastBlockForMaster = bSaveLastBlockForMaster;
+	TFunctionRef<void(int32)> LocalBody(Body);
 	while (true)
 	{
 		int32 MyIndex = IndexToDo.Increment() - 1;
-		if (MyIndex < Num)
+		if (bLocalSaveLastBlockForMaster)
 		{
-			Body(MyIndex);
-			check(!bExited);
+			if (!bMaster && MyIndex >= LocalNum - 1)
+			{
+				break; // leave the last block for the master, hoping to avoid an event
+			}
+			else if (bMaster && MyIndex > LocalNum - 1)
+			{
+				MyIndex = LocalNum - 1; // I am the master, I need to take this block, hoping to avoid an event
+			}
+		}
+		if (MyIndex < LocalNum)
+		{
+			int32 ThisBlockSize = LocalBlockSize;
+			if (MyIndex == LocalNum - 1)
+			{
+				ThisBlockSize += LastBlockExtraNum;
+			}
+			for (int32 LocalIndex = 0; LocalIndex < ThisBlockSize; LocalIndex++)
+			{
+				LocalBody(MyIndex * LocalBlockSize + LocalIndex);
+			}
+			checkSlow(!bExited);
 			int32 LocalNumCompleted = NumCompleted.Increment();
-			if (LocalNumCompleted == Num)
+			if (LocalNumCompleted == LocalNum)
 			{
-				FPlatformMisc::MemoryBarrier();
-				check(IndexToDo.GetValue() >= Num);
 				return true;
 			}
-			check(LocalNumCompleted < Num);
+			checkSlow(LocalNumCompleted < LocalNum);
 		}
-		else
+		if (MyIndex >= LocalNum - 1)
 		{
 			break;
 		}
@@ -129,23 +170,20 @@ inline void ParallelFor(int32 Num, TFunctionRef<void(int32)> Body, bool bForceSi
 		}
 		return;
 	}
-	check(Num);
-	FParallelForData* DataPtr = new FParallelForData(Num, Body);
+	FParallelForData* DataPtr = new FParallelForData(Num, AnyThreadTasks + 1, Num > AnyThreadTasks + 1, Body);
 	TSharedRef<FParallelForData, ESPMode::ThreadSafe> Data = MakeShareable(DataPtr);
 	TGraphTask<FParallelForTask>::CreateTask().ConstructAndDispatchWhenReady(Data, AnyThreadTasks - 1);		
 	// this thread can help too and this is important to prevent deadlock on recursion 
-	if (!Data->Process(0, Data))
+	if (!Data->Process(0, Data, true))
 	{
-		FPlatformMisc::MemoryBarrier();
 		Data->Event->Wait();
 		check(Data->bTriggered);
 	}
 	else
 	{
-		FPlatformMisc::MemoryBarrier();
 		check(!Data->bTriggered);
 	}
-	check(Data->NumCompleted.GetValue() == Num);
+	check(Data->NumCompleted.GetValue() == Data->Num);
 	Data->bExited = true;
 	// DoneEvent waits here if some other thread finishes the last item
 	// Data must live on until all of the tasks are cleared which might be long after this function exits
@@ -180,24 +218,22 @@ inline void ParallelForWithPreWork(int32 Num, TFunctionRef<void(int32)> Body, TF
 		return;
 	}
 	check(Num);
-	FParallelForData* DataPtr = new FParallelForData(Num, Body);
+	FParallelForData* DataPtr = new FParallelForData(Num, AnyThreadTasks, false, Body);
 	TSharedRef<FParallelForData, ESPMode::ThreadSafe> Data = MakeShareable(DataPtr);
 	TGraphTask<FParallelForTask>::CreateTask().ConstructAndDispatchWhenReady(Data, AnyThreadTasks - 1);		
 	// do the prework
 	CurrentThreadWorkToDoBeforeHelping();
 	// this thread can help too and this is important to prevent deadlock on recursion 
-	if (!Data->Process(0, Data))
+	if (!Data->Process(0, Data, true))
 	{
-		FPlatformMisc::MemoryBarrier();
 		Data->Event->Wait();
 		check(Data->bTriggered);
 	}
 	else
 	{
-		FPlatformMisc::MemoryBarrier();
 		check(!Data->bTriggered);
 	}
-	check(Data->NumCompleted.GetValue() == Num);
+	check(Data->NumCompleted.GetValue() == Data->Num);
 	Data->bExited = true;
 	// Data must live on until all of the tasks are cleared which might be long after this function exits
 }
diff --git a/Engine/Source/Runtime/Core/Public/Async/TaskGraphInterfaces.h b/Engine/Source/Runtime/Core/Public/Async/TaskGraphInterfaces.h
--- a/Engine/Source/Runtime/Core/Public/Async/TaskGraphInterfaces.h
+++ b/Engine/Source/Runtime/Core/Public/Async/TaskGraphInterfaces.h
@@ -17,6 +17,7 @@
 
 // what level of checking to perform...normally checkSlow but could be ensure or check
 #define checkThreadGraph checkSlow
+//#define checkThreadGraph(x) ((x)||((*(char*)3) = 0))
 
 DECLARE_STATS_GROUP(TEXT("Task Graph Tasks"), STATGROUP_TaskGraphTasks, STATCAT_Advanced);
 
@@ -231,9 +232,19 @@ public:
  *	Base class for all tasks. 
  *	Tasks go through a very specific life stage progression, and this is verified.
  **/
+
+#define USE_VIRTUAL_BYPASS USE_NEW_LOCK_FREE_LISTS
+
 class FBaseGraphTask
 {
 public:
+#if USE_VIRTUAL_BYPASS
+	typedef void (*TExecuteTask)(FBaseGraphTask*, TArray<FBaseGraphTask*>& NewTasks, ENamedThreads::Type CurrentThread);
+#endif
+
+#if USE_NEW_LOCK_FREE_LISTS
+	FBaseGraphTask* LockFreePointerQueueNext;
+#endif
 	// Allocator for small tasks.
 	enum
 	{
@@ -246,10 +257,18 @@ protected:
 	 *	Constructor
 	 *	@param InNumberOfPrerequistitesOutstanding; the number of prerequisites outstanding. We actually add one to this to prevent the task from firing while we are setting up the task
 	 **/
-	FBaseGraphTask(int32 InNumberOfPrerequistitesOutstanding)
+	FBaseGraphTask(
+#if USE_VIRTUAL_BYPASS
+		TExecuteTask InExecuteTaskPtr,
+#endif
+		int32 InNumberOfPrerequistitesOutstanding
+		)
 		: ThreadToExecuteOn(ENamedThreads::AnyThread)
 		, NumberOfPrerequistitesOutstanding(InNumberOfPrerequistitesOutstanding + 1) // + 1 is not a prerequisite, it is a lock to prevent it from executing while it is getting prerequisites, one it is safe to execute, call PrerequisitesComplete
 	{
+#if USE_VIRTUAL_BYPASS
+		ExecuteTaskPtr = InExecuteTaskPtr;
+#endif
 		checkThreadGraph(LifeStage.Increment() == int32(LS_Contructed));
 	}
 	/** 
@@ -277,7 +296,12 @@ protected:
 		}
 	}
 	/** destructor, just checks the life stage **/
-	virtual ~FBaseGraphTask()
+#if USE_VIRTUAL_BYPASS
+	FORCEINLINE
+#else
+	virtual 
+#endif
+		~FBaseGraphTask()
 	{
 		checkThreadGraph(LifeStage.Increment() == int32(LS_Deconstucted));
 	}
@@ -303,7 +327,9 @@ protected:
 	}
 
 private:
-	friend class FTaskThread;
+	friend class FNamedTaskThread;
+	friend class FTaskThreadBase;
+	friend class FTaskThreadAnyThread;
 	friend class FGraphEvent;
 	friend class FTaskGraphImplementation;
 
@@ -313,7 +339,9 @@ private:
 	 *	Virtual call to actually execute the task. This should also call the destructor and free any memory.
 	 *	@param CurrentThread; provides the index of the thread we are running on. This is handy for submitting new taks.
 	 **/
+#if !USE_VIRTUAL_BYPASS
 	virtual void ExecuteTask(TArray<FBaseGraphTask*>& NewTasks, ENamedThreads::Type CurrentThread)=0;
+#endif
 
 	// API called from other parts of the system
 
@@ -322,10 +350,14 @@ private:
 	 *	Just checks the life stage and passes off to the virtual ExecuteTask method.
 	 *	@param CurrentThread; provides the index of the thread we are running on. This is handy for submitting new taks.
 	 **/
-	void Execute(TArray<FBaseGraphTask*>& NewTasks, ENamedThreads::Type CurrentThread)
+	FORCEINLINE void Execute(TArray<FBaseGraphTask*>& NewTasks, ENamedThreads::Type CurrentThread)
 	{
 		checkThreadGraph(LifeStage.Increment() == int32(LS_Executing));
+#if USE_VIRTUAL_BYPASS
+		ExecuteTaskPtr(this, NewTasks, CurrentThread);
+#else
 		ExecuteTask(NewTasks, CurrentThread);
+#endif
 	}
 
 	// Internal Use
@@ -340,13 +372,16 @@ private:
 		FTaskGraphInterface::Get().QueueTask(this, ThreadToExecuteOn, CurrentThreadIfKnown);
 	}
 
+#if USE_VIRTUAL_BYPASS
+	TExecuteTask ExecuteTaskPtr;
+#endif
 	/**	Thread to execute on, can be ENamedThreads::AnyThread to execute on any unnamed thread **/
 	ENamedThreads::Type			ThreadToExecuteOn;
 	/**	Number of prerequisites outstanding. When this drops to zero, the thread is queued for execution.  **/
 	FThreadSafeCounter			NumberOfPrerequistitesOutstanding; 
 
 
-#if !(UE_BUILD_SHIPPING || UE_BUILD_TEST)
+#if !UE_BUILD_SHIPPING
 	// Life stage verification
 	// Tasks go through 8 steps, in order. In non-final builds, we track them with a thread safe counter and verify that the progression is correct.
 	enum ELifeStage
@@ -379,12 +414,20 @@ public:
 	**/
 	static CORE_API FGraphEventRef CreateGraphEvent();
 
+	// the returned event will have ref count zero; be sure to add one!
+	static CORE_API FGraphEvent* CreateGraphEventWithInlineStorage();
 	/**
 	 *	Attempts to a new subsequent task. If this event has already fired, false is returned and action must be taken to ensure that the task will still fire even though this event cannot be a prerequisite (because it is already finished).
 	 *	@return true if the task was successfully set up as a subsequent. false if the event has already fired.
 	**/
 	bool AddSubsequent(class FBaseGraphTask* Task)
 	{
+#if USE_NEW_LOCK_FREE_LISTS
+		if (bComplete)
+		{
+			return false;
+		}
+#endif
 		return SubsequentList.PushIfNotClosed(Task);
 	}
 
@@ -411,7 +454,6 @@ public:
 	 *	"Complete" the event. This grabs the list of subsequents and atomically closes it. Then for each subsequent it reduces the number of prerequisites outstanding and if that drops to zero, the task is queued.
 	 *	@param CurrentThreadIfKnown if the current thread is known, provide it here. Otherwise it will be determined via TLS if any task ends up being queued.
 	**/
-	CORE_API void DispatchSubsequents(ENamedThreads::Type CurrentThreadIfKnown = ENamedThreads::AnyThread);
 	CORE_API void DispatchSubsequents(TArray<FBaseGraphTask*>& NewTasks, ENamedThreads::Type CurrentThreadIfKnown = ENamedThreads::AnyThread);
 
 	/**
@@ -421,8 +463,19 @@ public:
 	**/
 	bool IsComplete() const
 	{
+#if USE_NEW_LOCK_FREE_LISTS
+		FPlatformMisc::MemoryBarrier();
+		return bComplete;
+#else
 		return SubsequentList.IsClosed();
+#endif
 	}
+
+#if USE_NEW_LOCK_FREE_LISTS
+	void* GetInlineStorage();
+	static CORE_API FGraphEvent* GetBundle(int32 NumBundle);
+#endif
+
 private:
 	friend class TRefCountPtr<FGraphEvent>;
 	friend class TLockFreeClassAllocator_TLSCache<FGraphEvent>;
@@ -436,7 +489,13 @@ private:
 	/**
 	 *	Hidden Constructor
 	**/
-	FGraphEvent()
+	friend struct FGraphEventAndSmallTaskStorage;
+	FGraphEvent(bool bInInline = false)
+#if USE_NEW_LOCK_FREE_LISTS
+		: bComplete(false)
+		, bInline(bInInline)
+		, LockFreePointerQueueNext(nullptr)
+#endif
 	{
 	}
 
@@ -454,7 +513,9 @@ public:
 	**/
 	uint32 AddRef()
 	{
-		return ReferenceCount.Increment();
+		int32 RefCount = ReferenceCount.Increment();
+		checkThreadGraph(RefCount > 0);
+		return RefCount;
 	}
 	/** 
 	 *	Decreases the reference count and destroys the graph event if it is zero.
@@ -471,26 +532,54 @@ public:
 		return RefCount;
 	}
 
-#if !(UE_BUILD_SHIPPING || UE_BUILD_TEST)
-	/**
-	 * Checks if the pointers look 'ok'. Used for debugging (temporarily)
-	 */
-	bool CheckSubsequents() const
-	{
-		return SubsequentList.CheckPointers();
-	}
-#endif
 
 private:
 
-	/** Number of outstanding references to this graph event **/
-	FThreadSafeCounter									ReferenceCount;
 	/** Threadsafe list of subsequents for the event **/
-	TClosableLockFreePointerList<FBaseGraphTask>		SubsequentList;
+	TClosableLockFreePointerListUnorderedSingleConsumer<FBaseGraphTask>		SubsequentList;
 	/** List of events to wait for until firing. This is not thread safe as it is only legal to fill it in within the context of an executing task. **/
-	FGraphEventArray									EventsToWaitFor;
+	FGraphEventArray														EventsToWaitFor;
+	/** Number of outstanding references to this graph event **/
+	FThreadSafeCounter														ReferenceCount;
+#if USE_NEW_LOCK_FREE_LISTS
+	void Reset();
+	bool bComplete;
+	bool bInline;
+public:
+	void *LockFreePointerQueueNext;
+#endif
+};
+
+#if USE_NEW_LOCK_FREE_LISTS
+struct FGraphEventAndSmallTaskStorage : public FGraphEvent
+{
+	enum
+	{
+		StorageAlignment = 16,
+		StorageSize = 256 - StorageAlignment * ((sizeof(FGraphEvent) + StorageAlignment - 1) / StorageAlignment)
+	};
+	TAlignedBytes<StorageSize, 16> TaskStorage;
+
+	FGraphEventAndSmallTaskStorage()
+		: FGraphEvent(true)
+	{
+	}
+
+	static CORE_API FGraphEventAndSmallTaskStorage* GetBundle(int32 NumBundle);
 };
 
+FORCEINLINE void* FGraphEvent::GetInlineStorage()
+{
+	if (!bInline)
+	{
+		return nullptr;
+	}
+	return &(((FGraphEventAndSmallTaskStorage *)this)->TaskStorage);
+}
+
+#endif
+
+
 /** 
  The user defined task type can take arguments to a constructor. These arguments (unfortunately) must not be references.
  The API required of TTask:
@@ -728,6 +817,16 @@ public:
 	static FConstructor CreateTask(const FGraphEventArray* Prerequisites = NULL, ENamedThreads::Type CurrentThreadIfKnown = ENamedThreads::AnyThread)
 	{
 		int32 NumPrereq = Prerequisites ? Prerequisites->Num() : 0;
+#if USE_NEW_LOCK_FREE_LISTS
+		if (sizeof(TGraphTask) <= FGraphEventAndSmallTaskStorage::StorageSize && ALIGNOF(TGraphTask) <= FGraphEventAndSmallTaskStorage::StorageAlignment && TTask::GetSubsequentsMode() != ESubsequentsMode::FireAndForget)
+		{
+			FGraphEventRef Event(FGraphEvent::CreateGraphEventWithInlineStorage());
+			void *Mem = Event->GetInlineStorage();
+			checkThreadGraph(Mem);
+			return FConstructor(new (Mem) TGraphTask(Event, NumPrereq), Prerequisites, CurrentThreadIfKnown);
+		}
+		else
+#endif
 		if (sizeof(TGraphTask) <= FBaseGraphTask::SMALL_TASK_SIZE)
 		{
 			void *Mem = FBaseGraphTask::GetSmallTaskAllocator().Allocate();
@@ -760,7 +859,15 @@ private:
 	 *	Dispatches the subsequents.
 	 *	Destroys myself.
 	 **/
-	virtual void ExecuteTask(TArray<FBaseGraphTask*>& NewTasks, ENamedThreads::Type CurrentThread)
+#if USE_VIRTUAL_BYPASS
+	static void ExecuteTask(FBaseGraphTask *This, TArray<FBaseGraphTask*>& NewTasks, ENamedThreads::Type CurrentThread)
+	{
+		((TGraphTask*)This)->ExecuteTaskInner(NewTasks, CurrentThread);
+	}
+	FORCEINLINE void ExecuteTaskInner(TArray<FBaseGraphTask*>& NewTasks, ENamedThreads::Type CurrentThread)
+#else 
+	virtual void ExecuteTask(TArray<FBaseGraphTask*>& NewTasks, ENamedThreads::Type CurrentThread) final override
+#endif
 	{
 		checkThreadGraph(TaskConstructed);
 
@@ -777,23 +884,6 @@ private:
 		{
 			FScopeCycleCounter Scope(Task.GetStatId(), true); 
 			Task.DoTask(CurrentThread, Subsequents);
-#if (UE_BUILD_DEBUG || UE_BUILD_DEBUGGAME) && 0 // enable this to find nasty overwrite bugs
-			
-			if (TTask::GetSubsequentsMode() == ESubsequentsMode::TrackSubsequents)
-			{
-				if (!Subsequents->CheckSubsequents())
-				{
-					LogPossiblyInvalidSubsequentsTask(
-#if STATS
-						*Task.GetStatId().GetName().ToString()
-#else
-						TEXT("Unknown")
-#endif
-					);
-				}
-			}	
-#endif
-
 			Task.~TTask();
 			checkThreadGraph(ENamedThreads::GetThreadIndex(CurrentThread) <= ENamedThreads::RenderThread || FMemStack::Get().IsEmpty()); // you must mark and pop memstacks if you use them in tasks! Named threads are excepted.
 		}
@@ -806,9 +896,18 @@ private:
 			Subsequents->DispatchSubsequents(NewTasks, CurrentThread);
 		}
 
+#if USE_NEW_LOCK_FREE_LISTS
+		FGraphEventRef LocalSubsequents; // need to make sure this dies just a wee bit later, after the task is all destroyed
+		if (Subsequents.GetReference() && Subsequents->GetInlineStorage() == this)
+		{
+			LocalSubsequents.Swap(Subsequents);
+			this->TGraphTask::~TGraphTask(); // this will free the ref, if that is the last one it will recycle
+		}
+		else
+#endif
 		if (sizeof(TGraphTask) <= FBaseGraphTask::SMALL_TASK_SIZE)
 		{
-			this->~TGraphTask();
+			this->TGraphTask::~TGraphTask();
 			FBaseGraphTask::GetSmallTaskAllocator().Free(this);
 		}
 		else
@@ -821,20 +920,30 @@ private:
 
 	/** 
 	 *	Private constructor, constructs the base class with the number of prerequisites.
-	 *	@param InSubsequents subsequents to associate with this task
+	 *	@param InSubsequents subsequents to associate with this task. Thsi refernence is destroyed in the process!
 	 *	@param NumberOfPrerequistitesOutstanding the number of prerequisites this task will have when it is built.
 	**/
 	TGraphTask(FGraphEventRef InSubsequents, int32 NumberOfPrerequistitesOutstanding)
-		: FBaseGraphTask(NumberOfPrerequistitesOutstanding)
+		: FBaseGraphTask(
+#if USE_VIRTUAL_BYPASS
+		&ExecuteTask,
+#endif
+		NumberOfPrerequistitesOutstanding)
 		, TaskConstructed(false)
-		, Subsequents(InSubsequents)
 	{
+		Subsequents.Swap(InSubsequents);
 	}
 
 	/** 
 	 *	Private destructor, just checks that the task appears to be completed
 	**/
-	virtual ~TGraphTask()
+#if !USE_VIRTUAL_BYPASS
+	virtual 
+#endif
+		~TGraphTask() 
+#if !USE_VIRTUAL_BYPASS
+		final override
+#endif
 	{
 		checkThreadGraph(!TaskConstructed);
 	}
@@ -1295,7 +1404,7 @@ public:
 **/
 class FCompletionList
 {
-	TLockFreePointerList<FGraphEvent>	Prerequisites;
+	TLockFreePointerListUnordered<FGraphEvent>	Prerequisites;
 public:
 	/**
 	 * Adds a task to the completion list, can be called from any thread
diff --git a/Engine/Source/Runtime/Core/Public/Containers/LockFreeFixedSizeAllocator.h b/Engine/Source/Runtime/Core/Public/Containers/LockFreeFixedSizeAllocator.h
--- a/Engine/Source/Runtime/Core/Public/Containers/LockFreeFixedSizeAllocator.h
+++ b/Engine/Source/Runtime/Core/Public/Containers/LockFreeFixedSizeAllocator.h
@@ -4,6 +4,9 @@
 
 #include "LockFreeList.h"
 
+
+
+#if USE_NEW_LOCK_FREE_LISTS
 /**
  * Thread safe, lock free pooling allocator of fixed size blocks that
  * never returns free space until program shutdown.
@@ -12,6 +15,14 @@
 template<int32 SIZE, typename TTrackingCounter = FNoopCounter>
 class TLockFreeFixedSizeAllocator
 {
+	/**
+	 * Fake struct to use as a type for recycling memory blocks with intrusive lockfree lists
+	 */
+	struct FFakeType
+	{
+		FFakeType* LockFreePointerQueueNext;
+	};
+
 public:
 
 	/** Destructor, returns all memory via FMemory::Free **/
@@ -34,6 +45,7 @@ public:
 	 */
 	void* Allocate()
 	{
+		static_assert(sizeof(FFakeType) <= SIZE, "can't make an efficient recycler for blocks smaller than a pointer");
 		NumUsed.Increment();
 		void *Memory = FreeList.Pop();
 		if (Memory)
@@ -56,7 +68,7 @@ public:
 	void Free(void *Item)
 	{
 		NumUsed.Decrement();
-		FreeList.Push(Item);
+		FreeList.Push((FFakeType*)Item);
 		NumFree.Increment();
 	}
 
@@ -85,7 +97,7 @@ public:
 private:
 
 	/** Lock free list of free memory blocks. */
-	TLockFreePointerList<void> FreeList;
+	FLockFreePointerListFIFOIntrusive<FFakeType> FreeList;
 
 	/** Total number of blocks outstanding and not in the free list. */
 	TTrackingCounter NumUsed; 
@@ -94,32 +106,27 @@ private:
 	TTrackingCounter NumFree;
 };
 
+#else
 /**
  * Thread safe, lock free pooling allocator of fixed size blocks that
- * never returns free space, even at shutdown
+ * never returns free space until program shutdown.
  * alignment isn't handled, assumes FMemory::Malloc will work
  */
 template<int32 SIZE, typename TTrackingCounter = FNoopCounter>
-class TLockFreeFixedSizeAllocator_TLSCache	
+class TLockFreeFixedSizeAllocator
 {
-	enum
-	{
-		NUM_PER_BUNDLE=256,
-	};
 public:
 
-	TLockFreeFixedSizeAllocator_TLSCache()
-	{
-		static_assert(SIZE >= sizeof(void*) && SIZE % sizeof(void*) == 0, "Blocks in TLockFreeFixedSizeAllocator must be at least the size of a pointer.");
-		check(IsInGameThread());
-		TlsSlot = FPlatformTLS::AllocTlsSlot();
-		check(FPlatformTLS::IsValidTlsSlot(TlsSlot));
-	}
-	/** Destructor, leaks all of the memory **/
-	~TLockFreeFixedSizeAllocator_TLSCache()
+	/** Destructor, returns all memory via FMemory::Free **/
+	~TLockFreeFixedSizeAllocator()
 	{
-		FPlatformTLS::FreeTlsSlot(TlsSlot);
-		TlsSlot = 0;
+		check(!NumUsed.GetValue());
+		while (void* Mem = FreeList.Pop())
+		{
+			FMemory::Free(Mem);
+			NumFree.Decrement();
+		}
+		check(!NumFree.GetValue());
 	}
 
 	/**
@@ -130,41 +137,17 @@ public:
 	 */
 	void* Allocate()
 	{
-		FThreadLocalCache& TLS = GetTLS();
-
-		if (!TLS.PartialBundle)
+		NumUsed.Increment();
+		void *Memory = FreeList.Pop();
+		if (Memory)
 		{
-			if (TLS.FullBundle)
-			{
-				TLS.PartialBundle = TLS.FullBundle;
-				TLS.FullBundle = nullptr;
-			}
-			else
-			{
-				TLS.PartialBundle = GlobalFreeListBundles.Pop();
-				if (!TLS.PartialBundle)
-				{
-					TLS.PartialBundle = (void**)FMemory::Malloc(SIZE * NUM_PER_BUNDLE);
-					void **Next = TLS.PartialBundle;
-					for (int32 Index = 0; Index < NUM_PER_BUNDLE - 1; Index++)
-					{
-						void* NextNext = (void*)(((uint8*)Next) + SIZE);
-						*Next = NextNext;
-						Next = (void**)NextNext;
-					}
-					*Next = nullptr;
-					NumFree.Add(NUM_PER_BUNDLE);
-				}
-			}
-			TLS.NumPartial = NUM_PER_BUNDLE;
+			NumFree.Decrement();
 		}
-		NumUsed.Increment();
-		NumFree.Decrement();
-		void* Result = (void*)TLS.PartialBundle;
-		TLS.PartialBundle = (void**)*TLS.PartialBundle;
-		TLS.NumPartial--;
-		check(TLS.NumPartial >= 0 && ((!!TLS.NumPartial) == (!!TLS.PartialBundle)));
-		return Result;
+		else
+		{
+			Memory = FMemory::Malloc(SIZE);
+		}
+		return Memory;
 	}
 
 	/**
@@ -176,22 +159,8 @@ public:
 	void Free(void *Item)
 	{
 		NumUsed.Decrement();
+		FreeList.Push(Item);
 		NumFree.Increment();
-		FThreadLocalCache& TLS = GetTLS();
-		if (TLS.NumPartial >= NUM_PER_BUNDLE)
-		{
-			if (TLS.FullBundle)
-			{
-				GlobalFreeListBundles.Push(TLS.FullBundle);
-				//TLS.FullBundle = nullptr;
-			}
-			TLS.FullBundle = TLS.PartialBundle;
-			TLS.PartialBundle = nullptr;
-			TLS.NumPartial = 0;
-		}
-		*(void**)Item = (void*)TLS.PartialBundle;
-		TLS.PartialBundle = (void**)Item;
-		TLS.NumPartial++;
 	}
 
 	/**
@@ -218,38 +187,8 @@ public:
 
 private:
 
-	/** struct for the TLS cache. */
-	struct FThreadLocalCache
-	{
-		void **FullBundle;
-		void **PartialBundle;
-		int32 NumPartial;
-
-		FThreadLocalCache()
-			: FullBundle(nullptr)
-			, PartialBundle(nullptr)
-			, NumPartial(0)
-		{
-		}
-	};
-
-	FThreadLocalCache& GetTLS()
-	{
-		checkSlow(FPlatformTLS::IsValidTlsSlot(TlsSlot));
-		FThreadLocalCache* TLS = (FThreadLocalCache*)FPlatformTLS::GetTlsValue(TlsSlot);
-		if (!TLS)
-		{
-			TLS = new FThreadLocalCache();
-			FPlatformTLS::SetTlsValue(TlsSlot, TLS);
-		}
-		return *TLS;
-	}
-
-	/** Slot for TLS struct. */
-	uint32 TlsSlot;
-
-	/** Lock free list of free memory blocks, these are all linked into a bundle of NUM_PER_BUNDLE. */
-	TLockFreePointerList<void*> GlobalFreeListBundles;
+	/** Lock free list of free memory blocks. */
+	TLockFreePointerListUnordered<void> FreeList;
 
 	/** Total number of blocks outstanding and not in the free list. */
 	TTrackingCounter NumUsed; 
@@ -257,6 +196,17 @@ private:
 	/** Total number of blocks in the free list. */
 	TTrackingCounter NumFree;
 };
+#endif
+
+/**
+ * Thread safe, lock free pooling allocator of fixed size blocks that
+ * never returns free space, even at shutdown
+ * alignment isn't handled, assumes FMemory::Malloc will work
+ */
+template<int32 SIZE, typename TTrackingCounter = FNoopCounter>
+class TLockFreeFixedSizeAllocator_TLSCache : public TLockFreeFixedSizeAllocator_TLSCacheBase<SIZE, TLockFreePointerListUnordered<void*>, TTrackingCounter>
+{
+};
 
 /**
  * Thread safe, lock free pooling allocator of memory for instances of T.
diff --git a/Engine/Source/Runtime/Core/Public/Containers/LockFreeList.h b/Engine/Source/Runtime/Core/Public/Containers/LockFreeList.h
--- a/Engine/Source/Runtime/Core/Public/Containers/LockFreeList.h
+++ b/Engine/Source/Runtime/Core/Public/Containers/LockFreeList.h
@@ -6,7 +6,8 @@
 CORE_API DECLARE_LOG_CATEGORY_EXTERN(LogLockFreeList, Log, All);
 
 // what level of checking to perform...normally checkLockFreePointerList but could be ensure or check
-#define checkLockFreePointerList checkSlow /*check*/
+#define checkLockFreePointerList checkSlow
+//#define checkLockFreePointerList(x) ((x)||((*(char*)3) = 0))
 
 #define MONITOR_LINK_ALLOCATION (0)
 
@@ -16,41 +17,6 @@ CORE_API DECLARE_LOG_CATEGORY_EXTERN(LogLockFreeList, Log, All);
 	typedef FNoopCounter		FLockFreeListCounter;
 #endif // MONITOR_LINK_ALLOCATION==1
 
-struct CORE_API FLockFreeListStats
-{
-	/** Number of lock free list operations like Push, Pop, PopAll etc. */
-	static FThreadSafeCounter NumOperations;
-	/** Cycles spent in lock free list operations like Push, Pop, PopAll etc. */
-	static FThreadSafeCounter Cycles;
-};
-
-#define MONITOR_LOCKFREELIST_PERFORMANCE (0)
-
-/** Helper struct used to test the performance of the lock free list operations. */
-struct FThreadSafePerfCounter
-{
-	const uint32 StartCycles;
-
-	FThreadSafePerfCounter()
-		: StartCycles( FPlatformTime::Cycles() )
-	{}
-
-	~FThreadSafePerfCounter()
-	{
-		const uint32 EndCycles = FPlatformTime::Cycles();
-		const int32 Delta = int32( (uint32)EndCycles - (uint32)StartCycles );
-		FLockFreeListStats::Cycles.Add( Delta );
-		FLockFreeListStats::NumOperations.Increment();
-	}
-};
-
-
-#if	MONITOR_LOCKFREELIST_PERFORMANCE==1
-	typedef FThreadSafePerfCounter	FLockFreeListPerfCounter;
-#else
-	typedef FNoopStruct				FLockFreeListPerfCounter;
-#endif // MONITOR_LOCKFREELIST_PERFORMANCE==1
-
 /** Enumerates lock free list alignments. */
 namespace ELockFreeAligment
 {
diff --git a/Engine/Source/Runtime/Core/Public/Containers/LockFreeListImpl.h b/Engine/Source/Runtime/Core/Public/Containers/LockFreeListImpl.h
--- a/Engine/Source/Runtime/Core/Public/Containers/LockFreeListImpl.h
+++ b/Engine/Source/Runtime/Core/Public/Containers/LockFreeListImpl.h
@@ -9,7 +9,7 @@
  *	T is the type of the pointer the list will contain.
  */
 template <class T>
-class TLockFreePointerList : private FLockFreeVoidPointerListGeneric 
+class TLockFreePointerListLIFOBase : private FLockFreeVoidPointerListGeneric 
 {
 public:
 	/**	
@@ -19,7 +19,6 @@ public:
 	 */
 	void Push(T *NewItem)
 	{
-		FLockFreeListPerfCounter LockFreePerfCounter;
 		FLockFreeVoidPointerListGeneric::Push(NewItem);
 	}
 
@@ -29,7 +28,6 @@ public:
 	 */
 	T* Pop()
 	{
-		FLockFreeListPerfCounter LockFreePerfCounter;
 		return (T*)FLockFreeVoidPointerListGeneric::Pop();
 	}
 
@@ -40,19 +38,9 @@ public:
 	 */
 	void PopAll(TArray<T *>& Output)
 	{
-		FLockFreeListPerfCounter LockFreePerfCounter;
 		FLockFreeVoidPointerListGeneric::PopAll<TArray<T *>, T * >(Output);
 	}	
 
-	/**	
-	 *	If the list is empty, replace it with the other list and null the other list.
-	 *	@return true if this call actively closed the list
-	 */
-	bool ReplaceListIfEmpty(TLockFreePointerList<T>& NotThreadSafeTempListToReplaceWith)
-	{
-		return FLockFreeVoidPointerListGeneric::ReplaceListIfEmpty(NotThreadSafeTempListToReplaceWith);
-	}
-
 	/**	
 	 *	Check if the list is empty.
 	 *
@@ -60,33 +48,34 @@ public:
 	 *	CAUTION: This methods safety depends on external assumptions. For example, if another thread could add to the list at any time, the return value is no better than a best guess.
 	 *	As typically used, the list is not being access concurrently when this is called.
 	 */
-	bool IsEmpty() const  
+	FORCEINLINE bool IsEmpty() const  
 	{
 		return FLockFreeVoidPointerListGeneric::IsEmpty();
 	}
-
-#if	!(UE_BUILD_SHIPPING || UE_BUILD_TEST)
-	/**
-	 * Checks if the pointers look 'ok'. Used for debugging.
+protected:
+	/**	
+	 *	If the list is empty, replace it with the other list and null the other list.
+	 *	@return true if this call actively closed the list
 	 */
-	bool CheckPointers() const
+	FORCEINLINE bool ReplaceListIfEmpty(TLockFreePointerListLIFOBase<T>& NotThreadSafeTempListToReplaceWith)
 	{
-		return FLockFreeVoidPointerListGeneric::CheckPointers();
+		return FLockFreeVoidPointerListGeneric::ReplaceListIfEmpty(NotThreadSafeTempListToReplaceWith);
 	}
+};
 
-	/**
-	 * @return number of elements via iterating through this list, used for debugging.
-	 */
-	int32 NumVerified() const
+template <class T>
+class TLockFreePointerListLIFO : public TLockFreePointerListLIFOBase<T> 
+{
+public:
+	FORCEINLINE bool ReplaceListIfEmpty(TLockFreePointerListLIFO<T>& NotThreadSafeTempListToReplaceWith)
 	{
-		return FLockFreeVoidPointerListGeneric::NumVerified();
+		return TLockFreePointerListLIFOBase<T>::ReplaceListIfEmpty(NotThreadSafeTempListToReplaceWith);
 	}
-#endif // !(UE_BUILD_SHIPPING || UE_BUILD_TEST)
 
 };
 
 template <class T>
-class TClosableLockFreePointerList : private FLockFreeVoidPointerListGeneric 
+class TClosableLockFreePointerListLIFO : private FLockFreeVoidPointerListGeneric 
 {
 public:
 	/**	
@@ -97,7 +86,6 @@ public:
 	 */
 	bool PushIfNotClosed(T *NewItem)
 	{
-		FLockFreeListPerfCounter LockFreePerfCounter;
 		return FLockFreeVoidPointerListGeneric::PushIfNotClosed(NewItem);
 	}
 
@@ -108,7 +96,6 @@ public:
 	 */
 	void PopAllAndClose(TArray<T *>& Output)
 	{
-		FLockFreeListPerfCounter LockFreePerfCounter;
 		FLockFreeVoidPointerListGeneric::PopAllAndClose<TArray<T *>, T * >(Output);
 	}	
 
@@ -121,28 +108,10 @@ public:
 	{
 		return FLockFreeVoidPointerListGeneric::IsClosed();
 	}
-
-#if	!(UE_BUILD_SHIPPING || UE_BUILD_TEST)
-	/**
-	 * Checks if the pointers look 'ok'. Used for debugging.
-	 */
-	bool CheckPointers() const
-	{
-		return FLockFreeVoidPointerListGeneric::CheckPointers();
-	}
-
-	/**
-	 * @return number of elements via iterating through this list, used for debugging.
-	 */
-	int32 NumVerified() const
-	{
-		return FLockFreeVoidPointerListGeneric::NumVerified();
-	}
-#endif // !(UE_BUILD_SHIPPING || UE_BUILD_TEST)
 };
 
 template <class T>
-class TReopenableLockFreePointerList : private FLockFreeVoidPointerListGeneric
+class TReopenableLockFreePointerListLIFO : private FLockFreeVoidPointerListGeneric
 {
 public:
 	/**	
@@ -153,7 +122,6 @@ public:
 	 */
 	bool ReopenIfClosedAndPush(T *NewItem)
 	{
-		FLockFreeListPerfCounter LockFreePerfCounter;
 		bool bWasReopenedByMe = FLockFreeVoidPointerListGeneric::ReopenIfClosedAndPush(NewItem);
 		return bWasReopenedByMe;
 	}
@@ -164,7 +132,6 @@ public:
 	 */
 	T* PopIfNotClosed()
 	{
-		FLockFreeListPerfCounter LockFreePerfCounter;
 		return (T*)FLockFreeVoidPointerListGeneric::PopIfNotClosed();
 	}
 
@@ -176,7 +143,6 @@ public:
 	 */
 	void PopAll(TArray<T *>& Output)
 	{
-		FLockFreeListPerfCounter LockFreePerfCounter;
 		FLockFreeVoidPointerListGeneric::PopAll<TArray<T *>, T * >(Output);
 	}	
 	/**	
@@ -186,25 +152,122 @@ public:
 	 */
 	bool CloseIfEmpty()
 	{
-		FLockFreeListPerfCounter LockFreePerfCounter;
 		return FLockFreeVoidPointerListGeneric::CloseIfEmpty();
 	}
+};
+
+#if !USE_NEW_LOCK_FREE_LISTS
+// where you don't care what order they pop in. We choose the fastest implementation we have.
+template <class T>
+class TLockFreePointerListUnordered : public TLockFreePointerListLIFOBase<T>
+{
+
+};
+template <class T>
+class TClosableLockFreePointerListUnorderedSingleConsumer : public TClosableLockFreePointerListLIFO<T>
+{
+
+};
+
+#else
+
+template <class T>
+class TClosableLockFreePointerListFIFOSingleConsumer
+{
+public:
+	/**	
+	 *	Push an item onto the head of the list, unless the list is closed
+	 *
+	 *	@param NewItem, the new item to push on the list, cannot be NULL
+	 *	@return true if the item was pushed on the list, false if the list was closed.
+	 */
+	FORCEINLINE bool PushIfNotClosed(T *NewItem)
+	{
+		return Inner.PushIfNotClosed(NewItem);
+	}
+
+	/**	
+	 *	Pop an item from the list or return NULL if the list is empty
+	 *	@return The new item, if any
+	 *	CAUTION: This method should not be used unless the list is known to not be closed.
+	 */
+	FORCEINLINE T* Pop()
+	{
+		return (T*)Inner.Pop();
+	}
 
-#if	!(UE_BUILD_SHIPPING || UE_BUILD_TEST)
-	/**
-	 * Checks if the pointers look 'ok'. Used for debugging.
+	/**	
+	 *	Close the list if it is empty.
+	 *	@return true if this call actively closed the list.
 	 */
-	bool CheckPointers() const
+	FORCEINLINE bool CloseIfEmpty()
 	{
-		return FLockFreeVoidPointerListGeneric::CheckPointers();
+		return Inner.CloseIfEmpty();
 	}
 
-	/**
-	 * @return number of elements via iterating through this list, used for debugging.
+	/**	
+	 *	Check if the list is closed
+	 *	@return true if the list is closed.
+	 */
+	FORCEINLINE bool IsClosed() const
+	{
+		return Inner.IsClosed();
+	}
+	/**	
+	 *	Not thread safe, used to reset the list for recycling without freeing the stub
+	 *	@return true if the list is closed.
 	 */
-	int32 NumVerified() const
+	FORCEINLINE void Reset()
 	{
-		return FLockFreeVoidPointerListGeneric::NumVerified();
+		return Inner.Reset();
 	}
-#endif // !(UE_BUILD_SHIPPING || UE_BUILD_TEST)
+private:
+	FCloseableLockFreePointerListFIFOBaseSingleConsumer Inner;
 };
+
+template <class T>
+class TReopenableLockFreePointerListFIFOSingleConsumer
+{
+public:
+	/**	
+	 *	Push an item onto the head of the list, opening it first if necessary.
+	 *
+	 *	@param NewItem, the new item to push on the list, cannot be NULL.
+	 *	@return true if the list needed to be opened first, false if the list was not closed before our push.
+	 */
+	FORCEINLINE bool ReopenIfClosedAndPush(T *NewItem)
+	{
+		return Inner.ReopenIfClosedAndPush(NewItem);
+	}
+
+	/**	
+	 *	Pop an item from the list or return NULL if the list is empty
+	 *	@return The new item, if any
+	 *	CAUTION: This method should not be used unless the list is known to not be closed.
+	 */
+	FORCEINLINE T* Pop()
+	{
+		return (T*)Inner.Pop();
+	}
+
+	/**	
+	 *	Close the list if it is empty.
+	 *	@return true if this call actively closed the list.
+	 */
+	FORCEINLINE bool CloseIfEmpty()
+	{
+		return Inner.CloseIfEmpty();
+	}
+private:
+	FCloseableLockFreePointerListFIFOBaseSingleConsumer Inner;
+};
+
+template <class T>
+class TClosableLockFreePointerListUnorderedSingleConsumer : public TClosableLockFreePointerListFIFOSingleConsumer<T>
+{
+
+};
+
+
+#endif
+
diff --git a/Engine/Source/Runtime/Core/Public/Containers/LockFreeVoidPointerListBase.h b/Engine/Source/Runtime/Core/Public/Containers/LockFreeVoidPointerListBase.h
--- a/Engine/Source/Runtime/Core/Public/Containers/LockFreeVoidPointerListBase.h
+++ b/Engine/Source/Runtime/Core/Public/Containers/LockFreeVoidPointerListBase.h
@@ -4,6 +4,8 @@
 
 #define CHECK_NON_CONCURRENT_ASSUMPTIONS (0)
 
+#define USE_NEW_LOCK_FREE_LISTS (PLATFORM_PS4)
+
 
 /** 
  * Base class for a lock free list of pointers 
@@ -205,44 +207,6 @@ public:
 		return Head == nullptr;
 	}
 
-#if !(UE_BUILD_SHIPPING || UE_BUILD_TEST)
-	/**
-	 * Checks if item pointers look 'ok'. Used for debugging.
-	 */
-	bool CheckPointers() const
-	{
-		const int32 AlignmentCheck = ELockFreeAligment::LF64bit - 1;
-		FPlatformMisc::MemoryBarrier();
-		for (FLink* Link = Head; Link; Link = Link->Next)
-		{
-			void* ItemPtr = Link->Item;
-			if (!ItemPtr || (UPTRINT)ItemPtr < 0x100)
-			{
-				return false;
-			}	
-			if ((UPTRINT)ItemPtr & AlignmentCheck)
-			{
-				return false;
-			}
-		}
-		return true;
-	}
-
-	/**
-	 * @return number of elements via iterating through this list, used for debugging.
-	 */
-	int32 NumVerified() const
-	{
-		FPlatformMisc::MemoryBarrier();
-
-		int32 NumElements = 0;	
-		for (FLink* Link = Head; Link; Link = Link->Next)
-		{
-			NumElements++;
-		}
-		return NumElements;
-	}
-#endif
 
 private:
 	/**	Internal single-linked list link structure. */
@@ -644,3 +608,1524 @@ inline void FLockFreeVoidPointerListBase::FLink::Unlock(bool bShouldNeverCauseAF
 }
 
 typedef FLockFreeVoidPointerListBase	FLockFreeVoidPointerListGeneric;
+
+// New lock free lists
+
+template<class T, typename TBundleRecycler, typename TTrackingCounter = FNoopCounter>
+class TPointerSet_TLSCacheBase : public FNoncopyable
+{
+	enum
+	{
+		NUM_PER_BUNDLE=256,
+	};
+public:
+
+	TPointerSet_TLSCacheBase()
+	{
+		check(IsInGameThread());
+		TlsSlot = FPlatformTLS::AllocTlsSlot();
+		check(FPlatformTLS::IsValidTlsSlot(TlsSlot));
+	}
+	/** Destructor, leaks all of the memory **/
+	~TPointerSet_TLSCacheBase()
+	{
+		FPlatformTLS::FreeTlsSlot(TlsSlot);
+		TlsSlot = 0;
+	}
+
+	/**
+	 * Allocates a memory block of size SIZE.
+	 *
+	 * @return Pointer to the allocated memory.
+	 * @see Free
+	 */
+	T* Pop()
+	{
+		FThreadLocalCache& TLS = GetTLS();
+
+		if (!TLS.PartialBundle)
+		{
+			if (TLS.FullBundle)
+			{
+				TLS.PartialBundle = TLS.FullBundle;
+				TLS.FullBundle = nullptr;
+			}
+			else
+			{
+				TLS.PartialBundle = GlobalFreeListBundles.Pop();
+				if (!TLS.PartialBundle)
+				{
+					TLS.PartialBundle = T::GetBundle(NUM_PER_BUNDLE);
+					NumFree.Add(NUM_PER_BUNDLE);
+				}
+			}
+			TLS.NumPartial = NUM_PER_BUNDLE;
+		}
+		NumUsed.Increment();
+		NumFree.Decrement();
+		T* Result = TLS.PartialBundle;
+		TLS.PartialBundle = (T*)(TLS.PartialBundle->LockFreePointerQueueNext);
+		TLS.NumPartial--;
+		checkLockFreePointerList(TLS.NumPartial >= 0 && ((!!TLS.NumPartial) == (!!TLS.PartialBundle)));
+		return Result;
+	}
+
+	/**
+	 * Puts a memory block previously obtained from Allocate() back on the free list for future use.
+	 *
+	 * @param Item The item to free.
+	 * @see Allocate
+	 */
+	void Push(T *Item)
+	{
+		NumUsed.Decrement();
+		NumFree.Increment();
+		FThreadLocalCache& TLS = GetTLS();
+		if (TLS.NumPartial >= NUM_PER_BUNDLE)
+		{
+			if (TLS.FullBundle)
+			{
+				GlobalFreeListBundles.Push(TLS.FullBundle);
+				//TLS.FullBundle = nullptr;
+			}
+			TLS.FullBundle = TLS.PartialBundle;
+			TLS.PartialBundle = nullptr;
+			TLS.NumPartial = 0;
+		}
+		Item->LockFreePointerQueueNext = TLS.PartialBundle;
+		TLS.PartialBundle = Item;
+		TLS.NumPartial++;
+	}
+
+	/**
+	 * Gets the number of allocated memory blocks that are currently in use.
+	 *
+	 * @return Number of used memory blocks.
+	 * @see GetNumFree
+	 */
+	const TTrackingCounter& GetNumUsed() const
+	{
+		return NumUsed;
+	}
+
+	/**
+	 * Gets the number of allocated memory blocks that are currently unused.
+	 *
+	 * @return Number of unused memory blocks.
+	 * @see GetNumUsed
+	 */
+	const TTrackingCounter& GetNumFree() const
+	{
+		return NumFree;
+	}
+
+private:
+
+	/** struct for the TLS cache. */
+	struct FThreadLocalCache
+	{
+		T* FullBundle;
+		T* PartialBundle;
+		int32 NumPartial;
+
+		FThreadLocalCache()
+			: FullBundle(nullptr)
+			, PartialBundle(nullptr)
+			, NumPartial(0)
+		{
+		}
+	};
+
+	FThreadLocalCache& GetTLS()
+	{
+		checkSlow(FPlatformTLS::IsValidTlsSlot(TlsSlot));
+		FThreadLocalCache* TLS = (FThreadLocalCache*)FPlatformTLS::GetTlsValue(TlsSlot);
+		if (!TLS)
+		{
+			TLS = new FThreadLocalCache();
+			FPlatformTLS::SetTlsValue(TlsSlot, TLS);
+		}
+		return *TLS;
+	}
+
+	/** Slot for TLS struct. */
+	uint32 TlsSlot;
+
+	/** Lock free list of free memory blocks, these are all linked into a bundle of NUM_PER_BUNDLE. */
+	TBundleRecycler GlobalFreeListBundles;
+
+	/** Total number of blocks outstanding and not in the free list. */
+	TTrackingCounter NumUsed; 
+
+	/** Total number of blocks in the free list. */
+	TTrackingCounter NumFree;
+};
+
+/**
+ * Thread safe, lock free pooling allocator of fixed size blocks that
+ * never returns free space, even at shutdown
+ * alignment isn't handled, assumes FMemory::Malloc will work
+ */
+template<int32 SIZE, typename TBundleRecycler, typename TTrackingCounter = FNoopCounter>
+class TLockFreeFixedSizeAllocator_TLSCacheBase : public FNoncopyable
+{
+	enum
+	{
+		NUM_PER_BUNDLE=256,
+	};
+public:
+
+	TLockFreeFixedSizeAllocator_TLSCacheBase()
+	{
+		static_assert(SIZE >= sizeof(void*) && SIZE % sizeof(void*) == 0, "Blocks in TLockFreeFixedSizeAllocator must be at least the size of a pointer.");
+		check(IsInGameThread());
+		TlsSlot = FPlatformTLS::AllocTlsSlot();
+		check(FPlatformTLS::IsValidTlsSlot(TlsSlot));
+	}
+	/** Destructor, leaks all of the memory **/
+	~TLockFreeFixedSizeAllocator_TLSCacheBase()
+	{
+		FPlatformTLS::FreeTlsSlot(TlsSlot);
+		TlsSlot = 0;
+	}
+
+	/**
+	 * Allocates a memory block of size SIZE.
+	 *
+	 * @return Pointer to the allocated memory.
+	 * @see Free
+	 */
+	FORCEINLINE void* Allocate()
+	{
+		FThreadLocalCache& TLS = GetTLS();
+
+		if (!TLS.PartialBundle)
+		{
+			if (TLS.FullBundle)
+			{
+				TLS.PartialBundle = TLS.FullBundle;
+				TLS.FullBundle = nullptr;
+			}
+			else
+			{
+				TLS.PartialBundle = GlobalFreeListBundles.Pop();
+				if (!TLS.PartialBundle)
+				{
+					TLS.PartialBundle = (void**)FMemory::Malloc(SIZE * NUM_PER_BUNDLE);
+					void **Next = TLS.PartialBundle;
+					for (int32 Index = 0; Index < NUM_PER_BUNDLE - 1; Index++)
+					{
+						void* NextNext = (void*)(((uint8*)Next) + SIZE);
+						*Next = NextNext;
+						Next = (void**)NextNext;
+					}
+					*Next = nullptr;
+					NumFree.Add(NUM_PER_BUNDLE);
+				}
+			}
+			TLS.NumPartial = NUM_PER_BUNDLE;
+		}
+		NumUsed.Increment();
+		NumFree.Decrement();
+		void* Result = (void*)TLS.PartialBundle;
+		TLS.PartialBundle = (void**)*TLS.PartialBundle;
+		TLS.NumPartial--;
+		check(TLS.NumPartial >= 0 && ((!!TLS.NumPartial) == (!!TLS.PartialBundle)));
+		return Result;
+	}
+
+	/**
+	 * Puts a memory block previously obtained from Allocate() back on the free list for future use.
+	 *
+	 * @param Item The item to free.
+	 * @see Allocate
+	 */
+	FORCEINLINE void Free(void *Item)
+	{
+		NumUsed.Decrement();
+		NumFree.Increment();
+		FThreadLocalCache& TLS = GetTLS();
+		if (TLS.NumPartial >= NUM_PER_BUNDLE)
+		{
+			if (TLS.FullBundle)
+			{
+				GlobalFreeListBundles.Push(TLS.FullBundle);
+				//TLS.FullBundle = nullptr;
+			}
+			TLS.FullBundle = TLS.PartialBundle;
+			TLS.PartialBundle = nullptr;
+			TLS.NumPartial = 0;
+		}
+		*(void**)Item = (void*)TLS.PartialBundle;
+		TLS.PartialBundle = (void**)Item;
+		TLS.NumPartial++;
+	}
+
+	/**
+	 * Gets the number of allocated memory blocks that are currently in use.
+	 *
+	 * @return Number of used memory blocks.
+	 * @see GetNumFree
+	 */
+	const TTrackingCounter& GetNumUsed() const
+	{
+		return NumUsed;
+	}
+
+	/**
+	 * Gets the number of allocated memory blocks that are currently unused.
+	 *
+	 * @return Number of unused memory blocks.
+	 * @see GetNumUsed
+	 */
+	const TTrackingCounter& GetNumFree() const
+	{
+		return NumFree;
+	}
+
+private:
+
+	/** struct for the TLS cache. */
+	struct FThreadLocalCache
+	{
+		void **FullBundle;
+		void **PartialBundle;
+		int32 NumPartial;
+
+		FThreadLocalCache()
+			: FullBundle(nullptr)
+			, PartialBundle(nullptr)
+			, NumPartial(0)
+		{
+		}
+	};
+
+	FThreadLocalCache& GetTLS()
+	{
+		checkSlow(FPlatformTLS::IsValidTlsSlot(TlsSlot));
+		FThreadLocalCache* TLS = (FThreadLocalCache*)FPlatformTLS::GetTlsValue(TlsSlot);
+		if (!TLS)
+		{
+			TLS = new FThreadLocalCache();
+			FPlatformTLS::SetTlsValue(TlsSlot, TLS);
+		}
+		return *TLS;
+	}
+
+	/** Slot for TLS struct. */
+	uint32 TlsSlot;
+
+	/** Lock free list of free memory blocks, these are all linked into a bundle of NUM_PER_BUNDLE. */
+	TBundleRecycler GlobalFreeListBundles;
+
+	/** Total number of blocks outstanding and not in the free list. */
+	TTrackingCounter NumUsed; 
+
+	/** Total number of blocks in the free list. */
+	TTrackingCounter NumFree;
+};
+
+// in these cases it is critical that no thread starves because they must advance and other threads are blocked.
+
+CORE_API void LockFreeCriticalSpin(int32& SpinCount);
+
+#if !UE_BUILD_SHIPPING && !UE_BUILD_TEST
+
+CORE_API void DoTestCriticalStall();
+extern CORE_API int32 GTestCriticalStalls;
+
+FORCEINLINE void TestCriticalStall()
+{
+	if (GTestCriticalStalls)
+	{
+		DoTestCriticalStall();
+	}
+}
+#else
+FORCEINLINE void TestCriticalStall()
+{
+}
+#endif
+
+template<class T>
+class FLockFreePointerQueueBaseSingleConsumerIntrusive : public FNoncopyable
+{
+public:
+	FLockFreePointerQueueBaseSingleConsumerIntrusive()
+		: Head(GetStub())
+		, Tail(GetStub())
+	{
+		GetStub()->LockFreePointerQueueNext = nullptr;
+	}
+
+	void Push(T* Link)
+	{
+		Link->LockFreePointerQueueNext = nullptr;
+		T* Prev = (T*)FPlatformAtomics::InterlockedExchangePtr((void**)&Head, Link);
+		TestCriticalStall();
+		Prev->LockFreePointerQueueNext = Link;
+		//FPlatformMisc::MemoryBarrier();
+	}
+
+	T* Pop()
+	{
+		int32 SpinCount = 0;
+		while (true)
+		{
+			T* LocalTail = Tail; 
+			T* LocalTailNext = (T*)(LocalTail->LockFreePointerQueueNext); 
+			if (LocalTail == GetStub())
+			{
+				if (!LocalTailNext)
+				{
+					if (LocalTail != Head)
+					{
+						LockFreeCriticalSpin(SpinCount);
+						continue;
+					}
+					return nullptr;
+				}
+				Tail = LocalTailNext;
+				LocalTail = LocalTailNext;
+				LocalTailNext = (T*)(LocalTailNext->LockFreePointerQueueNext);
+			}
+			if (LocalTailNext)
+			{
+				Tail = LocalTailNext;
+				return LocalTail;
+			}
+			if (LocalTail != Head)
+			{
+				LockFreeCriticalSpin(SpinCount);
+				continue;
+			}
+			Push(GetStub());
+			LocalTailNext = (T*)(LocalTail->LockFreePointerQueueNext); 
+			if (LocalTailNext)
+			{
+				Tail = LocalTailNext;
+				return LocalTail;
+			}
+			LockFreeCriticalSpin(SpinCount);
+		}
+		return nullptr;
+	}
+	/**	
+	 *	Check if the list is empty.
+	 *
+	 *	@return true if the list is empty.
+	 *	CAUTION: This methods safety depends on external assumptions. For example, if another thread could add to the list at any time, the return value is no better than a best guess.
+	 *	As typically used, the list is not being access concurrently when this is called.
+	 */
+	bool IsEmpty() const
+	{
+		//FPlatformMisc::MemoryBarrier();
+		T* LocalTail = Tail; 
+		return Head == LocalTail && GetStub() == LocalTail;
+	}
+
+private:
+
+	T* Head;
+	T* Tail;
+
+	T* GetStub() const
+	{
+		return (T*)&StubStorage; // we aren't going to do anything with this other than mess with internal links
+	}
+	TAlignedBytes<sizeof(T), ALIGNOF(void*)> StubStorage;
+};
+
+struct FLockFreeLink
+{
+	FLockFreeLink* Next;
+	void *Payload;
+};
+
+class FLockFreePointerQueueBaseSingleConsumer : public FNoncopyable
+{
+public:
+	FLockFreePointerQueueBaseSingleConsumer(FLockFreeLink* Stub)
+		: Head(Stub)
+		, Tail(Stub)
+	{
+
+	}
+
+	void Push(FLockFreeLink* Link)
+	{
+		Link->Next = nullptr;
+		FLockFreeLink* Prev = (FLockFreeLink*)FPlatformAtomics::InterlockedExchangePtr((void**)&Head, Link);
+		TestCriticalStall();
+		Prev->Next = Link;
+		//FPlatformMisc::MemoryBarrier();
+	}
+
+	FLockFreeLink* Pop()
+	{
+		int32 SpinCount = 0;
+		while (true)
+		{
+			FLockFreeLink* LocalTail = Tail; 
+			FLockFreeLink* LocalTailNext = LocalTail->Next; 
+			if (LocalTailNext)
+			{
+				Tail = LocalTailNext;
+				LocalTail->Payload = LocalTailNext->Payload;
+				return LocalTail;
+			}
+			else if (Head == LocalTail)
+			{
+				break;
+			}
+			LockFreeCriticalSpin(SpinCount);
+		}
+		return nullptr;
+	}
+
+	/**	
+	 *	Check if the list is empty.
+	 *	@return true if the list is empty.
+	 *	CAUTION: This methods safety depends on external assumptions. For example, if another thread could add to the list at any time, the return value is no better than a best guess.
+	 *	As typically used, the list is not being access concurrently when this is called.
+	 */
+	bool IsEmpty() const
+	{
+		//FPlatformMisc::MemoryBarrier();
+		return Head == Tail;
+	}
+
+private:
+
+	FLockFreeLink* Head;
+	FLockFreeLink* Tail;
+};
+
+class FCloseableLockFreePointerQueueBaseSingleConsumer : public FNoncopyable
+{
+public:
+	FCloseableLockFreePointerQueueBaseSingleConsumer(FLockFreeLink* Stub)
+		: Head(Stub)
+		, Tail(Stub)
+	{
+
+	}
+
+	/**	
+	 *	Push an item onto the head of the list, unless the list is closed
+	 *	@param NewItem, the new item to push on the list, cannot be nullptr
+	 *	@return true if the item was pushed on the list, false if the list was closed.
+	 */
+	bool PushIfNotClosed(FLockFreeLink* Link)
+	{
+		checkLockFreePointerList(Link);
+		Link->Next = nullptr;
+		while (true)
+		{
+			FLockFreeLink* LocalHead = Head; 
+			if (!IsClosed(LocalHead))
+			{
+				if (FPlatformAtomics::InterlockedCompareExchangePointer((void**)&Head, Link, LocalHead) == LocalHead)
+				{
+					TestCriticalStall();
+					LocalHead->Next = Link;
+					return true;
+				}
+			}
+			else
+			{
+				break;
+			}
+		}
+		return false;
+	}
+
+	/**	
+	 *	Push an item onto the head of the list, opening it first if necessary
+	 *	@param NewItem, the new item to push on the list, cannot be nullptr
+	 *	@return true if the list needed to be opened first, false if the list was not closed before our push
+	 */
+	bool ReopenIfClosedAndPush(FLockFreeLink *Link)
+	{
+		bool bWasReopenedByMe = false;
+		Link->Next = nullptr;
+		while (true)
+		{
+			FLockFreeLink* LocalHead = Head; 
+			if (FPlatformAtomics::InterlockedCompareExchangePointer((void**)&Head, Link, LocalHead) == LocalHead)
+			{
+				TestCriticalStall();
+				FLockFreeLink* LocalHeadOpen = ClearClosed(LocalHead); 
+				bWasReopenedByMe = IsClosed(LocalHead);
+				LocalHeadOpen->Next = Link;
+				break;
+			}
+		}
+		return bWasReopenedByMe;
+	}
+
+	/**	
+	 *	Close the list if it is empty
+	 *	@return true if this call actively closed the list
+	 */
+	bool CloseIfEmpty()
+	{
+		int32 SpinCount = 0;
+		while (true)
+		{
+			FPlatformMisc::MemoryBarrier();
+			FLockFreeLink* LocalTail = Tail; 
+			FLockFreeLink* LocalHead = Head; 
+			checkLockFreePointerList(!IsClosed(LocalHead));
+			if (LocalTail == LocalHead)
+			{
+				FLockFreeLink* LocalHeadClosed = SetClosed(LocalHead); 
+				checkLockFreePointerList(LocalHeadClosed != LocalHead);
+				if (FPlatformAtomics::InterlockedCompareExchangePointer((void**)&Head, LocalHeadClosed, LocalHead) == LocalHead)
+				{
+					return true;
+				}
+			}
+			else
+			{
+				break;
+			}
+			LockFreeCriticalSpin(SpinCount);
+		}
+		return false;
+	}
+	/**	
+	 *	Pop an item from the list or return nullptr if the list is empty
+	 *	@return The popped item, if any
+	 *	CAUTION: This method should not be used unless the list is known to not be closed.
+	 */
+	FLockFreeLink* Pop()
+	{
+		int32 SpinCount = 0;
+		while (true)
+		{
+			FLockFreeLink* LocalTail = Tail; 
+			checkLockFreePointerList(!IsClosed(LocalTail));
+			FLockFreeLink* LocalTailNext = LocalTail->Next; 
+			checkLockFreePointerList(!IsClosed(LocalTailNext));
+			if (LocalTailNext)
+			{
+				Tail = LocalTailNext;
+				LocalTail->Payload = LocalTailNext->Payload;
+				return LocalTail;
+			}
+			else if (Head == LocalTail)
+			{
+				break;
+			}
+			checkLockFreePointerList(!IsClosed(Head));
+			LockFreeCriticalSpin(SpinCount);
+		}
+		return nullptr;
+	}
+	/**	
+	 *	Check if the list is empty.
+	 *
+	 *	@return true if the list is empty.
+	 *	CAUTION: This methods safety depends on external assumptions. For example, if another thread could add to the list at any time, the return value is no better than a best guess.
+	 *	As typically used, the list is not being access concurrently when this is called.
+	 */
+	bool IsEmpty() const
+	{
+		FPlatformMisc::MemoryBarrier();
+		return ClearClosed(Head) == Tail;
+	}
+
+	/**	
+	 *	Check if the list is closed
+	 *	@return true if the list is closed.
+	 *	CAUTION: This methods safety depends on external assumptions. For example, if another thread could open or close the list at any time, the return value is no better than a best guess.
+	 *	As typically used, the list is only closed once it is lifetime, so a return of true here means it is closed forever.
+	 */
+	FORCEINLINE bool IsClosed() const  // caution, if this returns false, that does not mean the list is open!
+	{
+		FPlatformMisc::MemoryBarrier();
+		return IsClosed(Head);
+	}
+	/**	
+	 *	Not thread safe, used to reset the list for recycling without freeing the stub
+	 *	@return true if the list is closed.
+	 */
+	void Reset()
+	{
+		Head = ClearClosed(Head);
+		FPlatformMisc::MemoryBarrier();
+		checkLockFreePointerList(Head == Tail); // we don't clear the list here, we assume it is already clear and possibly closed
+	}
+
+private:
+	// we will use the lowest bit for the closed state
+	FORCEINLINE static bool IsClosed(FLockFreeLink* HeadToTest)
+	{
+		return !!(UPTRINT(HeadToTest) & 1);
+	}
+	FORCEINLINE static FLockFreeLink* SetClosed(FLockFreeLink* In)
+	{
+		return (FLockFreeLink*)(UPTRINT(In) | 1);
+	}
+	FORCEINLINE static FLockFreeLink* ClearClosed(FLockFreeLink* In)
+	{
+		return (FLockFreeLink*)(UPTRINT(In) & ~UPTRINT(1));
+	}
+
+	FLockFreeLink* Head;
+	FLockFreeLink* Tail;
+};
+
+template<class T>
+class FCloseableLockFreePointerQueueBaseSingleBaseConsumerIntrusive : public FNoncopyable
+{
+public:
+	FCloseableLockFreePointerQueueBaseSingleBaseConsumerIntrusive()
+		: Head(GetStub())
+		, Tail(GetStub())
+	{
+		GetStub()->LockFreePointerQueueNext = nullptr;
+	}
+#if 0
+	/**	
+	 *	Push an item onto the head of the list, unless the list is closed
+	 *	@param NewItem, the new item to push on the list, cannot be nullptr
+	 *	@return true if the item was pushed on the list, false if the list was closed.
+	 */
+	bool PushIfNotClosed(T* Link)
+	{
+		checkLockFreePointerList(Link);
+		Link->LockFreePointerQueueNext = nullptr;
+		while (true)
+		{
+			T* LocalHead = Head; 
+			if (!IsClosed(LocalHead))
+			{
+				if (FPlatformAtomics::InterlockedCompareExchangePointer((void**)&Head, Link, LocalHead) == LocalHead)
+				{
+					TestCriticalStall();
+					LocalHead->LockFreePointerQueueNext = Link;
+					return true;
+				}
+			}
+			else
+			{
+				break;
+			}
+		}
+		return false;
+	}
+#endif
+	/**	
+	 *	Push an item onto the head of the list, opening it first if necessary
+	 *	@param NewItem, the new item to push on the list, cannot be nullptr
+	 *	@return true if the list needed to be opened first, false if the list was not closed before our push
+	 */
+	bool ReopenIfClosedAndPush(T *Link)
+	{
+		bool bWasReopenedByMe = false;
+		checkLockFreePointerList(Link);
+		Link->LockFreePointerQueueNext = nullptr;
+		while (true)
+		{
+			T* LocalHead = Head; 
+			if (FPlatformAtomics::InterlockedCompareExchangePointer((void**)&Head, Link, LocalHead) == LocalHead)
+			{
+				TestCriticalStall();
+				T* LocalHeadOpen = ClearClosed(LocalHead); 
+				bWasReopenedByMe = IsClosed(LocalHead);
+				LocalHeadOpen->LockFreePointerQueueNext = Link;
+				break;
+			}
+		}
+		return bWasReopenedByMe;
+	}
+
+	/**	
+	 *	Close the list if it is empty
+	 *	@return true if this call actively closed the list
+	 */
+	bool CloseIfEmpty()
+	{
+		int32 SpinCount = 0;
+		while (true)
+		{
+			T* LocalTail = Tail; 
+			T* LocalHead = Head; 
+			checkLockFreePointerList(!IsClosed(LocalHead));
+			if (LocalTail == LocalHead)
+			{
+				T* LocalHeadClosed = SetClosed(LocalHead); 
+				checkLockFreePointerList(LocalHeadClosed != LocalHead);
+				if (FPlatformAtomics::InterlockedCompareExchangePointer((void**)&Head, LocalHeadClosed, LocalHead) == LocalHead)
+				{
+					return true;
+				}
+			}
+			else
+			{
+				break;
+			}
+			LockFreeCriticalSpin(SpinCount);
+		}
+		return false;
+	}
+
+	/**	
+	 *	Pop an item from the list or return nullptr if the list is empty
+	 *	@return The popped item, if any
+	 *	CAUTION: This method should not be used unless the list is known to not be closed.
+	 */
+	FORCEINLINE T* Pop()
+	{
+		int32 SpinCount = 0;
+		while (true)
+		{
+			T* LocalTail = Tail; 
+			checkLockFreePointerList(!IsClosed(LocalTail));
+			T* LocalTailNext = (T*)(LocalTail->LockFreePointerQueueNext); 
+			checkLockFreePointerList(!IsClosed(LocalTailNext));
+			if (LocalTail == GetStub())
+			{
+				if (!LocalTailNext)
+				{
+					if (LocalTail != Head)
+					{
+						LockFreeCriticalSpin(SpinCount);
+						continue;
+					}
+					return nullptr;
+				}
+				Tail = LocalTailNext;
+				LocalTail = LocalTailNext;
+				LocalTailNext = (T*)(LocalTailNext->LockFreePointerQueueNext);
+			}
+			if (LocalTailNext)
+			{
+				Tail = LocalTailNext;
+				return LocalTail;
+			}
+			if (LocalTail != Head)
+			{
+				LockFreeCriticalSpin(SpinCount);
+				continue;
+			}
+			Push(GetStub());
+			LocalTailNext = (T*)(LocalTail->LockFreePointerQueueNext); 
+			if (LocalTailNext)
+			{
+				Tail = LocalTailNext;
+				return LocalTail;
+			}
+			LockFreeCriticalSpin(SpinCount);
+		}
+		return nullptr;
+	}
+	/**	
+	 *	Check if the list is closed
+	 *	@return true if the list is closed.
+	 *	CAUTION: This methods safety depends on external assumptions. For example, if another thread could open or close the list at any time, the return value is no better than a best guess.
+	 *	As typically used, the list is only closed once it is lifetime, so a return of true here means it is closed forever.
+	 */
+	FORCEINLINE bool IsClosed() const  // caution, if this returns false, that does not mean the list is open!
+	{
+		FPlatformMisc::MemoryBarrier();
+		return IsClosed(Head);
+	}
+#if 0
+	/**	
+	 *	Check if the list is empty.
+	 *
+	 *	@return true if the list is empty.
+	 *	CAUTION: This methods safety depends on external assumptions. For example, if another thread could add to the list at any time, the return value is no better than a best guess.
+	 *	As typically used, the list is not being access concurrently when this is called.
+	 */
+	bool IsEmpty() const
+	{
+		FPlatformMisc::MemoryBarrier();
+		return ClearClosed(Head) == Tail;
+	}
+
+	/**	
+	 *	Not thread safe, used to reset the list for recycling without freeing the stub
+	 *	@return true if the list is closed.
+	 */
+	void Reset()
+	{
+		Head = ClearClosed(Head);
+		FPlatformMisc::MemoryBarrier();
+		checkLockFreePointerList(Head == Tail); // we don't clear the list here, we assume it is already clear and possibly closed
+	}
+#endif
+private:
+	void Push(T* Link)
+	{
+		Link->LockFreePointerQueueNext = nullptr;
+		T* Prev = (T*)FPlatformAtomics::InterlockedExchangePtr((void**)&Head, Link);
+		TestCriticalStall();
+		Prev->LockFreePointerQueueNext = Link;
+		//FPlatformMisc::MemoryBarrier();
+	}
+
+	// we will use the lowest bit for the closed state
+	FORCEINLINE static bool IsClosed(T* HeadToTest)
+	{
+		return !!(UPTRINT(HeadToTest) & 1);
+	}
+	FORCEINLINE static T* SetClosed(T* In)
+	{
+		return (T*)(UPTRINT(In) | 1);
+	}
+	FORCEINLINE static T* ClearClosed(T* In)
+	{
+		return (T*)(UPTRINT(In) & ~UPTRINT(1));
+	}
+
+
+	T* Head;
+	T* Tail;
+
+	T* GetStub() const
+	{
+		return (T*)&StubStorage; // we aren't going to do anything with this other than mess with internal links
+	}
+	TAlignedBytes<sizeof(T), ALIGNOF(void*)> StubStorage;
+};
+
+
+class FDequeueCache : public FNoncopyable
+{
+	enum
+	{
+		DequeueCacheSize = 256, // probably needs to be a power of two so wraparound works right
+	};
+
+	typedef uint32 TCounter; // must be unsigned
+	typedef int32 TSignedCounter; // must be signed
+
+public:
+	FDequeueCache()
+		: NumPopped(0)
+		, NumPushed(0)
+	{
+		FMemory::Memzero(&Available[0], sizeof(void*) * DequeueCacheSize);
+	}
+
+	FORCEINLINE void* Pop()
+	{
+		int32 SpinCount = 0;
+		while (true)
+		{
+			TCounter LocalNumPushed = NumPushed;
+			TCounter LocalNumPopped = NumPopped;
+			if (TSignedCounter(LocalNumPushed - LocalNumPopped) < 0) 
+			{
+				LockFreeCriticalSpin(SpinCount);
+				continue;
+			}
+			if (LocalNumPushed == LocalNumPopped)
+			{
+				break;
+			}
+			uint32 Slot = LocalNumPopped % DequeueCacheSize;
+			void *ReturnPtr = Available[Slot];
+			TCounter Result = (TCounter)FPlatformAtomics::InterlockedCompareExchange((volatile TSignedCounter*)&NumPopped, LocalNumPopped + 1, LocalNumPopped);
+			if (Result == LocalNumPopped)
+			{
+				check(ReturnPtr);
+				return ReturnPtr;
+			}
+		}
+		return nullptr;
+	}
+
+	FORCEINLINE bool IsEmpty() const
+	{
+		return NumPushed == NumPopped;
+	}
+
+	FORCEINLINE bool IsFull() const
+	{
+		return NumPushed - NumPopped >= TCounter(DequeueCacheSize);
+	}
+
+	// must not be called concurrently!
+	FORCEINLINE bool Push(void* Item)
+	{
+		TCounter LocalNumPushed = NumPushed;
+		TCounter LocalNumPopped = NumPopped;
+		if (LocalNumPushed - LocalNumPopped >= TCounter(DequeueCacheSize))
+		{
+			return false;
+		}
+		check(TSignedCounter(LocalNumPushed - LocalNumPopped) >= 0); // wraparound is ok, this would mean we have at least 16GB of queue for 32 bits...which is not possible
+		Available[LocalNumPushed % DequeueCacheSize] = Item;
+		FPlatformMisc::MemoryBarrier();
+		verify(++NumPushed == LocalNumPushed + 1);
+		return true;
+	}
+
+private:
+
+	TCounter NumPopped;
+	void *Available[DequeueCacheSize];
+	TCounter NumPushed;
+};
+
+class FSpinLocked_LIFO
+{
+public:
+	void** Pop()
+	{
+		int32 SpinCount = 0;
+		while (true)
+		{
+			if (!Lock && FPlatformAtomics::InterlockedCompareExchange((volatile int32*)&Lock, 1, 0) == 0)
+			{
+				TestCriticalStall();
+				void** Result = nullptr;
+				if (Queue.Num())
+				{
+					Result = Queue.Pop(false);
+				}
+				FPlatformMisc::MemoryBarrier();
+				Lock = 0;
+				return Result;
+			}
+			LockFreeCriticalSpin(SpinCount);
+		}
+	}
+	void Push(void** Item)
+	{
+		int32 SpinCount = 0;
+		while (true)
+		{
+			if (!Lock && FPlatformAtomics::InterlockedCompareExchange((volatile int32*)&Lock, 1, 0) == 0)
+			{
+				TestCriticalStall();
+				Queue.Add(Item);
+				FPlatformMisc::MemoryBarrier();
+				Lock = 0;
+				return;
+			}
+			LockFreeCriticalSpin(SpinCount);
+		}
+	}	
+private:
+	TArray<void **> Queue;
+	int32 Lock;
+};
+
+class FLockFreePointerQueueBaseLinkAllocator : public FNoncopyable
+{
+public:
+	static CORE_API FLockFreePointerQueueBaseLinkAllocator& Get();
+
+	FORCEINLINE void Free(FLockFreeLink* Link)
+	{
+		Link->~FLockFreeLink();
+		InnerAllocator.Free(Link);
+	}
+
+	FORCEINLINE FLockFreeLink* Alloc()
+	{
+		return new (InnerAllocator.Allocate()) FLockFreeLink();
+	}
+
+	/**
+	 * Gets the number of allocated memory blocks that are currently in use.
+	 *
+	 * @return Number of used memory blocks.
+	 */
+	const int32 GetNumUsed() const
+	{
+		return InnerAllocator.GetNumUsed().GetValue();
+	}
+
+	/**
+	 * Gets the number of allocated memory blocks that are currently unused.
+	 *
+	 * @return Number of allocated, but unused memory blocks.
+	 */
+	const int32 GetNumFree() const
+	{
+		return InnerAllocator.GetNumFree().GetValue();
+	}
+
+private:
+	TLockFreeFixedSizeAllocator_TLSCacheBase<sizeof(FLockFreeLink), FSpinLocked_LIFO,
+#if 1 || UE_BUILD_SHIPPING || UE_BUILD_TEST
+		FNoopCounter
+#else
+		FThreadSafeCounter
+#endif
+	> InnerAllocator;
+};
+
+
+template<class T>
+class FLockFreePointerListFIFOIntrusive : public FNoncopyable
+{
+public:
+
+	FLockFreePointerListFIFOIntrusive()
+		: DequeueLock(0)
+	{
+	}
+
+	void Push(T *Link)
+	{
+		IncomingQueue.Push(Link);
+	}
+
+	FORCEINLINE T* Pop(bool bOkToSpeculativelyReturnNull = false)
+	{
+		int32 SpinCount = 0;
+		while (true)
+		{
+			T* Result = (T*)DequeueCache.Pop();
+			if (Result)
+			{
+				return Result;
+			}
+			if (bOkToSpeculativelyReturnNull && IncomingQueue.IsEmpty())
+			{
+				return nullptr;
+			}
+
+			if (!DequeueLock && FPlatformAtomics::InterlockedCompareExchange((volatile int32*)&DequeueLock, 1, 0) == 0)
+			{
+				TestCriticalStall();
+				Result = (T*)DequeueCache.Pop();
+				if (Result)
+				{
+					// someone beat us to it
+					FPlatformMisc::MemoryBarrier();
+					DequeueLock = 0;
+					return Result;
+				}
+				T* Link = IncomingQueue.Pop();
+				if (!Link)
+				{
+					FPlatformMisc::MemoryBarrier();
+					DequeueLock = 0;
+					return nullptr;
+				}
+				check(!DequeueCache.IsFull()); // can't be full because it was empty and we have the lock
+				verify(DequeueCache.Push(Link));
+				while (!DequeueCache.IsFull())
+				{
+					T* Repush = IncomingQueue.Pop(); 
+					if (!Repush)
+					{
+						break;
+					}
+					verify(DequeueCache.Push(Repush));
+				}
+				FPlatformMisc::MemoryBarrier();
+				DequeueLock = 0;
+			}
+			LockFreeCriticalSpin(SpinCount);
+		}
+	}
+
+	/**	
+	 *	Check if the list is empty.
+	 *
+	 *	@return true if the list is empty.
+	 *	CAUTION: This methods safety depends on external assumptions. For example, if another thread could add to the list at any time, the return value is no better than a best guess.
+	 *	As typically used, the list is not being access concurrently when this is called.
+	 */
+	bool IsEmpty()  
+	{
+		int32 SpinCount = 0;
+		bool bResult = true;
+		while (true)
+		{
+			if (!DequeueCache.IsEmpty())
+			{
+				bResult = false;
+				break;
+			}
+			if (!DequeueLock && FPlatformAtomics::InterlockedCompareExchange((volatile int32*)&DequeueLock, 1, 0) == 0)
+			{
+				TestCriticalStall();
+				bool bAny = false;
+				if (!DequeueCache.IsFull())
+				{
+					T* Repush = IncomingQueue.Pop(); 
+					if (!Repush)
+					{
+						break;
+					}
+					bResult = false;
+					verify(DequeueCache.Push(Repush));
+				}
+				else
+				{
+					bResult = false; // we didn't pop one because the dequeue cache is full...guess it isn't empty
+				}
+				FPlatformMisc::MemoryBarrier();
+				DequeueLock = 0;
+				break;
+			}
+			LockFreeCriticalSpin(SpinCount);
+		}
+		return bResult;
+	}
+	/**	
+	 *	Check if the list is empty. This is a faster, less rigorous test that can miss a queue full of stuff. Used to optimize thread restarts.
+	 */
+	bool IsEmptyFast()  
+	{
+		return DequeueCache.IsEmpty() && IncomingQueue.IsEmpty();
+	}
+
+	/**	
+	 *	Pop all items from the list.
+	 *
+	 *	@param Output The array to hold the returned items. Must be empty.
+	 */
+	void PopAll(TArray<T *>& Output)
+	{
+		while (true)
+		{
+			T* Item = Pop();
+			if (!Item)
+			{
+				break;
+			}
+			Output.Add(Item);
+		}
+	}	
+
+private:
+	FLockFreePointerQueueBaseSingleConsumerIntrusive<T> IncomingQueue;
+	FDequeueCache DequeueCache;
+	int32 DequeueLock;
+};
+
+// FIFO of void *'s
+class FLockFreePointerListFIFOBase : public FNoncopyable
+{
+public:
+
+	FLockFreePointerListFIFOBase()
+		: LinkAllocator(FLockFreePointerQueueBaseLinkAllocator::Get())
+		, IncomingQueue(LinkAllocator.Alloc())
+		, DequeueLock(0)
+	{
+	}
+
+	void Push(void *Item)
+	{
+		FLockFreeLink* Link = LinkAllocator.Alloc();
+		Link->Payload = Item;
+		IncomingQueue.Push(Link);
+	}
+
+	void* Pop(bool bOkToSpeculativelyReturnNull = false)
+	{
+		int32 SpinCount = 0;
+		while (true)
+		{
+			void* Result = DequeueCache.Pop();
+			if (Result)
+			{
+				return Result;
+			}
+			if (bOkToSpeculativelyReturnNull && IncomingQueue.IsEmpty())
+			{
+				return nullptr;
+			}
+			if (!DequeueLock && FPlatformAtomics::InterlockedCompareExchange((volatile int32*)&DequeueLock, 1, 0) == 0)
+			{
+				TestCriticalStall();
+				Result = DequeueCache.Pop();
+				if (Result)
+				{
+					// someone beat us to it
+					FPlatformMisc::MemoryBarrier();
+					DequeueLock = 0;
+					return Result;
+				}
+				FLockFreeLink* Link = IncomingQueue.Pop();
+				if (!Link)
+				{
+					FPlatformMisc::MemoryBarrier();
+					DequeueLock = 0;
+					return nullptr;
+				}
+				check(!DequeueCache.IsFull() && Link->Payload); // can't be full because it was empty and we have the lock
+				verify(DequeueCache.Push(Link->Payload));
+				LinkAllocator.Free(Link);
+				while (!DequeueCache.IsFull())
+				{
+					FLockFreeLink* Repush = IncomingQueue.Pop(); 
+					if (!Repush)
+					{
+						break;
+					}
+					check(Repush->Payload);
+					verify(DequeueCache.Push(Repush->Payload));
+					LinkAllocator.Free(Repush);
+				}
+				FPlatformMisc::MemoryBarrier();
+				DequeueLock = 0;
+			}
+			LockFreeCriticalSpin(SpinCount);
+		}
+	}
+
+	/**	
+	 *	Check if the list is empty.
+	 *
+	 *	@return true if the list is empty.
+	 *	CAUTION: This methods safety depends on external assumptions. For example, if another thread could add to the list at any time, the return value is no better than a best guess.
+	 *	As typically used, the list is not being access concurrently when this is called.
+	 */
+	bool IsEmpty()  
+	{
+		bool bResult = true;
+		int32 SpinCount = 0;
+		while (true)
+		{
+			if (!DequeueCache.IsEmpty())
+			{
+				bResult = false;
+				break;
+			}
+			if (!DequeueLock && FPlatformAtomics::InterlockedCompareExchange((volatile int32*)&DequeueLock, 1, 0) == 0)
+			{
+				TestCriticalStall();
+				bool bAny = false;
+				if (!DequeueCache.IsFull())
+				{
+					FLockFreeLink* Repush = IncomingQueue.Pop(); 
+					if (!Repush)
+					{
+						break;
+					}
+					bResult = false;
+					check(Repush->Payload);
+					verify(DequeueCache.Push(Repush->Payload));
+					LinkAllocator.Free(Repush);
+				}
+				else
+				{
+					bResult = false; // we didn't pop one because the dequeue cache is full...guess it isn't empty
+				}
+				FPlatformMisc::MemoryBarrier();
+				DequeueLock = 0;
+				break;
+			}
+			LockFreeCriticalSpin(SpinCount);
+		}
+		return bResult;
+	}
+	/**	
+	 *	Check if the list is empty. This is a faster, less rigorous test that can miss a queue full of stuff. Used to optimize thread restarts.
+	 */
+	bool IsEmptyFast()  
+	{
+		return DequeueCache.IsEmpty() && IncomingQueue.IsEmpty();
+	}
+
+private:
+	FLockFreePointerQueueBaseLinkAllocator& LinkAllocator;
+	FLockFreePointerQueueBaseSingleConsumer IncomingQueue;
+	FDequeueCache DequeueCache;
+	int32 DequeueLock;
+};
+
+// closeable FIFO of void *'s
+class FCloseableLockFreePointerListFIFOBaseSingleConsumer : public FNoncopyable
+{
+public:
+
+	FCloseableLockFreePointerListFIFOBaseSingleConsumer()
+		: LinkAllocator(FLockFreePointerQueueBaseLinkAllocator::Get())
+		, IncomingQueue(LinkAllocator.Alloc())
+	{
+	}
+	/**	
+	 *	Push an item onto the head of the list, unless the list is closed
+	 *
+	 *	@param NewItem, the new item to push on the list, cannot be NULL
+	 *	@return true if the item was pushed on the list, false if the list was closed.
+	 */
+	bool PushIfNotClosed(void *NewItem)
+	{
+		FLockFreeLink* Link = LinkAllocator.Alloc();
+		Link->Payload = NewItem;
+		if (!IncomingQueue.PushIfNotClosed(Link))
+		{
+			LinkAllocator.Free(Link);
+			return false;
+		}
+		return true;
+	}
+
+	/**	
+	 *	Pop all items from the list and atomically close it.
+	 *
+	 *	@param Output The array to hold the returned items. Must be empty.
+	 */
+	template<class T>
+	void PopAllAndClose(TArray<T *>& Output)
+	{
+		while (true)
+		{
+			FLockFreeLink* Link = IncomingQueue.Pop();
+			if (Link)
+			{
+				Output.Add((T*)Link->Payload);
+				LinkAllocator.Free(Link);
+			}
+			else if (IncomingQueue.CloseIfEmpty())
+			{
+				break;
+			}
+		}
+	}	
+
+	/**	
+	 *	Check if the list is closed
+	 *
+	 *	@return true if the list is closed.
+	 */
+	FORCEINLINE bool IsClosed() const
+	{
+		return IncomingQueue.IsClosed();
+	}
+
+	/**	
+	 *	Push an item onto the head of the list, opening it first if necessary.
+	 *
+	 *	@param NewItem, the new item to push on the list, cannot be NULL.
+	 *	@return true if the list needed to be opened first, false if the list was not closed before our push.
+	 */
+	bool ReopenIfClosedAndPush(void *NewItem)
+	{
+		FLockFreeLink* Link = LinkAllocator.Alloc();
+		Link->Payload = NewItem;
+		bool bWasReopenedByMe = IncomingQueue.ReopenIfClosedAndPush(Link);
+		return bWasReopenedByMe;
+	}
+
+	/**	
+	 *	Pop an item from the list or return NULL if the list is empty or closed.
+	 *	CAUTION: This method should not be used unless the list is known to not be closed.
+	 *	@return The new item, if any
+	 */
+	void* Pop()
+	{
+		FLockFreeLink* Link = IncomingQueue.Pop();
+		if (Link)
+		{
+			void* Result = Link->Payload;
+			LinkAllocator.Free(Link);
+			return Result;
+		}
+		return nullptr;
+	}
+
+	/**	
+	 *	Close the list if it is empty.
+	 *
+	 *	@return true if this call actively closed the list.
+	 */
+	bool CloseIfEmpty()
+	{
+		return IncomingQueue.CloseIfEmpty();
+	}
+
+	/**	
+	 *	Not thread safe, used to reset the list for recycling without freeing the stub
+	 *	@return true if the list is closed.
+	 */
+	FORCEINLINE void Reset()
+	{
+		return IncomingQueue.Reset();
+	}
+
+private:
+	FLockFreePointerQueueBaseLinkAllocator& LinkAllocator;
+	FCloseableLockFreePointerQueueBaseSingleConsumer IncomingQueue;
+};
+
+
+template <class T>
+class TLockFreePointerListFIFO : private FLockFreePointerListFIFOBase 
+{
+public:
+	/**	
+	 *	Push an item onto the head of the list.
+	 *
+	 *	@param NewItem, the new item to push on the list, cannot be NULL.
+	 */
+	FORCEINLINE void Push(T *NewItem)
+	{
+		FLockFreePointerListFIFOBase::Push(NewItem);
+	}
+
+	/**	
+	 *	Pop an item from the list or return NULL if the list is empty.
+	 *	@return The popped item, if any.
+	 */
+	FORCEINLINE T* Pop()
+	{
+		return (T*)FLockFreePointerListFIFOBase::Pop();
+	}
+
+	/**	
+	 *	Pop all items from the list.
+	 *
+	 *	@param Output The array to hold the returned items. Must be empty.
+	 */
+	void PopAll(TArray<T *>& Output)
+	{
+		while (true)
+		{
+			T* Item = Pop();
+			if (!Item)
+			{
+				break;
+			}
+			Output.Add(Item);
+		}
+	}	
+
+	/**	
+	 *	Check if the list is empty.
+	 *
+	 *	@return true if the list is empty.
+	 *	CAUTION: This methods safety depends on external assumptions. For example, if another thread could add to the list at any time, the return value is no better than a best guess.
+	 *	As typically used, the list is not being access concurrently when this is called.
+	 */
+	FORCEINLINE bool IsEmpty() const  
+	{
+		return FLockFreePointerListFIFOBase::IsEmpty();
+	}
+	/**	
+	 *	Check if the list is empty. This is a faster, less rigorous test that can miss a queue full of stuff. Used to optimize thread restarts.
+	 */
+	FORCEINLINE bool IsEmptyFast()  
+	{
+		return FLockFreePointerListFIFOBase::IsEmptyFast();
+	}
+};
+
+#if USE_NEW_LOCK_FREE_LISTS
+
+// a list where you don't care what order they pop in. We choose the fastest implementation we have.
+template <class T>
+class TLockFreePointerListUnordered : public TLockFreePointerListFIFO<T>
+{
+
+};
+
+#endif
+
+
diff --git a/Engine/Source/Runtime/Core/Public/Stats/Stats2.h b/Engine/Source/Runtime/Core/Public/Stats/Stats2.h
--- a/Engine/Source/Runtime/Core/Public/Stats/Stats2.h
+++ b/Engine/Source/Runtime/Core/Public/Stats/Stats2.h
@@ -1164,7 +1164,7 @@ private:
 	};
 
 	/** Lock free pool of FThreadStats instances. */
-	TLockFreePointerList<FThreadStats> Pool;
+	TLockFreePointerListUnordered<FThreadStats> Pool;
 
 public:
 	/** Default constructor. */
diff --git a/Engine/Source/Runtime/Core/Public/Templates/RefCounting.h b/Engine/Source/Runtime/Core/Public/Templates/RefCounting.h
--- a/Engine/Source/Runtime/Core/Public/Templates/RefCounting.h
+++ b/Engine/Source/Runtime/Core/Public/Templates/RefCounting.h
@@ -53,7 +53,7 @@ class TRefCountPtr
 
 public:
 
-	TRefCountPtr():
+	FORCEINLINE TRefCountPtr():
 		Reference(nullptr)
 	{ }
 
@@ -75,6 +75,12 @@ public:
 		}
 	}
 
+	FORCEINLINE TRefCountPtr(TRefCountPtr&& Copy)
+	{
+		Reference = Copy.Reference;
+		Copy.Reference = nullptr;
+	}
+
 	~TRefCountPtr()
 	{
 		if(Reference)
@@ -99,33 +105,48 @@ public:
 		return *this;
 	}
 
-	TRefCountPtr& operator=(const TRefCountPtr& InPtr)
+	FORCEINLINE TRefCountPtr& operator=(const TRefCountPtr& InPtr)
 	{
 		return *this = InPtr.Reference;
 	}
 
-	ReferencedType* operator->() const
+	TRefCountPtr& operator=(TRefCountPtr&& InPtr)
+	{
+		if (this != &InPtr)
+		{
+			ReferencedType* OldReference = Reference;
+			Reference = InPtr.Reference;
+			InPtr.Reference = nullptr;
+			if(OldReference)
+			{
+				OldReference->Release();
+			}
+		}
+		return *this;
+	}
+
+	FORCEINLINE ReferencedType* operator->() const
 	{
 		return Reference;
 	}
 
-	operator ReferenceType() const
+	FORCEINLINE operator ReferenceType() const
 	{
 		return Reference;
 	}
 
-	ReferencedType** GetInitReference()
+	FORCEINLINE ReferencedType** GetInitReference()
 	{
 		*this = nullptr;
 		return &Reference;
 	}
 
-	ReferencedType* GetReference() const
+	FORCEINLINE ReferencedType* GetReference() const
 	{
 		return Reference;
 	}
 
-	friend bool IsValidRef(const TRefCountPtr& InReference)
+	FORCEINLINE friend bool IsValidRef(const TRefCountPtr& InReference)
 	{
 		return InReference.Reference != nullptr;
 	}
@@ -135,25 +156,23 @@ public:
 		return Reference != nullptr;
 	}
 
-	void SafeRelease()
+	FORCEINLINE void SafeRelease()
 	{
 		*this = nullptr;
 	}
 
 	uint32 GetRefCount()
 	{
-		if(Reference)
-		{
-			Reference->AddRef();
-			return Reference->Release();
-		}
-		else
+		uint32 Result = 0;
+		if (Reference)
 		{
-			return 0;
+			Result = Reference->GetRefCount();
+			check(Result > 0); // you should never have a zero ref count if there is a live ref counted pointer (*this is live)
 		}
+		return Result;
 	}
 
-	void Swap(TRefCountPtr& InPtr) // this does not change the reference count, and so is faster
+	FORCEINLINE void Swap(TRefCountPtr& InPtr) // this does not change the reference count, and so is faster
 	{
 		ReferencedType* OldReference = Reference;
 		Reference = InPtr.Reference;
diff --git a/Engine/Source/Runtime/CoreUObject/Private/UObject/GarbageCollection.cpp b/Engine/Source/Runtime/CoreUObject/Private/UObject/GarbageCollection.cpp
--- a/Engine/Source/Runtime/CoreUObject/Private/UObject/GarbageCollection.cpp
+++ b/Engine/Source/Runtime/CoreUObject/Private/UObject/GarbageCollection.cpp
@@ -268,7 +268,7 @@ public:
 private:
 
 	/** Holds the collection of recycled arrays. */
-	TLockFreePointerList< TArray<UObject*> > Pool;
+	TLockFreePointerListUnordered< TArray<UObject*> > Pool;
 
 #if UE_BUILD_DEBUG
 	/** Number of arrays currently acquired from the pool by GC */
diff --git a/Engine/Source/Runtime/CoreUObject/Public/UObject/UObjectArray.h b/Engine/Source/Runtime/CoreUObject/Public/UObject/UObjectArray.h
--- a/Engine/Source/Runtime/CoreUObject/Public/UObject/UObjectArray.h
+++ b/Engine/Source/Runtime/CoreUObject/Public/UObject/UObjectArray.h
@@ -401,7 +401,7 @@ private:
 	/** Synchronization object for all live objects.											*/
 	FCriticalSection ObjObjectsCritical;
 	/** Available object indices.											*/
-	TLockFreePointerList<int32> ObjAvailableList;
+	TLockFreePointerListUnordered<int32> ObjAvailableList;
 #if WITH_EDITOR
 	/** Available object index count.										*/
 	FThreadSafeCounter ObjAvailableCount;
diff --git a/Engine/Source/Runtime/Engine/Private/Particles/ParticleGpuSimulation.cpp b/Engine/Source/Runtime/Engine/Private/Particles/ParticleGpuSimulation.cpp
--- a/Engine/Source/Runtime/Engine/Private/Particles/ParticleGpuSimulation.cpp
+++ b/Engine/Source/Runtime/Engine/Private/Particles/ParticleGpuSimulation.cpp
@@ -2515,10 +2515,10 @@ public:
 	}
 };
 
-static struct FNewParticleAlloc
+struct FNewParticleAlloc
 {
 	TLockFreeFixedSizeAllocator<sizeof(TArray<FNewParticle>), FThreadSafeCounter> FreeTArrayFNewParticleArrays;
-	TLockFreePointerList<TArray<FNewParticle>>	FreeTArrayFNewParticle;
+	TLockFreePointerListUnordered<TArray<FNewParticle>>	FreeTArrayFNewParticle;
 	~FNewParticleAlloc()
 	{
 		while (true)
@@ -2533,7 +2533,13 @@ static struct FNewParticleAlloc
 		}
 	}
 
-} GNewParticleAlloc;
+};
+
+static FNewParticleAlloc& GNewParticleAlloc()
+{
+	static FNewParticleAlloc Singleton;
+	return Singleton;
+}
 
 static const int MaxNumParticlesToRecycle  = 512; // these can be quite large
 
@@ -2543,13 +2549,13 @@ static void FreeNewParticleArray(TArray<FNewParticle>& NewParticles)
 	NewParticles.Reset();
 	const int MaxNumToRecycledArrays  = 100; 
 	int32 CurrentSize = NewParticles.GetSlack();
-	if (CurrentSize > 0 && CurrentSize <= MaxNumParticlesToRecycle && GNewParticleAlloc.FreeTArrayFNewParticleArrays.GetNumUsed().GetValue() < MaxNumToRecycledArrays)
+	if (CurrentSize > 0 && CurrentSize <= MaxNumParticlesToRecycle && GNewParticleAlloc().FreeTArrayFNewParticleArrays.GetNumUsed().GetValue() < MaxNumToRecycledArrays)
 	{
-		TArray<FNewParticle>* Recycle = new (GNewParticleAlloc.FreeTArrayFNewParticleArrays.Allocate()) TArray<FNewParticle>;
+		TArray<FNewParticle>* Recycle = new (GNewParticleAlloc().FreeTArrayFNewParticleArrays.Allocate()) TArray<FNewParticle>;
 		Exchange(*Recycle, NewParticles);
 		check(Recycle->Num() == 0 && Recycle->GetSlack());
 		check(NewParticles.Num() == 0 && NewParticles.GetSlack() == 0);
-		GNewParticleAlloc.FreeTArrayFNewParticle.Push(Recycle);
+		GNewParticleAlloc().FreeTArrayFNewParticle.Push(Recycle);
 	}
 }
 
@@ -2557,12 +2563,12 @@ static void GetNewParticleArray(TArray<FNewParticle>& NewParticles, int32 NumPar
 {
 	if (NumParticlesNeeded <= MaxNumParticlesToRecycle)
 	{
-		TArray<FNewParticle>* Recycle = GNewParticleAlloc.FreeTArrayFNewParticle.Pop();
+		TArray<FNewParticle>* Recycle = GNewParticleAlloc().FreeTArrayFNewParticle.Pop();
 		if (Recycle)
 		{
 			Exchange(*Recycle, NewParticles);
 			Recycle->~TArray<FNewParticle>(); // this probably doesn't do anything, but type safety and all
-			GNewParticleAlloc.FreeTArrayFNewParticleArrays.Free(Recycle);
+			GNewParticleAlloc().FreeTArrayFNewParticleArrays.Free(Recycle);
 			check(NewParticles.Num() == 0 && NewParticles.GetSlack());
 		}
 	}
diff --git a/Engine/Source/Runtime/Engine/Private/PhysicsEngine/PhysScene.cpp b/Engine/Source/Runtime/Engine/Private/PhysicsEngine/PhysScene.cpp
--- a/Engine/Source/Runtime/Engine/Private/PhysicsEngine/PhysScene.cpp
+++ b/Engine/Source/Runtime/Engine/Private/PhysicsEngine/PhysScene.cpp
@@ -728,7 +728,9 @@ void FPhysScene::TickPhysScene(uint32 SceneType, FGraphEventRef& InOutCompletion
 
 	if (!bTaskOutstanding)
 	{
-		InOutCompletionEvent->DispatchSubsequents(); // nothing to do, so nothing to wait for
+		check(IsInGameThread());
+		static TArray<FBaseGraphTask*> NewTasks;
+		InOutCompletionEvent->DispatchSubsequents(NewTasks, ENamedThreads::GameThread); // nothing to do, so nothing to wait for
 	}
 #if WITH_SUBSTEPPING
 	bSubstepping = UPhysicsSettings::Get()->bSubstepping;
diff --git a/Engine/Source/Runtime/Engine/Private/PhysicsEngine/PhysSubstepTasks.h b/Engine/Source/Runtime/Engine/Private/PhysicsEngine/PhysSubstepTasks.h
--- a/Engine/Source/Runtime/Engine/Private/PhysicsEngine/PhysSubstepTasks.h
+++ b/Engine/Source/Runtime/Engine/Private/PhysicsEngine/PhysSubstepTasks.h
@@ -9,7 +9,7 @@ void FinishSceneStat(uint32 Scene);
 #include "pxtask/PxTask.h"
 class PhysXCompletionTask : public PxLightCpuTask
 {
-	FGraphEventRef& EventToFire;
+	FGraphEventRef EventToFire;
 	uint32 Scene;
 public:
 	PhysXCompletionTask(FGraphEventRef& InEventToFire, uint32 InScene, PxTaskManager* TaskManager)
@@ -25,7 +25,11 @@ public:
 	{
 		PxLightCpuTask::release();
 		FinishSceneStat(Scene);
-		EventToFire->DispatchSubsequents();
+		if (EventToFire.GetReference())
+		{
+			TArray<FBaseGraphTask*> NewTasks;
+			EventToFire->DispatchSubsequents(NewTasks);
+		}
 		delete this;
 	}
 	virtual const char *getName() const
diff --git a/Engine/Source/Runtime/Engine/Private/SkeletalRenderGPUSkin.cpp b/Engine/Source/Runtime/Engine/Private/SkeletalRenderGPUSkin.cpp
--- a/Engine/Source/Runtime/Engine/Private/SkeletalRenderGPUSkin.cpp
+++ b/Engine/Source/Runtime/Engine/Private/SkeletalRenderGPUSkin.cpp
@@ -1011,7 +1011,7 @@ const FTwoVectors& FSkeletalMeshObjectGPUSkin::GetCustomLeftRightVectors(int32 S
 FDynamicSkelMeshObjectDataGPUSkin
 -----------------------------------------------------------------------------*/
 
-static TLockFreePointerList<FDynamicSkelMeshObjectDataGPUSkin>	FreeDynamicSkelMeshObjectDataGPUSkins;
+static TLockFreePointerListUnordered<FDynamicSkelMeshObjectDataGPUSkin>	FreeDynamicSkelMeshObjectDataGPUSkins;
 
 void FDynamicSkelMeshObjectDataGPUSkin::Clear()
 {
diff --git a/Engine/Source/Runtime/Engine/Private/UnrealEngine.cpp b/Engine/Source/Runtime/Engine/Private/UnrealEngine.cpp
--- a/Engine/Source/Runtime/Engine/Private/UnrealEngine.cpp
+++ b/Engine/Source/Runtime/Engine/Private/UnrealEngine.cpp
@@ -1104,6 +1104,7 @@ struct FABTest
 {
 	FRandomStream Stream;
 	bool bABTestActive;
+	bool bFrameLog;
 	FString ABTestCmds[2];
 	int32 ABTestNumSamples;
 	int32 RemainingCoolDown;
@@ -1131,10 +1132,14 @@ struct FABTest
 	uint32 Totals[2];
 	uint32 Counts[2];
 
+	double TotalTime;
+	uint32 TotalFrames;
+	uint32 Spikes;
 
 	FABTest()
 		: Stream(9567)
 		, bABTestActive(false)
+		, bFrameLog(false)
 		, ABTestNumSamples(0)
 		, RemainingCoolDown(0)
 		, CurrentTest(0)
@@ -1149,6 +1154,31 @@ struct FABTest
 	{
 	}
 
+	void StartFrameLog()
+	{
+		TotalTime = 0.0;
+		TotalFrames = 0;
+		Spikes = 0;
+		bFrameLog = true;
+		UE_LOG(LogConsoleResponse, Display, TEXT("Starting frame log."));
+	}
+
+	void FrameLogTick(double Delta)
+	{
+		if (Delta > .034)
+		{
+			Spikes++;
+		}
+		else
+		{
+			TotalFrames++;
+			TotalTime += Delta;
+		}
+		if (TotalFrames > 0 && TotalFrames % 1000 == 0)
+		{
+			FPlatformMisc::LowLevelOutputDebugStringf(TEXT("%8d frames   %6.3fms/f    %8d spikes rejected "), TotalFrames, float(1000.0 * TotalTime / TotalFrames), Spikes);
+		}
+	}
 	void Tick()
 	{
 		static double LastTime = 0;
@@ -1270,6 +1300,10 @@ struct FABTest
 				SwitchTest(1 - CurrentTest);
 			}
 		}
+		else if (bFrameLog)
+		{
+			FrameLogTick(FPlatformTime::Seconds() - LastTime);
+		}
 		LastTime = FPlatformTime::Seconds();
 	}
 
@@ -1281,6 +1315,12 @@ struct FABTest
 	static void ABTestCmdFunc(const TArray<FString>& Args)
 	{
 		FString ABTestCmds[2];
+		if (Args.Num() == 1 && Args[0].Compare(FString(TEXT("framelog")), ESearchCase::IgnoreCase) == 0)
+		{
+			Get().Stop();
+			Get().StartFrameLog();
+			return;
+		}
 		if (Args.Num() == 1 && Args[0].Compare(FString(TEXT("stop")), ESearchCase::IgnoreCase) == 0)
 		{
 			Get().Stop();
@@ -1330,9 +1370,17 @@ struct FABTest
 
 	void Stop()
 	{
-		UE_LOG(LogConsoleResponse, Display, TEXT("Running 'A' console command and stopping test."));
-		SwitchTest(0);
-		bABTestActive = false;
+		if (bABTestActive)
+		{
+			UE_LOG(LogConsoleResponse, Display, TEXT("Running 'A' console command and stopping test."));
+			SwitchTest(0);
+			bABTestActive = false;
+		}
+		else if (bFrameLog)
+		{
+			UE_LOG(LogConsoleResponse, Display, TEXT("Stopping frame log."));
+			bFrameLog = false;
+		}
 	}
 	
 	void Start(FString* InABTestCmds)
diff --git a/Engine/Source/Runtime/Engine/Public/DistanceFieldAtlas.h b/Engine/Source/Runtime/Engine/Public/DistanceFieldAtlas.h
--- a/Engine/Source/Runtime/Engine/Public/DistanceFieldAtlas.h
+++ b/Engine/Source/Runtime/Engine/Public/DistanceFieldAtlas.h
@@ -219,10 +219,12 @@ private:
 	TArray<FAsyncDistanceFieldTask*> ReferencedTasks;
 
 	/** Tasks that have not yet started processing yet. */
-	TLockFreePointerList<FAsyncDistanceFieldTask> TaskQueue;
+	// consider changing this from FIFO to Unordered, which may be faster
+	TLockFreePointerListLIFO<FAsyncDistanceFieldTask> TaskQueue;
 
 	/** Tasks that have completed processing. */
-	TLockFreePointerList<FAsyncDistanceFieldTask> CompletedTasks;
+	// consider changing this from FIFO to Unordered, which may be faster
+	TLockFreePointerListLIFO<FAsyncDistanceFieldTask> CompletedTasks;
 
 	class IMeshUtilities* MeshUtilities;
 
diff --git a/Engine/Source/Runtime/PakFile/Private/SignedArchiveReader.h b/Engine/Source/Runtime/PakFile/Private/SignedArchiveReader.h
--- a/Engine/Source/Runtime/PakFile/Private/SignedArchiveReader.h
+++ b/Engine/Source/Runtime/PakFile/Private/SignedArchiveReader.h
@@ -115,7 +115,7 @@ class FChunkCacheWorker : public FRunnable
 	/** Stops this thread */
 	FThreadSafeCounter StopTaskCounter;
 	/** Available chunk requests */
-	TLockFreePointerList<FChunkRequest> FreeChunkRequests;
+	TLockFreePointerListUnordered<FChunkRequest> FreeChunkRequests;
 	/** Public decryption key */
 	FEncryptionKey DecryptionKey;
 
diff --git a/Engine/Source/Runtime/RHI/Private/RHI.cpp b/Engine/Source/Runtime/RHI/Private/RHI.cpp
--- a/Engine/Source/Runtime/RHI/Private/RHI.cpp
+++ b/Engine/Source/Runtime/RHI/Private/RHI.cpp
@@ -72,7 +72,7 @@ const FClearValueBinding FClearValueBinding::DepthNear((float)ERHIZBuffer::NearP
 const FClearValueBinding FClearValueBinding::DepthFar((float)ERHIZBuffer::FarPlane, 0);
 
 
-TLockFreePointerList<FRHIResource> FRHIResource::PendingDeletes;
+TLockFreePointerListUnordered<FRHIResource> FRHIResource::PendingDeletes;
 FRHIResource* FRHIResource::CurrentlyDeleting = nullptr;
 
 #if !DISABLE_RHI_DEFFERED_DELETE
diff --git a/Engine/Source/Runtime/RHI/Private/RHICommandList.cpp b/Engine/Source/Runtime/RHI/Private/RHICommandList.cpp
--- a/Engine/Source/Runtime/RHI/Private/RHICommandList.cpp
+++ b/Engine/Source/Runtime/RHI/Private/RHICommandList.cpp
@@ -548,7 +548,9 @@ struct FRHICommandRHIThreadFence : public FRHICommand<FRHICommandRHIThreadFence>
 	}
 	void Execute(FRHICommandListBase& CmdList)
 	{
-		Fence->DispatchSubsequents();
+		check(IsInRHIThread());
+		static TArray<FBaseGraphTask*> NewTasks;
+		Fence->DispatchSubsequents(NewTasks, ENamedThreads::RHIThread);
 		Fence = nullptr;
 	}
 };
@@ -720,6 +722,8 @@ DECLARE_DWORD_COUNTER_STAT(TEXT("Num Async Chains Links"), STAT_ChainLinkCount,
 DECLARE_CYCLE_STAT(TEXT("Wait for Async CmdList"), STAT_ChainWait, STATGROUP_RHICMDLIST);
 DECLARE_CYCLE_STAT(TEXT("Async Chain Execute"), STAT_ChainExecute, STATGROUP_RHICMDLIST);
 
+FGraphEvent* GEventToWaitFor = nullptr;
+
 struct FRHICommandWaitForAndSubmitSubList : public FRHICommand<FRHICommandWaitForAndSubmitSubList>
 {
 	FGraphEventRef EventToWaitFor;
@@ -732,6 +736,12 @@ struct FRHICommandWaitForAndSubmitSubList : public FRHICommand<FRHICommandWaitFo
 	void Execute(FRHICommandListBase& CmdList)
 	{
 		INC_DWORD_STAT_BY(STAT_ChainLinkCount, 1);
+		if (EventToWaitFor.GetReference() && !EventToWaitFor->IsComplete() && !(!GRHIThread || !IsInRHIThread()))
+		{
+			GEventToWaitFor = EventToWaitFor.GetReference();
+			FPlatformMisc::DebugBreak();
+			check(EventToWaitFor->IsComplete());
+		}
 		if (EventToWaitFor.GetReference() && !EventToWaitFor->IsComplete())
 		{
 			check(!GRHIThread || !IsInRHIThread()); // things should not be dispatched if they can't complete without further waits
diff --git a/Engine/Source/Runtime/RHI/Public/RHIResources.h b/Engine/Source/Runtime/RHI/Public/RHIResources.h
--- a/Engine/Source/Runtime/RHI/Public/RHIResources.h
+++ b/Engine/Source/Runtime/RHI/Public/RHIResources.h
@@ -79,7 +79,7 @@ private:
 	mutable FThreadSafeCounter NumRefs;
 	mutable int32 MarkedForDelete;
 	bool bDoNotDeferDelete;
-	static TLockFreePointerList<FRHIResource> PendingDeletes;
+	static TLockFreePointerListUnordered<FRHIResource> PendingDeletes;
 	static FRHIResource* CurrentlyDeleting;
 };
 
diff --git a/Engine/Source/Runtime/RenderCore/Private/RenderingThread.cpp b/Engine/Source/Runtime/RenderCore/Private/RenderingThread.cpp
--- a/Engine/Source/Runtime/RenderCore/Private/RenderingThread.cpp
+++ b/Engine/Source/Runtime/RenderCore/Private/RenderingThread.cpp
@@ -916,7 +916,7 @@ FRHICommandListImmediate& GetImmediateCommandList_ForRenderCommand()
 }
 
 /** The set of deferred cleanup objects which are pending cleanup. */
-static TLockFreePointerList<FDeferredCleanupInterface>	PendingCleanupObjectsList;
+static TLockFreePointerListUnordered<FDeferredCleanupInterface>	PendingCleanupObjectsList;
 
 FPendingCleanupObjects::FPendingCleanupObjects()
 {
diff --git a/Engine/Source/Runtime/Slate/Public/Framework/Notifications/NotificationManager.h b/Engine/Source/Runtime/Slate/Public/Framework/Notifications/NotificationManager.h
--- a/Engine/Source/Runtime/Slate/Public/Framework/Notifications/NotificationManager.h
+++ b/Engine/Source/Runtime/Slate/Public/Framework/Notifications/NotificationManager.h
@@ -90,5 +90,5 @@ private:
 	TArray< FRegionalNotificationList > RegionalLists;
 
 	/** Thread safe queue of notifications to display */
-	TLockFreePointerList<FNotificationInfo> PendingNotifications;
+	TLockFreePointerListLIFO<FNotificationInfo> PendingNotifications;
 };
