Crash reporter improvements - Removed unused code

diff --git a/Engine/Source/Runtime/Core/Public/GenericPlatform/GenericPlatformCrashContext.h b/Engine/Source/Runtime/Core/Public/GenericPlatform/GenericPlatformCrashContext.h
--- a/Engine/Source/Runtime/Core/Public/GenericPlatform/GenericPlatformCrashContext.h
+++ b/Engine/Source/Runtime/Core/Public/GenericPlatform/GenericPlatformCrashContext.h
@@ -60,23 +60,6 @@ enum class ECrashDumpMode : int32
 	FullDump = 1,
 };
 
-/*-----------------------------------------------------------------------------
-	LexicalConversion
------------------------------------------------------------------------------*/
-
-namespace LexicalConversion
-{
-	inline FString ToString( const ECrashDescVersions& Value )
-	{
-		return TTypeToString<int32>::ToString( int32( Value ) );
-	}
-
-	inline FString ToString( const ECrashDumpMode& Value )
-	{
-		return TTypeToString<int32>::ToString( int32( Value ) );
-	}
-}
-
 /**
  *	Contains a runtime crash's properties that are common for all platforms.
  *	This may change in the future.
