Cached some string conversions that were happening in FCulture

diff --git a/Engine/Source/Runtime/Core/Private/Internationalization/Culture.cpp b/Engine/Source/Runtime/Core/Private/Internationalization/Culture.cpp
--- a/Engine/Source/Runtime/Core/Private/Internationalization/Culture.cpp
+++ b/Engine/Source/Runtime/Core/Private/Internationalization/Culture.cpp
@@ -23,27 +23,40 @@ FCulturePtr FCulture::Create(const FText& InDisplayName, const FString& InEnglis
 #if UE_ENABLE_ICU
 FCulture::FCulture(const FString& LocaleName)
 	: Implementation( new FICUCultureImplementation( LocaleName ) )
-{
-
-}
 #else
 FCulture::FCulture(const FText& InDisplayName, const FString& InEnglishName, const int InKeyboardLayoutId, const int InLCID, const FString& InName, const FString& InNativeName, const FString& InUnrealLegacyThreeLetterISOLanguageName, const FString& InThreeLetterISOLanguageName, const FString& InTwoLetterISOLanguageName, const FNumberFormattingRules& InNumberFormattingRule, const FTextFormattingRules& InTextFormattingRule, const FDateTimeFormattingRules& InDateTimeFormattingRule) 
 	: Implementation( new FLegacyCultureImplementation(InDisplayName, InEnglishName, InKeyboardLayoutId, InLCID, InName, InNativeName, InUnrealLegacyThreeLetterISOLanguageName, InThreeLetterISOLanguageName, InTwoLetterISOLanguageName) )
 	, DateTimeFormattingRule(InDateTimeFormattingRule)
 	, TextFormattingRule(InTextFormattingRule)
 	, NumberFormattingRule(InNumberFormattingRule)
+#endif
+	, CachedDisplayName(Implementation->GetDisplayName())
+	, CachedEnglishName(Implementation->GetEnglishName())
+	, CachedName(Implementation->GetName())
+	, CachedNativeName(Implementation->GetNativeName())
+	, CachedUnrealLegacyThreeLetterISOLanguageName(Implementation->GetUnrealLegacyThreeLetterISOLanguageName())
+	, CachedThreeLetterISOLanguageName(Implementation->GetThreeLetterISOLanguageName())
+	, CachedTwoLetterISOLanguageName(Implementation->GetTwoLetterISOLanguageName())
+	, CachedNativeLanguage(Implementation->GetNativeLanguage())
+#if UE_ENABLE_ICU
+	, CachedRegion(Implementation->GetRegion())
+#endif
+	, CachedNativeRegion(Implementation->GetNativeRegion())
+#if UE_ENABLE_ICU
+	, CachedScript(Implementation->GetScript())
+	, CachedVariant(Implementation->GetVariant())
+#endif
 { 
 }
-#endif
 
-FString FCulture::GetDisplayName() const
+const FString& FCulture::GetDisplayName() const
 {
-	return Implementation->GetDisplayName();
+	return CachedDisplayName;
 }
 
-FString FCulture::GetEnglishName() const
+const FString& FCulture::GetEnglishName() const
 {
-	return Implementation->GetEnglishName();
+	return CachedEnglishName;
 }
 
 int FCulture::GetKeyboardLayoutId() const
@@ -85,64 +98,52 @@ FString FCulture::GetCanonicalName(const FString& Name)
 	return FImplementation::GetCanonicalName(Name);
 }
 
-FString FCulture::GetName() const
+const FString& FCulture::GetName() const
 {
-	return Implementation->GetName();
+	return CachedName;
 }
 
-FString FCulture::GetNativeName() const
+const FString& FCulture::GetNativeName() const
 {
-	return Implementation->GetNativeName();
+	return CachedNativeName;
 }
 
-FString FCulture::GetUnrealLegacyThreeLetterISOLanguageName() const
+const FString& FCulture::GetUnrealLegacyThreeLetterISOLanguageName() const
 {
-	return Implementation->GetUnrealLegacyThreeLetterISOLanguageName();
+	return CachedUnrealLegacyThreeLetterISOLanguageName;
 }
 
-FString FCulture::GetThreeLetterISOLanguageName() const
+const FString& FCulture::GetThreeLetterISOLanguageName() const
 {
-	return Implementation->GetThreeLetterISOLanguageName();
+	return CachedThreeLetterISOLanguageName;
 }
 
-FString FCulture::GetTwoLetterISOLanguageName() const
+const FString& FCulture::GetTwoLetterISOLanguageName() const
 {
-	return Implementation->GetTwoLetterISOLanguageName();
+	return CachedTwoLetterISOLanguageName;
 }
 
-FString FCulture::GetNativeLanguage() const
+const FString& FCulture::GetNativeLanguage() const
 {
-	return Implementation->GetNativeLanguage();
+	return CachedNativeLanguage;
 }
 
-FString FCulture::GetRegion() const
+const FString& FCulture::GetRegion() const
 {
-#if UE_ENABLE_ICU
-	return Implementation->GetRegion();
-#else
-	return TEXT("");
-#endif
+	return CachedRegion;
 }
 
-FString FCulture::GetNativeRegion() const
+const FString& FCulture::GetNativeRegion() const
 {
-	return Implementation->GetNativeRegion();
+	return CachedNativeRegion;
 }
 
-FString FCulture::GetScript() const
+const FString& FCulture::GetScript() const
 {
-#if UE_ENABLE_ICU
-	return Implementation->GetScript();
-#else
-	return TEXT("");
-#endif
+	return CachedScript;
 }
 
-FString FCulture::GetVariant() const
+const FString& FCulture::GetVariant() const
 {
-#if UE_ENABLE_ICU
-	return Implementation->GetVariant();
-#else
-	return TEXT("");
-#endif
+	return CachedVariant;
 }
diff --git a/Engine/Source/Runtime/Core/Public/Internationalization/Culture.h b/Engine/Source/Runtime/Core/Public/Internationalization/Culture.h
--- a/Engine/Source/Runtime/Core/Public/Internationalization/Culture.h
+++ b/Engine/Source/Runtime/Core/Public/Internationalization/Culture.h
@@ -29,9 +29,9 @@ public:
 									const FDateTimeFormattingRules& InDateTimeFormattingRule);
 #endif
 
-	FString GetDisplayName() const;
+	const FString& GetDisplayName() const;
 
-	FString GetEnglishName() const;
+	const FString& GetEnglishName() const;
 
 	int GetKeyboardLayoutId() const;
 
@@ -41,25 +41,25 @@ public:
 
 	static FString GetCanonicalName(const FString& Name);
 
-	FString GetName() const;
+	const FString& GetName() const;
 	
-	FString GetNativeName() const;
+	const FString& GetNativeName() const;
 
-	FString GetUnrealLegacyThreeLetterISOLanguageName() const;
+	const FString& GetUnrealLegacyThreeLetterISOLanguageName() const;
 
-	FString GetThreeLetterISOLanguageName() const;
+	const FString& GetThreeLetterISOLanguageName() const;
 
-	FString GetTwoLetterISOLanguageName() const;
+	const FString& GetTwoLetterISOLanguageName() const;
 
-	FString GetNativeLanguage() const;
+	const FString& GetNativeLanguage() const;
 
-	FString GetRegion() const;
+	const FString& GetRegion() const;
 
-	FString GetNativeRegion() const;
+	const FString& GetNativeRegion() const;
 
-	FString GetScript() const;
+	const FString& GetScript() const;
 
-	FString GetVariant() const;
+	const FString& GetVariant() const;
 
 public:
 #if UE_ENABLE_ICU
@@ -97,6 +97,19 @@ private:
 	const FTextFormattingRules TextFormattingRule;
 	const FNumberFormattingRules NumberFormattingRule;
 #endif
+
+	FString CachedDisplayName;
+	FString CachedEnglishName;
+	FString CachedName;
+	FString CachedNativeName;
+	FString CachedUnrealLegacyThreeLetterISOLanguageName;
+	FString CachedThreeLetterISOLanguageName;
+	FString CachedTwoLetterISOLanguageName;
+	FString CachedNativeLanguage;
+	FString CachedRegion;
+	FString CachedNativeRegion;
+	FString CachedScript;
+	FString CachedVariant;
 };
 
 #include "CulturePointer.h"
\ No newline at end of file
