Added a re-usable FString <-> icu::UnicodeString converter object

This can be used to avoid the overhead of creating a new ICU UConverter object for each string conversion.

diff --git a/Engine/Source/Runtime/Core/Private/Internationalization/ICUUtilities.cpp b/Engine/Source/Runtime/Core/Private/Internationalization/ICUUtilities.cpp
--- a/Engine/Source/Runtime/Core/Private/Internationalization/ICUUtilities.cpp
+++ b/Engine/Source/Runtime/Core/Private/Internationalization/ICUUtilities.cpp
@@ -8,38 +8,45 @@
 
 namespace ICUUtilities
 {
-	void ConvertString(const FString& Source, icu::UnicodeString& Destination, const bool ShouldNullTerminate)
+	FStringConverter::FStringConverter()
+		: ICUConverter(nullptr)
+	{
+		UErrorCode ICUStatus = U_ZERO_ERROR;
+		ICUConverter = ucnv_open(FPlatformString::GetEncodingName(), &ICUStatus);
+		check(U_SUCCESS(ICUStatus));
+	}
+
+	FStringConverter::~FStringConverter()
+	{
+		ucnv_close(ICUConverter);
+	}
+
+	void FStringConverter::ConvertString(const FString& Source, icu::UnicodeString& Destination, const bool ShouldNullTerminate)
 	{
-		if( Source.Len() )
+		if (Source.Len() > 0)
 		{
-			const char* const SourceEncoding = FPlatformString::GetEncodingName();
-	
 			UErrorCode ICUStatus = U_ZERO_ERROR;
-		
-			UConverter* ICUConverter = ucnv_open(SourceEncoding, &ICUStatus);
-			if( U_SUCCESS(ICUStatus) ) 
-			{
-				// Get the internal buffer of the string, we're going to use it as scratch space
-				const int32_t DestinationCapacityUChars = Source.Len() * 2;
-				UChar* InternalStringBuffer = Destination.getBuffer(DestinationCapacityUChars);
 
-				// Perform the conversion into the string buffer
-				const int32_t SourceSizeBytes = Source.Len() * sizeof(TCHAR);
-				const int32_t DestinationLength = ucnv_toUChars(ICUConverter, InternalStringBuffer, DestinationCapacityUChars, reinterpret_cast<const char*>(*Source), SourceSizeBytes, &ICUStatus);
+			ucnv_reset(ICUConverter);
 
-				// Optionally null terminate the string
-				if( ShouldNullTerminate )
-				{
-					InternalStringBuffer[DestinationLength] = 0;
-				}
+			// Get the internal buffer of the string, we're going to use it as scratch space
+			const int32_t DestinationCapacityUChars = Source.Len() * 2;
+			UChar* InternalStringBuffer = Destination.getBuffer(DestinationCapacityUChars);
 
-				// Size it back down to the correct size and release our lock on the string buffer
-				Destination.releaseBuffer(DestinationLength);
+			// Perform the conversion into the string buffer
+			const int32_t SourceSizeBytes = Source.Len() * sizeof(TCHAR);
+			const int32_t DestinationLength = ucnv_toUChars(ICUConverter, InternalStringBuffer, DestinationCapacityUChars, reinterpret_cast<const char*>(*Source), SourceSizeBytes, &ICUStatus);
 
-				ucnv_close(ICUConverter);
+			// Optionally null terminate the string
+			if (ShouldNullTerminate)
+			{
+				InternalStringBuffer[DestinationLength] = 0;
 			}
 
-			check( U_SUCCESS(ICUStatus) );
+			// Size it back down to the correct size and release our lock on the string buffer
+			Destination.releaseBuffer(DestinationLength);
+
+			check(U_SUCCESS(ICUStatus));
 		}
 		else
 		{
@@ -47,42 +54,76 @@ namespace ICUUtilities
 		}
 	}
 
-	icu::UnicodeString ConvertString(const FString& Source, const bool ShouldNullTerminate)
+	icu::UnicodeString FStringConverter::ConvertString(const FString& Source, const bool ShouldNullTerminate)
 	{
 		icu::UnicodeString Destination;
 		ConvertString(Source, Destination, ShouldNullTerminate);
 		return Destination;
 	}
 
-	void ConvertString(const icu::UnicodeString& Source, FString& Destination)
+	void FStringConverter::ConvertString(const icu::UnicodeString& Source, FString& Destination)
 	{
-		if( Source.length() )
+		if (Source.length() > 0)
 		{
-			const char* const DestinationEncoding = FPlatformString::GetEncodingName();
-
 			UErrorCode ICUStatus = U_ZERO_ERROR;
 
-			UConverter* ICUConverter = ucnv_open(DestinationEncoding, &ICUStatus);
-			if( U_SUCCESS(ICUStatus) ) 
-			{
-				// Get the internal buffer of the string, we're going to use it as scratch space
-				TArray<TCHAR>& InternalStringBuffer = Destination.GetCharArray();
+			ucnv_reset(ICUConverter);
+			
+			// Get the internal buffer of the string, we're going to use it as scratch space
+			TArray<TCHAR>& InternalStringBuffer = Destination.GetCharArray();
 				
-				// Work out the maximum size required and resize the buffer so it can hold enough data
-				const int32_t DestinationCapacityBytes = UCNV_GET_MAX_BYTES_FOR_STRING(Source.length(), ucnv_getMaxCharSize(ICUConverter));
-				const int32 DestinationCapacityTCHARs = DestinationCapacityBytes / sizeof(TCHAR);
-				InternalStringBuffer.SetNumUninitialized(DestinationCapacityTCHARs);
-
-				// Perform the conversion into the string buffer, and then null terminate the FString and size it back down to the correct size
-				const int32_t DestinationSizeBytes = ucnv_fromUChars(ICUConverter, reinterpret_cast<char*>(InternalStringBuffer.GetData()), DestinationCapacityBytes, Source.getBuffer(), Source.length(), &ICUStatus);
-				const int32 DestinationSizeTCHARs = DestinationSizeBytes / sizeof(TCHAR);
-				InternalStringBuffer[DestinationSizeTCHARs] = 0;
-				InternalStringBuffer.SetNum(DestinationSizeTCHARs + 1, /*bAllowShrinking*/false); // the array size includes null
-
-				ucnv_close(ICUConverter);
-			}
+			// Work out the maximum size required and resize the buffer so it can hold enough data
+			const int32_t DestinationCapacityBytes = UCNV_GET_MAX_BYTES_FOR_STRING(Source.length(), ucnv_getMaxCharSize(ICUConverter));
+			const int32 DestinationCapacityTCHARs = DestinationCapacityBytes / sizeof(TCHAR);
+			InternalStringBuffer.SetNumUninitialized(DestinationCapacityTCHARs);
+
+			// Perform the conversion into the string buffer, and then null terminate the FString and size it back down to the correct size
+			const int32_t DestinationSizeBytes = ucnv_fromUChars(ICUConverter, reinterpret_cast<char*>(InternalStringBuffer.GetData()), DestinationCapacityBytes, Source.getBuffer(), Source.length(), &ICUStatus);
+			const int32 DestinationSizeTCHARs = DestinationSizeBytes / sizeof(TCHAR);
+			InternalStringBuffer[DestinationSizeTCHARs] = 0;
+			InternalStringBuffer.SetNum(DestinationSizeTCHARs + 1, /*bAllowShrinking*/false); // the array size includes null
+
+			check(U_SUCCESS(ICUStatus));
+		}
+		else
+		{
+			Destination.Empty();
+		}
+	}
 
-			check( U_SUCCESS(ICUStatus) );
+	FString FStringConverter::ConvertString(const icu::UnicodeString& Source)
+	{
+		FString Destination;
+		ConvertString(Source, Destination);
+		return Destination;
+	}
+
+	void ConvertString(const FString& Source, icu::UnicodeString& Destination, const bool ShouldNullTerminate)
+	{
+		if (Source.Len() > 0)
+		{
+			FStringConverter StringConverter;
+			StringConverter.ConvertString(Source, Destination, ShouldNullTerminate);
+		}
+		else
+		{
+			Destination.remove();
+		}
+	}
+
+	icu::UnicodeString ConvertString(const FString& Source, const bool ShouldNullTerminate)
+	{
+		icu::UnicodeString Destination;
+		ConvertString(Source, Destination, ShouldNullTerminate);
+		return Destination;
+	}
+
+	void ConvertString(const icu::UnicodeString& Source, FString& Destination)
+	{
+		if (Source.length() > 0)
+		{
+			FStringConverter StringConverter;
+			StringConverter.ConvertString(Source, Destination);
 		}
 		else
 		{
@@ -97,4 +138,4 @@ namespace ICUUtilities
 		return Destination;
 	}
 }
-#endif
\ No newline at end of file
+#endif
diff --git a/Engine/Source/Runtime/Core/Private/Internationalization/ICUUtilities.h b/Engine/Source/Runtime/Core/Private/Internationalization/ICUUtilities.h
--- a/Engine/Source/Runtime/Core/Private/Internationalization/ICUUtilities.h
+++ b/Engine/Source/Runtime/Core/Private/Internationalization/ICUUtilities.h
@@ -7,10 +7,38 @@
 
 namespace ICUUtilities
 {
-	/** Convert FString <-> icu::UnicodeString */
+	/** 
+	 * An object that can convert between FString and icu::UnicodeString
+	 * Note: This object is not thread-safe.
+	 */
+	class FStringConverter
+	{
+	public:
+		FStringConverter();
+		~FStringConverter();
+
+		/** Convert FString -> icu::UnicodeString */
+		void ConvertString(const FString& Source, icu::UnicodeString& Destination, const bool ShouldNullTerminate = true);
+		icu::UnicodeString ConvertString(const FString& Source, const bool ShouldNullTerminate = true);
+
+		/** Convert icu::UnicodeString -> FString */
+		void ConvertString(const icu::UnicodeString& Source, FString& Destination);
+		FString ConvertString(const icu::UnicodeString& Source);
+
+	private:
+		/** Non-copyable */
+		FStringConverter(const FStringConverter&);
+		FStringConverter& operator=(const FStringConverter&);
+
+		UConverter* ICUConverter;
+	};
+
+	/** Convert FString -> icu::UnicodeString */
 	void ConvertString(const FString& Source, icu::UnicodeString& Destination, const bool ShouldNullTerminate = true);
 	icu::UnicodeString ConvertString(const FString& Source, const bool ShouldNullTerminate = true);
+
+	/** Convert icu::UnicodeString -> FString */
 	void ConvertString(const icu::UnicodeString& Source, FString& Destination);
 	FString ConvertString(const icu::UnicodeString& Source);
 }
-#endif
\ No newline at end of file
+#endif
