The ICU BiDi object now uses a re-usable string converter instance

diff --git a/Engine/Source/Runtime/Core/Private/Internationalization/ICUText.cpp b/Engine/Source/Runtime/Core/Private/Internationalization/ICUText.cpp
--- a/Engine/Source/Runtime/Core/Private/Internationalization/ICUText.cpp
+++ b/Engine/Source/Runtime/Core/Private/Internationalization/ICUText.cpp
@@ -318,7 +318,7 @@ public:
 			return ETextDirection::LeftToRight;
 		}
 
-		ICUUtilities::ConvertString(InString, ICUString);
+		StringConverter.ConvertString(InString, ICUString);
 
 		return Internal::ComputeTextDirection(ICUBiDi, ICUString);
 	}
@@ -337,14 +337,19 @@ public:
 			return ETextDirection::LeftToRight;
 		}
 
-		ICUUtilities::ConvertString(InString, ICUString);
+		StringConverter.ConvertString(InString, ICUString);
 
 		return Internal::ComputeTextDirection(ICUBiDi, ICUString, OutTextDirectionInfo);
 	}
 
 private:
+	/** Non-copyable */
+	FICUTextBiDi(const FICUTextBiDi&);
+	FICUTextBiDi& operator=(const FICUTextBiDi&);
+
 	UBiDi* ICUBiDi;
 	icu::UnicodeString ICUString;
+	ICUUtilities::FStringConverter StringConverter;
 };
 
 } // namespace Internal
