Stopped FName instances with the same string but different numerical suffixes all getting the same hash

diff --git a/Engine/Source/Runtime/Core/Public/UObject/NameTypes.h b/Engine/Source/Runtime/Core/Public/UObject/NameTypes.h
--- a/Engine/Source/Runtime/Core/Public/UObject/NameTypes.h
+++ b/Engine/Source/Runtime/Core/Public/UObject/NameTypes.h
@@ -1028,7 +1028,7 @@ Expose_TNameOf(FName)
 
 inline uint32 GetTypeHash( const FName N )
 {
-	return N.GetComparisonIndex();
+	return N.GetComparisonIndex() + N.GetNumber();
 }
 
 
