TBitArraySetRange::SetRange added.

diff --git a/Engine/Source/Runtime/Core/Public/Containers/BitArray.h b/Engine/Source/Runtime/Core/Public/Containers/BitArray.h
--- a/Engine/Source/Runtime/Core/Public/Containers/BitArray.h
+++ b/Engine/Source/Runtime/Core/Public/Containers/BitArray.h
@@ -350,6 +350,68 @@ public:
 		}
 	}
 
+	/**
+	 * Sets or unsets a range of bits within the array.
+	 * @param  Index  The index of the first bit to set.
+	 * @param  Num    The number of bits to set.
+	 * @param  Value  The value to set the bits to.
+	 */
+	FORCENOINLINE void SetRange(int32 Index, int32 Num, bool Value)
+	{
+		check(Index >= 0 && Num >= 0 && Index + Num <= NumBits);
+
+		if (Num == 0)
+		{
+			return;
+		}
+
+		// Work out which uint32 index to set from, and how many
+		uint32 StartIndex = Index / 32;
+		uint32 Count      = (Index + Num + 31) / 32 - StartIndex;
+
+		// Work out masks for the start/end of the sequence
+		uint32 StartMask  = 0xFFFFFFFFu << (Index % 32);
+		uint32 EndMask    = 0xFFFFFFFFu >> (32 - (Index + Num) % 32) % 32;
+
+		uint32* Data = GetData() + StartIndex;
+		if (Value)
+		{
+			if (Count == 1)
+			{
+				*Data |= StartMask & EndMask;
+			}
+			else
+			{
+				*Data++ |= StartMask;
+				Count -= 2;
+				while (Count != 0)
+				{
+					*Data++ = ~0;
+					--Count;
+				}
+				*Data |= EndMask;
+			}
+		}
+		else
+		{
+			if (Count == 1)
+			{
+				*Data &= ~(StartMask & EndMask);
+			}
+			else
+			{
+				*Data++ &= ~StartMask;
+				Count -= 2;
+				while (Count != 0)
+				{
+					*Data++ = 0;
+					--Count;
+				}
+				*Data &= ~EndMask;
+			}
+		}
+	}
+
 	/**
 	 * Removes bits from the array.
 	 * @param BaseIndex - The index of the first bit to remove.
