Changed FString format arguments to use an FString value rather than a pointer.

This protects us against references to temporary variables. Where a string copy is to be avoided, pass by const TCHAR* instead.

diff --git a/Engine/Source/Runtime/Core/Private/Misc/StringFormatter.cpp b/Engine/Source/Runtime/Core/Private/Misc/StringFormatter.cpp
--- a/Engine/Source/Runtime/Core/Private/Misc/StringFormatter.cpp
+++ b/Engine/Source/Runtime/Core/Private/Misc/StringFormatter.cpp
@@ -12,7 +12,7 @@ FStringFormatArg::FStringFormatArg( const int64 Value ) : Type(Int), IntValue(Va
 FStringFormatArg::FStringFormatArg( const uint64 Value ) : Type(UInt), UIntValue(Value) {}
 FStringFormatArg::FStringFormatArg( const float Value ) : Type(Double), DoubleValue(Value) {}
 FStringFormatArg::FStringFormatArg( const double Value ) : Type(Double), DoubleValue(Value) {}
-FStringFormatArg::FStringFormatArg( const FString& Value ) : Type(String), StringValue(&Value) {}
+FStringFormatArg::FStringFormatArg( FString Value ) : Type(String), StringValue(MoveTemp(Value)) {}
 FStringFormatArg::FStringFormatArg( const TCHAR* Value ) : Type(StringLiteral), StringLiteralValue(Value) {}
 FStringFormatArg::FStringFormatArg( const FStringFormatArg& RHS )
 {
@@ -34,7 +34,7 @@ void AppendToString(const FStringFormatArg& Arg, FString& StringToAppendTo)
 		case FStringFormatArg::Int: 			StringToAppendTo.Append(LexicalConversion::ToString(Arg.IntValue)); break;
 		case FStringFormatArg::UInt: 			StringToAppendTo.Append(LexicalConversion::ToString(Arg.UIntValue)); break;
 		case FStringFormatArg::Double: 			StringToAppendTo.Append(LexicalConversion::ToString(Arg.DoubleValue)); break;
-		case FStringFormatArg::String: 			StringToAppendTo.AppendChars(**Arg.StringValue, Arg.StringValue->Len()); break;
+		case FStringFormatArg::String: 			StringToAppendTo.AppendChars(*Arg.StringValue, Arg.StringValue.Len()); break;
 		case FStringFormatArg::StringLiteral: 	StringToAppendTo += Arg.StringLiteralValue; break;
 	}
 }
diff --git a/Engine/Source/Runtime/Core/Public/Containers/UnrealString.h b/Engine/Source/Runtime/Core/Public/Containers/UnrealString.h
--- a/Engine/Source/Runtime/Core/Public/Containers/UnrealString.h
+++ b/Engine/Source/Runtime/Core/Public/Containers/UnrealString.h
@@ -10,9 +10,9 @@
 #include "Math/UnrealMathUtility.h"
 #include "Misc/Crc.h"
 #include "Misc/CString.h"
-#include "Misc/StringFormatArg.h"
 #include "Templates/MemoryOps.h"
 
+struct FStringFormatArg;
 
 /** Determines case sensitivity options for string comparisons. */
 namespace ESearchCase
@@ -2052,3 +2052,5 @@ public:
  * @return the index in the given string of the closing parenthesis
  */
 CORE_API int32 FindMatchingClosingParenthesis(const FString& TargetString, const int32 StartSearch = 0);
+
+#include "Misc/StringFormatArg.h"
diff --git a/Engine/Source/Runtime/Core/Public/Misc/StringFormatArg.h b/Engine/Source/Runtime/Core/Public/Misc/StringFormatArg.h
--- a/Engine/Source/Runtime/Core/Public/Misc/StringFormatArg.h
+++ b/Engine/Source/Runtime/Core/Public/Misc/StringFormatArg.h
@@ -19,19 +19,20 @@ struct CORE_API FStringFormatArg
 		uint64 UIntValue;
 		/** Value as double */
 		double DoubleValue;
-		/** Value as a pointer to an FString */
-		const FString* StringValue;
 		/** Value as a string literal */
 		const TCHAR* StringLiteralValue;
 	};
 
+	/** Value as an FString */
+	FString StringValue;
+
 	FStringFormatArg( const int32 Value );
 	FStringFormatArg( const uint32 Value );
 	FStringFormatArg( const int64 Value );
 	FStringFormatArg( const uint64 Value );
 	FStringFormatArg( const float Value );
 	FStringFormatArg( const double Value );
-	FStringFormatArg( const FString& Value );
+	FStringFormatArg( FString Value );
 	FStringFormatArg( const TCHAR* Value );
 
 	/** Copyable */
