Use Reverse algorithm for TMultiMap::MultiFind and TMultiMap::MultiFindPointer.

diff --git a/Engine/Source/Runtime/Core/Public/Containers/Map.h b/Engine/Source/Runtime/Core/Public/Containers/Map.h
--- a/Engine/Source/Runtime/Core/Public/Containers/Map.h
+++ b/Engine/Source/Runtime/Core/Public/Containers/Map.h
@@ -5,7 +5,7 @@
 #include "Containers/Array.h"
 #include "Containers/Set.h"
 #include "Containers/UnrealString.h"
-
+#include "Algo/Reverse.h"
 
 #define ExchangeB(A,B) {bool T=A; A=B; B=T;}
 
@@ -995,16 +995,7 @@ public:
 
 		if(bMaintainOrder)
 		{
-			// Create an array with the values in reverse enumerated order; i.e. the order they were inserted in the map.
-			TArray<ValueType, Allocator> OrderedValues;
-			OrderedValues.Empty(OutValues.Num());
-			for(int32 Index = OutValues.Num() - 1;Index >= 0;Index--)
-			{
-				new(OrderedValues) ValueType(OutValues[Index]);
-			}
-
-			// Swap the ordered array into the output array.
-			Exchange(OrderedValues,OutValues);
+			Algo::Reverse(OutValues);
 		}
 	}
 
@@ -1024,16 +1015,7 @@ public:
 
 		if(bMaintainOrder)
 		{
-			// Create an array with the values in reverse enumerated order; i.e. the order they were inserted in the map.
-			TArray<const ValueType*, Allocator> OrderedValues;
-			OrderedValues.Empty(OutValues.Num());
-			for(int32 Index = OutValues.Num() - 1;Index >= 0;Index--)
-			{
-				new(OrderedValues) const ValueType*(OutValues[Index]);
-			}
-
-			// Swap the ordered array into the output array.
-			Exchange(OrderedValues,OutValues);
+			Algo::Reverse(OutValues);
 		}
 	}
 
