Custom deleter support for TSharedPtr.

diff --git a/Engine/Source/Runtime/Core/Public/Templates/SharedPointer.h b/Engine/Source/Runtime/Core/Public/Templates/SharedPointer.h
--- a/Engine/Source/Runtime/Core/Public/Templates/SharedPointer.h
+++ b/Engine/Source/Runtime/Core/Public/Templates/SharedPointer.h
@@ -148,7 +148,21 @@ public:
 	template< class OtherType >
 	FORCEINLINE explicit TSharedRef( OtherType* InObject )
 		: Object( InObject )
-		, SharedReferenceCount( InObject )
+		, SharedReferenceCount( SharedPointerInternals::NewDefaultReferenceController( InObject ) )
+	{
+		Init(InObject);
+	}
+
+	/**
+	 * Constructs a shared reference that owns the specified object.  Must not be nullptr.
+	 *
+	 * @param  InObject   Object this shared pointer to retain a reference to
+	 * @param  InDeleter  Deleter object used to destroy the object when it is no longer referenced.
+	 */
+	template< class OtherType, class DeleterType >
+	FORCEINLINE TSharedRef( OtherType* InObject, DeleterType&& InDeleter )
+		: Object( InObject )
+		, SharedReferenceCount( SharedPointerInternals::NewCustomReferenceController( InObject, Forward< DeleterType >( InDeleter ) ) )
 	{
 		Init(InObject);
 	}
@@ -161,7 +175,7 @@ public:
 	 */
 	TSharedRef()
 		: Object(new ObjectType())
-		, SharedReferenceCount(Object)
+		, SharedReferenceCount(SharedPointerInternals::NewDefaultReferenceController(Object))
 	{
 		EnsureRetrievingVTablePtrDuringCtor(TEXT("TSharedRef()"));
 		Init(Object);
@@ -178,7 +192,7 @@ public:
 	template< class OtherType >
 	FORCEINLINE TSharedRef( SharedPointerInternals::FRawPtrProxy< OtherType > const& InRawPtrProxy )
 		: Object( InRawPtrProxy.Object )
-		, SharedReferenceCount( InRawPtrProxy.Object )
+		, SharedReferenceCount( InRawPtrProxy.ReferenceController )
 	{
 		// If the following assert goes off, it means a TSharedRef was initialized from a nullptr object pointer.
 		// Shared references must never be nullptr, so either pass a valid object or consider using TSharedPtr instead.
@@ -497,7 +511,24 @@ public:
 	template< class OtherType >
 	FORCEINLINE explicit TSharedPtr( OtherType* InObject )
 		: Object( InObject )
-		, SharedReferenceCount( InObject )
+		, SharedReferenceCount( SharedPointerInternals::NewDefaultReferenceController( InObject ) )
+	{
+		// If the object happens to be derived from TSharedFromThis, the following method
+		// will prime the object with a weak pointer to itself.
+		SharedPointerInternals::EnableSharedFromThis( this, InObject, InObject );
+	}
+
+	/**
+	 * Constructs a shared pointer that owns the specified object.  Note that passing nullptr here will
+	 * still create a tracked reference to a nullptr pointer. (Consistent with std::shared_ptr)
+	 *
+	 * @param  InObject   Object this shared pointer to retain a reference to
+	 * @param  InDeleter  Deleter object used to destroy the object when it is no longer referenced.
+	 */
+	template< class OtherType, class DeleterType >
+	FORCEINLINE TSharedPtr( OtherType* InObject, DeleterType&& InDeleter )
+		: Object( InObject )
+		, SharedReferenceCount( SharedPointerInternals::NewCustomReferenceController( InObject, Forward< DeleterType >( InDeleter ) ) )
 	{
 		// If the object happens to be derived from TSharedFromThis, the following method
 		// will prime the object with a weak pointer to itself.
@@ -513,7 +544,7 @@ public:
 	template< class OtherType >
 	FORCEINLINE TSharedPtr( SharedPointerInternals::FRawPtrProxy< OtherType > const& InRawPtrProxy )
 		: Object( InRawPtrProxy.Object )
-		, SharedReferenceCount( InRawPtrProxy.Object )
+		, SharedReferenceCount( InRawPtrProxy.ReferenceController )
 	{
 		// If the object happens to be derived from TSharedFromThis, the following method
 		// will prime the object with a weak pointer to itself.
@@ -1500,6 +1531,19 @@ FORCEINLINE SharedPointerInternals::FRawPtrProxy< ObjectType > MakeShareable( Ob
 }
 
 
+/**
+ * MakeShareable utility function.  Wrap object pointers with MakeShareable to allow them to be implicitly
+ * converted to shared pointers!  This is useful in assignment operations, or when returning a shared
+ * pointer from a function.
+ */
+// NOTE: The following is an Unreal extension to standard shared_ptr behavior
+template< class ObjectType, class DeleterType >
+FORCEINLINE SharedPointerInternals::FRawPtrProxy< ObjectType > MakeShareable( ObjectType* InObject, DeleterType&& InDeleter )
+{
+	return SharedPointerInternals::FRawPtrProxy< ObjectType >( InObject, Forward< DeleterType >( InDeleter ) );
+}
+
+
 /**
  * Given a TArray of TWeakPtr's, will remove any invalid pointers.
  * @param  PointerArray  The pointer array to prune invalid pointers out of
diff --git a/Engine/Source/Runtime/Core/Public/Templates/SharedPointerInternals.h b/Engine/Source/Runtime/Core/Public/Templates/SharedPointerInternals.h
--- a/Engine/Source/Runtime/Core/Public/Templates/SharedPointerInternals.h
+++ b/Engine/Source/Runtime/Core/Public/Templates/SharedPointerInternals.h
@@ -43,7 +43,6 @@ template< class ObjectType, ESPMode Mode = ESPMode::Fast > class TSharedFromThis
  */
 namespace SharedPointerInternals
 {
-
 	// Forward declarations
 	template< ESPMode Mode > class FWeakReferencer;
 
@@ -54,31 +53,15 @@ namespace SharedPointerInternals
 	// NOTE: The following is an Unreal extension to standard shared_ptr behavior
 	struct FNullTag {};
 
-	
-	/** Proxy structure for implicitly converting raw pointers to shared/weak pointers */
-	// NOTE: The following is an Unreal extension to standard shared_ptr behavior
-	template< class ObjectType >
-	struct FRawPtrProxy
-	{
-		/** The object pointer */
-		ObjectType* Object;
 
-		/** Construct implicitly from an object */
-		FORCEINLINE FRawPtrProxy( ObjectType* InObject )
-			: Object( InObject )
-		{
-		}
-	};
-
-
-	struct FReferenceControllerState
+	class FReferenceControllerBase
 	{
+	public:
 		/** Constructor */
-		FORCEINLINE FReferenceControllerState(void* InObject, void (*InDestroyObject)(void*))
+		FORCEINLINE explicit FReferenceControllerBase(void* InObject)
 			: SharedReferenceCount(1)
 			, WeakReferenceCount(1)
 			, Object(InObject)
-			, DestroyObject(InDestroyObject)
 		{ }
 
 		// NOTE: The primary reason these reference counters are 32-bit values (and not 16-bit to save
@@ -96,12 +79,83 @@ namespace SharedPointerInternals
 		void* Object;
 
 		/** Destroys the object associated with this reference counter.  */
-		void (*DestroyObject)(void*);
+		virtual void DestroyObject() = 0;
+
+		/** Destroys the object associated with this reference counter.  */
+		virtual ~FReferenceControllerBase() {};
 
 	private:
+		FReferenceControllerBase( FReferenceControllerBase const& );
+		FReferenceControllerBase& operator=( FReferenceControllerBase const& );
+	};
+
+	template <typename ObjectType, typename DeleterType>
+	class TReferenceControllerWithDeleter : private DeleterType, public FReferenceControllerBase
+	{
+	public:
+		explicit TReferenceControllerWithDeleter(void* Object, DeleterType&& Deleter)
+			: DeleterType             (MoveTemp(Deleter))
+			, FReferenceControllerBase(Object)
+		{
+		}
 
-		FReferenceControllerState( FReferenceControllerState const& );
-		FReferenceControllerState& operator=( FReferenceControllerState const& );
+		virtual void DestroyObject()
+		{
+			(*static_cast<DeleterType*>(this))((ObjectType*)static_cast<FReferenceControllerBase*>(this)->Object);
+		}
+	};
+
+
+	/** Deletes an object via the standard delete operator */
+	template <typename Type>
+	struct DefaultDeleter
+	{
+		FORCEINLINE void operator()(Type* Object) const
+		{
+			delete Object;
+		}
+	};
+
+	/** Creates a reference controller which just calls delete */
+	template <typename ObjectType>
+	inline FReferenceControllerBase* NewDefaultReferenceController(ObjectType* Object)
+	{
+		return new TReferenceControllerWithDeleter<ObjectType, DefaultDeleter<ObjectType>>(Object, DefaultDeleter<ObjectType>());
+	}
+
+	/** Creates a custom reference controller with a specified deleter */
+	template <typename ObjectType, typename DeleterType>
+	inline FReferenceControllerBase* NewCustomReferenceController(ObjectType* Object, DeleterType&& Deleter)
+	{
+		return new TReferenceControllerWithDeleter<ObjectType, typename TRemoveReference<DeleterType>::Type>(Object, Forward<DeleterType>(Deleter));
+	}
+
+
+	/** Proxy structure for implicitly converting raw pointers to shared/weak pointers */
+	// NOTE: The following is an Unreal extension to standard shared_ptr behavior
+	template< class ObjectType >
+	struct FRawPtrProxy
+	{
+		/** The object pointer */
+		ObjectType* Object;
+
+		/** Reference controller used to destroy the object */
+		FReferenceControllerBase* ReferenceController;
+
+		/** Construct implicitly from an object */
+		FORCEINLINE FRawPtrProxy( ObjectType* InObject )
+			: Object             ( InObject )
+			, ReferenceController( NewDefaultReferenceController( InObject ) )
+		{
+		}
+
+		/** Construct implicitly from an object and a custom deleter */
+		template< class Deleter >
+		FORCEINLINE FRawPtrProxy( ObjectType* InObject, Deleter&& InDeleter )
+			: Object             ( InObject )
+			, ReferenceController( NewCustomReferenceController( InObject, Forward< Deleter >( InDeleter ) ) )
+		{
+		}
 	};
 
 
@@ -118,16 +172,16 @@ namespace SharedPointerInternals
 	struct FReferenceControllerOps<ESPMode::ThreadSafe>
 	{
 		/** Returns the shared reference count */
-		static FORCEINLINE const int32 GetSharedReferenceCount(const FReferenceControllerState* State)
+		static FORCEINLINE const int32 GetSharedReferenceCount(const FReferenceControllerBase* ReferenceController)
 		{
 			// This reference count may be accessed by multiple threads
-			return static_cast< int32 const volatile& >( State->SharedReferenceCount );
+			return static_cast< int32 const volatile& >( ReferenceController->SharedReferenceCount );
 		}
 
 		/** Adds a shared reference to this counter */
-		static FORCEINLINE void AddSharedReference(FReferenceControllerState* State)
+		static FORCEINLINE void AddSharedReference(FReferenceControllerBase* ReferenceController)
 		{
-			FPlatformAtomics::InterlockedIncrement( &State->SharedReferenceCount );
+			FPlatformAtomics::InterlockedIncrement( &ReferenceController->SharedReferenceCount );
 		}
 
 		/**
@@ -135,13 +189,13 @@ namespace SharedPointerInternals
 		 *
 		 * @return  True if the shared reference was added successfully
 		 */
-		static bool ConditionallyAddSharedReference(FReferenceControllerState* State)
+		static bool ConditionallyAddSharedReference(FReferenceControllerBase* ReferenceController)
 		{
 			for( ; ; )
 			{
 				// Peek at the current shared reference count.  Remember, this value may be updated by
 				// multiple threads.
-				const int32 OriginalCount = static_cast< int32 const volatile& >( State->SharedReferenceCount );
+				const int32 OriginalCount = static_cast< int32 const volatile& >( ReferenceController->SharedReferenceCount );
 				if( OriginalCount == 0 )
 				{
 					// Never add a shared reference if the pointer has already expired
@@ -149,7 +203,7 @@ namespace SharedPointerInternals
 				}
 
 				// Attempt to increment the reference count.
-				const int32 ActualOriginalCount = FPlatformAtomics::InterlockedCompareExchange( &State->SharedReferenceCount, OriginalCount + 1, OriginalCount );
+				const int32 ActualOriginalCount = FPlatformAtomics::InterlockedCompareExchange( &ReferenceController->SharedReferenceCount, OriginalCount + 1, OriginalCount );
 
 				// We need to make sure that we never revive a counter that has already expired, so if the
 				// actual value what we expected (because it was touched by another thread), then we'll try
@@ -162,37 +216,37 @@ namespace SharedPointerInternals
 		}
 
 		/** Releases a shared reference to this counter */
-		static FORCEINLINE void ReleaseSharedReference(FReferenceControllerState* State)
+		static FORCEINLINE void ReleaseSharedReference(FReferenceControllerBase* ReferenceController)
 		{
-			checkSlow( State->SharedReferenceCount > 0 );
+			checkSlow( ReferenceController->SharedReferenceCount > 0 );
 
-			if( FPlatformAtomics::InterlockedDecrement( &State->SharedReferenceCount ) == 0 )
+			if( FPlatformAtomics::InterlockedDecrement( &ReferenceController->SharedReferenceCount ) == 0 )
 			{
 				// Last shared reference was released!  Destroy the referenced object.
-				State->DestroyObject(State->Object);
+				ReferenceController->DestroyObject();
 
 				// No more shared referencers, so decrement the weak reference count by one.  When the weak
 				// reference count reaches zero, this object will be deleted.
-				ReleaseWeakReference(State);
+				ReleaseWeakReference(ReferenceController);
 			}
 		}
 
 
 		/** Adds a weak reference to this counter */
-		static FORCEINLINE void AddWeakReference(FReferenceControllerState* State)
+		static FORCEINLINE void AddWeakReference(FReferenceControllerBase* ReferenceController)
 		{
-			FPlatformAtomics::InterlockedIncrement( &State->WeakReferenceCount );
+			FPlatformAtomics::InterlockedIncrement( &ReferenceController->WeakReferenceCount );
 		}
 
 		/** Releases a weak reference to this counter */
-		static void ReleaseWeakReference(FReferenceControllerState* State)
+		static void ReleaseWeakReference(FReferenceControllerBase* ReferenceController)
 		{
-			checkSlow( State->WeakReferenceCount > 0 );
+			checkSlow( ReferenceController->WeakReferenceCount > 0 );
 
-			if( FPlatformAtomics::InterlockedDecrement( &State->WeakReferenceCount ) == 0 )
+			if( FPlatformAtomics::InterlockedDecrement( &ReferenceController->WeakReferenceCount ) == 0 )
 			{
 				// No more references to this reference count.  Destroy it!
-				delete State;
+				delete ReferenceController;
 			}
 		}
 	};
@@ -202,15 +256,15 @@ namespace SharedPointerInternals
 	struct FReferenceControllerOps<ESPMode::NotThreadSafe>
 	{
 		/** Returns the shared reference count */
-		static FORCEINLINE const int32 GetSharedReferenceCount(const FReferenceControllerState* State)
+		static FORCEINLINE const int32 GetSharedReferenceCount(const FReferenceControllerBase* ReferenceController)
 		{
-			return State->SharedReferenceCount;
+			return ReferenceController->SharedReferenceCount;
 		}
 
 		/** Adds a shared reference to this counter */
-		static FORCEINLINE void AddSharedReference(FReferenceControllerState* State)
+		static FORCEINLINE void AddSharedReference(FReferenceControllerBase* ReferenceController)
 		{
-			++State->SharedReferenceCount;
+			++ReferenceController->SharedReferenceCount;
 		}
 
 		/**
@@ -218,62 +272,54 @@ namespace SharedPointerInternals
 		 *
 		 * @return  True if the shared reference was added successfully
 		 */
-		static bool ConditionallyAddSharedReference(FReferenceControllerState* State)
+		static bool ConditionallyAddSharedReference(FReferenceControllerBase* ReferenceController)
 		{
-			if( State->SharedReferenceCount == 0 )
+			if( ReferenceController->SharedReferenceCount == 0 )
 			{
 				// Never add a shared reference if the pointer has already expired
 				return false;
 			}
 
-			++State->SharedReferenceCount;
+			++ReferenceController->SharedReferenceCount;
 			return true;
 		}
 
 		/** Releases a shared reference to this counter */
-		static FORCEINLINE void ReleaseSharedReference(FReferenceControllerState* State)
+		static FORCEINLINE void ReleaseSharedReference(FReferenceControllerBase* ReferenceController)
 		{
-			checkSlow( State->SharedReferenceCount > 0 );
+			checkSlow( ReferenceController->SharedReferenceCount > 0 );
 
-			if( --State->SharedReferenceCount == 0 )
+			if( --ReferenceController->SharedReferenceCount == 0 )
 			{
 				// Last shared reference was released!  Destroy the referenced object.
-				State->DestroyObject(State->Object);
+				ReferenceController->DestroyObject();
 
 				// No more shared referencers, so decrement the weak reference count by one.  When the weak
 				// reference count reaches zero, this object will be deleted.
-				ReleaseWeakReference(State);
+				ReleaseWeakReference(ReferenceController);
 			}
 		}
 
 		/** Adds a weak reference to this counter */
-		static FORCEINLINE void AddWeakReference(FReferenceControllerState* State)
+		static FORCEINLINE void AddWeakReference(FReferenceControllerBase* ReferenceController)
 		{
-			++State->WeakReferenceCount;
+			++ReferenceController->WeakReferenceCount;
 		}
 
 		/** Releases a weak reference to this counter */
-		static void ReleaseWeakReference(FReferenceControllerState* State)
+		static void ReleaseWeakReference(FReferenceControllerBase* ReferenceController)
 		{
-			checkSlow( State->WeakReferenceCount > 0 );
+			checkSlow( ReferenceController->WeakReferenceCount > 0 );
 
-			if( --State->WeakReferenceCount == 0 )
+			if( --ReferenceController->WeakReferenceCount == 0 )
 			{
 				// No more references to this reference count.  Destroy it!
-				delete State;
+				delete ReferenceController;
 			}
 		}
 	};
 
 
-	/** Type-erasing function deleting objects of a given type via a common pointer type (void*) */
-	template <typename Type>
-	void DestroyObject(void* Object)
-	{
-		delete (Type*)Object;
-	}
-
-
 	/**
 	 * FSharedReferencer is a wrapper around a pointer to a reference controller that is used by either a
 	 * TSharedRef or a TSharedPtr to keep track of a referenced object's lifetime
@@ -289,11 +335,10 @@ namespace SharedPointerInternals
 		FORCEINLINE FSharedReferencer()
 			: ReferenceController( nullptr )
 		{ }
-		
+
 		/** Constructor that counts a single reference to the specified object */
-		template< class OtherType > 
-		inline explicit FSharedReferencer( OtherType* InObject )
-			: ReferenceController( new FReferenceControllerState( InObject, &DestroyObject<OtherType> ) )
+		inline explicit FSharedReferencer( FReferenceControllerBase* InReferenceController )
+			: ReferenceController( InReferenceController )
 		{ }
 		
 		/** Copy constructor creates a new reference to the existing object */
@@ -434,7 +479,7 @@ namespace SharedPointerInternals
 	private:
 
 		/** Pointer to the reference controller for the object a shared reference/pointer is referencing */
-		FReferenceControllerState* ReferenceController;
+		FReferenceControllerBase* ReferenceController;
 	};
 
 
@@ -541,7 +586,7 @@ namespace SharedPointerInternals
 
 		/** Assigns a new reference controller to this counter object, first adding a reference to it, then
 		    releasing the previous object. */
-		inline void AssignReferenceController( FReferenceControllerState* NewReferenceController )
+		inline void AssignReferenceController( FReferenceControllerBase* NewReferenceController )
 		{
 			// Only proceed if the new reference counter is different than our current
 			if( NewReferenceController != ReferenceController )
@@ -571,7 +616,7 @@ namespace SharedPointerInternals
 	private:
 
 		/** Pointer to the reference controller for the object a TWeakPtr is referencing */
-		FReferenceControllerState* ReferenceController;
+		FReferenceControllerBase* ReferenceController;
 	};
 
 
