Fix example code

diff --git a/Engine/Source/Runtime/Core/Public/Templates/Function.h b/Engine/Source/Runtime/Core/Public/Templates/Function.h
--- a/Engine/Source/Runtime/Core/Public/Templates/Function.h
+++ b/Engine/Source/Runtime/Core/Public/Templates/Function.h
@@ -279,7 +279,7 @@ namespace UE4Function_Private
  * {
  *     for (const FString& Str : SomeBunchOfStrings)
  *     {
- *         int32 Int = Func(Str);
+ *         int32 Int = Convert(Str);
  *         DoSomething(Int);
  *     }
  * }
