Code cleanup pass

diff --git a/Engine/Source/Runtime/Core/Private/Tests/Misc/DateTimeTest.cpp b/Engine/Source/Runtime/Core/Private/Tests/Misc/DateTimeTest.cpp
--- a/Engine/Source/Runtime/Core/Private/Tests/Misc/DateTimeTest.cpp
+++ b/Engine/Source/Runtime/Core/Private/Tests/Misc/DateTimeTest.cpp
@@ -4,8 +4,20 @@
 #include "AutomationTest.h"
 
 
+#define TestUnixEquivalent( Desc, A, B ) if( (A).ToUnixTimestamp() != (B) ) AddError(FString::Printf(TEXT("%s - A=%d B=%d"),Desc,(A).ToUnixTimestamp(),(B)));
+#define TestYear( Desc, A, B ) if( (A.GetYear()) != (B) ) AddError(FString::Printf(TEXT("%s - A=%d B=%d"),Desc,(A.GetYear()),(B)));
+#define TestMonth( Desc, A, B ) if( (A.GetMonth()) != (B) ) AddError(FString::Printf(TEXT("%s - A=%d B=%d"),Desc,(A.GetMonth()),B));
+#define TestMonthOfYear( Desc, A, B ) if( (A.GetMonthOfYear()) != (B) ) AddError(FString::Printf(TEXT("%s - A=%d B=%d"),Desc,(static_cast<int32>(A.GetMonthOfYear())),static_cast<int32>(B)));
+#define TestDay( Desc, A, B ) if( (A.GetDay()) != (B) ) AddError(FString::Printf(TEXT("%s - A=%d B=%d"),Desc,(A.GetDay()),(B)));
+#define TestHour( Desc, A, B ) if( (A.GetHour()) != (B) ) AddError(FString::Printf(TEXT("%s - A=%d B=%d"),Desc,(A.GetHour()),(B)));
+#define TestMinute( Desc, A, B ) if( (A.GetMinute()) != (B) ) AddError(FString::Printf(TEXT("%s - A=%d B=%d"),Desc,(A.GetMinute()),(B)));
+#define TestSecond( Desc, A, B ) if( (A.GetSecond()) != (B) ) AddError(FString::Printf(TEXT("%s - A=%d B=%d"),Desc,(A.GetSecond()),(B)));
+#define TestMillisecond( Desc, A, B ) if( (A.GetMillisecond()) != (B) ) AddError(FString::Printf(TEXT("%s - A=%d B=%d"),Desc,(A.GetMillisecond()),(B)));
+
+
 IMPLEMENT_SIMPLE_AUTOMATION_TEST(FDateTimeTest, "System.Core.Misc.DateTime", EAutomationTestFlags::ApplicationContextMask | EAutomationTestFlags::SmokeFilter)
 
+
 bool FDateTimeTest::RunTest( const FString& Parameters )
 {
 	const int64 UnixEpochTimestamp = 0;
@@ -20,24 +32,11 @@ bool FDateTimeTest::RunTest( const FString& Parameters )
 
 	const FDateTime TestDateTime(2013, 8, 14, 12, 34, 56, 789);
 	
-#define TestUnixEquivalent( Desc, A, B ) if( (A).ToUnixTimestamp() != (B) ) AddError(FString::Printf(TEXT("%s - A=%d B=%d"),Desc,(A).ToUnixTimestamp(),(B)));
-
 	TestUnixEquivalent( TEXT("Testing Unix Epoch Ticks"), UnixEpoch, UnixEpochTimestamp);
 	TestUnixEquivalent( TEXT("Testing Unix Billennium Ticks"), UnixBillennium, UnixBillenniumTimestamp);
 	TestUnixEquivalent( TEXT("Testing Unix Ones Ticks"), UnixOnes, UnixOnesTimestamp);
 	TestUnixEquivalent( TEXT("Testing Unix Decimal Sequence Ticks"), UnixDecimalSequence, UnixDecimalSequenceTimestamp);
 
-#undef TestUnixEquivalent
-
-#define TestYear( Desc, A, B ) if( (A.GetYear()) != (B) ) AddError(FString::Printf(TEXT("%s - A=%d B=%d"),Desc,(A.GetYear()),(B)));
-#define TestMonth( Desc, A, B ) if( (A.GetMonth()) != (B) ) AddError(FString::Printf(TEXT("%s - A=%d B=%d"),Desc,(A.GetMonth()),B));
-#define TestMonthOfYear( Desc, A, B ) if( (A.GetMonthOfYear()) != (B) ) AddError(FString::Printf(TEXT("%s - A=%d B=%d"),Desc,(static_cast<int32>(A.GetMonthOfYear())),static_cast<int32>(B)));
-#define TestDay( Desc, A, B ) if( (A.GetDay()) != (B) ) AddError(FString::Printf(TEXT("%s - A=%d B=%d"),Desc,(A.GetDay()),(B)));
-#define TestHour( Desc, A, B ) if( (A.GetHour()) != (B) ) AddError(FString::Printf(TEXT("%s - A=%d B=%d"),Desc,(A.GetHour()),(B)));
-#define TestMinute( Desc, A, B ) if( (A.GetMinute()) != (B) ) AddError(FString::Printf(TEXT("%s - A=%d B=%d"),Desc,(A.GetMinute()),(B)));
-#define TestSecond( Desc, A, B ) if( (A.GetSecond()) != (B) ) AddError(FString::Printf(TEXT("%s - A=%d B=%d"),Desc,(A.GetSecond()),(B)));
-#define TestMillisecond( Desc, A, B ) if( (A.GetMillisecond()) != (B) ) AddError(FString::Printf(TEXT("%s - A=%d B=%d"),Desc,(A.GetMillisecond()),(B)));
-
 	TestYear( TEXT("Testing Unix Epoch Year"), UnixEpoch, 1970 );
 	TestMonth( TEXT("Testing Unix Epoch Month"), UnixEpoch, 1 );
 	TestMonthOfYear( TEXT("Testing Unix Epoch MonthOfYear"), UnixEpoch, EMonthOfYear::January );
@@ -83,6 +82,11 @@ bool FDateTimeTest::RunTest( const FString& Parameters )
 	TestSecond( TEXT("Testing Test Date Time Second"), TestDateTime, 56 );
 	TestMillisecond( TEXT("Testing Test Date Time Millisecond"), TestDateTime, 789 );
 
+	return true;
+}
+
+
+#undef TestUnixEquivalent
 #undef TestYear
 #undef TestMonth
 #undef TestDay
@@ -90,6 +94,3 @@ bool FDateTimeTest::RunTest( const FString& Parameters )
 #undef TestMinute
 #undef TestSecond
 #undef TestMillisecond
-
-	return true;
-}
