Core: Added unit test for ISO date/time conversion of empty strings

See also: https://github.com/EpicGames/UnrealEngine/pull/1523

diff --git a/Engine/Source/Runtime/Core/Private/Tests/Misc/DateTimeTest.cpp b/Engine/Source/Runtime/Core/Private/Tests/Misc/DateTimeTest.cpp
--- a/Engine/Source/Runtime/Core/Private/Tests/Misc/DateTimeTest.cpp
+++ b/Engine/Source/Runtime/Core/Private/Tests/Misc/DateTimeTest.cpp
@@ -82,6 +82,10 @@ bool FDateTimeTest::RunTest( const FString& Parameters )
 	TestSecond( TEXT("Testing Test Date Time Second"), TestDateTime, 56 );
 	TestMillisecond( TEXT("Testing Test Date Time Millisecond"), TestDateTime, 789 );
 
+	FDateTime ParsedDateTime;
+
+	TestFalse(TEXT("Parsing an empty ISO string must fail"), FDateTime::ParseIso8601(TEXT(""), ParsedDateTime));
+
 	return true;
 }
 
