[INTEGRATE] Fixed Clang compiler warning on Windows platform

diff --git a/Engine/Source/Runtime/Core/Private/Windows/WindowsPlatformMemory.cpp b/Engine/Source/Runtime/Core/Private/Windows/WindowsPlatformMemory.cpp
--- a/Engine/Source/Runtime/Core/Private/Windows/WindowsPlatformMemory.cpp
+++ b/Engine/Source/Runtime/Core/Private/Windows/WindowsPlatformMemory.cpp
@@ -104,11 +104,13 @@ FPlatformMemoryStats FWindowsPlatformMemory::GetStats()
 	FPlatformMemoryStats MemoryStats;
 
 	// Gather platform memory stats.
-	MEMORYSTATUSEX MemoryStatusEx = {0};
+	MEMORYSTATUSEX MemoryStatusEx;
+	FPlatformMemory::Memzero( &MemoryStatusEx, sizeof( MemoryStatusEx ) );
 	MemoryStatusEx.dwLength = sizeof( MemoryStatusEx );
 	::GlobalMemoryStatusEx( &MemoryStatusEx );
 
-	PROCESS_MEMORY_COUNTERS ProcessMemoryCounters = {0};
+	PROCESS_MEMORY_COUNTERS ProcessMemoryCounters;
+	FPlatformMemory::Memzero( &ProcessMemoryCounters, sizeof( ProcessMemoryCounters ) );
 	::GetProcessMemoryInfo( ::GetCurrentProcess(), &ProcessMemoryCounters, sizeof(ProcessMemoryCounters) );
 
 	MemoryStats.AvailablePhysical = MemoryStatusEx.ullAvailPhys;
@@ -141,11 +143,13 @@ const FPlatformMemoryConstants& FWindowsPlatformMemory::GetConstants()
 	if( MemoryConstants.TotalPhysical == 0 )
 	{
 		// Gather platform memory constants.
-		MEMORYSTATUSEX MemoryStatusEx = {0};
+		MEMORYSTATUSEX MemoryStatusEx;
+		FPlatformMemory::Memzero( &MemoryStatusEx, sizeof( MemoryStatusEx ) );
 		MemoryStatusEx.dwLength = sizeof( MemoryStatusEx );
 		::GlobalMemoryStatusEx( &MemoryStatusEx );
 
-		SYSTEM_INFO SystemInfo = { 0 };
+		SYSTEM_INFO SystemInfo;
+		FPlatformMemory::Memzero( &SystemInfo, sizeof( SystemInfo ) );
 		::GetSystemInfo(&SystemInfo);
 
 		MemoryConstants.TotalPhysical = MemoryStatusEx.ullTotalPhys;
