Fix check(0) firing in task graph when in single threaded mode.

diff --git a/Engine/Source/Runtime/Core/Private/Async/TaskGraph.cpp b/Engine/Source/Runtime/Core/Private/Async/TaskGraph.cpp
--- a/Engine/Source/Runtime/Core/Private/Async/TaskGraph.cpp
+++ b/Engine/Source/Runtime/Core/Private/Async/TaskGraph.cpp
@@ -984,6 +984,20 @@ public:
 			}
 		}
 	}
+
+	virtual void ProcessTasksUntilIdle(int32 QueueIndex) override
+	{
+		if (!FPlatformProcess::SupportsMultithreading()) 
+		{
+			Queue.QuitWhenIdle.Set(1);
+			ProcessTasks();
+			Queue.QuitWhenIdle.Reset();
+		}
+		else
+		{
+			FTaskThreadBase::ProcessTasksUntilIdle(QueueIndex);
+		}
+	}
 	// Calls meant to be called from any thread.
 
 	/** 
