Added runtime means of controlling engine's fixed-time-step mechanics #UE4

Merged over from Orion branch.

diff --git a/Engine/Source/Runtime/Core/Private/Misc/App.cpp b/Engine/Source/Runtime/Core/Private/Misc/App.cpp
--- a/Engine/Source/Runtime/Core/Private/Misc/App.cpp
+++ b/Engine/Source/Runtime/Core/Private/Misc/App.cpp
@@ -14,6 +14,7 @@ FString FApp::SessionName = FString();
 FString FApp::SessionOwner = FString();
 bool FApp::Standalone = true;
 bool FApp::bIsBenchmarking = false;
+bool FApp::bUseFixedTimeStep = false;
 double FApp::FixedDeltaTime = 1 / 30.0;
 double FApp::CurrentTime = 0.0;
 double FApp::LastTime = 0.0;
diff --git a/Engine/Source/Runtime/Core/Public/Misc/App.h b/Engine/Source/Runtime/Core/Public/Misc/App.h
--- a/Engine/Source/Runtime/Core/Public/Misc/App.h
+++ b/Engine/Source/Runtime/Core/Public/Misc/App.h
@@ -338,10 +338,19 @@ public:
 	 */
 	static bool UseFixedTimeStep()
 	{
-		static bool bUseFixedTimeStep = FParse::Param(FCommandLine::Get(), TEXT("UseFixedTimeStep"));
 		return bUseFixedTimeStep;
 	}
 
+	/**
+	 * Enables or disabled usage of fixed time step.
+	 *
+	 * @param whether to use fixed time step or not
+	 */
+	static void SetUseFixedTimeStep(bool bVal)
+	{
+		bUseFixedTimeStep = bVal;
+	}
+
 	/**
 	 * Gets current time in seconds.
 	 *
@@ -445,6 +454,9 @@ private:
 	/** Holds a flag Whether we are in benchmark mode or not. */
 	static bool bIsBenchmarking;
 
+	/** Holds a flag whether we want to use a fixed time step or not. */
+	static bool bUseFixedTimeStep;
+
 	/** Holds time step if a fixed delta time is wanted. */
 	static double FixedDeltaTime;
 
diff --git a/Engine/Source/Runtime/Launch/Private/LaunchEngineLoop.cpp b/Engine/Source/Runtime/Launch/Private/LaunchEngineLoop.cpp
--- a/Engine/Source/Runtime/Launch/Private/LaunchEngineLoop.cpp
+++ b/Engine/Source/Runtime/Launch/Private/LaunchEngineLoop.cpp
@@ -980,7 +980,9 @@ int32 FEngineLoop::PreInit( const TCHAR* CmdLine )
 #else
 	FApp::SetBenchmarking(false);
 #endif // !UE_BUILD_SHIPPING
-
+	
+	FApp::SetUseFixedTimeStep(FParse::Param(FCommandLine::Get(), TEXT("UseFixedTimeStep")));
+	
 	FString CheckToken = Token;
 	bool bFoundValidToken = false;
 	while (!bFoundValidToken && (CheckToken.Len() > 0))
