GitHub #1540 - Fix HexToBytes off-by-one

Also fix to assumption that sizeof(TCHAR) == 2.

diff --git a/Engine/Source/Runtime/Core/Public/Containers/UnrealString.h b/Engine/Source/Runtime/Core/Public/Containers/UnrealString.h
--- a/Engine/Source/Runtime/Core/Public/Containers/UnrealString.h
+++ b/Engine/Source/Runtime/Core/Public/Containers/UnrealString.h
@@ -1773,7 +1773,7 @@ inline const uint8 TCharToNibble( const TCHAR Char )
 inline int32 HexToBytes( const FString& HexString, uint8* OutBytes )
 {
 	int32 NumBytes = 0;
-	const bool bPadNibble = ( HexString.Len() % 2 ) == 1;
+	const bool bPadNibble = ( HexString.Len() % sizeof(TCHAR) ) == 1;
 	const TCHAR* CharPos = *HexString;
 	if( bPadNibble )
 	{
@@ -1785,7 +1785,7 @@ inline int32 HexToBytes( const FString& HexString, uint8* OutBytes )
 		OutBytes[ NumBytes ] += TCharToNibble( *CharPos++ );
 		++NumBytes;
 	}
-	return NumBytes - 1;
+	return NumBytes;
 }
 
 /** Namespace that houses lexical conversion for various types. User defined conversions can be implemented externally */
