Undoing overzealous sizeof.

diff --git a/Engine/Source/Runtime/Core/Public/Containers/UnrealString.h b/Engine/Source/Runtime/Core/Public/Containers/UnrealString.h
--- a/Engine/Source/Runtime/Core/Public/Containers/UnrealString.h
+++ b/Engine/Source/Runtime/Core/Public/Containers/UnrealString.h
@@ -1773,7 +1773,7 @@ inline const uint8 TCharToNibble( const TCHAR Char )
 inline int32 HexToBytes( const FString& HexString, uint8* OutBytes )
 {
 	int32 NumBytes = 0;
-	const bool bPadNibble = ( HexString.Len() % sizeof(TCHAR) ) == 1;
+	const bool bPadNibble = ( HexString.Len() % 2 ) == 1;
 	const TCHAR* CharPos = *HexString;
 	if( bPadNibble )
 	{
