Async: Fixed missing return value

diff --git a/Engine/Source/Runtime/Core/Public/Async/AsyncResult.h b/Engine/Source/Runtime/Core/Public/Async/AsyncResult.h
--- a/Engine/Source/Runtime/Core/Public/Async/AsyncResult.h
+++ b/Engine/Source/Runtime/Core/Public/Async/AsyncResult.h
@@ -66,6 +66,8 @@ public:
 
 		Other.Progress.Reset();
 		Other.Task.Reset();
+
+		return *this;
 	}
 
 public:
