Comment fix.

diff --git a/Engine/Source/Runtime/Core/Public/Templates/Function.h b/Engine/Source/Runtime/Core/Public/Templates/Function.h
--- a/Engine/Source/Runtime/Core/Public/Templates/Function.h
+++ b/Engine/Source/Runtime/Core/Public/Templates/Function.h
@@ -492,7 +492,7 @@ private:
  * Example:
  *
  * // Something.h
- * TFunction<FString (int32)> GetTransform();
+ * TFunction<FString (int32)> GetTransform(const FString& Prefix);
  *
  * // Something.cpp
  * TFunction<FString (int32)> GetTransform(const FString& Prefix)
