Cleaning out programmer-only comments that were harvested by Doxygen.

diff --git a/Engine/Source/Editor/Persona/Private/SAnimCompositeEditor.h b/Engine/Source/Editor/Persona/Private/SAnimCompositeEditor.h
--- a/Engine/Source/Editor/Persona/Private/SAnimCompositeEditor.h
+++ b/Engine/Source/Editor/Persona/Private/SAnimCompositeEditor.h
@@ -84,7 +84,7 @@ public:
 	void			PreAnimUpdate();
 	void			PostAnimUpdate();
 
-	// Begin SAnimEditorBase interface
+	//~ Begin SAnimEditorBase Interface
 	virtual TSharedRef<SWidget> CreateDocumentAnchor() override;
-	// End SAnimEditorBase interface
+	//~ End SAnimEditorBase Interface
 };
diff --git a/Engine/Source/Editor/Persona/Private/SMontageEditor.h b/Engine/Source/Editor/Persona/Private/SMontageEditor.h
--- a/Engine/Source/Editor/Persona/Private/SMontageEditor.h
+++ b/Engine/Source/Editor/Persona/Private/SMontageEditor.h
@@ -49,9 +49,9 @@ private:
 
 	TWeakPtr<FPersona> WeakPersona;
 protected:
-	// Begin SAnimEditorBase interface
+	//~ Begin SAnimEditorBase Interface
 	virtual TSharedRef<class SAnimationScrubPanel> ConstructAnimScrubPanel() override;
-	// End SAnimEditorBase interface
+	//~ End SAnimEditorBase Interface
 
 public:
 	void Construct(const FArguments& InArgs);
@@ -141,7 +141,7 @@ public:
 	void			PreAnimUpdate();
 	void			PostAnimUpdate();
 
-	// Begin SAnimEditorBase interface
+	//~ Begin SAnimEditorBase Interface
 	virtual TSharedRef<SWidget> CreateDocumentAnchor() override;
-	// End SAnimEditorBase interface
+	//~ End SAnimEditorBase Interface
 };
diff --git a/Engine/Source/Editor/PhAT/Private/PhAT.h b/Engine/Source/Editor/PhAT/Private/PhAT.h
--- a/Engine/Source/Editor/PhAT/Private/PhAT.h
+++ b/Engine/Source/Editor/PhAT/Private/PhAT.h
@@ -88,11 +88,11 @@ public:
 	/** Returns whether a PIE session is running. */
 	static bool IsPIERunning();
 
-	// FTickableEditorObject interface
+	//~ Begin FTickableEditorObject Interface
 	virtual void Tick(float DeltaTime) override;
 	virtual bool IsTickable() const override { return true; }
 	virtual TStatId GetStatId() const override;
-	// End of FTickableEditorObject interface
+	//~ End FTickableEditorObject Interface
 
 	void OnFocusSelection();
 
@@ -104,7 +104,7 @@ private:
 		PHFM_Bodies
 	};
 
-	// Begin FEditorUndoClient Interface
+	//~ Begin FEditorUndoClient Interface
 	virtual void PostUndo(bool bSuccess) override;
 	virtual void PostRedo(bool bSuccess) override;
 	// End of FEditorUndoClient
diff --git a/Engine/Source/Editor/PropertyEditor/Private/Presentation/PropertyTable/PropertyTableColumn.h b/Engine/Source/Editor/PropertyEditor/Private/Presentation/PropertyTable/PropertyTableColumn.h
--- a/Engine/Source/Editor/PropertyEditor/Private/Presentation/PropertyTable/PropertyTableColumn.h
+++ b/Engine/Source/Editor/PropertyEditor/Private/Presentation/PropertyTable/PropertyTableColumn.h
@@ -15,7 +15,7 @@ public:
 
 	virtual ~FPropertyTableColumn() {}
 
-	// Begin IPropertyTable Interface
+	//~ Begin IPropertyTable Interface
 
 	virtual FName GetId() const override;
 
@@ -58,7 +58,7 @@ public:
 	DECLARE_DERIVED_EVENT( FPropertyTableColumn, IPropertyTableColumn::FFrozenStateChanged, FFrozenStateChanged );
 	FFrozenStateChanged* OnFrozenStateChanged() override { return &FrozenStateChanged; }
 
-	// End IPropertyTable Interface
+	//~ End IPropertyTable Interface
 
 private:
 
diff --git a/Engine/Source/Editor/PropertyEditor/Private/Presentation/PropertyTable/PropertyTablePropertyNameCell.h b/Engine/Source/Editor/PropertyEditor/Private/Presentation/PropertyTable/PropertyTablePropertyNameCell.h
--- a/Engine/Source/Editor/PropertyEditor/Private/Presentation/PropertyTable/PropertyTablePropertyNameCell.h
+++ b/Engine/Source/Editor/PropertyEditor/Private/Presentation/PropertyTable/PropertyTablePropertyNameCell.h
@@ -11,7 +11,7 @@ public:
 
 	FPropertyTablePropertyNameCell( const TSharedRef< class FPropertyTablePropertyNameColumn >& InColumn, const TSharedRef< class IPropertyTableRow >& InRow );
 
-	// Begin IPropertyTableCell Interface
+	//~ Begin IPropertyTableCell Interface
 
 	virtual void EnterEditMode() override;
 
@@ -49,7 +49,7 @@ public:
 
 	virtual void SetValueFromString( const FString& InString ) override {}
 
-	// End IPropertyTableCell Interface
+	//~ End IPropertyTableCell Interface
 
 private:
 
diff --git a/Engine/Source/Editor/PropertyEditor/Private/Presentation/PropertyTable/PropertyTablePropertyNameColumn.h b/Engine/Source/Editor/PropertyEditor/Private/Presentation/PropertyTable/PropertyTablePropertyNameColumn.h
--- a/Engine/Source/Editor/PropertyEditor/Private/Presentation/PropertyTable/PropertyTablePropertyNameColumn.h
+++ b/Engine/Source/Editor/PropertyEditor/Private/Presentation/PropertyTable/PropertyTablePropertyNameColumn.h
@@ -20,7 +20,7 @@ public:
 
 	virtual ~FPropertyTablePropertyNameColumn() {}
 
-	// Begin IPropertyTableColumn Interface
+	//~ Begin IPropertyTableColumn Interface
 
 	virtual bool CanSelectCells() const override { return true; }
 
@@ -66,7 +66,7 @@ public:
 	DECLARE_DERIVED_EVENT( FPropertyTableColumn, IPropertyTableColumn::FFrozenStateChanged, FFrozenStateChanged );
 	FFrozenStateChanged* OnFrozenStateChanged() override { return &FrozenStateChanged; }
 
-	// End IPropertyTableColumn Interface
+	//~ End IPropertyTableColumn Interface
 
 private:
 
diff --git a/Engine/Source/Editor/Sequencer/Private/Sequencer.h b/Engine/Source/Editor/Sequencer/Private/Sequencer.h
--- a/Engine/Source/Editor/Sequencer/Private/Sequencer.h
+++ b/Engine/Source/Editor/Sequencer/Private/Sequencer.h
@@ -43,13 +43,13 @@ public:
 
 public:
 
-	// FGCObject interface
+	//~ Begin FGCObject Interface
 
 	virtual void AddReferencedObjects(FReferenceCollector& Collector) override;
 
 public:
 
-	// FTickableEditorObject interface
+	//~ Begin FTickableEditorObject Interface
 
 	virtual void Tick(float DeltaTime) override;
 	virtual bool IsTickable() const override { return true; }
@@ -57,7 +57,7 @@ public:
 
 public:
 
-	// ISequencer interface
+	//~ Begin ISequencer Interface
 
 	virtual void Close() override;
 	virtual TSharedRef<SWidget> GetSequencerWidget() const override { return SequencerWidget.ToSharedRef(); }
@@ -435,7 +435,7 @@ protected:
 	void ActivateDetailKeyframeHandler();
 	void DeactivateDetailKeyframeHandler();
 
-	// Begin FEditorUndoClient Interface
+	//~ Begin FEditorUndoClient Interface
 	virtual void PostUndo(bool bSuccess) override;
 	virtual void PostRedo(bool bSuccess) override { PostUndo(bSuccess); }
 	// End of FEditorUndoClient
diff --git a/Engine/Source/Editor/StaticMeshEditor/Private/StaticMeshEditor.h b/Engine/Source/Editor/StaticMeshEditor/Private/StaticMeshEditor.h
--- a/Engine/Source/Editor/StaticMeshEditor/Private/StaticMeshEditor.h
+++ b/Engine/Source/Editor/StaticMeshEditor/Private/StaticMeshEditor.h
@@ -35,9 +35,9 @@ public:
 	/** Creates details for a static mesh */
 	TSharedRef<class IDetailCustomization> MakeStaticMeshDetails();
 
-	// FGCObject interface
+	//~ Begin FGCObject Interface
 	virtual void AddReferencedObjects( FReferenceCollector& Collector ) override;
-	// End of FGCObject interface
+	//~ End FGCObject Interface
 
 	/** IToolkit interface */
 	virtual FName GetToolkitFName() const override;
@@ -231,11 +231,11 @@ private:
 	/** Opens the convex decomposition tab. */
 	void OnConvexDecomposition();
 
-	// Begin FAssetEditorToolkit interface.
+	//~ Begin FAssetEditorToolkit Interface.
 	virtual bool OnRequestClose() override;
-	// End FAssetEditorToolkit interface.
+	//~ End FAssetEditorToolkit Interface.
 
-	// Begin FEditorUndoClient Interface
+	//~ Begin FEditorUndoClient Interface
 	virtual void PostUndo( bool bSuccess ) override;
 	virtual void PostRedo( bool bSuccess ) override;
 	// End of FEditorUndoClient
diff --git a/Engine/Source/Editor/UMGEditor/Classes/WidgetBlueprintFactory.h b/Engine/Source/Editor/UMGEditor/Classes/WidgetBlueprintFactory.h
--- a/Engine/Source/Editor/UMGEditor/Classes/WidgetBlueprintFactory.h
+++ b/Engine/Source/Editor/UMGEditor/Classes/WidgetBlueprintFactory.h
@@ -17,10 +17,10 @@ class UWidgetBlueprintFactory : public UFactory
 	UPROPERTY(EditAnywhere, Category=WidgetBlueprintFactory, meta=(AllowAbstract = ""))
 	TSubclassOf<class UUserWidget> ParentClass;
 
-	// Begin UFactory Interface
+	//~ Begin UFactory Interface
 	virtual bool ConfigureProperties() override;
 	virtual bool ShouldShowInNewMenu() const override;
 	virtual UObject* FactoryCreateNew(UClass* Class, UObject* InParent, FName Name, EObjectFlags Flags, UObject* Context, FFeedbackContext* Warn, FName CallingContext) override;
 	virtual UObject* FactoryCreateNew(UClass* Class, UObject* InParent, FName Name, EObjectFlags Flags, UObject* Context, FFeedbackContext* Warn) override;
-	// Begin UFactory Interface	
+	//~ Begin UFactory Interface	
 };
diff --git a/Engine/Source/Editor/UMGEditor/Private/Nodes/K2Node_CreateDragDropOperation.h b/Engine/Source/Editor/UMGEditor/Private/Nodes/K2Node_CreateDragDropOperation.h
--- a/Engine/Source/Editor/UMGEditor/Private/Nodes/K2Node_CreateDragDropOperation.h
+++ b/Engine/Source/Editor/UMGEditor/Private/Nodes/K2Node_CreateDragDropOperation.h
@@ -11,15 +11,15 @@ class UMGEDITOR_API UK2Node_CreateDragDropOperation : public UK2Node_ConstructOb
 {
 	GENERATED_UCLASS_BODY()
 
-	// Begin UEdGraphNode interface.
+	//~ Begin UEdGraphNode Interface.
 	virtual void AllocateDefaultPins() override;
 	virtual void ExpandNode(class FKismetCompilerContext& CompilerContext, UEdGraph* SourceGraph) override;
-	// End UEdGraphNode interface.
+	//~ End UEdGraphNode Interface.
 
-	// Begin UK2Node interface
+	//~ Begin UK2Node Interface
 	virtual FText GetMenuCategory() const override;
 	virtual FName GetCornerIcon() const override;
-	// End UK2Node interface.
+	//~ End UK2Node Interface.
 
 protected:
 	/** Gets the default node title when no class is selected */
diff --git a/Engine/Source/Editor/UMGEditor/Private/Nodes/K2Node_CreateWidget.h b/Engine/Source/Editor/UMGEditor/Private/Nodes/K2Node_CreateWidget.h
--- a/Engine/Source/Editor/UMGEditor/Private/Nodes/K2Node_CreateWidget.h
+++ b/Engine/Source/Editor/UMGEditor/Private/Nodes/K2Node_CreateWidget.h
@@ -11,16 +11,16 @@ class UMGEDITOR_API UK2Node_CreateWidget : public UK2Node_ConstructObjectFromCla
 {
 	GENERATED_UCLASS_BODY()
 
-	// Begin UEdGraphNode interface.
+	//~ Begin UEdGraphNode Interface.
 	virtual void AllocateDefaultPins() override;
 	virtual FLinearColor GetNodeTitleColor() const override;
 	virtual void ExpandNode(class FKismetCompilerContext& CompilerContext, UEdGraph* SourceGraph) override;
-	// End UEdGraphNode interface.
+	//~ End UEdGraphNode Interface.
 
-	// Begin UK2Node interface
+	//~ Begin UK2Node Interface
 	virtual FText GetMenuCategory() const override;
 	virtual FName GetCornerIcon() const override;
-	// End UK2Node interface.
+	//~ End UK2Node Interface.
 
 	/** Get the owning player pin */
 	UEdGraphPin* GetOwningPlayerPin() const;
diff --git a/Engine/Source/Editor/UMGEditor/Public/WidgetBlueprint.h b/Engine/Source/Editor/UMGEditor/Public/WidgetBlueprint.h
--- a/Engine/Source/Editor/UMGEditor/Public/WidgetBlueprint.h
+++ b/Engine/Source/Editor/UMGEditor/Public/WidgetBlueprint.h
@@ -175,10 +175,10 @@ class UMGEDITOR_API UWidgetBlueprint : public UBlueprint
 
 public:
 
-	// Begin UObject interface
+	//~ Begin UObject Interface
 	virtual bool NeedsLoadForClient() const override;
 	virtual bool NeedsLoadForServer() const override;
-	// End UObject interface
+	//~ End UObject Interface
 
 public:
 
@@ -211,7 +211,7 @@ public:
 	virtual void PostLoad() override;
 	virtual void PostDuplicate(bool bDuplicateForPIE) override;
 	
-	// UBlueprint interface
+	//~ Begin UBlueprint Interface
 	virtual UClass* GetBlueprintClass() const override;
 
 	virtual bool AllowsDynamicBinding() const override;
@@ -225,7 +225,7 @@ public:
 
 	/** UWidget blueprints are never data only, should always compile on load (data only blueprints cannot declare new variables) */
 	virtual bool AlwaysCompileOnLoad() const override { return true; }
-	// End of UBlueprint interface
+	//~ End UBlueprint Interface
 
 	virtual void GatherDependencies(TSet<TWeakObjectPtr<UBlueprint>>& InDependencies) const override;
 
diff --git a/Engine/Source/Editor/UndoHistory/Private/Widgets/SUndoHistory.h b/Engine/Source/Editor/UndoHistory/Private/Widgets/SUndoHistory.h
--- a/Engine/Source/Editor/UndoHistory/Private/Widgets/SUndoHistory.h
+++ b/Engine/Source/Editor/UndoHistory/Private/Widgets/SUndoHistory.h
@@ -41,11 +41,11 @@ public:
 
 public:
 
-	// Begin SWidget interface
+	//~ Begin SWidget Interface
 
 	virtual void Tick( const FGeometry& AllottedGeometry, const double InCurrentTime, const float InDeltaTime ) override;
 
-	// End SWidget interface
+	//~ End SWidget Interface
 
 protected:
 
diff --git a/Engine/Source/Editor/UndoHistory/Public/UndoHistoryModule.h b/Engine/Source/Editor/UndoHistory/Public/UndoHistoryModule.h
--- a/Engine/Source/Editor/UndoHistory/Public/UndoHistoryModule.h
+++ b/Engine/Source/Editor/UndoHistory/Public/UndoHistoryModule.h
@@ -13,13 +13,13 @@ class FUndoHistoryModule : public IUndoHistoryModule
 {
 public:
 
-	// Begin IModuleInterface interface
+	//~ Begin IModuleInterface Interface
 
 	virtual void StartupModule() override;
 	virtual void ShutdownModule() override;
 	virtual bool SupportsDynamicReloading() override;
 
-	// End IModuleInterface interface
+	//~ End IModuleInterface Interface
 
 	static void ExecuteOpenUndoHistory()
 	{
diff --git a/Engine/Source/Editor/UnrealEd/Classes/Animation/DebugSkelMeshComponent.h b/Engine/Source/Editor/UnrealEd/Classes/Animation/DebugSkelMeshComponent.h
--- a/Engine/Source/Editor/UnrealEd/Classes/Animation/DebugSkelMeshComponent.h
+++ b/Engine/Source/Editor/UnrealEd/Classes/Animation/DebugSkelMeshComponent.h
@@ -166,11 +166,11 @@ class UDebugSkelMeshComponent : public USkeletalMeshComponent
 	/** true if wind effects on clothing is enabled */
 	bool bEnableWind;
 
-	// Begin USceneComponent interface.
+	//~ Begin USceneComponent Interface.
 	virtual FBoxSphereBounds CalcBounds(const FTransform& LocalToWorld) const override;
-	// End USceneComponent interface.
+	//~ End USceneComponent Interface.
 
-	// Begin UPrimitiveComponent interface.
+	//~ Begin UPrimitiveComponent Interface.
 	virtual FPrimitiveSceneProxy* CreateSceneProxy() override;
 	
 	// engine only draw bounds IF selected
@@ -181,18 +181,18 @@ class UDebugSkelMeshComponent : public USkeletalMeshComponent
 	{
 		return bDisplayBound;
 	}
-	// End UPrimitiveComponent interface.
+	//~ End UPrimitiveComponent Interface.
 
-	// Begin SkinnedMeshComponent Interface
+	//~ Begin SkinnedMeshComponent Interface
 	virtual bool ShouldCPUSkin() override;
 	virtual void PostInitMeshObject(class FSkeletalMeshObject* MeshObject) override;
 	virtual void RefreshBoneTransforms(FActorComponentTickFunction* TickFunction = NULL) override;
-	// End SkinnedMeshComponent Interface
+	//~ End SkinnedMeshComponent Interface
 
-	// Begin SkeletalMeshComponent Interface
+	//~ Begin SkeletalMeshComponent Interface
 	virtual void InitAnim(bool bForceReinit) override;
 	virtual bool IsWindEnabled() const override;
-	// End SkeletalMeshComponent Interface
+	//~ End SkeletalMeshComponent Interface
 	// Preview.
 	// @todo document
 	UNREALED_API bool IsPreviewOn() const;
diff --git a/Engine/Source/Editor/UnrealEd/Classes/Builders/ConeBuilder.h b/Engine/Source/Editor/UnrealEd/Classes/Builders/ConeBuilder.h
--- a/Engine/Source/Editor/UnrealEd/Classes/Builders/ConeBuilder.h
+++ b/Engine/Source/Editor/UnrealEd/Classes/Builders/ConeBuilder.h
@@ -1,8 +1,8 @@
 // Copyright 1998-2015 Epic Games, Inc. All Rights Reserved.
 
-//=============================================================================
+//~=============================================================================
 // ConeBuilder: Builds a 3D cone brush, compatible with cylinder of same size.
-//=============================================================================
+//~=============================================================================
 
 #pragma once
 #include "ConeBuilder.generated.h"
@@ -47,13 +47,13 @@ public:
 	UPROPERTY(EditAnywhere, Category=BrushSettings)
 	uint32 Hollow:1;
 
-	// Begin UObject interface
+	//~ Begin UObject Interface
 	virtual void PostEditChangeProperty(struct FPropertyChangedEvent& PropertyChangedEvent) override;
-	// End UObject interface
+	//~ End UObject Interface
 
-	// Begin UBrushBuilder Interface
+	//~ Begin UBrushBuilder Interface
 	virtual bool Build( UWorld* InWorld, ABrush* InBrush = NULL ) override;
-	// End UBrushBuilder Interface
+	//~ End UBrushBuilder Interface
 
 	// @todo document
 	virtual void BuildCone( int32 Direction, bool InAlignToSide, int32 InSides, float InZ, float Radius, FName Item );
diff --git a/Engine/Source/Editor/UnrealEd/Classes/Builders/CubeBuilder.h b/Engine/Source/Editor/UnrealEd/Classes/Builders/CubeBuilder.h
--- a/Engine/Source/Editor/UnrealEd/Classes/Builders/CubeBuilder.h
+++ b/Engine/Source/Editor/UnrealEd/Classes/Builders/CubeBuilder.h
@@ -44,13 +44,13 @@ public:
 	uint32 Tessellated:1;
 
 
-	// Begin UObject interface
+	//~ Begin UObject Interface
 	virtual void PostEditChangeProperty(struct FPropertyChangedEvent& PropertyChangedEvent) override;
-	// End UObject interface
+	//~ End UObject Interface
 
-	// Begin UBrushBuilder Interface
+	//~ Begin UBrushBuilder Interface
 	virtual bool Build( UWorld* InWorld, ABrush* InBrush = NULL ) override;
-	// End UBrushBuilder Interface
+	//~ End UBrushBuilder Interface
 
 	// @todo document
 	virtual void BuildCube( int32 Direction, float dx, float dy, float dz, bool _tessellated );
diff --git a/Engine/Source/Editor/UnrealEd/Classes/Builders/CurvedStairBuilder.h b/Engine/Source/Editor/UnrealEd/Classes/Builders/CurvedStairBuilder.h
--- a/Engine/Source/Editor/UnrealEd/Classes/Builders/CurvedStairBuilder.h
+++ b/Engine/Source/Editor/UnrealEd/Classes/Builders/CurvedStairBuilder.h
@@ -48,9 +48,9 @@ public:
 	uint32 CounterClockwise:1;
 
 
-	// Begin UBrushBuilder Interface
+	//~ Begin UBrushBuilder Interface
 	virtual bool Build( UWorld* InWorld, ABrush* InBrush = NULL ) override;
-	// End UBrushBuilder Interface
+	//~ End UBrushBuilder Interface
 
 	// @todo document
 	virtual void BuildCurvedStair( int32 Direction );
diff --git a/Engine/Source/Editor/UnrealEd/Classes/Builders/CylinderBuilder.h b/Engine/Source/Editor/UnrealEd/Classes/Builders/CylinderBuilder.h
--- a/Engine/Source/Editor/UnrealEd/Classes/Builders/CylinderBuilder.h
+++ b/Engine/Source/Editor/UnrealEd/Classes/Builders/CylinderBuilder.h
@@ -5,9 +5,9 @@
  */
 
 
-//=============================================================================
+//~=============================================================================
 // CylinderBuilder: Builds a 3D cylinder brush.
-//=============================================================================
+//~=============================================================================
 
 #pragma once
 #include "CylinderBuilder.generated.h"
@@ -48,13 +48,13 @@ public:
 	UPROPERTY(EditAnywhere, Category=BrushSettings)
 	uint32 Hollow:1;
 
-	// Begin UObject interface
+	//~ Begin UObject Interface
 	virtual void PostEditChangeProperty(struct FPropertyChangedEvent& PropertyChangedEvent) override;
-	// End UObject interface
+	//~ End UObject Interface
 
-	// Begin UBrushBuilder Interface
+	//~ Begin UBrushBuilder Interface
 	virtual bool Build( UWorld* InWorld, ABrush* InBrush = NULL ) override;
-	// End UBrushBuilder Interface
+	//~ End UBrushBuilder Interface
 
 	// @todo document
 	virtual void BuildCylinder( int32 Direction, bool InAlignToSide, int32 InSides, float InZ, float Radius );
diff --git a/Engine/Source/Editor/UnrealEd/Classes/Builders/LinearStairBuilder.h b/Engine/Source/Editor/UnrealEd/Classes/Builders/LinearStairBuilder.h
--- a/Engine/Source/Editor/UnrealEd/Classes/Builders/LinearStairBuilder.h
+++ b/Engine/Source/Editor/UnrealEd/Classes/Builders/LinearStairBuilder.h
@@ -1,8 +1,8 @@
 // Copyright 1998-2015 Epic Games, Inc. All Rights Reserved.
 
-//=============================================================================
+//~=============================================================================
 // LinearStairBuilder: Builds a Linear Staircase.
-//=============================================================================
+//~=============================================================================
 
 #pragma once
 #include "LinearStairBuilder.generated.h"
@@ -40,9 +40,9 @@ public:
 	FName GroupName;
 
 
-	// Begin UBrushBuilder Interface
+	//~ Begin UBrushBuilder Interface
 	virtual bool Build( UWorld* InWorld, ABrush* InBrush = NULL ) override;
-	// End UBrushBuilder Interface
+	//~ End UBrushBuilder Interface
 };
 
 
diff --git a/Engine/Source/Editor/UnrealEd/Classes/Builders/SheetBuilder.h b/Engine/Source/Editor/UnrealEd/Classes/Builders/SheetBuilder.h
--- a/Engine/Source/Editor/UnrealEd/Classes/Builders/SheetBuilder.h
+++ b/Engine/Source/Editor/UnrealEd/Classes/Builders/SheetBuilder.h
@@ -1,8 +1,8 @@
 // Copyright 1998-2015 Epic Games, Inc. All Rights Reserved.
 
-//=============================================================================
+//~=============================================================================
 // SheetBuilder: Builds a simple sheet.
-//=============================================================================
+//~=============================================================================
 
 #pragma once
 #include "SheetBuilder.generated.h"
@@ -44,9 +44,9 @@ public:
 	FName GroupName;
 
 
-	// Begin UBrushBuilder Interface
+	//~ Begin UBrushBuilder Interface
 	virtual bool Build( UWorld* InWorld, ABrush* InBrush = NULL ) override;
-	// End UBrushBuilder Interface
+	//~ End UBrushBuilder Interface
 };
 
 
diff --git a/Engine/Source/Editor/UnrealEd/Classes/Commandlets/AudioTestCommandlet.h b/Engine/Source/Editor/UnrealEd/Classes/Commandlets/AudioTestCommandlet.h
--- a/Engine/Source/Editor/UnrealEd/Classes/Commandlets/AudioTestCommandlet.h
+++ b/Engine/Source/Editor/UnrealEd/Classes/Commandlets/AudioTestCommandlet.h
@@ -9,7 +9,7 @@ class UAudioTestCommandlet : public UCommandlet
 {
 	GENERATED_UCLASS_BODY()
 
-	// Begin UCommandlet Interface
+	//~ Begin UCommandlet Interface
 	virtual int32 Main(const FString& Params) override;
-	// End UCommandlet Interface
+	//~ End UCommandlet Interface
 };
diff --git a/Engine/Source/Editor/UnrealEd/Classes/Commandlets/CompressAnimationsCommandlet.h b/Engine/Source/Editor/UnrealEd/Classes/Commandlets/CompressAnimationsCommandlet.h
--- a/Engine/Source/Editor/UnrealEd/Classes/Commandlets/CompressAnimationsCommandlet.h
+++ b/Engine/Source/Editor/UnrealEd/Classes/Commandlets/CompressAnimationsCommandlet.h
@@ -8,9 +8,9 @@ UCLASS()
 class UCompressAnimationsCommandlet : public UCommandlet
 {
 	GENERATED_UCLASS_BODY()
-	// Begin UCommandlet Interface
+	//~ Begin UCommandlet Interface
 	virtual int32 Main(const FString& Params) override;
-	// End UCommandlet Interface
+	//~ End UCommandlet Interface
 };
 
 
diff --git a/Engine/Source/Editor/UnrealEd/Classes/Editor/EditorEngine.h b/Engine/Source/Editor/UnrealEd/Classes/Editor/EditorEngine.h
--- a/Engine/Source/Editor/UnrealEd/Classes/Editor/EditorEngine.h
+++ b/Engine/Source/Editor/UnrealEd/Classes/Editor/EditorEngine.h
@@ -686,13 +686,13 @@ public:
 	/**	Broadcasts that an object has been reimported. THIS SHOULD NOT BE PUBLIC */
 	void BroadcastObjectReimported(UObject* InObject);
 
-	// Begin UObject interface.
+	//~ Begin UObject Interface.
 	virtual void FinishDestroy() override;
 	virtual void PostEditChangeProperty(FPropertyChangedEvent& PropertyChangedEvent) override;
 	static void AddReferencedObjects(UObject* InThis, FReferenceCollector& Collector);
-	// End UObject interface.
+	//~ End UObject Interface.
 
-	// Begin UEngine interface.
+	//~ Begin UEngine Interface.
 public:
 	virtual void Init(IEngineLoop* InEngineLoop) override;
 	virtual float GetMaxTickRate(float DeltaTime, bool bAllowFrameRateSmoothing = true) const override;
@@ -722,11 +722,11 @@ private:
 	virtual void CreateStartupAnalyticsAttributes(TArray<FAnalyticsEventAttribute>& StartSessionAttributes) const override;
 	virtual void VerifyLoadMapWorldCleanup() override;
 public:
-	// End UEngine interface.
+	//~ End UEngine Interface.
 	
-	// Begin FExec Interface
+	//~ Begin FExec Interface
 	virtual bool Exec( UWorld* InWorld, const TCHAR* Cmd, FOutputDevice& Ar=*GLog ) override;
-	// End FExec Interface
+	//~ End FExec Interface
 
 	bool	CommandIsDeprecated( const TCHAR* CommandStr, FOutputDevice& Ar );
 	
@@ -833,7 +833,7 @@ public:
 	 */
 	void ExecFile( UWorld* InWorld, const TCHAR* InFilename, FOutputDevice& Ar );
 
-	// Transaction interfaces.
+	//~ Begin Transaction Interfaces.
 	int32 BeginTransaction(const TCHAR* SessionContext, const FText& Description, UObject* PrimaryObject);
 	int32 BeginTransaction(const FText& Description);
 	int32 EndTransaction();
diff --git a/Engine/Source/Editor/UnrealEd/Classes/Editor/GroupActor.h b/Engine/Source/Editor/UnrealEd/Classes/Editor/GroupActor.h
--- a/Engine/Source/Editor/UnrealEd/Classes/Editor/GroupActor.h
+++ b/Engine/Source/Editor/UnrealEd/Classes/Editor/GroupActor.h
@@ -1,8 +1,8 @@
 // Copyright 1998-2015 Epic Games, Inc. All Rights Reserved.
 
-//=============================================================================
+//~=============================================================================
 // GroupActor: Collects a group of actors, allowing for management and universal transformation.
-//=============================================================================
+//~=============================================================================
 
 #pragma once
 #include "GameFramework/Actor.h"
@@ -26,7 +26,7 @@ class AGroupActor : public AActor
 #endif // WITH_EDITORONLY_DATA
 
 
-	// Begin UObject interface
+	//~ Begin UObject Interface
 	virtual void PostLoad() override;
 	virtual void PostEditChangeProperty(FPropertyChangedEvent& PropertyChangedEvent) override;
 	virtual void PostEditUndo() override;
@@ -39,14 +39,14 @@ class AGroupActor : public AActor
 	{ 
 		return false;
 	}
-	// End UObject interface
+	//~ End UObject Interface
 
-	// Begin AActor interface
+	//~ Begin AActor Interface
 	virtual void PostActorCreated() override;
 	virtual void InvalidateLightingCacheDetailed(bool bTranslationOnly) override;
 	virtual void PostEditMove(bool bFinished) override;
 	virtual void SetIsTemporarilyHiddenInEditor( bool bIsHidden ) override;
-	// End AActor interface
+	//~ End AActor Interface
 
 	// @todo document
 	virtual bool IsSelected() const;
diff --git a/Engine/Source/Editor/UnrealEd/Classes/Editor/TransBuffer.h b/Engine/Source/Editor/UnrealEd/Classes/Editor/TransBuffer.h
--- a/Engine/Source/Editor/UnrealEd/Classes/Editor/TransBuffer.h
+++ b/Engine/Source/Editor/UnrealEd/Classes/Editor/TransBuffer.h
@@ -47,12 +47,12 @@ public:
 
 public:
 
-	// Begin UObject interface.
+	//~ Begin UObject Interface.
 
 	virtual void Serialize( FArchive& Ar ) override;
 	virtual void FinishDestroy() override;
 
-	// End UObject interface.
+	//~ End UObject Interface.
 
 protected:
 	
@@ -88,7 +88,7 @@ protected:
 
 public:
 
-	// Begin UTransactor interface.
+	//~ Begin UTransactor Interface.
 
 	virtual int32 Begin( const TCHAR* SessionContext, const FText& Description ) override;
 	virtual int32 End() override;
@@ -116,7 +116,7 @@ public:
 		return ActiveCount > 0;
 	}
 
-	// End UTransactor interface.
+	//~ End UTransactor Interface.
 
 public:
 
diff --git a/Engine/Source/Editor/UnrealEd/Classes/Editor/UnrealEdEngine.h b/Engine/Source/Editor/UnrealEd/Classes/Editor/UnrealEdEngine.h
--- a/Engine/Source/Editor/UnrealEd/Classes/Editor/UnrealEdEngine.h
+++ b/Engine/Source/Editor/UnrealEd/Classes/Editor/UnrealEdEngine.h
@@ -168,18 +168,18 @@ public:
 	/** Manages currently active visualizer and routes interactions to it */
 	FComponentVisualizerManager	ComponentVisManager;
 
-	// Begin UObject interface.
+	//~ Begin UObject Interface.
 	~UUnrealEdEngine();
 	virtual void FinishDestroy() override;
 	virtual void Serialize( FArchive& Ar ) override;
-	// End UObject interface.
+	//~ End UObject Interface.
 
-	// Begin FNotify interface.
+	//~ Begin FNotify Interface.
 	virtual void NotifyPreChange( UProperty* PropertyAboutToChange ) override;
 	virtual void NotifyPostChange( const FPropertyChangedEvent& PropertyChangedEvent, UProperty* PropertyThatChanged ) override;
-	// End FNotify interface.
+	//~ End FNotify Interface.
 
-	// Begin UEditorEngine Interface
+	//~ Begin UEditorEngine Interface
 	virtual void SelectActor(AActor* Actor, bool InSelected, bool bNotify, bool bSelectEvenIfHidden = false, bool bForceRefresh = false) override;
 	virtual bool CanSelectActor(AActor* Actor, bool InSelected, bool bSelectEvenIfHidden=false, bool bWarnIfLevelLocked=false) const override;
 	virtual void SelectGroup(AGroupActor* InGroupActor, bool bForceSelection=false, bool bInSelected=true, bool bNotify=true) override;
@@ -203,18 +203,18 @@ public:
 	virtual void CloseEditor() override;
 	virtual void OnOpenMatinee() override;
 	virtual bool IsAutosaving() const override;
-	// End UEditorEngine Interface 
+	//~ End UEditorEngine Interface 
 	
-	// Begin FExec Interface
+	//~ Begin FExec Interface
 	virtual bool Exec( UWorld* InWorld, const TCHAR* Cmd, FOutputDevice& Ar=*GLog ) override;
-	// End FExec Interface
+	//~ End FExec Interface
 
 
-	// Begin UEngine Interface.
+	//~ Begin UEngine Interface.
 	virtual void Init(IEngineLoop* InEngineLoop) override;
 	virtual void PreExit() override;
 	virtual void Tick(float DeltaSeconds, bool bIdleMode) override;
-	// End UEngine Interface.
+	//~ End UEngine Interface.
 
 	/** Builds a list of sprite categories for use in menus */
 	static void MakeSortedSpriteInfo(TArray<struct FSpriteCategoryInfo>& OutSortedSpriteInfo);
diff --git a/Engine/Source/Editor/UnrealEd/Classes/Exporters/AnimSequenceExporterFBX.h b/Engine/Source/Editor/UnrealEd/Classes/Exporters/AnimSequenceExporterFBX.h
--- a/Engine/Source/Editor/UnrealEd/Classes/Exporters/AnimSequenceExporterFBX.h
+++ b/Engine/Source/Editor/UnrealEd/Classes/Exporters/AnimSequenceExporterFBX.h
@@ -1,8 +1,8 @@
 // Copyright 1998-2015 Epic Games, Inc. All Rights Reserved.
 
-//=============================================================================
+//~=============================================================================
 // AnimSequenceExporterFBX
-//=============================================================================
+//~=============================================================================
 
 #pragma once
 #include "Exporters/Exporter.h"
@@ -14,9 +14,9 @@ class UAnimSequenceExporterFBX : public UExporter
 	GENERATED_UCLASS_BODY()
 
 
-	// Begin UExporter Interface
+	//~ Begin UExporter Interface
 	virtual bool ExportBinary( UObject* Object, const TCHAR* Type, FArchive& Ar, FFeedbackContext* Warn, int32 FileIndex = 0, uint32 PortFlags=0 ) override;
-	// End UExporter Interface
+	//~ End UExporter Interface
 };
 
 
diff --git a/Engine/Source/Editor/UnrealEd/Classes/Exporters/LevelExporterFBX.h b/Engine/Source/Editor/UnrealEd/Classes/Exporters/LevelExporterFBX.h
--- a/Engine/Source/Editor/UnrealEd/Classes/Exporters/LevelExporterFBX.h
+++ b/Engine/Source/Editor/UnrealEd/Classes/Exporters/LevelExporterFBX.h
@@ -1,8 +1,8 @@
 // Copyright 1998-2015 Epic Games, Inc. All Rights Reserved.
 
-//=============================================================================
+//~=============================================================================
 // LevelExporterFBX
-//=============================================================================
+//~=============================================================================
 
 #pragma once
 #include "LevelExporterFBX.generated.h"
@@ -16,9 +16,9 @@ public:
 public:
 	ULevelExporterFBX(const FObjectInitializer& ObjectInitializer = FObjectInitializer());
 
-	// Begin UExporter Interface
+	//~ Begin UExporter Interface
 	virtual bool ExportBinary( UObject* Object, const TCHAR* Type, FArchive& Ar, FFeedbackContext* Warn, int32 FileIndex = 0, uint32 PortFlags=0 ) override;
-	// End UExporter Interface
+	//~ End UExporter Interface
 };
 
 
diff --git a/Engine/Source/Editor/UnrealEd/Classes/Exporters/SoundExporterOGG.h b/Engine/Source/Editor/UnrealEd/Classes/Exporters/SoundExporterOGG.h
--- a/Engine/Source/Editor/UnrealEd/Classes/Exporters/SoundExporterOGG.h
+++ b/Engine/Source/Editor/UnrealEd/Classes/Exporters/SoundExporterOGG.h
@@ -12,10 +12,10 @@ public:
 public:
 	USoundExporterOGG(const FObjectInitializer& ObjectInitializer = FObjectInitializer());
 
-	// Begin UExporter Interface
+	//~ Begin UExporter Interface
 	virtual bool ExportBinary( UObject* Object, const TCHAR* Type, FArchive& Ar, FFeedbackContext* Warn, int32 FileIndex = 0, uint32 PortFlags=0 ) override;
 	virtual bool SupportsObject(UObject* Object) const override;
-	// End UExporter Interface
+	//~ End UExporter Interface
 };
 
 
diff --git a/Engine/Source/Editor/UnrealEd/Classes/Exporters/SoundExporterWAV.h b/Engine/Source/Editor/UnrealEd/Classes/Exporters/SoundExporterWAV.h
--- a/Engine/Source/Editor/UnrealEd/Classes/Exporters/SoundExporterWAV.h
+++ b/Engine/Source/Editor/UnrealEd/Classes/Exporters/SoundExporterWAV.h
@@ -12,10 +12,10 @@ public:
 public:
 	USoundExporterWAV(const FObjectInitializer& ObjectInitializer = FObjectInitializer());
 
-	// Begin UExporter Interface
+	//~ Begin UExporter Interface
 	virtual bool ExportBinary( UObject* Object, const TCHAR* Type, FArchive& Ar, FFeedbackContext* Warn, int32 FileIndex = 0, uint32 PortFlags=0 ) override;
 	virtual bool SupportsObject(UObject* Object) const override;
-	// End UExporter Interface
+	//~ End UExporter Interface
 };
 
 
diff --git a/Engine/Source/Editor/UnrealEd/Classes/Factories/AimOffsetBlendSpaceFactory1D.h b/Engine/Source/Editor/UnrealEd/Classes/Factories/AimOffsetBlendSpaceFactory1D.h
--- a/Engine/Source/Editor/UnrealEd/Classes/Factories/AimOffsetBlendSpaceFactory1D.h
+++ b/Engine/Source/Editor/UnrealEd/Classes/Factories/AimOffsetBlendSpaceFactory1D.h
@@ -1,8 +1,8 @@
 // Copyright 1998-2015 Epic Games, Inc. All Rights Reserved.
 
-//=============================================================================
+//~=============================================================================
 // UAimOffsetBlendSpaceFactory1D
-//=============================================================================
+//~=============================================================================
 
 #pragma once
 #include "AimOffsetBlendSpaceFactory1D.generated.h"
@@ -12,9 +12,9 @@ class UAimOffsetBlendSpaceFactory1D : public UBlendSpaceFactory1D
 {
 	GENERATED_UCLASS_BODY()
 
-	// Begin UFactory Interface
+	//~ Begin UFactory Interface
 	virtual UObject* FactoryCreateNew(UClass* Class,UObject* InParent,FName Name,EObjectFlags Flags,UObject* Context,FFeedbackContext* Warn) override;
-	// Begin UFactory Interface
+	//~ Begin UFactory Interface
 };
 
 
diff --git a/Engine/Source/Editor/UnrealEd/Classes/Factories/AimOffsetBlendSpaceFactoryNew.h b/Engine/Source/Editor/UnrealEd/Classes/Factories/AimOffsetBlendSpaceFactoryNew.h
--- a/Engine/Source/Editor/UnrealEd/Classes/Factories/AimOffsetBlendSpaceFactoryNew.h
+++ b/Engine/Source/Editor/UnrealEd/Classes/Factories/AimOffsetBlendSpaceFactoryNew.h
@@ -1,8 +1,8 @@
 // Copyright 1998-2015 Epic Games, Inc. All Rights Reserved.
 
-//=============================================================================
+//~=============================================================================
 // UAimOffsetBlendSpaceFactoryNew
-//=============================================================================
+//~=============================================================================
 
 #pragma once
 #include "AimOffsetBlendSpaceFactoryNew.generated.h"
@@ -12,9 +12,9 @@ class UAimOffsetBlendSpaceFactoryNew : public UBlendSpaceFactoryNew
 {
 	GENERATED_UCLASS_BODY()
 
-	// Begin UFactory Interface
+	//~ Begin UFactory Interface
 	virtual UObject* FactoryCreateNew(UClass* Class,UObject* InParent,FName Name,EObjectFlags Flags,UObject* Context,FFeedbackContext* Warn) override;
-	// Begin UFactory Interface
+	//~ Begin UFactory Interface
 };
 
 
diff --git a/Engine/Source/Editor/UnrealEd/Classes/Factories/Factory.h b/Engine/Source/Editor/UnrealEd/Classes/Factories/Factory.h
--- a/Engine/Source/Editor/UnrealEd/Classes/Factories/Factory.h
+++ b/Engine/Source/Editor/UnrealEd/Classes/Factories/Factory.h
@@ -69,9 +69,9 @@ public:
 	DEPRECATED(4.8, "bAllowOneTimeWarningMessages is due to be removed in future.")
 	static bool bAllowOneTimeWarningMessages;
 
-	// Begin UObject interface.
+	//~ Begin UObject Interface.
 	static void AddReferencedObjects(UObject* InThis, FReferenceCollector& Collector);
-	// End UObject interface.
+	//~ End UObject Interface.
 
 	/** Helper function to sort an array of factories by their import priority - use as a predicate for Sort */
 	static bool SortFactoriesByPriority(const UFactory& A, const UFactory& B);
diff --git a/Engine/Source/Editor/UnrealEd/Classes/Factories/FbxFactory.h b/Engine/Source/Editor/UnrealEd/Classes/Factories/FbxFactory.h
--- a/Engine/Source/Editor/UnrealEd/Classes/Factories/FbxFactory.h
+++ b/Engine/Source/Editor/UnrealEd/Classes/Factories/FbxFactory.h
@@ -16,18 +16,18 @@ class UNREALED_API UFbxFactory : public UFactory
 	void EnableShowOption() { bShowOption = true; }
 
 
-	// Begin UObject Interface
+	//~ Begin UObject Interface
 	virtual void CleanUp() override;
 	virtual bool ConfigureProperties() override;
 	virtual void PostInitProperties() override;
-	// End UObject Interface
+	//~ End UObject Interface
 
-	// Begin UFactory Interface
+	//~ Begin UFactory Interface
 	virtual bool DoesSupportClass(UClass * Class) override;
 	virtual UClass* ResolveSupportedClass() override;
 	virtual UObject* FactoryCreateBinary(UClass* InClass, UObject* InParent, FName InName, EObjectFlags Flags, UObject* Context, const TCHAR* Type, const uint8*& Buffer, const uint8* BufferEnd, FFeedbackContext* Warn, bool& bOutOperationCanceled) override;
 	virtual bool FactoryCanImport(const FString& Filename) override;
-	// End UFactory Interface
+	//~ End UFactory Interface
 	
 	/**
 	 * Detect mesh type to import: Static Mesh or Skeletal Mesh.
diff --git a/Engine/Source/Editor/UnrealEd/Classes/Factories/ModelFactory.h b/Engine/Source/Editor/UnrealEd/Classes/Factories/ModelFactory.h
--- a/Engine/Source/Editor/UnrealEd/Classes/Factories/ModelFactory.h
+++ b/Engine/Source/Editor/UnrealEd/Classes/Factories/ModelFactory.h
@@ -1,8 +1,8 @@
 // Copyright 1998-2015 Epic Games, Inc. All Rights Reserved.
 
-//=============================================================================
+//~=============================================================================
 // ModelFactory
-//=============================================================================
+//~=============================================================================
 
 #pragma once
 #include "ModelFactory.generated.h"
@@ -13,9 +13,9 @@ class UModelFactory : public UFactory
 	GENERATED_UCLASS_BODY()
 
 
-	// Begin UFactory Interface
+	//~ Begin UFactory Interface
 	virtual UObject* FactoryCreateText( UClass* Class, UObject* InParent, FName Name, EObjectFlags Flags, UObject* Context, const TCHAR* Type, const TCHAR*& Buffer, const TCHAR* BufferEnd, FFeedbackContext* Warn ) override;
-	// End UFactory Interface
+	//~ End UFactory Interface
 };
 
 
diff --git a/Engine/Source/Editor/UnrealEd/Classes/Factories/PackFactory.h b/Engine/Source/Editor/UnrealEd/Classes/Factories/PackFactory.h
--- a/Engine/Source/Editor/UnrealEd/Classes/Factories/PackFactory.h
+++ b/Engine/Source/Editor/UnrealEd/Classes/Factories/PackFactory.h
@@ -11,7 +11,7 @@ class UPackFactory : public UFactory
 {
 	GENERATED_UCLASS_BODY()
 
-	// Begin UFactory Interface
+	//~ Begin UFactory Interface
 	virtual UObject* FactoryCreateBinary( UClass* Class, UObject* InParent, FName Name, EObjectFlags Flags, UObject* Context, const TCHAR* Type, const uint8*& Buffer, const uint8* BufferEnd, FFeedbackContext* Warn ) override;
-	// End UFactory Interface
+	//~ End UFactory Interface
 };
\ No newline at end of file
diff --git a/Engine/Source/Editor/UnrealEd/Classes/Factories/PackageFactory.h b/Engine/Source/Editor/UnrealEd/Classes/Factories/PackageFactory.h
--- a/Engine/Source/Editor/UnrealEd/Classes/Factories/PackageFactory.h
+++ b/Engine/Source/Editor/UnrealEd/Classes/Factories/PackageFactory.h
@@ -1,8 +1,8 @@
 // Copyright 1998-2015 Epic Games, Inc. All Rights Reserved.
 
-//=============================================================================
+//~=============================================================================
 // PackageFactory
-//=============================================================================
+//~=============================================================================
 
 #pragma once
 #include "PackageFactory.generated.h"
@@ -13,9 +13,9 @@ class UPackageFactory : public UFactory
 	GENERATED_UCLASS_BODY()
 
 
-	// Begin UFactory Interface
+	//~ Begin UFactory Interface
 	virtual UObject* FactoryCreateText( UClass* Class, UObject* InParent, FName Name, EObjectFlags Flags, UObject* Context, const TCHAR* Type, const TCHAR*& Buffer, const TCHAR* BufferEnd, FFeedbackContext* Warn ) override;
-	// End UFactory Interface
+	//~ End UFactory Interface
 };
 
 
diff --git a/Engine/Source/Editor/UnrealEd/Classes/Factories/ReimportCurveFactory.h b/Engine/Source/Editor/UnrealEd/Classes/Factories/ReimportCurveFactory.h
--- a/Engine/Source/Editor/UnrealEd/Classes/Factories/ReimportCurveFactory.h
+++ b/Engine/Source/Editor/UnrealEd/Classes/Factories/ReimportCurveFactory.h
@@ -9,12 +9,12 @@ class UReimportCurveFactory : public UCSVImportFactory, public FReimportHandler
 	GENERATED_UCLASS_BODY()
 
 
-	// Begin FReimportHandler interface
+	//~ Begin FReimportHandler Interface
 	virtual bool CanReimport( UObject* Obj, TArray<FString>& OutFilenames ) override;
 	virtual void SetReimportPaths( UObject* Obj, const TArray<FString>& NewReimportPaths ) override;
 	virtual EReimportResult::Type Reimport( UObject* Obj ) override;
 	virtual int32 GetPriority() const override;
-	// End FReimportHandler interface
+	//~ End FReimportHandler Interface
 };
 
 
diff --git a/Engine/Source/Editor/UnrealEd/Classes/Factories/ReimportFbxAnimSequenceFactory.h b/Engine/Source/Editor/UnrealEd/Classes/Factories/ReimportFbxAnimSequenceFactory.h
--- a/Engine/Source/Editor/UnrealEd/Classes/Factories/ReimportFbxAnimSequenceFactory.h
+++ b/Engine/Source/Editor/UnrealEd/Classes/Factories/ReimportFbxAnimSequenceFactory.h
@@ -1,8 +1,8 @@
 // Copyright 1998-2015 Epic Games, Inc. All Rights Reserved.
 
-//=============================================================================
+//~=============================================================================
 // ReimportFbxAnimSequenceFactory
-//=============================================================================
+//~=============================================================================
 
 #pragma once
 #include "ReimportFbxAnimSequenceFactory.generated.h"
@@ -13,16 +13,16 @@ class UReimportFbxAnimSequenceFactory : public UFbxFactory, public FReimportHand
 	GENERATED_UCLASS_BODY()
 
 
-	// Begin FReimportHandler interface
+	//~ Begin FReimportHandler Interface
 	virtual bool CanReimport( UObject* Obj, TArray<FString>& OutFilenames ) override;
 	virtual void SetReimportPaths( UObject* Obj, const TArray<FString>& NewReimportPaths ) override;
 	virtual EReimportResult::Type Reimport( UObject* Obj ) override;
 	virtual int32 GetPriority() const override;
-	// End FReimportHandler interface
+	//~ End FReimportHandler Interface
 
-	// Begin UFactory Interface
+	//~ Begin UFactory Interface
 	virtual bool FactoryCanImport(const FString& Filename) override;
-	// End UFactory Interface
+	//~ End UFactory Interface
 };
 
 
diff --git a/Engine/Source/Editor/UnrealEd/Classes/Factories/ReimportFbxSkeletalMeshFactory.h b/Engine/Source/Editor/UnrealEd/Classes/Factories/ReimportFbxSkeletalMeshFactory.h
--- a/Engine/Source/Editor/UnrealEd/Classes/Factories/ReimportFbxSkeletalMeshFactory.h
+++ b/Engine/Source/Editor/UnrealEd/Classes/Factories/ReimportFbxSkeletalMeshFactory.h
@@ -1,8 +1,8 @@
 // Copyright 1998-2015 Epic Games, Inc. All Rights Reserved.
 
-//=============================================================================
+//~=============================================================================
 // ReimportFbxSkeletalMeshFactory
-//=============================================================================
+//~=============================================================================
 
 #pragma once
 #include "ReimportFbxSkeletalMeshFactory.generated.h"
@@ -13,16 +13,16 @@ class UReimportFbxSkeletalMeshFactory : public UFbxFactory, public FReimportHand
 	GENERATED_UCLASS_BODY()
 
 
-	// Begin FReimportHandler interface
+	//~ Begin FReimportHandler Interface
 	virtual bool CanReimport( UObject* Obj, TArray<FString>& OutFilenames ) override;
 	virtual void SetReimportPaths( UObject* Obj, const TArray<FString>& NewReimportPaths ) override;
 	virtual EReimportResult::Type Reimport( UObject* Obj ) override;
 	virtual int32 GetPriority() const override;
-	// End FReimportHandler interface
+	//~ End FReimportHandler Interface
 
-	// Begin UFactory Interface
+	//~ Begin UFactory Interface
 	virtual bool FactoryCanImport(const FString& Filename) override;
-	// End UFactory Interface
+	//~ End UFactory Interface
 };
 
 
diff --git a/Engine/Source/Editor/UnrealEd/Classes/Factories/ReimportSoundFactory.h b/Engine/Source/Editor/UnrealEd/Classes/Factories/ReimportSoundFactory.h
--- a/Engine/Source/Editor/UnrealEd/Classes/Factories/ReimportSoundFactory.h
+++ b/Engine/Source/Editor/UnrealEd/Classes/Factories/ReimportSoundFactory.h
@@ -1,8 +1,8 @@
 // Copyright 1998-2015 Epic Games, Inc. All Rights Reserved.
 
-//=============================================================================
+//~=============================================================================
 // ReimportSoundFactory
-//=============================================================================
+//~=============================================================================
 
 #pragma once
 #include "ReimportSoundFactory.generated.h"
@@ -13,12 +13,12 @@ class UReimportSoundFactory : public USoundFactory, public FReimportHandler
 	GENERATED_UCLASS_BODY()
 
 
-	// Begin FReimportHandler interface
+	//~ Begin FReimportHandler Interface
 	virtual bool CanReimport( UObject* Obj, TArray<FString>& OutFilenames ) override;
 	virtual void SetReimportPaths( UObject* Obj, const TArray<FString>& NewReimportPaths ) override;
 	virtual EReimportResult::Type Reimport( UObject* Obj ) override;
 	virtual int32 GetPriority() const override;
-	// End FReimportHandler interface
+	//~ End FReimportHandler Interface
 };
 
 
diff --git a/Engine/Source/Editor/UnrealEd/Classes/Factories/ReimportSoundSurroundFactory.h b/Engine/Source/Editor/UnrealEd/Classes/Factories/ReimportSoundSurroundFactory.h
--- a/Engine/Source/Editor/UnrealEd/Classes/Factories/ReimportSoundSurroundFactory.h
+++ b/Engine/Source/Editor/UnrealEd/Classes/Factories/ReimportSoundSurroundFactory.h
@@ -1,8 +1,8 @@
 // Copyright 1998-2015 Epic Games, Inc. All Rights Reserved.
 
-//=============================================================================
+//~=============================================================================
 // ReimportSoundSurroundFactory
-//=============================================================================
+//~=============================================================================
 
 #pragma once
 #include "ReimportSoundSurroundFactory.generated.h"
@@ -15,10 +15,10 @@ class UReimportSoundSurroundFactory : public USoundSurroundFactory, public FReim
 	UPROPERTY()
 	TArray<FString> ReimportPaths;
 
-	// Begin FReimportHandler interface
+	//~ Begin FReimportHandler Interface
 	virtual bool CanReimport(UObject* Obj, TArray<FString>& OutFilenames) override;
 	virtual void SetReimportPaths(UObject* Obj, const TArray<FString>& NewReimportPaths) override;
 	virtual EReimportResult::Type Reimport(UObject* Obj) override;
 	virtual int32 GetPriority() const override;
-	// End FReimportHandler interface
+	//~ End FReimportHandler Interface
 };
diff --git a/Engine/Source/Editor/UnrealEd/Classes/Factories/ReimportTextureFactory.h b/Engine/Source/Editor/UnrealEd/Classes/Factories/ReimportTextureFactory.h
--- a/Engine/Source/Editor/UnrealEd/Classes/Factories/ReimportTextureFactory.h
+++ b/Engine/Source/Editor/UnrealEd/Classes/Factories/ReimportTextureFactory.h
@@ -1,8 +1,8 @@
 // Copyright 1998-2015 Epic Games, Inc. All Rights Reserved.
 
-//=============================================================================
+//~=============================================================================
 // ReimportTextureFactory
-//=============================================================================
+//~=============================================================================
 
 #pragma once
 #include "ReimportTextureFactory.generated.h"
@@ -16,18 +16,18 @@ class UReimportTextureFactory : public UTextureFactory, public FReimportHandler
 	class UTexture* pOriginalTex;
 
 
-	// Begin FReimportHandler interface
+	//~ Begin FReimportHandler Interface
 	virtual bool CanReimport( UObject* Obj, TArray<FString>& OutFilenames ) override;
 	virtual void SetReimportPaths( UObject* Obj, const TArray<FString>& NewReimportPaths ) override;
 	virtual EReimportResult::Type Reimport( UObject* Obj ) override;
 	virtual int32 GetPriority() const override;
-	// End FReimportHandler interface
+	//~ End FReimportHandler Interface
 
 private:
-	// Begin UTextureFactory Interface
+	//~ Begin UTextureFactory Interface
 	virtual UTexture2D* CreateTexture2D( UObject* InParent, FName Name, EObjectFlags Flags ) override;
 	virtual UTextureCube* CreateTextureCube( UObject* InParent, FName Name, EObjectFlags Flags ) override;
-	// End UTextureFactory Interface
+	//~ End UTextureFactory Interface
 };
 
 
diff --git a/Engine/Source/Editor/UnrealEd/Classes/Factories/TextureCubeThumbnailRenderer.h b/Engine/Source/Editor/UnrealEd/Classes/Factories/TextureCubeThumbnailRenderer.h
--- a/Engine/Source/Editor/UnrealEd/Classes/Factories/TextureCubeThumbnailRenderer.h
+++ b/Engine/Source/Editor/UnrealEd/Classes/Factories/TextureCubeThumbnailRenderer.h
@@ -14,9 +14,9 @@ class UTextureCubeThumbnailRenderer : public UTextureThumbnailRenderer
 	GENERATED_UCLASS_BODY()
 
 
-	// Begin UThumbnailRenderer Interface
+	//~ Begin UThumbnailRenderer Interface
 	virtual void GetThumbnailSize(UObject* Object, float Zoom, uint32& OutWidth, uint32& OutHeight) const override;
 	virtual void Draw(UObject* Object, int32 X, int32 Y, uint32 Width, uint32 Height, FRenderTarget*, FCanvas* Canvas) override;
-	// End UThumbnailRenderer Interface
+	//~ End UThumbnailRenderer Interface
 };
 
diff --git a/Engine/Source/Editor/UnrealEd/Classes/Factories/TextureFactory.h b/Engine/Source/Editor/UnrealEd/Classes/Factories/TextureFactory.h
--- a/Engine/Source/Editor/UnrealEd/Classes/Factories/TextureFactory.h
+++ b/Engine/Source/Editor/UnrealEd/Classes/Factories/TextureFactory.h
@@ -92,14 +92,14 @@ class UNREALED_API UTextureFactory : public UFactory
 public:
 	UTextureFactory(const FObjectInitializer& ObjectInitializer = FObjectInitializer::Get());
 
-	// Begin UObject Interface
+	//~ Begin UObject Interface
 	virtual void PostInitProperties() override;
-	// End UObject Interface
+	//~ End UObject Interface
 
-	// Begin UFactory Interface
+	//~ Begin UFactory Interface
 	virtual bool DoesSupportClass(UClass* Class) override;
 	virtual UObject* FactoryCreateBinary( UClass* Class, UObject* InParent, FName Name, EObjectFlags Flags, UObject* Context, const TCHAR* Type, const uint8*& Buffer, const uint8* BufferEnd, FFeedbackContext* Warn ) override;
-	// End UFactory Interface
+	//~ End UFactory Interface
 	
 
 	/** Create a texture given the appropriate input parameters	*/
diff --git a/Engine/Source/Editor/UnrealEd/Classes/Factories/TrueTypeFontFactory.h b/Engine/Source/Editor/UnrealEd/Classes/Factories/TrueTypeFontFactory.h
--- a/Engine/Source/Editor/UnrealEd/Classes/Factories/TrueTypeFontFactory.h
+++ b/Engine/Source/Editor/UnrealEd/Classes/Factories/TrueTypeFontFactory.h
@@ -1,8 +1,8 @@
 // Copyright 1998-2015 Epic Games, Inc. All Rights Reserved.
 
-//=============================================================================
+//~=============================================================================
 // TrueTypeFontFactory
-//=============================================================================
+//~=============================================================================
 
 #pragma once
 #include "TrueTypeFontFactory.generated.h"
@@ -24,11 +24,11 @@ class UNREALED_API UTrueTypeFontFactory : public UTextureFactory, public FReimpo
 	UPROPERTY()
 	bool bFontSelected;
 
-	// Begin UObject Interface
+	//~ Begin UObject Interface
 	virtual void PostInitProperties() override;
-	// End UObject Interface
+	//~ End UObject Interface
 
-	// Begin UFactory Interface
+	//~ Begin UFactory Interface
 	virtual bool ConfigureProperties() override;
 	virtual UObject* FactoryCreateNew(UClass* Class,UObject* InParent,FName Name,EObjectFlags Flags,UObject* Context,FFeedbackContext* Warn) override;
 	virtual bool ShouldShowInNewMenu() const override
@@ -36,14 +36,14 @@ class UNREALED_API UTrueTypeFontFactory : public UTextureFactory, public FReimpo
 		// Don't show this factory in the content browser menu; it's invoked manually when changing the UFont cache type to "Offline" 
 		return false;
 	}
-	// Begin UFactory Interface	
+	//~ Begin UFactory Interface	
 
-	// Begin FReimportHandler interface
+	//~ Begin FReimportHandler Interface
 	virtual bool CanReimport( UObject* Obj, TArray<FString>& OutFilenames ) override;
 	virtual void SetReimportPaths( UObject* Obj, const TArray<FString>& NewReimportPaths ) override;
 	virtual EReimportResult::Type Reimport( UObject* Obj ) override;
 	virtual int32 GetPriority() const override;
-	// End FReimportHandler interface
+	//~ End FReimportHandler Interface
 
 	/** Creates the import options structure for this font */
 	void SetupFontImportOptions();
diff --git a/Engine/Source/Editor/UnrealEd/Classes/MaterialEditor/MaterialEditorInstanceConstant.h b/Engine/Source/Editor/UnrealEd/Classes/MaterialEditor/MaterialEditorInstanceConstant.h
--- a/Engine/Source/Editor/UnrealEd/Classes/MaterialEditor/MaterialEditorInstanceConstant.h
+++ b/Engine/Source/Editor/UnrealEd/Classes/MaterialEditor/MaterialEditorInstanceConstant.h
@@ -247,12 +247,12 @@ class UNREALED_API UMaterialEditorInstanceConstant : public UObject
 	UPROPERTY(EditAnywhere, Category=MaterialEditorInstanceConstant)
 	uint32 bUseOldStyleMICEditorGroups:1;
 
-	// Begin UObject interface.
+	//~ Begin UObject Interface.
 	virtual void PostEditChangeProperty(FPropertyChangedEvent& PropertyChangedEvent) override;
 #if WITH_EDITOR
 	virtual void PostEditUndo() override;
 #endif
-	// End UObject interface.
+	//~ End UObject Interface.
 
 	/** Regenerates the parameter arrays. */
 	void RegenerateArrays();
diff --git a/Engine/Source/Editor/UnrealEd/Classes/MaterialEditor/PreviewMaterial.h b/Engine/Source/Editor/UnrealEd/Classes/MaterialEditor/PreviewMaterial.h
--- a/Engine/Source/Editor/UnrealEd/Classes/MaterialEditor/PreviewMaterial.h
+++ b/Engine/Source/Editor/UnrealEd/Classes/MaterialEditor/PreviewMaterial.h
@@ -12,9 +12,9 @@ class UNREALED_API UPreviewMaterial : public UMaterial
 	GENERATED_UCLASS_BODY()
 
 
-	// Begin UMaterial interface.
+	//~ Begin UMaterial Interface.
 	virtual FMaterialResource* AllocateResource() override;
 	virtual bool IsAsset()  const override  { return false; }
-	// End UMaterial interface.
+	//~ End UMaterial Interface.
 };
 
diff --git a/Engine/Source/Editor/UnrealEd/Classes/MaterialGraph/MaterialGraphNode.h b/Engine/Source/Editor/UnrealEd/Classes/MaterialGraph/MaterialGraphNode.h
--- a/Engine/Source/Editor/UnrealEd/Classes/MaterialGraph/MaterialGraphNode.h
+++ b/Engine/Source/Editor/UnrealEd/Classes/MaterialGraph/MaterialGraphNode.h
@@ -43,13 +43,13 @@ public:
 	/** Get the Material value type of an output pin */
 	uint32 GetOutputType(const UEdGraphPin* OutputPin);
 
-	// UObject interface
+	//~ Begin UObject Interface
 	virtual void PostEditChangeProperty(FPropertyChangedEvent& PropertyChangedEvent) override;
 	virtual void PostEditImport() override;
 	virtual void PostDuplicate(bool bDuplicateForPIE) override;
-	// End of UObject interface
+	//~ End UObject Interface
 
-	// Begin UEdGraphNode interface.
+	//~ Begin UEdGraphNode Interface.
 	virtual bool CanPasteHere(const UEdGraph* TargetGraph) const override;
 	virtual FText GetNodeTitle(ENodeTitleType::Type TitleType) const override;
 	virtual FLinearColor GetNodeTitleColor() const override;
@@ -63,14 +63,14 @@ public:
 	virtual void OnCommentBubbleToggled( bool bInCommentBubbleVisible ) override;
 	virtual void GetPinHoverText(const UEdGraphPin& Pin, FString& HoverTextOut) const override;
 	virtual FString GetDocumentationExcerptName() const override;
-	// End UEdGraphNode interface.
+	//~ End UEdGraphNode Interface.
 
-	// UMaterialGraphNode_Base interface
+	//~ Begin UMaterialGraphNode_Base Interface
 	virtual void CreateInputPins() override;
 	virtual void CreateOutputPins() override;
 	virtual UNREALED_API int32 GetInputIndex(const UEdGraphPin* InputPin) const override;
 	virtual uint32 GetInputType(const UEdGraphPin* InputPin) const override;
-	// End of UMaterialGraphNode_Base interface
+	//~ End UMaterialGraphNode_Base Interface
 
 private:
 	/** Make sure the MaterialExpression is owned by the Material */
diff --git a/Engine/Source/Editor/UnrealEd/Classes/MaterialGraph/MaterialGraphNode_Base.h b/Engine/Source/Editor/UnrealEd/Classes/MaterialGraph/MaterialGraphNode_Base.h
--- a/Engine/Source/Editor/UnrealEd/Classes/MaterialGraph/MaterialGraphNode_Base.h
+++ b/Engine/Source/Editor/UnrealEd/Classes/MaterialGraph/MaterialGraphNode_Base.h
@@ -40,13 +40,13 @@ class UMaterialGraphNode_Base : public UEdGraphNode
 	 */
 	void InsertNewNode(UEdGraphPin* FromPin, UEdGraphPin* NewLinkPin, TSet<UEdGraphNode*>& OutNodeList);
 
-	// Begin UEdGraphNode interface.
+	//~ Begin UEdGraphNode Interface.
 	virtual void AllocateDefaultPins() override;
 	virtual void ReconstructNode() override;
 	virtual void AutowireNewNode(UEdGraphPin* FromPin) override;
 	virtual bool CanCreateUnderSpecifiedSchema(const UEdGraphSchema* Schema) const override;
 	virtual FString GetDocumentationLink() const override;
-	// End UEdGraphNode interface.
+	//~ End UEdGraphNode Interface.
 
 protected:
 	void ModifyAndCopyPersistentPinData(UEdGraphPin& TargetPin, const UEdGraphPin& SourcePin) const;
diff --git a/Engine/Source/Editor/UnrealEd/Classes/MaterialGraph/MaterialGraphNode_Comment.h b/Engine/Source/Editor/UnrealEd/Classes/MaterialGraph/MaterialGraphNode_Comment.h
--- a/Engine/Source/Editor/UnrealEd/Classes/MaterialGraph/MaterialGraphNode_Comment.h
+++ b/Engine/Source/Editor/UnrealEd/Classes/MaterialGraph/MaterialGraphNode_Comment.h
@@ -21,20 +21,20 @@ class UMaterialGraphNode_Comment : public UEdGraphNode_Comment
 	/** Fix up the node's owner after being copied */
 	UNREALED_API void PostCopyNode();
 
-	// UObject interface
+	//~ Begin UObject Interface
 	virtual void PostEditChangeProperty(FPropertyChangedEvent& PropertyChangedEvent) override;
 	virtual void PostEditImport() override;
 	virtual void PostDuplicate(bool bDuplicateForPIE) override;
-	// End of UObject interface
+	//~ End UObject Interface
 
-	// Begin UEdGraphNode interface.
+	//~ Begin UEdGraphNode Interface.
 	virtual void PrepareForCopying() override;
 	virtual void GetContextMenuActions(const FGraphNodeContextMenuBuilder& Context) const override;
 	virtual bool CanCreateUnderSpecifiedSchema(const UEdGraphSchema* Schema) const override;
 	virtual void PostPlacedNewNode() override;
 	virtual void OnRenameNode(const FString& NewName) override;
 	virtual void ResizeNode(const FVector2D& NewSize) override;
-	// End UEdGraphNode interface.
+	//~ End UEdGraphNode Interface.
 
 private:
 	/** Make sure the MaterialExpressionComment is owned by the Material */
diff --git a/Engine/Source/Editor/UnrealEd/Classes/MaterialGraph/MaterialGraphNode_Root.h b/Engine/Source/Editor/UnrealEd/Classes/MaterialGraph/MaterialGraphNode_Root.h
--- a/Engine/Source/Editor/UnrealEd/Classes/MaterialGraph/MaterialGraphNode_Root.h
+++ b/Engine/Source/Editor/UnrealEd/Classes/MaterialGraph/MaterialGraphNode_Root.h
@@ -13,19 +13,19 @@ class UMaterialGraphNode_Root : public UMaterialGraphNode_Base
 	UPROPERTY()
 	class UMaterial* Material;
 
-	// Begin UEdGraphNode interface.
+	//~ Begin UEdGraphNode Interface.
 	virtual FText GetNodeTitle(ENodeTitleType::Type TitleType) const override;
 	virtual FLinearColor GetNodeTitleColor() const override;
 	virtual FText GetTooltipText() const override;
 	virtual bool CanUserDeleteNode() const override { return false; }
 	virtual bool CanDuplicateNode() const override { return false; }
 	virtual void PostPlacedNewNode() override;
-	// End UEdGraphNode interface.
+	//~ End UEdGraphNode Interface.
 
-	// UMaterialGraphNode_Base interface
+	//~ Begin UMaterialGraphNode_Base Interface
 	virtual void CreateInputPins() override;
 	virtual bool IsRootNode() const override {return true;}
 	virtual int32 GetInputIndex(const UEdGraphPin* InputPin) const override;
 	virtual uint32 GetInputType(const UEdGraphPin* InputPin) const override;
-	// End of UMaterialGraphNode_Base interface
+	//~ End UMaterialGraphNode_Base Interface
 };
diff --git a/Engine/Source/Editor/UnrealEd/Classes/MaterialGraph/MaterialGraphSchema.h b/Engine/Source/Editor/UnrealEd/Classes/MaterialGraph/MaterialGraphSchema.h
--- a/Engine/Source/Editor/UnrealEd/Classes/MaterialGraph/MaterialGraphSchema.h
+++ b/Engine/Source/Editor/UnrealEd/Classes/MaterialGraph/MaterialGraphSchema.h
@@ -28,9 +28,9 @@ struct UNREALED_API FMaterialGraphSchemaAction_NewNode : public FEdGraphSchemaAc
 		, MaterialExpressionClass(NULL)
 	{}
 
-	// FEdGraphSchemaAction interface
+	//~ Begin FEdGraphSchemaAction Interface
 	virtual UEdGraphNode* PerformAction(class UEdGraph* ParentGraph, UEdGraphPin* FromPin, const FVector2D Location, bool bSelectNewNode = true) override;
-	// End of FEdGraphSchemaAction interface
+	//~ End FEdGraphSchemaAction Interface
 
 	/**
 	 * Sets the type of a Function input based on an EMaterialValueType value.
@@ -63,9 +63,9 @@ struct UNREALED_API FMaterialGraphSchemaAction_NewFunctionCall : public FEdGraph
 		: FEdGraphSchemaAction(InNodeCategory, InMenuDesc, InToolTip, InGrouping)
 	{}
 
-	// FEdGraphSchemaAction interface
+	//~ Begin FEdGraphSchemaAction Interface
 	virtual UEdGraphNode* PerformAction(class UEdGraph* ParentGraph, UEdGraphPin* FromPin, const FVector2D Location, bool bSelectNewNode = true) override;
-	// End of FEdGraphSchemaAction interface
+	//~ End FEdGraphSchemaAction Interface
 };
 
 /** Action to add a comment node to the graph */
@@ -86,9 +86,9 @@ struct UNREALED_API FMaterialGraphSchemaAction_NewComment : public FEdGraphSchem
 		: FEdGraphSchemaAction(InNodeCategory, InMenuDesc, InToolTip, InGrouping)
 	{}
 
-	// FEdGraphSchemaAction interface
+	//~ Begin FEdGraphSchemaAction Interface
 	virtual UEdGraphNode* PerformAction(class UEdGraph* ParentGraph, UEdGraphPin* FromPin, const FVector2D Location, bool bSelectNewNode = true) override;
-	// End of FEdGraphSchemaAction interface
+	//~ End FEdGraphSchemaAction Interface
 };
 
 /** Action to paste clipboard contents into the graph */
@@ -109,9 +109,9 @@ struct UNREALED_API FMaterialGraphSchemaAction_Paste : public FEdGraphSchemaActi
 		: FEdGraphSchemaAction(InNodeCategory, InMenuDesc, InToolTip, InGrouping)
 	{}
 
-	// FEdGraphSchemaAction interface
+	//~ Begin FEdGraphSchemaAction Interface
 	virtual UEdGraphNode* PerformAction(class UEdGraph* ParentGraph, UEdGraphPin* FromPin, const FVector2D Location, bool bSelectNewNode = true) override;
-	// End of FEdGraphSchemaAction interface
+	//~ End FEdGraphSchemaAction Interface
 };
 
 UCLASS(MinimalAPI)
@@ -190,7 +190,7 @@ class UMaterialGraphSchema : public UEdGraphSchema
 	/** Gets the type of this pin (must be part of a UMaterialGraphNode_Base) */
 	UNREALED_API static uint32 GetMaterialValueType(const UEdGraphPin* MaterialPin);
 
-	// Begin UEdGraphSchema interface
+	//~ Begin UEdGraphSchema Interface
 	virtual void GetGraphContextActions(FGraphContextMenuBuilder& ContextMenuBuilder) const override;
 	virtual void GetContextMenuActions(const UEdGraph* CurrentGraph, const UEdGraphNode* InGraphNode, const UEdGraphPin* InGraphPin, class FMenuBuilder* MenuBuilder, bool bIsDebugging) const override;
 	virtual const FPinConnectionResponse CanCreateConnection(const UEdGraphPin* A, const UEdGraphPin* B) const override;
@@ -206,7 +206,7 @@ class UMaterialGraphSchema : public UEdGraphSchema
 	virtual bool IsCacheVisualizationOutOfDate(int32 InVisualizationCacheID) const override;
 	virtual int32 GetCurrentVisualizationCacheID() const override;
 	virtual void ForceVisualizationCacheClear() const override;
-	// End UEdGraphSchema interface
+	//~ End UEdGraphSchema Interface
 
 private:
 	/** Adds actions for all Material Functions */
diff --git a/Engine/Source/Editor/UnrealEd/Classes/Niagara/NiagaraGraph.h b/Engine/Source/Editor/UnrealEd/Classes/Niagara/NiagaraGraph.h
--- a/Engine/Source/Editor/UnrealEd/Classes/Niagara/NiagaraGraph.h
+++ b/Engine/Source/Editor/UnrealEd/Classes/Niagara/NiagaraGraph.h
@@ -9,9 +9,9 @@ class UNiagaraGraph : public UEdGraph
 {
 	GENERATED_UCLASS_BODY()
 
-	// Begin UObject interface
+	//~ Begin UObject Interface
 	virtual void PostEditChangeProperty(struct FPropertyChangedEvent& PropertyChangedEvent)override;
-	// End UObjet interface
+	//~ End UObjet Interface
 	
 	/** Get the source that owns this graph */
 	class UNiagaraScriptSource* GetSource() const;
diff --git a/Engine/Source/Editor/UnrealEd/Classes/Niagara/NiagaraNodeConstant.h b/Engine/Source/Editor/UnrealEd/Classes/Niagara/NiagaraNodeConstant.h
--- a/Engine/Source/Editor/UnrealEd/Classes/Niagara/NiagaraNodeConstant.h
+++ b/Engine/Source/Editor/UnrealEd/Classes/Niagara/NiagaraNodeConstant.h
@@ -27,18 +27,18 @@ public:
 
 	bool IsExposedToEffectEditor()	{ return bExposeToEffectEditor; }
 
-	// Begin UObject interface
+	//~ Begin UObject Interface
 	virtual void PostEditChangeProperty(struct FPropertyChangedEvent& PropertyChangedEvent) override;
-	// End UObject interface
+	//~ End UObject Interface
 	
-	// Begin EdGraphNode interface
+	//~ Begin EdGraphNode Interface
 	virtual void AllocateDefaultPins() override;
 	virtual FText GetNodeTitle(ENodeTitleType::Type TitleType) const override;
 	virtual FLinearColor GetNodeTitleColor() const override;
-	// End EdGraphNode interface
+	//~ End EdGraphNode Interface
 
-	// Begin UNiagaraNode interface
+	//~ Begin UNiagaraNode Interface
 	UNREALED_API virtual void Compile(class INiagaraCompiler* Compiler, TArray<FNiagaraNodeResult>& Outputs) override;
-	// End UNiagaraNode interface
+	//~ End UNiagaraNode Interface
 };
 
diff --git a/Engine/Source/Editor/UnrealEd/Classes/Niagara/NiagaraNodeFunctionCall.h b/Engine/Source/Editor/UnrealEd/Classes/Niagara/NiagaraNodeFunctionCall.h
--- a/Engine/Source/Editor/UnrealEd/Classes/Niagara/NiagaraNodeFunctionCall.h
+++ b/Engine/Source/Editor/UnrealEd/Classes/Niagara/NiagaraNodeFunctionCall.h
@@ -20,15 +20,15 @@ public:
 	virtual void PostEditChangeProperty(struct FPropertyChangedEvent& PropertyChangedEvent)override;
 	//End UObject interface
 
-	// Begin UNiagaraNode interface
+	//~ Begin UNiagaraNode Interface
 	virtual void Compile(class INiagaraCompiler* Compiler, TArray<FNiagaraNodeResult>& Outputs)override;
 	//End UNiagaraNode interface
 
-	// Begin EdGraphNode interface
+	//~ Begin EdGraphNode Interface
 	virtual void AllocateDefaultPins() override;
 	virtual FText GetNodeTitle(ENodeTitleType::Type TitleType) const override;
 	virtual FText GetTooltipText() const override;
 	virtual FLinearColor GetNodeTitleColor() const override;
-	// End EdGraphNode interface
+	//~ End EdGraphNode Interface
 };
 
diff --git a/Engine/Source/Editor/UnrealEd/Classes/Niagara/NiagaraNodeGetAttr.h b/Engine/Source/Editor/UnrealEd/Classes/Niagara/NiagaraNodeGetAttr.h
--- a/Engine/Source/Editor/UnrealEd/Classes/Niagara/NiagaraNodeGetAttr.h
+++ b/Engine/Source/Editor/UnrealEd/Classes/Niagara/NiagaraNodeGetAttr.h
@@ -14,14 +14,14 @@ public:
 	UPROPERTY()
 	FName	AttrName;
 
-	// Begin EdGraphNode interface
+	//~ Begin EdGraphNode Interface
 	virtual void AllocateDefaultPins() override;
 	virtual FText GetNodeTitle(ENodeTitleType::Type TitleType) const override;
 	virtual FLinearColor GetNodeTitleColor() const override;
-	// End EdGraphNode interface
+	//~ End EdGraphNode Interface
 
-	// Begin UNiagaraNode interface
+	//~ Begin UNiagaraNode Interface
 	UNREALED_API virtual void Compile(class INiagaraCompiler* Compiler, TArray<FNiagaraNodeResult>& Output) override;
-	// End UNiagaraNode interface
+	//~ End UNiagaraNode Interface
 };
 
diff --git a/Engine/Source/Editor/UnrealEd/Classes/Niagara/NiagaraNodeInput.h b/Engine/Source/Editor/UnrealEd/Classes/Niagara/NiagaraNodeInput.h
--- a/Engine/Source/Editor/UnrealEd/Classes/Niagara/NiagaraNodeInput.h
+++ b/Engine/Source/Editor/UnrealEd/Classes/Niagara/NiagaraNodeInput.h
@@ -35,16 +35,16 @@ public:
 	UPROPERTY(EditAnywhere, Category = Constant, meta = (editcondition = "bCanBeExposed"))
 	uint32 bExposeWhenConstant:1;
 
-	// Begin UObject interface
+	//~ Begin UObject Interface
 	virtual void PostEditChangeProperty(struct FPropertyChangedEvent& PropertyChangedEvent) override;
-	// End UObject interface
+	//~ End UObject Interface
 
-	// Begin EdGraphNode interface
+	//~ Begin EdGraphNode Interface
 	virtual void AllocateDefaultPins() override;
 	virtual FText GetNodeTitle(ENodeTitleType::Type TitleType) const override;
 	virtual FLinearColor GetNodeTitleColor() const override;
 	virtual void AutowireNewNode(UEdGraphPin* FromPin)override;
-	// End EdGraphNode interface
+	//~ End EdGraphNode Interface
 
 	virtual void Compile(class INiagaraCompiler* Compiler, TArray<FNiagaraNodeResult>& Outputs)override;
 	
diff --git a/Engine/Source/Editor/UnrealEd/Classes/Niagara/NiagaraNodeOp.h b/Engine/Source/Editor/UnrealEd/Classes/Niagara/NiagaraNodeOp.h
--- a/Engine/Source/Editor/UnrealEd/Classes/Niagara/NiagaraNodeOp.h
+++ b/Engine/Source/Editor/UnrealEd/Classes/Niagara/NiagaraNodeOp.h
@@ -14,16 +14,16 @@ public:
 	UPROPERTY()
 	FName OpName;
 
-	// Begin EdGraphNode interface
+	//~ Begin EdGraphNode Interface
 	virtual void AllocateDefaultPins() override;
 	virtual FText GetNodeTitle(ENodeTitleType::Type TitleType) const override;
 	virtual FText GetTooltipText() const override;
 	virtual FLinearColor GetNodeTitleColor() const override;
-	// End EdGraphNode interface
+	//~ End EdGraphNode Interface
 
-	// Begin UNiagaraNode interface
+	//~ Begin UNiagaraNode Interface
 	UNREALED_API virtual void Compile(class INiagaraCompiler* Compiler, TArray<FNiagaraNodeResult>& Outputs) override;
-	// End UNiagaraNode interface
+	//~ End UNiagaraNode Interface
 };
 
 
@@ -38,23 +38,23 @@ public:
 
 		void Construct(const FArguments& InArgs, class UNiagaraNode* InNode);
 
-	// SGraphNode interface
+	//~ Begin SGraphNode Interface
 	virtual void CreatePinWidgets() override;
-	// End of SGraphNode interface
+	//~ End SGraphNode Interface
 
-	// SNodePanel::SNode interface
+	//~ Begin SNodePanel::SNode Interface
 	virtual void MoveTo(const FVector2D& NewPosition, FNodeSet& NodeFilter) override;
-	// End of SNodePanel::SNode interface
+	//~ End SNodePanel::SNode Interface
 
 	UNiagaraNode* GetNiagaraGraphNode() const { return NiagaraNode; }
 
 protected:
-	// SGraphNode interface
+	//~ Begin SGraphNode Interface
 	virtual void AddPin(const TSharedRef<SGraphPin>& PinToAdd) override;
 	virtual void CreateBelowPinControls(TSharedPtr<SVerticalBox> MainBox) override;
 	virtual void SetDefaultTitleAreaWidget(TSharedRef<SOverlay> DefaultTitleAreaWidget) override;
 	virtual TSharedRef<SWidget> CreateNodeContentArea() override;
-	// End of SGraphNode interface
+	//~ End SGraphNode Interface
 
 	/** Creates a preview viewport if necessary */
 	//TSharedRef<SWidget> CreatePreviewWidget();
diff --git a/Engine/Source/Editor/UnrealEd/Classes/Niagara/NiagaraNodeOutput.h b/Engine/Source/Editor/UnrealEd/Classes/Niagara/NiagaraNodeOutput.h
--- a/Engine/Source/Editor/UnrealEd/Classes/Niagara/NiagaraNodeOutput.h
+++ b/Engine/Source/Editor/UnrealEd/Classes/Niagara/NiagaraNodeOutput.h
@@ -13,16 +13,16 @@ class UNiagaraNodeOutput : public UNiagaraNode
 
 public:
 
-	// Begin UObject interface
+	//~ Begin UObject Interface
 	virtual void PostEditChangeProperty(struct FPropertyChangedEvent& PropertyChangedEvent) override;
-	// End UObject interface
+	//~ End UObject Interface
 
-	// Begin EdGraphNode interface
+	//~ Begin EdGraphNode Interface
 	virtual void AllocateDefaultPins() override;
 	virtual bool CanUserDeleteNode() const override { return false; }
 	virtual FText GetNodeTitle(ENodeTitleType::Type TitleType) const override;
 	virtual FLinearColor GetNodeTitleColor() const override;
-	// End EdGraphNode interface
+	//~ End EdGraphNode Interface
 
 	virtual void Compile(class INiagaraCompiler* Compiler, TArray<FNiagaraNodeResult>& OutputExpressions)override;
 };
diff --git a/Engine/Source/Editor/UnrealEd/Classes/SoundCueGraph/SoundCueGraphSchema.h b/Engine/Source/Editor/UnrealEd/Classes/SoundCueGraph/SoundCueGraphSchema.h
--- a/Engine/Source/Editor/UnrealEd/Classes/SoundCueGraph/SoundCueGraphSchema.h
+++ b/Engine/Source/Editor/UnrealEd/Classes/SoundCueGraph/SoundCueGraphSchema.h
@@ -28,9 +28,9 @@ struct UNREALED_API FSoundCueGraphSchemaAction_NewNode : public FEdGraphSchemaAc
 		, SoundNodeClass(NULL)
 	{}
 
-	// FEdGraphSchemaAction interface
+	//~ Begin FEdGraphSchemaAction Interface
 	virtual UEdGraphNode* PerformAction(class UEdGraph* ParentGraph, UEdGraphPin* FromPin, const FVector2D Location, bool bSelectNewNode = true) override;
-	// End of FEdGraphSchemaAction interface
+	//~ End FEdGraphSchemaAction Interface
 
 private:
 	/** Connects new node to output of selected nodes */
@@ -51,9 +51,9 @@ struct UNREALED_API FSoundCueGraphSchemaAction_NewFromSelected : public FSoundCu
 		: FSoundCueGraphSchemaAction_NewNode(InNodeCategory, InMenuDesc, InToolTip, InGrouping) 
 	{}
 
-	// FEdGraphSchemaAction interface
+	//~ Begin FEdGraphSchemaAction Interface
 	virtual UEdGraphNode* PerformAction(class UEdGraph* ParentGraph, UEdGraphPin* FromPin, const FVector2D Location, bool bSelectNewNode = true) override;
-	// End of FEdGraphSchemaAction interface
+	//~ End FEdGraphSchemaAction Interface
 };
 
 /** Action to create new comment */
@@ -70,9 +70,9 @@ struct UNREALED_API FSoundCueGraphSchemaAction_NewComment : public FEdGraphSchem
 		: FEdGraphSchemaAction(InNodeCategory, InMenuDesc, InToolTip, InGrouping)
 	{}
 
-	// FEdGraphSchemaAction interface
+	//~ Begin FEdGraphSchemaAction Interface
 	virtual UEdGraphNode* PerformAction(class UEdGraph* ParentGraph, UEdGraphPin* FromPin, const FVector2D Location, bool bSelectNewNode = true) override;
-	// End of FEdGraphSchemaAction interface
+	//~ End FEdGraphSchemaAction Interface
 };
 
 /** Action to paste clipboard contents into the graph */
@@ -89,9 +89,9 @@ struct UNREALED_API FSoundCueGraphSchemaAction_Paste : public FEdGraphSchemaActi
 		: FEdGraphSchemaAction(InNodeCategory, InMenuDesc, InToolTip, InGrouping)
 	{}
 
-	// FEdGraphSchemaAction interface
+	//~ Begin FEdGraphSchemaAction Interface
 	virtual UEdGraphNode* PerformAction(class UEdGraph* ParentGraph, UEdGraphPin* FromPin, const FVector2D Location, bool bSelectNewNode = true) override;
-	// End of FEdGraphSchemaAction interface
+	//~ End FEdGraphSchemaAction Interface
 };
 
 UCLASS(MinimalAPI)
@@ -108,7 +108,7 @@ class USoundCueGraphSchema : public UEdGraphSchema
 	/** Attempts to connect the output of multiple nodes to the inputs of a single one */
 	void TryConnectNodes(const TArray<USoundNode*>& OutputNodes, USoundNode* InputNode) const;
 
-	// Begin EdGraphSchema interface
+	//~ Begin EdGraphSchema Interface
 	virtual void GetGraphContextActions(FGraphContextMenuBuilder& ContextMenuBuilder) const override;
 	virtual void GetContextMenuActions(const UEdGraph* CurrentGraph, const UEdGraphNode* InGraphNode, const UEdGraphPin* InGraphPin, class FMenuBuilder* MenuBuilder, bool bIsDebugging) const override;
 	virtual void CreateDefaultNodesForGraph(UEdGraph& Graph) const override;
@@ -121,7 +121,7 @@ class USoundCueGraphSchema : public UEdGraphSchema
 	virtual void DroppedAssetsOnGraph(const TArray<class FAssetData>& Assets, const FVector2D& GraphPosition, UEdGraph* Graph) const override;
 	virtual int32 GetNodeSelectionCount(const UEdGraph* Graph) const override;
 	virtual TSharedPtr<FEdGraphSchemaAction> GetCreateCommentAction() const override;
-	// End EdGraphSchema interface
+	//~ End EdGraphSchema Interface
 
 private:
 	/** Adds actions for creating every type of SoundNode */
diff --git a/Engine/Source/Editor/UnrealEd/Classes/TexAligner/TexAligner.h b/Engine/Source/Editor/UnrealEd/Classes/TexAligner/TexAligner.h
--- a/Engine/Source/Editor/UnrealEd/Classes/TexAligner/TexAligner.h
+++ b/Engine/Source/Editor/UnrealEd/Classes/TexAligner/TexAligner.h
@@ -1,10 +1,10 @@
 // Copyright 1998-2015 Epic Games, Inc. All Rights Reserved.
 
-//=============================================================================
+//~=============================================================================
 // TexAligner
 // Base class for all texture aligners.
 //
-//=============================================================================
+//~=============================================================================
 
 #pragma once
 
@@ -58,9 +58,9 @@ class UTexAligner : public UObject
 	FString Desc;
 
 
-	// Begin UObject Interface
+	//~ Begin UObject Interface
 	virtual void PostInitProperties() override;
-	// End UObject Interface
+	//~ End UObject Interface
 
 	// @todo document
 	UNREALED_API void Align( UWorld* InWorld, ETexAlign InTexAlignType );
diff --git a/Engine/Source/Editor/UnrealEd/Classes/TexAligner/TexAlignerBox.h b/Engine/Source/Editor/UnrealEd/Classes/TexAligner/TexAlignerBox.h
--- a/Engine/Source/Editor/UnrealEd/Classes/TexAligner/TexAlignerBox.h
+++ b/Engine/Source/Editor/UnrealEd/Classes/TexAligner/TexAlignerBox.h
@@ -1,10 +1,10 @@
 // Copyright 1998-2015 Epic Games, Inc. All Rights Reserved.
 
-//=============================================================================
+//~=============================================================================
 // TexAlignerBox
 // Aligns to the best U and V axis according to the polys normal.
 //
-//=============================================================================
+//~=============================================================================
 
 #pragma once
 #include "TexAlignerBox.generated.h"
@@ -15,12 +15,12 @@ class UTexAlignerBox : public UTexAligner
 	GENERATED_UCLASS_BODY()
 
 
-	// Begin UObject Interface
+	//~ Begin UObject Interface
 	virtual void PostInitProperties() override;
-	// End UObject Interface
+	//~ End UObject Interface
 
-	// Begin UTexAligner Interface
+	//~ Begin UTexAligner Interface
 	virtual void AlignSurf( ETexAlign InTexAlignType, UModel* InModel, FBspSurfIdx* InSurfIdx, FPoly* InPoly, FVector* InNormal ) override;
-	// End UTexAligner Interface
+	//~ End UTexAligner Interface
 };
 
diff --git a/Engine/Source/Editor/UnrealEd/Classes/TexAligner/TexAlignerDefault.h b/Engine/Source/Editor/UnrealEd/Classes/TexAligner/TexAlignerDefault.h
--- a/Engine/Source/Editor/UnrealEd/Classes/TexAligner/TexAlignerDefault.h
+++ b/Engine/Source/Editor/UnrealEd/Classes/TexAligner/TexAlignerDefault.h
@@ -1,10 +1,10 @@
 // Copyright 1998-2015 Epic Games, Inc. All Rights Reserved.
 
-//=============================================================================
+//~=============================================================================
 // TexAlignerDefault
 // Aligns to a default setting.
 //
-//=============================================================================
+//~=============================================================================
 
 #pragma once
 #include "TexAlignerDefault.generated.h"
@@ -15,12 +15,12 @@ class UTexAlignerDefault : public UTexAligner
 	GENERATED_UCLASS_BODY()
 
 
-	// Begin UObject Interface
+	//~ Begin UObject Interface
 	virtual void PostInitProperties() override;
-	// End UObject Interface
+	//~ End UObject Interface
 
-	// Begin UTexAligner Interface
+	//~ Begin UTexAligner Interface
 	virtual void AlignSurf( ETexAlign InTexAlignType, UModel* InModel, FBspSurfIdx* InSurfIdx, FPoly* InPoly, FVector* InNormal ) override;
-	// End UTexAligner Interface
+	//~ End UTexAligner Interface
 };
 
diff --git a/Engine/Source/Editor/UnrealEd/Public/AssetThumbnail.h b/Engine/Source/Editor/UnrealEd/Public/AssetThumbnail.h
--- a/Engine/Source/Editor/UnrealEd/Public/AssetThumbnail.h
+++ b/Engine/Source/Editor/UnrealEd/Public/AssetThumbnail.h
@@ -144,7 +144,7 @@ public:
 	/** Destructor to free all remaining resources */
 	UNREALED_API ~FAssetThumbnailPool();
 
-	// Begin FTickableObject interface
+	//~ Begin FTickableObject Interface
 	UNREALED_API virtual TStatId GetStatId() const override;
 
 	/** Checks if any new thumbnails are queued */
@@ -153,7 +153,7 @@ public:
 	/** Ticks the pool, rendering new thumbnails as needed */
 	UNREALED_API virtual void Tick( float DeltaTime ) override;
 
-	// End FTickableObject interface
+	//~ End FTickableObject Interface
 
 	/**
 	 * Releases all rendering resources held by the pool
diff --git a/Engine/Source/Editor/UnrealEd/Public/DebugToolExec.h b/Engine/Source/Editor/UnrealEd/Public/DebugToolExec.h
--- a/Engine/Source/Editor/UnrealEd/Public/DebugToolExec.h
+++ b/Engine/Source/Editor/UnrealEd/Public/DebugToolExec.h
@@ -24,9 +24,9 @@ protected:
 
 public:
 	
-	// Begin Exec Interface
+	//~ Begin Exec Interface
 	UNREALED_API virtual bool Exec( UWorld* InWorld, const TCHAR* Cmd, FOutputDevice& Ar ) override;
-	// End Exec Interface
+	//~ End Exec Interface
 
 };
 
diff --git a/Engine/Source/Editor/UnrealEd/Public/EdGraphNode_Comment.h b/Engine/Source/Editor/UnrealEd/Public/EdGraphNode_Comment.h
--- a/Engine/Source/Editor/UnrealEd/Public/EdGraphNode_Comment.h
+++ b/Engine/Source/Editor/UnrealEd/Public/EdGraphNode_Comment.h
@@ -47,11 +47,11 @@ public:
 
 public:
 
-	// Begin UObject Interface
+	//~ Begin UObject Interface
 	UNREALED_API static void AddReferencedObjects(UObject* InThis, FReferenceCollector& Collector);
-	// End UObject Interface
+	//~ End UObject Interface
 
-	// Begin UEdGraphNode interface
+	//~ Begin UEdGraphNode Interface
 	virtual void AllocateDefaultPins() override {}
 	UNREALED_API virtual FText GetTooltipText() const override;
 	UNREALED_API virtual FLinearColor GetNodeCommentColor() const override;
@@ -66,7 +66,7 @@ public:
 	UNREALED_API virtual FString GetDocumentationLink() const override;
 	UNREALED_API virtual FString GetDocumentationExcerptName() const override;
 	UNREALED_API virtual FName GetPaletteIcon(FLinearColor& OutColor) const override;
-	// End UEdGraphNode interface
+	//~ End UEdGraphNode Interface
 
 	/** Add a node that will be dragged when this comment is dragged */
 	UNREALED_API void	AddNodeUnderComment(UObject* Object);
diff --git a/Engine/Source/Editor/UnrealEd/Public/SCurveEditor.h b/Engine/Source/Editor/UnrealEd/Public/SCurveEditor.h
--- a/Engine/Source/Editor/UnrealEd/Public/SCurveEditor.h
+++ b/Engine/Source/Editor/UnrealEd/Public/SCurveEditor.h
@@ -419,7 +419,7 @@ private:
 	/** Function pointer to execute callback function when user select 'Create external curve'*/
 	FSimpleDelegate OnCreateAsset;
 
-	// SWidget interface
+	//~ Begin SWidget Interface
 	UNREALED_API virtual FVector2D ComputeDesiredSize(float) const override;
 
 	/** Paint a curve */
@@ -440,7 +440,7 @@ private:
 	void PaintMarquee(const FGeometry& AllottedGeometry, const FSlateRect& MyClippingRect, FSlateWindowElementList& OutDrawElements, int32 LayerId) const;
 
 protected:
-	// SWidget interface
+	//~ Begin SWidget Interface
 	UNREALED_API virtual FReply OnMouseButtonDown( const FGeometry& MyGeometry, const FPointerEvent& MouseEvent ) override;
 	UNREALED_API virtual FReply OnMouseButtonUp( const FGeometry& MyGeometry, const FPointerEvent& MouseEvent ) override;
 	UNREALED_API virtual FReply OnMouseMove( const FGeometry& MyGeometry, const FPointerEvent& MouseEvent ) override;
@@ -517,7 +517,7 @@ private:
 	/** Perform redo*/
 	void	RedoAction();
 
-	// Begin FEditorUndoClient Interface
+	//~ Begin FEditorUndoClient Interface
 	UNREALED_API virtual void PostUndo(bool bSuccess) override;
 	UNREALED_API virtual void PostRedo(bool bSuccess) override { PostUndo(bSuccess); }
 	// End of FEditorUndoClient
@@ -553,7 +553,7 @@ protected:
 	/** Get Time Step for vertical line drawing **/
 	UNREALED_API virtual float GetTimeStep(FTrackScaleInfo &ScaleInfo) const;
 	
-	// SWidget interface
+	//~ Begin SWidget Interface
 	UNREALED_API virtual int32 OnPaint( const FPaintArgs& Args, const FGeometry& AllottedGeometry, const FSlateRect& MyClippingRect, 
 		FSlateWindowElementList& OutDrawElements, int32 LayerId, const FWidgetStyle& InWidgetStyle, bool bParentEnabled ) const override;
 
diff --git a/Engine/Source/Runtime/AIModule/Classes/AIController.h b/Engine/Source/Runtime/AIModule/Classes/AIController.h
--- a/Engine/Source/Runtime/AIModule/Classes/AIController.h
+++ b/Engine/Source/Runtime/AIModule/Classes/AIController.h
@@ -66,7 +66,7 @@ struct FFocusKnowledge
 	TArray<FFocusItem> Priorities;
 };
 
-//=============================================================================
+//~=============================================================================
 /**
  * AIController is the base class of controllers for AI-controlled Pawns.
  * 
@@ -301,13 +301,13 @@ public:
 	 */
 	bool SuggestTossVelocity(FVector& OutTossVelocity, FVector Start, FVector End, float TossSpeed, bool bPreferHighArc, float CollisionRadius=0, bool bOnlyTraceUp=false);
 
-	// Begin AActor Interface
+	//~ Begin AActor Interface
 	virtual void Tick(float DeltaTime) override;
 	virtual void PostInitializeComponents() override;
 	virtual void PostRegisterAllComponents() override;
-	// End AActor Interface
+	//~ End AActor Interface
 
-	// Begin AController Interface
+	//~ Begin AController Interface
 	virtual void Possess(APawn* InPawn) override;
 	virtual void UnPossess() override;
 	virtual void DisplayDebug(UCanvas* Canvas, const FDebugDisplayInfo& DebugDisplay, float& YL, float& YPos) override;
@@ -326,7 +326,7 @@ public:
 	 * @return true if controller's pawn can see Other actor.
 	 */
 	virtual bool LineOfSightTo(const AActor* Other, FVector ViewPoint = FVector(ForceInit), bool bAlternateChecks = false) const override;
-	// End AController Interface
+	//~ End AController Interface
 
 	/** Notifies AIController of changes in given actors' perception */
 	virtual void ActorsPerceptionUpdated(const TArray<AActor*>& UpdatedActors);
diff --git a/Engine/Source/Runtime/AIModule/Classes/Blueprint/AIAsyncTaskBlueprintProxy.h b/Engine/Source/Runtime/AIModule/Classes/Blueprint/AIAsyncTaskBlueprintProxy.h
--- a/Engine/Source/Runtime/AIModule/Classes/Blueprint/AIAsyncTaskBlueprintProxy.h
+++ b/Engine/Source/Runtime/AIModule/Classes/Blueprint/AIAsyncTaskBlueprintProxy.h
@@ -28,9 +28,9 @@ public:
 
 	void OnNoPath();
 
-	// Begin UObject interface
+	//~ Begin UObject Interface
 	virtual void BeginDestroy() override;
-	// End UObject interface
+	//~ End UObject Interface
 
 	TWeakObjectPtr<AAIController> AIController;
 	FAIRequestID MoveRequestId;
diff --git a/Engine/Source/Runtime/AIModule/Classes/Navigation/PathFollowingComponent.h b/Engine/Source/Runtime/AIModule/Classes/Navigation/PathFollowingComponent.h
--- a/Engine/Source/Runtime/AIModule/Classes/Navigation/PathFollowingComponent.h
+++ b/Engine/Source/Runtime/AIModule/Classes/Navigation/PathFollowingComponent.h
@@ -123,9 +123,9 @@ class AIMODULE_API UPathFollowingComponent : public UActorComponent, public IAIR
 	/** delegate for move completion notify */
 	FMoveCompletedSignature OnMoveFinished;
 
-	// Begin UActorComponent Interface
+	//~ Begin UActorComponent Interface
 	virtual void TickComponent(float DeltaTime, enum ELevelTick TickType, FActorComponentTickFunction *ThisTickFunction) override;
-	// End UActorComponent Interface
+	//~ End UActorComponent Interface
 
 	/** initialize component to use */
 	virtual void Initialize();
diff --git a/Engine/Source/Runtime/AIModule/Classes/Perception/PawnSensingComponent.h b/Engine/Source/Runtime/AIModule/Classes/Perception/PawnSensingComponent.h
--- a/Engine/Source/Runtime/AIModule/Classes/Perception/PawnSensingComponent.h
+++ b/Engine/Source/Runtime/AIModule/Classes/Perception/PawnSensingComponent.h
@@ -120,9 +120,9 @@ public:
 	/** @Returns true if sensor can hear this noise. Only executed if the noise has been determined to be relevant (via IsNoiseRelevant) */
 	virtual bool CanHear(const FVector& NoiseLoc, float Loudness, bool bFailedLOS) const;
 
-	// Begin UActorComponent interface.
+	//~ Begin UActorComponent Interface.
 	virtual void InitializeComponent() override;
-	// End UActorComponent interface.
+	//~ End UActorComponent Interface.
 
 	/** Get position where hearing/seeing occurs (i.e. ear/eye position).  If we ever need different positions for hearing/seeing, we'll deal with that then! */
 	virtual FVector GetSensorLocation() const;
diff --git a/Engine/Source/Runtime/AutomationWorker/Private/AutomationWorkerModule.h b/Engine/Source/Runtime/AutomationWorker/Private/AutomationWorkerModule.h
--- a/Engine/Source/Runtime/AutomationWorker/Private/AutomationWorkerModule.h
+++ b/Engine/Source/Runtime/AutomationWorker/Private/AutomationWorkerModule.h
@@ -11,22 +11,22 @@ class FAutomationWorkerModule
 {
 public:
 
-	// Begin IModuleInterface interface
+	//~ Begin IModuleInterface Interface
 
 	virtual void StartupModule() override;
 	virtual void ShutdownModule() override;
 
 	virtual bool SupportsDynamicReloading() override;
 
-	// End IModuleInterface interface
+	//~ End IModuleInterface Interface
 
 public:
 
-	// Begin IAutomationWorkerModule interface
+	//~ Begin IAutomationWorkerModule Interface
 
 	virtual void Tick() override;
 
-	// End IAutomationWorkerModule interface
+	//~ End IAutomationWorkerModule Interface
 
 protected:
 
diff --git a/Engine/Source/Runtime/Core/Private/Internationalization/TextHistory.h b/Engine/Source/Runtime/Core/Private/Internationalization/TextHistory.h
--- a/Engine/Source/Runtime/Core/Private/Internationalization/TextHistory.h
+++ b/Engine/Source/Runtime/Core/Private/Internationalization/TextHistory.h
@@ -62,13 +62,13 @@ public:
 	FTextHistory_Base(TSharedPtr< FString, ESPMode::ThreadSafe > InSourceString);
 
 
-	// Begin FTextHistory interface
+	//~ Begin FTextHistory Interface
 	virtual FText ToText(bool bInAsSource) const override;
 	virtual void Serialize(FArchive& Ar) override;
 	virtual void SerializeForDisplayString(FArchive& Ar, FTextDisplayStringRef& InOutDisplayString) override;
 	virtual bool IsOutOfDate() override { return false; }
 	virtual TSharedPtr< FString, ESPMode::ThreadSafe > GetSourceString() const override;
-	// End FTextHistory interface
+	//~ End FTextHistory Interface
 
 private:
 	/** The source string for an FText */
@@ -82,11 +82,11 @@ public:
 	FTextHistory_NamedFormat() {};
 	FTextHistory_NamedFormat(const FText& InSourceText, const FFormatNamedArguments& InArguments);
 
-	// Begin FTextHistory interface
+	//~ Begin FTextHistory Interface
 	virtual FText ToText(bool bInAsSource) const override;
 	virtual void Serialize(FArchive& Ar) override;
 	virtual void GetSourceTextsFromFormatHistory(FText, TArray<FText>& OutSourceTexts) const override;
-	// End FTextHistory interface
+	//~ End FTextHistory Interface
 
 private:
 	/** The pattern used to format the text */
@@ -102,11 +102,11 @@ public:
 	FTextHistory_OrderedFormat() {};
 	FTextHistory_OrderedFormat(const FText& InSourceText, const FFormatOrderedArguments& InArguments);
 
-	// Begin FTextHistory interface
+	//~ Begin FTextHistory Interface
 	virtual FText ToText(bool bInAsSource) const override;
 	virtual void Serialize(FArchive& Ar) override;
 	virtual void GetSourceTextsFromFormatHistory(FText, TArray<FText>& OutSourceTexts) const override;
-	// End FTextHistory interface
+	//~ End FTextHistory Interface
 
 private:
 	/** The pattern used to format the text */
@@ -122,11 +122,11 @@ public:
 	FTextHistory_ArgumentDataFormat() {};
 	FTextHistory_ArgumentDataFormat(const FText& InSourceText, const TArray< struct FFormatArgumentData >& InArguments);
 
-	// Begin FTextHistory interface
+	//~ Begin FTextHistory Interface
 	virtual FText ToText(bool bInAsSource) const override;
 	virtual void Serialize(FArchive& Ar) override;
 	virtual void GetSourceTextsFromFormatHistory(FText, TArray<FText>& OutSourceTexts) const override;
-	// End FTextHistory interface
+	//~ End FTextHistory Interface
 
 private:
 	/** The pattern used to format the text */
@@ -144,9 +144,9 @@ public:
 
 	~FTextHistory_FormatNumber();
 
-	// Begin FTextHistory interface
+	//~ Begin FTextHistory Interface
 	virtual void Serialize(FArchive& Ar) override;
-	// End FTextHistory interface
+	//~ End FTextHistory Interface
 protected:
 	/** The source value to format from */
 	FFormatArgumentValue SourceValue;
@@ -163,10 +163,10 @@ public:
 	FTextHistory_AsNumber() {};
 	FTextHistory_AsNumber(const FFormatArgumentValue& InSourceValue, const FNumberFormattingOptions* const InFormatOptions, const FCulturePtr InTargetCulture);
 
-	// Begin FTextHistory interface
+	//~ Begin FTextHistory Interface
 	virtual FText ToText(bool bInAsSource) const override;
 	virtual void Serialize(FArchive& Ar) override;
-	// End FTextHistory interface
+	//~ End FTextHistory Interface
 };
 
 /**  Handles history for formatting using AsPercent */
@@ -176,10 +176,10 @@ public:
 	FTextHistory_AsPercent() {};
 	FTextHistory_AsPercent(const FFormatArgumentValue& InSourceValue, const FNumberFormattingOptions* const InFormatOptions, const FCulturePtr InTargetCulture);
 
-	// Begin FTextHistory interface
+	//~ Begin FTextHistory Interface
 	virtual FText ToText(bool bInAsSource) const override;
 	virtual void Serialize(FArchive& Ar) override;
-	// End FTextHistory interface
+	//~ End FTextHistory Interface
 };
 
 /**  Handles history for formatting using AsCurrency */
@@ -189,10 +189,10 @@ public:
 	FTextHistory_AsCurrency() {};
 	FTextHistory_AsCurrency(const FFormatArgumentValue& InSourceValue, const FString& CurrencyCode, const FNumberFormattingOptions* const InFormatOptions, const FCulturePtr InTargetCulture);
 
-	// Begin FTextHistory interface
+	//~ Begin FTextHistory Interface
 	virtual FText ToText(bool bInAsSource) const override;
 	virtual void Serialize(FArchive& Ar) override;
-	// End FTextHistory interface
+	//~ End FTextHistory Interface
 
 private:
 	/** The currency used to format the number. */
@@ -206,10 +206,10 @@ public:
 	FTextHistory_AsDate() {};
 	FTextHistory_AsDate(const FDateTime& InSourceDateTime, const EDateTimeStyle::Type InDateStyle, const FString& InTimeZone, const FCulturePtr InTargetCulture);
 
-	// Begin FTextHistory interface
+	//~ Begin FTextHistory Interface
 	virtual FText ToText(bool bInAsSource) const override;
 	virtual void Serialize(FArchive& Ar) override;
-	// End FTextHistory interface
+	//~ End FTextHistory Interface
 
 private:
 	/** The source date structure to format */
@@ -229,10 +229,10 @@ public:
 	FTextHistory_AsTime() {};
 	FTextHistory_AsTime(const FDateTime& InSourceDateTime, const EDateTimeStyle::Type InTimeStyle, const FString& InTimeZone, const FCulturePtr InTargetCulture);
 
-	// Begin FTextHistory interface
+	//~ Begin FTextHistory Interface
 	virtual FText ToText(bool bInAsSource) const override;
 	virtual void Serialize(FArchive& Ar) override;
-	// End FTextHistory interface
+	//~ End FTextHistory Interface
 
 private:
 	/** The source time structure to format */
@@ -252,10 +252,10 @@ public:
 	FTextHistory_AsDateTime() {};
 	FTextHistory_AsDateTime(const FDateTime& InSourceDateTime, const EDateTimeStyle::Type InDateStyle, const EDateTimeStyle::Type InTimeStyle, const FString& InTimeZone, const FCulturePtr InTargetCulture);
 
-	// Begin FTextHistory interface
+	//~ Begin FTextHistory Interface
 	virtual FText ToText(bool bInAsSource) const override;
 	virtual void Serialize(FArchive& Ar) override;
-	// End FTextHistory interfaces
+	//~ End FTextHistory Interfaces
 
 private:
 	/** The source date and time structure to format */
diff --git a/Engine/Source/Runtime/Core/Public/Android/AndroidMemory.h b/Engine/Source/Runtime/Core/Public/Android/AndroidMemory.h
--- a/Engine/Source/Runtime/Core/Public/Android/AndroidMemory.h
+++ b/Engine/Source/Runtime/Core/Public/Android/AndroidMemory.h
@@ -19,14 +19,14 @@ struct FPlatformMemoryStats : public FGenericPlatformMemoryStats
 **/
 struct CORE_API FAndroidPlatformMemory : public FGenericPlatformMemory
 {
-	// Begin FGenericPlatformMemory interface
+	//~ Begin FGenericPlatformMemory Interface
 	static void Init();
 	static FPlatformMemoryStats GetStats();
 	static const FPlatformMemoryConstants& GetConstants();
 	static FMalloc* BaseAllocator();
 	static void* BinnedAllocFromOS( SIZE_T Size );
 	static void BinnedFreeToOS( void* Ptr );
-	// End FGenericPlatformMemory interface
+	//~ End FGenericPlatformMemory Interface
 };
 
 typedef FAndroidPlatformMemory FPlatformMemory;
diff --git a/Engine/Source/Runtime/Core/Public/GenericPlatform/GenericPlatformFile.h b/Engine/Source/Runtime/Core/Public/GenericPlatform/GenericPlatformFile.h
--- a/Engine/Source/Runtime/Core/Public/GenericPlatform/GenericPlatformFile.h
+++ b/Engine/Source/Runtime/Core/Public/GenericPlatform/GenericPlatformFile.h
@@ -361,7 +361,7 @@ public:
 class CORE_API IPhysicalPlatformFile : public IPlatformFile
 {
 public:
-	// BEGIN IPlatformFile Interface
+	//~ Begin IPlatformFile Interface
 	virtual bool ShouldBeUsed(IPlatformFile* Inner, const TCHAR* CmdLine) const override
 	{
 		return true;
@@ -375,5 +375,5 @@ public:
 	{
 		return IPlatformFile::GetPhysicalTypeName();
 	}
-	// END IPlatformFile Interface
+	//~ End IPlatformFile Interface
 };
diff --git a/Engine/Source/Runtime/Core/Public/HTML5/HTML5PlatformMemory.h b/Engine/Source/Runtime/Core/Public/HTML5/HTML5PlatformMemory.h
--- a/Engine/Source/Runtime/Core/Public/HTML5/HTML5PlatformMemory.h
+++ b/Engine/Source/Runtime/Core/Public/HTML5/HTML5PlatformMemory.h
@@ -18,14 +18,14 @@ struct FPlatformMemoryStats : public FGenericPlatformMemoryStats
 **/
 struct CORE_API FHTML5PlatformMemory : public FGenericPlatformMemory
 {
-	// Begin FGenericPlatformMemory interface
+	//~ Begin FGenericPlatformMemory Interface
 	static void Init();
 	static const FPlatformMemoryConstants& GetConstants();
 	static FPlatformMemoryStats GetStats();
 	static FMalloc* BaseAllocator();
 	static void* BinnedAllocFromOS( SIZE_T Size );
 	static void BinnedFreeToOS( void* Ptr );
-	// End FGenericPlatformMemory interface
+	//~ End FGenericPlatformMemory Interface
 };
 
 typedef FHTML5PlatformMemory FPlatformMemory;
