Merging using UE4->Orion

diff --git a/Engine/Source/Runtime/Core/Public/Containers/Array.h b/Engine/Source/Runtime/Core/Public/Containers/Array.h
--- a/Engine/Source/Runtime/Core/Public/Containers/Array.h
+++ b/Engine/Source/Runtime/Core/Public/Containers/Array.h
@@ -1032,7 +1032,7 @@ public:
 
 	/**
 	 * Finds an item by key (assuming the ElementType overloads operator== for
-	 * the comparison).
+	 * the comparison). Time Complexity: O(n), starts iteration from the beginning so better performance if Key is in the front
 	 *
 	 * @param Key The key to search by.
 	 *
