Android: fixed compilation for r10c

diff --git a/Engine/Source/Runtime/Core/Private/Android/AndroidPlatformBacktrace.cpp b/Engine/Source/Runtime/Core/Private/Android/AndroidPlatformBacktrace.cpp
--- a/Engine/Source/Runtime/Core/Private/Android/AndroidPlatformBacktrace.cpp
+++ b/Engine/Source/Runtime/Core/Private/Android/AndroidPlatformBacktrace.cpp
@@ -5,13 +5,18 @@
 =============================================================================*/
 
 #include "CorePrivatePCH.h"
-#include <unwind.h>
 #include <cxxabi.h>
 #include <stdio.h>
-#include <link.h>
 
 #if PLATFORM_ANDROID_ARM
 
+#include <unwind.h>
+
+// missing in r9c llvm-3.3
+typedef uintptr_t _Unwind_Ptr;
+
+#ifndef _SYS_UCONTEXT_H_
+// missing r9c
 typedef struct ucontext {
     uint32_t uc_flags;
     struct ucontext* uc_link;
@@ -19,6 +24,8 @@ typedef struct ucontext {
     struct sigcontext uc_mcontext;
     uint32_t uc_sigmask;
 } ucontext_t;
+#endif // _SYS_UCONTEXT_H_
+
 
 /* Unwind state. */
 typedef struct {
@@ -45,8 +52,12 @@ static const uint32_t EXIDX_CANTUNWIND = 1;
  * Bionic exports a helpful function called __gnu_Unwind_Find_exidx that
  * handles both cases, so we use that here.
  */
-typedef long unsigned int* _Unwind_Ptr;
 //extern _Unwind_Ptr __gnu_Unwind_Find_exidx(_Unwind_Ptr pc, int *pcount);
+// have to define this since cannot include link.h
+__BEGIN_DECLS
+_Unwind_Ptr dl_unwind_find_exidx(_Unwind_Ptr pc, int* pcount);
+__END_DECLS
+
 static uintptr_t find_exidx(uintptr_t pc, size_t* out_exidx_size) {
 	int count;
 	uintptr_t start = (uintptr_t)dl_unwind_find_exidx((_Unwind_Ptr)pc, &count);
