Faster Memswap which does chunked iteration and has a constant-time swap optimization.
New FMath::CountTrailingZeros function.
TSparseArrayCompact will now compact in-place and not do anything if the array is already compacted.

diff --git a/Engine/Source/Runtime/Core/Private/GenericPlatform/GenericPlatformMemory.cpp b/Engine/Source/Runtime/Core/Private/GenericPlatform/GenericPlatformMemory.cpp
--- a/Engine/Source/Runtime/Core/Private/GenericPlatform/GenericPlatformMemory.cpp
+++ b/Engine/Source/Runtime/Core/Private/GenericPlatform/GenericPlatformMemory.cpp
@@ -185,18 +185,80 @@ void FGenericPlatformMemory::DumpPlatformAndAllocatorStats( class FOutputDevice&
 	GMalloc->DumpAllocatorStats( Ar );
 }
 
-void FGenericPlatformMemory::Memswap( void* Ptr1, void* Ptr2, SIZE_T Size )
+void FGenericPlatformMemory::MemswapImpl( void* RESTRICT Ptr1, void* RESTRICT Ptr2, SIZE_T Size )
 {
-	if (Ptr1 != Ptr2)
+	union PtrUnion
 	{
-		// check that Ptr1 and Ptr2 do not overlap in undefined ways
-		checkf(reinterpret_cast<uint8 *>(Ptr1)+Size <= reinterpret_cast<uint8 *>(Ptr2) || reinterpret_cast<uint8 *>(Ptr2)+Size <= reinterpret_cast<uint8 *>(Ptr1),
-			TEXT("Pointers given to FPlatformMemory::Memswap() point to overlapping memory areas, results are undefined."));
-
-		void* Temp = FMemory_Alloca(Size);
-		FPlatformMemory::Memcpy( Temp, Ptr1, Size );
-		FPlatformMemory::Memcpy( Ptr1, Ptr2, Size );
-		FPlatformMemory::Memcpy( Ptr2, Temp, Size );
+		void*   PtrVoid;
+		uint8*  Ptr8;
+		uint16* Ptr16;
+		uint32* Ptr32;
+		uint64* Ptr64;
+		UPTRINT PtrUint;
+	};
+
+	if (!Size)
+	{
+		return;
+	}
+
+	PtrUnion Union1 = { Ptr1 };
+	PtrUnion Union2 = { Ptr2 };
+
+	if (Union1.PtrUint & 1)
+	{
+		Valswap(*Union1.Ptr8++, *Union2.Ptr8++);
+		Size -= 1;
+		if (!Size)
+		{
+			return;
+		}
+	}
+	if (Union1.PtrUint & 2)
+	{
+		Valswap(*Union1.Ptr16++, *Union2.Ptr16++);
+		Size -= 2;
+		if (!Size)
+		{
+			return;
+		}
+	}
+	if (Union1.PtrUint & 4)
+	{
+		Valswap(*Union1.Ptr32++, *Union2.Ptr32++);
+		Size -= 4;
+		if (!Size)
+		{
+			return;
+		}
+	}
+
+	uint32 CommonAlignment = FMath::Min(FMath::CountTrailingZeros(Union1.PtrUint - Union2.PtrUint), 3u);
+	switch (CommonAlignment)
+	{
+		default:
+			for (; Size >= 8; Size -= 8)
+			{
+				Valswap(*Union1.Ptr64++, *Union2.Ptr64++);
+			}
+
+		case 2:
+			for (; Size >= 4; Size -= 4)
+			{
+				Valswap(*Union1.Ptr32++, *Union2.Ptr32++);
+			}
+
+		case 1:
+			for (; Size >= 2; Size -= 2)
+			{
+				Valswap(*Union1.Ptr16++, *Union2.Ptr16++);
+			}
+
+		case 0:
+			for (; Size >= 1; Size -= 1)
+			{
+				Valswap(*Union1.Ptr8++, *Union2.Ptr8++);
+			}
 	}
 }
 
diff --git a/Engine/Source/Runtime/Core/Public/Containers/BitArray.h b/Engine/Source/Runtime/Core/Public/Containers/BitArray.h
--- a/Engine/Source/Runtime/Core/Public/Containers/BitArray.h
+++ b/Engine/Source/Runtime/Core/Public/Containers/BitArray.h
@@ -134,21 +134,14 @@ private:
 /** Used to reference a bit in an unspecified bit array. */
 class FRelativeBitReference
 {
-	template<typename>
-	friend class TBitArray;
-	template<typename>
-	friend class TConstSetBitIterator;
-	template<typename,typename>
-	friend class TConstDualSetBitIterator;
 public:
+	FORCEINLINE explicit FRelativeBitReference(int32 BitIndex)
+		: DWORDIndex(BitIndex >> NumBitsPerDWORDLogTwo)
+		, Mask(1 << (BitIndex & (NumBitsPerDWORD - 1)))
+	{
+	}
 
-	FORCEINLINE FRelativeBitReference(int32 BitIndex)
-	:	DWORDIndex(BitIndex >> NumBitsPerDWORDLogTwo)
-	,	Mask(1 << (BitIndex & (NumBitsPerDWORD - 1)))
-	{}
-
-protected:
-	int32 DWORDIndex;
+	int32  DWORDIndex;
 	uint32 Mask;
 };
 
@@ -752,7 +745,7 @@ public:
 		}
 	}
 
-	/** Advancement operator. */
+	/** Forwards iteration operator. */
 	FORCEINLINE TConstSetBitIterator& operator++()
 	{
 		// Mark the current bit as visited.
diff --git a/Engine/Source/Runtime/Core/Public/Containers/SparseArray.h b/Engine/Source/Runtime/Core/Public/Containers/SparseArray.h
--- a/Engine/Source/Runtime/Core/Public/Containers/SparseArray.h
+++ b/Engine/Source/Runtime/Core/Public/Containers/SparseArray.h
@@ -4,6 +4,7 @@
 
 #include "Containers/Array.h"
 #include "Containers/BitArray.h"
+#include "MemoryOps.h"
 
 
 // Forward declarations.
@@ -341,16 +342,41 @@ public:
 	/** Compacts the allocated elements into a contiguous index range. */
 	void Compact()
 	{
-		// Copy the existing elements to a new array.
-		TSparseArray<ElementType,Allocator> CompactedArray;
-		CompactedArray.Empty(Num());
-		for(TConstIterator It(*this);It;++It)
+		int32 NumFree = NumFreeIndices;
+		if (NumFree == 0)
 		{
-			new(CompactedArray.AddUninitialized()) ElementType(*It);
+			return;
 		}
 
-		// Replace this array with the compacted array.
-		Exchange(*this,CompactedArray);
+		FElementOrFreeListLink* ElementData = Data.GetData();
+
+		int32 EndIndex    = Data.Num();
+		int32 TargetIndex = EndIndex - NumFree;
+		int32 FreeIndex   = FirstFreeIndex;
+		while (FreeIndex != -1)
+		{
+			int32 NextFreeIndex = GetData(FreeIndex).NextFreeIndex;
+			if (FreeIndex < TargetIndex)
+			{
+				// We need an element here
+				do
+				{
+					--EndIndex;
+				}
+				while (!AllocationFlags[EndIndex]);
+
+				RelocateConstructItems<FElementOrFreeListLink>(ElementData + FreeIndex, ElementData + EndIndex, 1);
+				AllocationFlags[FreeIndex] = true;
+			}
+
+			FreeIndex = NextFreeIndex;
+		}
+
+		Data           .RemoveAt(TargetIndex, NumFree);
+		AllocationFlags.RemoveAt(TargetIndex, NumFree);
+
+		NumFreeIndices = 0;
+		FirstFreeIndex = -1;
 	}
 
 	/** Sorts the elements using the provided comparison class. */
diff --git a/Engine/Source/Runtime/Core/Public/GenericPlatform/GenericPlatformMath.h b/Engine/Source/Runtime/Core/Public/GenericPlatform/GenericPlatformMath.h
--- a/Engine/Source/Runtime/Core/Public/GenericPlatform/GenericPlatformMath.h
+++ b/Engine/Source/Runtime/Core/Public/GenericPlatform/GenericPlatformMath.h
@@ -373,6 +373,28 @@ struct FGenericPlatformMath
 		return 31 - FloorLog2(Value);
 	}
 
+	/**
+	 * Counts the number of trailing zeros in the bit representation of the value
+	 *
+	 * @param Value the value to determine the number of trailing zeros for
+	 *
+	 * @return the number of zeros after the last "on" bit
+	 */
+	static FORCEINLINE uint32 CountTrailingZeros(uint32 Value)
+	{
+		if (Value == 0)
+		{
+			return 32;
+		}
+		uint32 Result = 0;
+		while ((Value & 1) == 0)
+		{
+			Value >>= 1;
+			++Result;
+		}
+		return Result;
+	}
+
 	/**
 	 * Returns smallest N such that (1<<N)>=Arg.
 	 * Note: CeilLogTwo(0)=0 because (1<<0)=1 >= 0.
diff --git a/Engine/Source/Runtime/Core/Public/GenericPlatform/GenericPlatformMemory.h b/Engine/Source/Runtime/Core/Public/GenericPlatform/GenericPlatformMemory.h
--- a/Engine/Source/Runtime/Core/Public/GenericPlatform/GenericPlatformMemory.h
+++ b/Engine/Source/Runtime/Core/Public/GenericPlatform/GenericPlatformMemory.h
@@ -283,7 +283,52 @@ struct CORE_API FGenericPlatformMemory
 		return memcpy( Dest, Src, Count );
 	}
 
-	static void Memswap( void* Ptr1, void* Ptr2, SIZE_T Size );
+private:
+	template <typename T>
+	static FORCEINLINE void Valswap(T& A, T& B)
+	{
+		// Usually such an implementation would use move semantics, but
+		// we're only ever going to call it on fundamental types and MoveTemp
+		// is not necessarily in scope here anyway, so we don't want to
+		// #include it if we don't need to.
+		T Tmp = A;
+		A = B;
+		B = Tmp;
+	}
+
+	static void MemswapImpl( void* Ptr1, void* Ptr2, SIZE_T Size );
+
+public:
+	static inline void Memswap( void* Ptr1, void* Ptr2, SIZE_T Size )
+	{
+		switch (Size)
+		{
+			case 1:
+				Valswap(*(uint8*)Ptr1, *(uint8*)Ptr2);
+				break;
+
+			case 2:
+				Valswap(*(uint16*)Ptr1, *(uint16*)Ptr2);
+				break;
+
+			case 4:
+				Valswap(*(uint32*)Ptr1, *(uint32*)Ptr2);
+				break;
+
+			case 8:
+				Valswap(*(uint64*)Ptr1, *(uint64*)Ptr2);
+				break;
+
+			case 16:
+				Valswap(((uint64*)Ptr1)[0], ((uint64*)Ptr2)[0]);
+				Valswap(((uint64*)Ptr1)[1], ((uint64*)Ptr2)[1]);
+				break;
+
+			default:
+				MemswapImpl(Ptr1, Ptr2, Size);
+				break;
+		}
+	}
 
 	/**
 	 * Maps a named shared memory region into process address space (creates or opens it)
diff --git a/Engine/Source/Runtime/Core/Public/Windows/WindowsPlatformMath.h b/Engine/Source/Runtime/Core/Public/Windows/WindowsPlatformMath.h
--- a/Engine/Source/Runtime/Core/Public/Windows/WindowsPlatformMath.h
+++ b/Engine/Source/Runtime/Core/Public/Windows/WindowsPlatformMath.h
@@ -99,6 +99,16 @@ struct FWindowsPlatformMath : public FGenericPlatformMath
 
 		return 32;
 	}
+	static FORCEINLINE uint32 CountTrailingZeros(uint32 Value)
+	{
+		if (Value == 0)
+		{
+			return 32;
+		}
+		uint32 BitIndex;	// 0-based, where the LSB is 0 and MSB is 31
+		_BitScanForward( (::DWORD *)&BitIndex, Value );	// Scans from LSB to MSB
+		return BitIndex;
+	}
 	static FORCEINLINE uint32 CeilLogTwo( uint32 Arg )
 	{
 		int32 Bitmask = ((int32)(CountLeadingZeros(Arg) << 26)) >> 31;
