Always assume compatibility with the current major/minor engine release. For hotfix releases, we should be writing out assets that are compatible with the original release.

diff --git a/Engine/Source/Runtime/Core/Private/Misc/EngineVersion.cpp b/Engine/Source/Runtime/Core/Private/Misc/EngineVersion.cpp
--- a/Engine/Source/Runtime/Core/Private/Misc/EngineVersion.cpp
+++ b/Engine/Source/Runtime/Core/Private/Misc/EngineVersion.cpp
@@ -14,8 +14,8 @@ const int32 GEngineNegotiationVersion	= 3077;
 // Global instance of the current engine version
 FEngineVersion FEngineVersion::CurrentVersion(ENGINE_MAJOR_VERSION, ENGINE_MINOR_VERSION, ENGINE_PATCH_VERSION, (BUILT_FROM_CHANGELIST | (ENGINE_IS_LICENSEE_VERSION << 31)), BRANCH_NAME);
 
-// Version which this engine maintains strict API and package compatibility with
-FEngineVersion FEngineVersion::CompatibleWithVersion(ENGINE_MAJOR_VERSION, ENGINE_MINOR_VERSION, ENGINE_PATCH_VERSION, (MODULE_API_VERSION | (ENGINE_IS_LICENSEE_VERSION << 31)), BRANCH_NAME);
+// Version which this engine maintains strict API and package compatibility with. By default, we always maintain compatibility with the current major/minor version, unless we're built at a different changelist.
+FEngineVersion FEngineVersion::CompatibleWithVersion(ENGINE_MAJOR_VERSION, ENGINE_MINOR_VERSION, 0, (MODULE_API_VERSION | (ENGINE_IS_LICENSEE_VERSION << 31)), BRANCH_NAME);
 
 
 FEngineVersionBase::FEngineVersionBase()
