Fixed a case where FastDecimalFormat would incorrectly format a negative fractional number whose integral part is zero

diff --git a/Engine/Source/Runtime/Core/Private/Internationalization/FastDecimalFormat.cpp b/Engine/Source/Runtime/Core/Private/Internationalization/FastDecimalFormat.cpp
--- a/Engine/Source/Runtime/Core/Private/Internationalization/FastDecimalFormat.cpp
+++ b/Engine/Source/Runtime/Core/Private/Internationalization/FastDecimalFormat.cpp
@@ -112,11 +112,10 @@ FORCEINLINE int32 IntegralToString_Common(const uint64 InVal, const FDecimalNumb
 		);
 }
 
-void FractionalToString_SplitAndRoundNumber(const double InValue, const int32 InNumDecimalPlaces, ERoundingMode InRoundingMode, double& OutIntegralPart, double& OutFractionalPart)
+void FractionalToString_SplitAndRoundNumber(const bool bIsNegative, const double InValue, const int32 InNumDecimalPlaces, ERoundingMode InRoundingMode, double& OutIntegralPart, double& OutFractionalPart)
 {
 	const int32 DecimalPlacesToRoundTo = FMath::Min(InNumDecimalPlaces, MaxFractionalPrintPrecision);
 
-	const bool bIsNegative = InValue < 0.0;
 	const bool bIsRoundingEntireNumber = DecimalPlacesToRoundTo == 0;
 
 	// We split the value before performing the rounding to avoid losing precision during the rounding calculations
@@ -233,11 +232,12 @@ FString FractionalToString(const double InVal, const FDecimalNumberFormattingRul
 		return InFormattingRules.NaNString;
 	}
 
+	const bool bIsNegative = InVal < 0.0;
+
 	double IntegralPart = 0.0;
 	double FractionalPart = 0.0;
-	FractionalToString_SplitAndRoundNumber(InVal, InFormattingOptions.MaximumFractionalDigits, InFormattingOptions.RoundingMode, IntegralPart, FractionalPart);
+	FractionalToString_SplitAndRoundNumber(bIsNegative, InVal, InFormattingOptions.MaximumFractionalDigits, InFormattingOptions.RoundingMode, IntegralPart, FractionalPart);
 
-	const bool bIsNegative = IntegralPart < 0.0;
 	if (bIsNegative)
 	{
 		IntegralPart = -IntegralPart;
