Linux: move setting LC_NUMERIC to an earlier stage.

diff --git a/Engine/Source/Runtime/Core/Private/Linux/LinuxPlatformMisc.cpp b/Engine/Source/Runtime/Core/Private/Linux/LinuxPlatformMisc.cpp
--- a/Engine/Source/Runtime/Core/Private/Linux/LinuxPlatformMisc.cpp
+++ b/Engine/Source/Runtime/Core/Private/Linux/LinuxPlatformMisc.cpp
@@ -173,13 +173,6 @@ void FLinuxPlatformMisc::PlatformInit()
 	UE_LOG(LogInit, Log, TEXT(" -reuseconn - allow libcurl to reuse HTTP connections (only matters if compiled with libcurl)"));
 	UE_LOG(LogInit, Log, TEXT(" -virtmemkb=NUMBER - sets process virtual memory (address space) limit (overrides VirtualMemoryLimitInKB value from .ini)"));
 
-	UE_LOG(LogInit, Log, TEXT("Setting LC_NUMERIC to en_US"));
-	if (setenv("LC_NUMERIC", "en_US", 1) != 0)
-	{
-		int ErrNo = errno;
-		UE_LOG(LogInit, Warning, TEXT("Unable to setenv(): errno=%d (%s)"), ErrNo, ANSI_TO_TCHAR(strerror(ErrNo)));
-	}
-
 	// skip for servers and programs, unless they request later
 	if (!UE_SERVER && !IS_PROGRAM)
 	{
diff --git a/Engine/Source/Runtime/Linux/LinuxCommonStartup/Private/LinuxCommonStartup.cpp b/Engine/Source/Runtime/Linux/LinuxCommonStartup/Private/LinuxCommonStartup.cpp
--- a/Engine/Source/Runtime/Linux/LinuxCommonStartup/Private/LinuxCommonStartup.cpp
+++ b/Engine/Source/Runtime/Linux/LinuxCommonStartup/Private/LinuxCommonStartup.cpp
@@ -165,7 +165,13 @@ int CommonLinuxMain(int argc, char *argv[], int (*RealMain)(const TCHAR * Comman
 
 	int ErrorLevel = 0;
 
+	if (setenv("LC_NUMERIC", "en_US", 1) != 0)
+	{
+		int ErrNo = errno;
+		fprintf(stderr, "Unable to setenv(LC_NUMERIC): errno=%d (%s)", ErrNo, strerror(ErrNo));
+	}
 	setlocale(LC_CTYPE, "");
+
 	for (int32 Option = 1; Option < argc; Option++)
 	{
 		GSavedCommandLine += TEXT(" ");
