Core: Reduced nesting

diff --git a/Engine/Source/Runtime/Core/Public/Containers/Queue.h b/Engine/Source/Runtime/Core/Public/Containers/Queue.h
--- a/Engine/Source/Runtime/Core/Public/Containers/Queue.h
+++ b/Engine/Source/Runtime/Core/Public/Containers/Queue.h
@@ -66,25 +66,25 @@ public:
 	 * @return true if a value was returned, false if the queue was empty.
 	 * @see Enqueue, IsEmpty, Peek
 	 */
-	bool Dequeue( ItemType& OutItem )
+	bool Dequeue(ItemType& OutItem)
 	{
 		TNode* Popped = Tail->NextNode;
 
-		if (Popped != nullptr)
+		if (Popped == nullptr)
 		{
-			OutItem = Popped->Item;
+			return false;
+		}
 
-			TNode* OldTail = Tail;
+		OutItem = Popped->Item;
 
-			Tail = Popped;
-			Tail->Item = ItemType();
+		TNode* OldTail = Tail;
 
-			delete OldTail;
+		Tail = Popped;
+		Tail->Item = ItemType();
 
-			return true;
-		}
+		delete OldTail;
 
-		return false;
+		return true;
 	}
 
 	/**
@@ -94,30 +94,30 @@ public:
 	 * @return true if the item was added, false otherwise.
 	 * @see Dequeue, IsEmpty, Peek
 	 */
-	bool Enqueue( const ItemType& Item )
+	bool Enqueue(const ItemType& Item)
 	{
 		TNode* NewNode = new TNode(Item);
 
-		if (NewNode != nullptr)
+		if (NewNode == nullptr)
 		{
-			TNode* OldHead;
-
-			if (Mode == EQueueMode::Mpsc)
-			{
-				OldHead = (TNode*)FPlatformAtomics::InterlockedExchangePtr((void**)&Head, NewNode);
-			}
-			else
-			{
-				OldHead = Head;
-				Head = NewNode;
-			}		
+			return false;
+		}
 
-			OldHead->NextNode = NewNode;
+		TNode* OldHead;
 
-			return true;
+		if (Mode == EQueueMode::Mpsc)
+		{
+			OldHead = (TNode*)FPlatformAtomics::InterlockedExchangePtr((void**)&Head, NewNode);
 		}
+		else
+		{
+			OldHead = Head;
+			Head = NewNode;
+		}		
+
+		OldHead->NextNode = NewNode;
 
-		return false;
+		return true;
 	}
 
 	/**
@@ -140,14 +140,14 @@ public:
 	 */
 	bool Peek( ItemType& OutItem )
 	{
-		if (Tail->NextNode != nullptr)
+		if (Tail->NextNode == nullptr)
 		{
-			OutItem = Tail->NextNode->Item;
-
-			return true;
+			return false;
 		}
 
-		return false;
+		OutItem = Tail->NextNode->Item;
+
+		return true;
 	}
 
 private:
@@ -161,7 +161,6 @@ private:
 		/** Holds the node's item. */
 		ItemType Item;
 
-
 		/** Default constructor. */
 		TNode()
 			: NextNode(nullptr)
