SessionFrontend: Added ability to set session name and owner from the command line, i.e. on consoles running from Visual Studio

See SESSION console command for details.

diff --git a/Engine/Source/Developer/SessionFrontend/Private/Widgets/Browser/SSessionBrowser.cpp b/Engine/Source/Developer/SessionFrontend/Private/Widgets/Browser/SSessionBrowser.cpp
--- a/Engine/Source/Developer/SessionFrontend/Private/Widgets/Browser/SSessionBrowser.cpp
+++ b/Engine/Source/Developer/SessionFrontend/Private/Widgets/Browser/SSessionBrowser.cpp
@@ -196,7 +196,7 @@ void SSessionBrowser::FilterSessions()
 
 		if (LocalOwner)
 		{
-			if (Instances[0]->GetInstanceId() != FApp::GetInstanceId())
+			if (!FApp::IsThisInstance(Instances[0]->GetInstanceId()))
 			{
 				OwnerGroupItem->AddChild(SessionItem.ToSharedRef());
 				SessionItem->SetParent(OwnerGroupItem);
@@ -226,7 +226,7 @@ void SSessionBrowser::FilterSessions()
 			NewItemMap.Add(InstanceInfo->GetInstanceId(), InstanceItem);
 
 			// add instance to group or session
-			if (InstanceInfo->GetInstanceId() == FApp::GetInstanceId())
+			if (FApp::IsThisInstance(InstanceInfo->GetInstanceId()))
 			{
 				AppGroupItem->AddChild(InstanceItem.ToSharedRef());
 				InstanceItem->SetParent(AppGroupItem);
diff --git a/Engine/Source/Developer/SessionFrontend/Private/Widgets/Browser/SSessionBrowserTreeSessionRow.h b/Engine/Source/Developer/SessionFrontend/Private/Widgets/Browser/SSessionBrowserTreeSessionRow.h
--- a/Engine/Source/Developer/SessionFrontend/Private/Widgets/Browser/SSessionBrowserTreeSessionRow.h
+++ b/Engine/Source/Developer/SessionFrontend/Private/Widgets/Browser/SSessionBrowserTreeSessionRow.h
@@ -56,7 +56,7 @@ public:
 							.VAlign(VAlign_Center)
 							[
 								SNew(STextBlock)
-									.Text(SanitizeSessionName(Item->GetSessionInfo()))
+									.Text(this, &SSessionBrowserTreeSessionRow::HandleSessionNameText)
 									.Font(FEditorStyle::GetFontStyle("PropertyWindow.NormalFont"))
 							]
 					]
@@ -97,7 +97,7 @@ protected:
 		{
 			const ISessionInstanceInfoPtr& FirstInstance = Instances[0];
 
-			if ((Instances.Num() == 1) && (FirstInstance->GetInstanceId() == FApp::GetInstanceId()))
+			if ((Instances.Num() == 1) && FApp::IsThisInstance(FirstInstance->GetInstanceId()))
 			{
 				return LOCTEXT("ThisApplicationSessionText", "This Application");
 			}
@@ -168,6 +168,12 @@ private:
 		return ToolTipTextBuilder.ToText();
 	}
 
+	/** Callback for getting the name of the session. */
+	FText HandleSessionNameText() const
+	{
+		return SanitizeSessionName(Item->GetSessionInfo());
+	}
+
 private:
 
 	/** The highlight string for this row. */
diff --git a/Engine/Source/Developer/SessionServices/Private/SessionInfo.cpp b/Engine/Source/Developer/SessionServices/Private/SessionInfo.cpp
--- a/Engine/Source/Developer/SessionServices/Private/SessionInfo.cpp
+++ b/Engine/Source/Developer/SessionServices/Private/SessionInfo.cpp
@@ -51,11 +51,7 @@ void FSessionInfo::UpdateFromMessage( const FSessionServicePong& Message, const
 	// update session info
 	Standalone = Message.Standalone;
 	SessionOwner = Message.SessionOwner;
-
-	if (SessionName.IsEmpty())
-	{
-		SessionName = Message.SessionName;
-	}
+	SessionName = Message.SessionName;
 
 	// update instance
 	TSharedPtr<FSessionInstanceInfo>& Instance = Instances.FindOrAdd(Context->GetSender());
diff --git a/Engine/Source/Developer/SessionServices/Private/SessionManager.cpp b/Engine/Source/Developer/SessionServices/Private/SessionManager.cpp
--- a/Engine/Source/Developer/SessionServices/Private/SessionManager.cpp
+++ b/Engine/Source/Developer/SessionServices/Private/SessionManager.cpp
@@ -290,7 +290,15 @@ void FSessionManager::HandleSessionPongMessage(const FSessionServicePong& Messag
 
 	if (Session.IsValid())
 	{
-		Session->UpdateFromMessage(Message, Context);
+		if (Session->GetSessionOwner() != Message.SessionOwner)
+		{
+			Session->UpdateFromMessage(Message, Context);
+			SessionsUpdatedDelegate.Broadcast();
+		}
+		else
+		{
+			Session->UpdateFromMessage(Message, Context);
+		}
 	}
 	else
 	{
diff --git a/Engine/Source/Developer/SessionServices/Private/SessionServicesModule.cpp b/Engine/Source/Developer/SessionServices/Private/SessionServicesModule.cpp
--- a/Engine/Source/Developer/SessionServices/Private/SessionServicesModule.cpp
+++ b/Engine/Source/Developer/SessionServices/Private/SessionServicesModule.cpp
@@ -11,7 +11,8 @@
  * @todo gmp: needs better error handling in case MessageBusPtr is invalid
  */
 class FSessionServicesModule
-	: public ISessionServicesModule
+	: public FSelfRegisteringExec
+	, public ISessionServicesModule
 {
 public:
 
@@ -21,6 +22,48 @@ public:
 		, SessionService(nullptr)
 	{ }
 
+public:
+
+	// FSelfRegisteringExec interface
+
+	virtual bool Exec(UWorld* InWorld, const TCHAR* Cmd, FOutputDevice& Ar) override
+	{
+		if (FParse::Command(&Cmd, TEXT("SESSION")))
+		{
+			if (FParse::Command(&Cmd, TEXT("STATUS")))
+			{
+				Ar.Logf(TEXT("Instance ID: %s"), *FApp::GetInstanceId().ToString());
+				Ar.Logf(TEXT("Instance Name: %s"), *FApp::GetInstanceName());
+				Ar.Logf(TEXT("Session ID: %s"), *FApp::GetSessionId().ToString());
+				Ar.Logf(TEXT("Session Name: %s"), *FApp::GetSessionName());
+				Ar.Logf(TEXT("Session Owner: %s"), *FApp::GetSessionOwner());
+				Ar.Logf(TEXT("Standalone: %s"), FApp::IsStandalone() ? *GYes.ToString() : *GNo.ToString());
+			}
+			else if (FParse::Command(&Cmd, TEXT("SETNAME")))
+			{
+				FApp::SetSessionName(FParse::Token(Cmd, false));
+			}
+			else if (FParse::Command(&Cmd, TEXT("SETOWNER")))
+			{
+				FApp::SetSessionOwner(FParse::Token(Cmd, false));
+			}
+			else
+			{
+				// show usage
+				Ar.Log(TEXT("Usage: SESSION <Command>"));
+				Ar.Log(TEXT(""));
+				Ar.Log(TEXT("Command"));
+				Ar.Log(TEXT("    STATUS = Displays the status of this application session"));
+				Ar.Log(TEXT("    SETNAME <Name> = Sets the name of this instance"));
+				Ar.Log(TEXT("    SETOWNER <Owner> = Sets the name of the owner of this instance"));
+			}
+
+			return true;
+		}
+
+		return false;
+	}
+
 public:
 
 	// ISessionServicesModule interface
diff --git a/Engine/Source/Runtime/Core/Public/Misc/App.h b/Engine/Source/Runtime/Core/Public/Misc/App.h
--- a/Engine/Source/Runtime/Core/Public/Misc/App.h
+++ b/Engine/Source/Runtime/Core/Public/Misc/App.h
@@ -15,16 +15,6 @@
 public:
 
 	/**
-	 * Checks whether this application can render anything.
-	 *
-	 * @return true if the application can render, false otherwise.
-	 */
-	FORCEINLINE static bool CanEverRender()
-	{
-		return (!IsRunningCommandlet() || IsAllowCommandletRendering()) && !IsRunningDedicatedServer();
-	}
-
-	/**
 	 * Gets the name of the version control branch that this application was built from.
 	 *
 	 * @return The branch name.
@@ -32,4 +22,18 @@
 	static FString GetBranchName();
 
 	/**
+	 * Gets the application's build configuration, i.e. Debug or Shipping.
+	 *
+	 * @return The build configuration.
+	 */
+	static EBuildConfigurations::Type GetBuildConfiguration();
+
+	/**
+	 * Gets the date at which this application was built.
+	 *
+	 * @return Build date string.
+	 */
+	static FString GetBuildDate();
+
+	/**
 	 * Gets the value of ENGINE_IS_PROMOTED_BUILD.
@@ -42,20 +46,6 @@
 	static FString GetEpicProductIdentifier();
 
 	/**
-	 * Gets the application's build configuration, i.e. Debug or Shipping.
-	 *
-	 * @return The build configuration.
-	 */
-	static EBuildConfigurations::Type GetBuildConfiguration();
-
-	/**
-	 * Gets the date at which this application was built.
-	 *
-	 * @return Build date string.
-	 */
-	static FString GetBuildDate();
-
-	/**
 	 * Gets the name of the currently running game.
 	 *
 	 * @return The game name.
@@ -66,3 +56,66 @@
 	}
 
 	/**
+	 * Gets the name of the application, i.e. "UE4" or "Rocket".
+	 *
+	 * @todo need better application name discovery. this is quite horrible and may not work on future platforms.
+	 * @return Application name string.
+	 */
+	static FString GetName()
+	{
+		FString ExecutableName = FPlatformProcess::ExecutableName();
+
+		int32 ChopIndex = ExecutableName.Find(TEXT("-"));
+
+		if (ExecutableName.FindChar(TCHAR('-'), ChopIndex))
+		{
+			return ExecutableName.Left(ChopIndex);
+		}
+
+		if (ExecutableName.FindChar(TCHAR('.'), ChopIndex))
+		{
+			return ExecutableName.Left(ChopIndex);
+		}
+
+		return ExecutableName;
+	}
+
+	/**
+	 * Reports if the game name has been set
+	 *
+	 * @return true if the game name has been set
+	 */
+	FORCEINLINE static bool HasGameName()
+	{
+		return (IsGameNameEmpty() == false) && (FCString::Stricmp(GInternalGameName, TEXT("None")) != 0);
+	}
+
+	/**
+	 * Checks whether this application is a game.
+	 *
+	 * Returns true if a normal or PIE game is active (basically !GIsEdit or || GIsPlayInEditorWorld)
+	 * This must NOT be accessed on threads other than the game thread!
+	 * Use View->Family->EnginShowFlags.Game on the rendering thread.
+	 *
+	 * @return true if the application is a game, false otherwise.
+	 */
+	FORCEINLINE static bool IsGame()
+	{
+#if WITH_EDITOR
+		return !GIsEditor || GIsPlayInEditorWorld || IsRunningGame();
+#else
+		return true;
+#endif
+	}
+
+	/**
+	 * Reports if the game name is empty
+	 *
+	 * @return true if the game name is empty
+	 */
+	FORCEINLINE static bool IsGameNameEmpty()
+	{
+		return (GInternalGameName[0] == 0);
+	}
+
+	/**
@@ -78,6 +131,8 @@
 		GInternalGameName[ARRAY_COUNT(GInternalGameName) - 1] = 0;
 	}
 
+public:
+
 	/**
 	 * Gets the globally unique identifier of this application instance.
 	 *
@@ -105,30 +160,5 @@
 	}
 
 	/**
-	 * Gets the name of the application, i.e. "UE4" or "Rocket".
-	 *
-	 * @todo need better application name discovery. this is quite horrible and may not work on future platforms.
-	 * @return Application name string.
-	 */
-	static FString GetName()
-	{
-		FString ExecutableName = FPlatformProcess::ExecutableName();
-
-		int32 ChopIndex = ExecutableName.Find(TEXT("-"));
-
-		if (ExecutableName.FindChar(TCHAR('-'), ChopIndex))
-		{
-			return ExecutableName.Left(ChopIndex);
-		}
-
-		if (ExecutableName.FindChar(TCHAR('.'), ChopIndex))
-		{
-			return ExecutableName.Left(ChopIndex);
-		}
-
-		return ExecutableName;
-	}
-
-	/**
 	 * Gets the identifier of the session that this application is part of.
 	 *
@@ -171,46 +201,66 @@
 	}
 
 	/**
-	* Reports if the game name is empty
-	*
-	* @return true if the game name is empty
-	*/
-	FORCEINLINE static bool IsGameNameEmpty()
-	{
-		return (GInternalGameName[0] == 0);
-	}
-
-	/**
-	 * Reports if the game name has been set
-	 *
-	 * @return true if the game name has been set
-	 */
-	FORCEINLINE static bool HasGameName()
-	{
-		return (IsGameNameEmpty() == false) && (FCString::Stricmp(GInternalGameName, TEXT("None")) != 0);
-	}
-
-	/**
 	 * Initializes the application session.
 	 */
 	static void InitializeSession();
 
 	/**
-	 * Checks whether this application is a game.
+	 * Checks whether this is a standalone application.
 	 *
-	 * Returns true if a normal or PIE game is active (basically !GIsEdit or || GIsPlayInEditorWorld)
-	 * This must NOT be accessed on threads other than the game thread!
-	 * Use View->Family->EnginShowFlags.Game on the rendering thread.
+	 * An application is standalone if it has not been launched as part of a session from UFE.
+	 * If an application was launched from UFE, it is part of a session and not considered standalone.
 	 *
-	 * @return true if the application is a game, false otherwise.
+	 * @return true if this application is standalone, false otherwise.
 	 */
-	FORCEINLINE static bool IsGame()
+	FORCEINLINE static bool IsStandalone()
 	{
-#if WITH_EDITOR
-		return !GIsEditor || GIsPlayInEditorWorld || IsRunningGame();
-#else
-		return true;
-#endif
+		return Standalone;
 	}
 
 	/**
+	 * Check whether the given instance ID identifies this instance.
+	 *
+	 * @param InInstanceId The instance ID to check.
+	 * @return true if the ID identifies this instance, false otherwise.
+	 */
+	FORCEINLINE static bool IsThisInstance(const FGuid& InInstanceId)
+	{
+		return (InInstanceId == InstanceId);
+	};
+
+	/**
+	 * Set a new session name.
+	 *
+	 * @param NewName The new session name.
+	 * @see SetSessionOwner
+	 */
+	FORCEINLINE static void SetSessionName(const FString& NewName)
+	{
+		SessionName = NewName;
+	}
+
+	/**
+	 * Set a new session owner.
+	 *
+	 * @param NewOwner The name of the new owner.
+	 * @see SetSessionName
+	 */
+	FORCEINLINE static void SetSessionOwner(const FString& NewOwner)
+	{
+		SessionOwner = NewOwner;
+	}
+
+public:
+
+	/**
+	 * Checks whether this application can render anything.
+	 *
+	 * @return true if the application can render, false otherwise.
+	 */
+	FORCEINLINE static bool CanEverRender()
+	{
+		return (!IsRunningCommandlet() || IsAllowCommandletRendering()) && !IsRunningDedicatedServer();
+	}
+
+	/**
@@ -242,18 +292,5 @@
 	static bool IsEngineInstalled();
 
 	/**
-	 * Checks whether this is a standalone application.
-	 *
-	 * An application is standalone if it has not been launched as part of a session from UFE.
-	 * If an application was launched from UFE, it is part of a session and not considered standalone.
-	 *
-	 * @return true if this application is standalone, false otherwise.
-	 */
-	FORCEINLINE static bool IsStandalone()
-	{
-		return Standalone;
-	}
-
-	/**
 	 * Checks whether this application runs unattended.
 	 *
