UE-17755/Time markers in the log file when using -LOGTIMESINCESTART are invalid

diff --git a/Engine/Source/Developer/DesktopPlatform/Private/Windows/WindowsNativeFeedbackContext.h b/Engine/Source/Developer/DesktopPlatform/Private/Windows/WindowsNativeFeedbackContext.h
--- a/Engine/Source/Developer/DesktopPlatform/Private/Windows/WindowsNativeFeedbackContext.h
+++ b/Engine/Source/Developer/DesktopPlatform/Private/Windows/WindowsNativeFeedbackContext.h
@@ -29,13 +29,6 @@ public:
 	void SetContext( FContextSupplier* InSupplier );
 
 private:
-	struct FBufferedLine
-	{
-		FString Text;
-		ELogVerbosity::Type Verbosity;
-		FName Category;
-	};
-
 	struct FWindowParams
 	{
 		FWindowsNativeFeedbackContext* Context;
diff --git a/Engine/Source/Runtime/Core/Private/HAL/OutputDevices.cpp b/Engine/Source/Runtime/Core/Private/HAL/OutputDevices.cpp
--- a/Engine/Source/Runtime/Core/Private/HAL/OutputDevices.cpp
+++ b/Engine/Source/Runtime/Core/Private/HAL/OutputDevices.cpp
@@ -790,7 +790,7 @@ void FOutputDeviceRedirector::UnsynchronizedFlushThreadedLogs( bool bUseAllDevic
 			FOutputDevice* OutputDevice = OutputDevices[OutputDeviceIndex];
 			if( OutputDevice->CanBeUsedOnAnyThread() || bUseAllDevices )
 			{
-				OutputDevice->Serialize( *BufferedLine.Data, BufferedLine.Verbosity, BufferedLine.Category );
+				OutputDevice->Serialize( *BufferedLine.Data, BufferedLine.Verbosity, BufferedLine.Category, BufferedLine.Time );
 			}
 		}
 	}
@@ -843,7 +843,7 @@ void FOutputDeviceRedirector::SerializeBacklog( FOutputDevice* OutputDevice )
 	for (int32 LineIndex = 0; LineIndex < BacklogLines.Num(); LineIndex++)
 	{
 		const FBufferedLine& BacklogLine = BacklogLines[ LineIndex ];
-		OutputDevice->Serialize( *BacklogLine.Data, BacklogLine.Verbosity, BacklogLine.Category );
+		OutputDevice->Serialize( *BacklogLine.Data, BacklogLine.Verbosity, BacklogLine.Category, BacklogLine.Time );
 	}
 }
 
@@ -877,18 +877,20 @@ void FOutputDeviceRedirector::SetCurrentThreadAsMasterThread()
 	MasterThreadID = FPlatformTLS::GetCurrentThreadId();
 }
 
-void FOutputDeviceRedirector::Serialize( const TCHAR* Data, ELogVerbosity::Type Verbosity, const class FName& Category )
+void FOutputDeviceRedirector::Serialize( const TCHAR* Data, ELogVerbosity::Type Verbosity, const class FName& Category, const double Time )
 {
+	const double RealTime = Time == -1.0f ? FPlatformTime::Seconds() - GStartTime : Time;
+
 	FScopeLock ScopeLock( &SynchronizationObject );
 
 	if ( bEnableBacklog )
 	{
-		new(BacklogLines) FBufferedLine(Data,Verbosity,Category);
+		new(BacklogLines)FBufferedLine( Data, Category, Verbosity, RealTime );
 	}
 
 	if(FPlatformTLS::GetCurrentThreadId() != MasterThreadID || OutputDevices.Num() == 0)
 	{
-		new(BufferedLines) FBufferedLine(Data,Verbosity,Category);
+		new(BufferedLines)FBufferedLine( Data, Category, Verbosity, RealTime );
 	}
 	else
 	{
@@ -898,11 +900,16 @@ void FOutputDeviceRedirector::Serialize( const TCHAR* Data, ELogVerbosity::Type
 
 		for( int32 OutputDeviceIndex=0; OutputDeviceIndex<OutputDevices.Num(); OutputDeviceIndex++ )
 		{
-			OutputDevices[OutputDeviceIndex]->Serialize( Data, Verbosity, Category );
+			OutputDevices[OutputDeviceIndex]->Serialize( Data, Verbosity, Category, RealTime );
 		}
 	}
 }
 
+void FOutputDeviceRedirector::Serialize( const TCHAR* Data, ELogVerbosity::Type Verbosity, const class FName& Category )
+{
+	Serialize( Data, Verbosity, Category, -1.0 );
+}
+
 /**
  * Passes on the flush request to all current output devices.
  */
@@ -1076,11 +1083,11 @@ void FOutputDeviceFile::CastAndSerializeData(const TCHAR* Data)
 	LogAr->Serialize((ANSICHAR*)(ConvertedData.Get()), ConvertedData.Length());
 }
 
-void FOutputDeviceFile::WriteDataToArchive(const TCHAR* Data, ELogVerbosity::Type Verbosity, const class FName& Category)
+void FOutputDeviceFile::WriteDataToArchive(const TCHAR* Data, ELogVerbosity::Type Verbosity, const class FName& Category, const double Time)
 {
 	if (!bSuppressEventTag)
 	{
-		FString Prefix = FOutputDevice::FormatLogLine(Verbosity, Category, NULL, GPrintLogTimes);
+		FString Prefix = FOutputDevice::FormatLogLine(Verbosity, Category, NULL, GPrintLogTimes, Time);
 		CastAndSerializeData(*Prefix);
 	}
 
@@ -1104,7 +1111,7 @@ void FOutputDeviceFile::WriteDataToArchive(const TCHAR* Data, ELogVerbosity::Typ
  * @param	Data	Text to log
  * @param	Event	Event name used for suppression purposes
  */
-void FOutputDeviceFile::Serialize( const TCHAR* Data, ELogVerbosity::Type Verbosity, const class FName& Category )
+void FOutputDeviceFile::Serialize( const TCHAR* Data, ELogVerbosity::Type Verbosity, const class FName& Category, const double Time )
 {
 #if ALLOW_LOG_FILE && !NO_LOGGING
 	static bool Entry=false;
@@ -1146,7 +1153,7 @@ void FOutputDeviceFile::Serialize( const TCHAR* Data, ELogVerbosity::Type Verbos
 
 		if( LogAr && Verbosity != ELogVerbosity::SetColor )
 		{
-			WriteDataToArchive(Data, Verbosity, Category);
+			WriteDataToArchive(Data, Verbosity, Category, Time);
 
 			static bool GForceLogFlush = false;
 			static bool GTestedCmdLine = false;
@@ -1165,12 +1172,16 @@ void FOutputDeviceFile::Serialize( const TCHAR* Data, ELogVerbosity::Type Verbos
 	else
 	{
 		Entry=true;
-		Serialize( Data, Verbosity, Category );
+		Serialize( Data, Verbosity, Category, Time );
 		Entry=false;
 	}
 #endif
 }
 
+void FOutputDeviceFile::Serialize( const TCHAR* Data, ELogVerbosity::Type Verbosity, const class FName& Category )
+{
+	Serialize( Data, Verbosity, Category, -1.0 );
+}
 
 
 void FOutputDeviceFile::WriteRaw( const TCHAR* C )
@@ -1184,14 +1195,14 @@ void FOutputDeviceFile::WriteRaw( const TCHAR* C )
  * @param	Data	Text to log
  * @param	Event	Event name used for suppression purposes
  */
-void FOutputDeviceDebug::Serialize( const TCHAR* Data, ELogVerbosity::Type Verbosity, const class FName& Category )
+void FOutputDeviceDebug::Serialize( const TCHAR* Data, ELogVerbosity::Type Verbosity, const class FName& Category, const double Time )
 {
 	static bool Entry=false;
 	if( !GIsCriticalError || Entry )
 	{
 		if (Verbosity != ELogVerbosity::SetColor)
 		{
-			FPlatformMisc::LowLevelOutputDebugStringf(TEXT("%s%s"),*FOutputDevice::FormatLogLine(Verbosity, Category, Data, GPrintLogTimes),LINE_TERMINATOR);
+			FPlatformMisc::LowLevelOutputDebugStringf(TEXT("%s%s"),*FOutputDevice::FormatLogLine(Verbosity, Category, Data, GPrintLogTimes, Time),LINE_TERMINATOR);
 		}
 	}
 	else
@@ -1202,6 +1213,11 @@ void FOutputDeviceDebug::Serialize( const TCHAR* Data, ELogVerbosity::Type Verbo
 	}
 }
 
+void FOutputDeviceDebug::Serialize( const TCHAR* Data, ELogVerbosity::Type Verbosity, const class FName& Category )
+{
+	Serialize( Data, Verbosity, Category, -1.0 );
+}
+
 /*-----------------------------------------------------------------------------
 	FOutputDeviceError subclasses.
 -----------------------------------------------------------------------------*/
diff --git a/Engine/Source/Runtime/Core/Private/Misc/OutputDevice.cpp b/Engine/Source/Runtime/Core/Private/Misc/OutputDevice.cpp
--- a/Engine/Source/Runtime/Core/Private/Misc/OutputDevice.cpp
+++ b/Engine/Source/Runtime/Core/Private/Misc/OutputDevice.cpp
@@ -45,14 +45,17 @@ const TCHAR* FOutputDevice::VerbosityToString(ELogVerbosity::Type Verbosity)
 	return TEXT("UknownVerbosity");
 }
 
-FString FOutputDevice::FormatLogLine(ELogVerbosity::Type Verbosity, const class FName& Category, const TCHAR* Message, ELogTimes::Type LogTime)
+FString FOutputDevice::FormatLogLine( ELogVerbosity::Type Verbosity, const class FName& Category, const TCHAR* Message /*= nullptr*/, ELogTimes::Type LogTime /*= ELogTimes::None*/, const double Time /*= -1.0*/ )
 {
 	FString Format;
 	switch (LogTime)
 	{
 		case ELogTimes::SinceGStartTime:
-			Format = FString::Printf(TEXT("[%07.2f][%3d]"), FPlatformTime::Seconds() - GStartTime, GFrameCounter % 1000);
+		{
+			const double RealTime = Time == -1.0f ? FPlatformTime::Seconds() - GStartTime : Time;
+			Format = FString::Printf( TEXT( "[%07.2f][%3d]" ), RealTime, GFrameCounter % 1000 );
 			break;
+		}
 
 		case ELogTimes::UTC:
 			Format = FString::Printf(TEXT("[%s][%3d]"), *FDateTime::UtcNow().ToString(TEXT("%Y.%m.%d-%H.%M.%S:%s")), GFrameCounter % 1000);
diff --git a/Engine/Source/Runtime/Core/Private/Windows/WindowsPlatformOutputDevices.cpp b/Engine/Source/Runtime/Core/Private/Windows/WindowsPlatformOutputDevices.cpp
--- a/Engine/Source/Runtime/Core/Private/Windows/WindowsPlatformOutputDevices.cpp
+++ b/Engine/Source/Runtime/Core/Private/Windows/WindowsPlatformOutputDevices.cpp
@@ -251,10 +251,12 @@ bool FOutputDeviceConsoleWindows::IsShown()
 	return ConsoleHandle != NULL;
 }
 
-void FOutputDeviceConsoleWindows::Serialize( const TCHAR* Data, ELogVerbosity::Type Verbosity, const class FName& Category )
+void FOutputDeviceConsoleWindows::Serialize( const TCHAR* Data, ELogVerbosity::Type Verbosity, const class FName& Category, const double Time )
 {
 	if( ConsoleHandle )
 	{
+		const double RealTime = Time == -1.0f ? FPlatformTime::Seconds() - GStartTime : Time;
+
 		static bool Entry=false;
 		if( !GIsCriticalError || Entry )
 		{
@@ -280,7 +282,7 @@ void FOutputDeviceConsoleWindows::Serialize( const TCHAR* Data, ELogVerbosity::T
 					}
 				}
 				TCHAR OutputString[MAX_SPRINTF]=TEXT(""); //@warning: this is safe as FCString::Sprintf only use 1024 characters max
-				FCString::Sprintf(OutputString,TEXT("%s%s"),*FOutputDevice::FormatLogLine(Verbosity, Category, Data, GPrintLogTimes),LINE_TERMINATOR);
+				FCString::Sprintf(OutputString,TEXT("%s%s"),*FOutputDevice::FormatLogLine(Verbosity, Category, Data, GPrintLogTimes,RealTime),LINE_TERMINATOR);
 				uint32 Written;
 				WriteConsole( ConsoleHandle, OutputString, FCString::Strlen(OutputString), (::DWORD*)&Written, NULL );
 
@@ -299,6 +301,11 @@ void FOutputDeviceConsoleWindows::Serialize( const TCHAR* Data, ELogVerbosity::T
 	}
 }
 
+void FOutputDeviceConsoleWindows::Serialize( const TCHAR* Data, ELogVerbosity::Type Verbosity, const class FName& Category )
+{
+	Serialize( Data, Verbosity, Category, -1.0 );
+}
+
 void FOutputDeviceConsoleWindows::SetColor( const TCHAR* Color )
 {
 	// here we can change the color of the text to display, it's in the format:
diff --git a/Engine/Source/Runtime/Core/Private/Windows/WindowsPlatformOutputDevicesPrivate.h b/Engine/Source/Runtime/Core/Private/Windows/WindowsPlatformOutputDevicesPrivate.h
--- a/Engine/Source/Runtime/Core/Private/Windows/WindowsPlatformOutputDevicesPrivate.h
+++ b/Engine/Source/Runtime/Core/Private/Windows/WindowsPlatformOutputDevicesPrivate.h
@@ -185,7 +185,8 @@ public:
 	 * @param Data	Text to display
 	 * @param Event	Event type, used for filtering/ suppression
 	 */
-	void Serialize( const TCHAR* Data, ELogVerbosity::Type Verbosity, const class FName& Category );
+	virtual void Serialize( const TCHAR* Data, ELogVerbosity::Type Verbosity, const class FName& Category, const double Time ) override;
+	virtual void Serialize( const TCHAR* Data, ELogVerbosity::Type Verbosity, const class FName& Category ) override;
 
 	void SetColor( const TCHAR* Color );
 };
diff --git a/Engine/Source/Runtime/Core/Public/HAL/OutputDevices.h b/Engine/Source/Runtime/Core/Public/HAL/OutputDevices.h
--- a/Engine/Source/Runtime/Core/Public/HAL/OutputDevices.h
+++ b/Engine/Source/Runtime/Core/Public/HAL/OutputDevices.h
@@ -249,14 +249,16 @@ public:
 struct FBufferedLine
 {
 	const FString Data;
-	const ELogVerbosity::Type Verbosity;
 	const FName Category;
+	const double Time;
+	const ELogVerbosity::Type Verbosity;
 
 	/** Initialization constructor. */
-	FBufferedLine( const TCHAR* InData, ELogVerbosity::Type InVerbosity, const class FName& InCategory )
+	FBufferedLine( const TCHAR* InData, const class FName& InCategory, ELogVerbosity::Type InVerbosity, const double InTime = -1 )
 		: Data( InData )
-		, Verbosity( InVerbosity )
 		, Category( InCategory )
+		, Time( InTime )
+		, Verbosity( InVerbosity )
 	{}
 };
 
@@ -356,7 +358,15 @@ public:
 	 * @param	Data	Text to log
 	 * @param	Event	Event name used for suppression purposes
 	 */
-	virtual void Serialize( const TCHAR* Data, ELogVerbosity::Type Verbosity, const class FName& Category) override;
+	virtual void Serialize( const TCHAR* Data, ELogVerbosity::Type Verbosity, const class FName& Category, const double Time) override;
+
+	/**
+	 * Serializes the passed in data via all current output devices.
+	 *
+	 * @param	Data	Text to log
+	 * @param	Event	Event name used for suppression purposes
+	 */
+	virtual void Serialize( const TCHAR* Data, ELogVerbosity::Type Verbosity, const class FName& Category ) override;
 
 	/**
 	 * Passes on the flush request to all current output devices.
@@ -408,7 +418,9 @@ public:
 	 */
 	void Flush() override;
 
-	virtual void Serialize( const TCHAR* Data, ELogVerbosity::Type Verbosity, const class FName& Category ) override;
+	virtual void Serialize( const TCHAR* Data, ELogVerbosity::Type Verbosity, const class FName& Category, const double Time ) override;
+
+	virtual void Serialize( const TCHAR* Data, ELogVerbosity::Type Verbosity, const class FName& Category) override;
 
 	virtual bool CanBeUsedOnAnyThread() const override
 	{
@@ -432,7 +444,7 @@ private:
 
 	void CastAndSerializeData(const TCHAR* Data);
 
-	void WriteDataToArchive(const TCHAR* Data, ELogVerbosity::Type Verbosity, const class FName& Category);
+	void WriteDataToArchive(const TCHAR* Data, ELogVerbosity::Type Verbosity, const class FName& Category, const double Time);
 };
 
 // Null output device.
@@ -458,6 +470,8 @@ public:
 	 * @param	Data	Text to log
 	 * @param	Event	Event name used for suppression purposes
 	 */
+	virtual void Serialize( const TCHAR* Data, ELogVerbosity::Type Verbosity, const class FName& Category, const double Time ) override;
+
 	virtual void Serialize( const TCHAR* Data, ELogVerbosity::Type Verbosity, const class FName& Category ) override;
 
 	virtual bool CanBeUsedOnAnyThread() const override
@@ -474,7 +488,7 @@ class FBufferedOutputDevice : public FOutputDevice
 public:
 	virtual void Serialize( const TCHAR* InData, ELogVerbosity::Type Verbosity, const class FName& Category ) override
 	{
-		new(BufferedLines)FBufferedLine( InData, Verbosity, Category );
+		new(BufferedLines)FBufferedLine( InData, Category, Verbosity );
 	}
 
 	/** Pushes buffered lines into the specified output device. */
diff --git a/Engine/Source/Runtime/Core/Public/Misc/OutputDevice.h b/Engine/Source/Runtime/Core/Public/Misc/OutputDevice.h
--- a/Engine/Source/Runtime/Core/Public/Misc/OutputDevice.h
+++ b/Engine/Source/Runtime/Core/Public/Misc/OutputDevice.h
@@ -294,11 +294,16 @@ public:
 
 	// static helpers
 	static const TCHAR* VerbosityToString(ELogVerbosity::Type Verbosity);
-	static FString FormatLogLine(ELogVerbosity::Type Verbosity, const class FName& Category, const TCHAR* Message = nullptr, ELogTimes::Type LogTime = ELogTimes::None);
+	static FString FormatLogLine(ELogVerbosity::Type Verbosity, const class FName& Category, const TCHAR* Message = nullptr, ELogTimes::Type LogTime = ELogTimes::None, const double Time = -1.0);
 
 
 	// FOutputDevice interface.
-	virtual void Serialize( const TCHAR* V, ELogVerbosity::Type Verbosity, const class FName& Category )=0;
+	virtual void Serialize( const TCHAR* V, ELogVerbosity::Type Verbosity, const class FName& Category ) = 0;
+	virtual void Serialize( const TCHAR* V, ELogVerbosity::Type Verbosity, const class FName& Category, const double Time )
+	{
+		Serialize( V, Verbosity, Category );
+	}
+
 	virtual void Flush()
 	{
 	}
diff --git a/Engine/Source/Runtime/CoreUObject/Private/Misc/TextBuffer.cpp b/Engine/Source/Runtime/CoreUObject/Private/Misc/TextBuffer.cpp
--- a/Engine/Source/Runtime/CoreUObject/Private/Misc/TextBuffer.cpp
+++ b/Engine/Source/Runtime/CoreUObject/Private/Misc/TextBuffer.cpp
@@ -32,6 +32,10 @@ void UTextBuffer::Serialize (const TCHAR* Data, ELogVerbosity::Type Verbosity, c
 	Text += (TCHAR*)Data;
 }
 
+void UTextBuffer::Serialize( const TCHAR* Data, ELogVerbosity::Type Verbosity, const class FName& Category, const double Time )
+{
+	Text += (TCHAR*)Data;
+}
 
 /* UObject interface
  *****************************************************************************/
diff --git a/Engine/Source/Runtime/CoreUObject/Public/Misc/TextBuffer.h b/Engine/Source/Runtime/CoreUObject/Public/Misc/TextBuffer.h
--- a/Engine/Source/Runtime/CoreUObject/Public/Misc/TextBuffer.h
+++ b/Engine/Source/Runtime/CoreUObject/Public/Misc/TextBuffer.h
@@ -49,6 +49,7 @@ public:
 
 	virtual void Serialize (FArchive& Ar) override;
 	virtual void Serialize (const TCHAR* Data, ELogVerbosity::Type Verbosity, const class FName& Category ) override;
+	virtual void Serialize( const TCHAR* V, ELogVerbosity::Type Verbosity, const class FName& Category, const double Time ) override;
 
 private:
 
