PR #1395: Windows Stack Walking improvement (Contributed by [REDACTED_PERSON])

diff --git a/Engine/Source/Runtime/Core/Private/Windows/WindowsPlatformStackWalk.cpp b/Engine/Source/Runtime/Core/Private/Windows/WindowsPlatformStackWalk.cpp
--- a/Engine/Source/Runtime/Core/Private/Windows/WindowsPlatformStackWalk.cpp
+++ b/Engine/Source/Runtime/Core/Private/Windows/WindowsPlatformStackWalk.cpp
@@ -570,7 +570,7 @@ bool FWindowsPlatformStackWalk::InitStackWalking()
 
 		// Initialize the symbol engine.
 		SymInitialize( GetCurrentProcess(), NULL, true );
-		LoadProcessModules();
+		//LoadProcessModules();
 	
 		GNeedToRefreshSymbols = false;
 		GStackWalkingInitialized = true;
@@ -580,7 +580,7 @@ bool FWindowsPlatformStackWalk::InitStackWalking()
 	{
 		// Refresh and reload symbols
 		SymRefreshModuleList( GetCurrentProcess() );
-		LoadProcessModules();
+		//LoadProcessModules();
 		GNeedToRefreshSymbols = false;
 	}
 #endif
