[NO NON-PII COMMIT MESSAGE]

diff --git a/Engine/Source/Runtime/Core/Private/Windows/WindowsPlatformCrashContext.cpp b/Engine/Source/Runtime/Core/Private/Windows/WindowsPlatformCrashContext.cpp
--- a/Engine/Source/Runtime/Core/Private/Windows/WindowsPlatformCrashContext.cpp
+++ b/Engine/Source/Runtime/Core/Private/Windows/WindowsPlatformCrashContext.cpp
@@ -32,7 +32,7 @@ static int32 ReportCrashCallCount = 0;
  */
 
 // @TODO yrx 2014-10-08 Move to FWindowsPlatformCrashContext
-bool WriteMinidump(const TCHAR* Path, LPEXCEPTION_POINTERS ExceptionInfo)
+bool WriteMinidump( const TCHAR* Path, LPEXCEPTION_POINTERS ExceptionInfo, bool bIsEnsure )
 {
 	// Try to create file for minidump.
 	HANDLE FileHandle = CreateFileW(Path, GENERIC_WRITE, 0, NULL, CREATE_ALWAYS, FILE_ATTRIBUTE_NORMAL, NULL);
@@ -63,7 +63,10 @@ bool WriteMinidump(const TCHAR* Path, LPEXCEPTION_POINTERS ExceptionInfo)
 	CrashContextStreamInformation.UserStreamArray = &CrashContextStream;
 
 	MINIDUMP_TYPE MinidumpType = MiniDumpNormal;//(MINIDUMP_TYPE)(MiniDumpWithPrivateReadWriteMemory|MiniDumpWithDataSegs|MiniDumpWithHandleData|MiniDumpWithFullMemoryInfo|MiniDumpWithThreadInfo|MiniDumpWithUnloadedModules);
-	if (CrashContext.IsFullCrashDump())
+
+	// For ensures use minidump to avoid severe hitches when writing 3GB+ files.
+	// However the crash dump mode will remain the same.
+	if (CrashContext.IsFullCrashDump() && bIsEnsure == false)
 	{
 		MinidumpType = (MINIDUMP_TYPE)(MiniDumpWithFullMemory|MiniDumpWithFullMemoryInfo|MiniDumpWithHandleData|MiniDumpWithThreadInfo|MiniDumpWithUnloadedModules);
 	}
@@ -258,7 +261,7 @@ private:
  * Create a Windows Error Report, add the user log and video, and add it to the WER queue
  * Launch CrashReportClient.exe to intercept the report and upload to our local site
  */
-int32 ReportCrashUsingCrashReportClient(EXCEPTION_POINTERS* ExceptionInfo, const TCHAR* ErrorMessage, EErrorReportUI ReportUI)
+int32 ReportCrashUsingCrashReportClient(EXCEPTION_POINTERS* ExceptionInfo, const TCHAR* ErrorMessage, EErrorReportUI ReportUI, bool bIsEnsure)
 {
 	// Flush out the log
 	GLog->Flush();
@@ -302,7 +305,7 @@ int32 ReportCrashUsingCrashReportClient(EXCEPTION_POINTERS* ExceptionInfo, const
 				WerReportAddFile( ReportHandle, *CrashContextXMLPath, WerFileTypeOther, WER_FILE_ANONYMOUS_DATA );
 
 				const FString MinidumpFileName = FPaths::Combine( *FPaths::GameLogDir(), *CrashContext.GetUniqueCrashName(), *FGenericCrashContext::UE4MinidumpName );
-				if( WriteMinidump( *MinidumpFileName, ExceptionInfo ) )
+				if (WriteMinidump( *MinidumpFileName, ExceptionInfo, bIsEnsure ))
 				{
 					WerReportAddFile( ReportHandle, *MinidumpFileName, WerFileTypeMinidump, WER_FILE_ANONYMOUS_DATA );
 				}
@@ -394,7 +397,7 @@ void NewReportEnsure( const TCHAR* ErrorMessage )
 		FPlatformMisc::RaiseException( 1 );
 	}
 #if !PLATFORM_SEH_EXCEPTIONS_DISABLED
-	__except(ReportCrashUsingCrashReportClient(GetExceptionInformation(), ErrorMessage, EErrorReportUI::ReportInUnattendedMode))
+	__except(ReportCrashUsingCrashReportClient(GetExceptionInformation(), ErrorMessage, EErrorReportUI::ReportInUnattendedMode, true))
 	CA_SUPPRESS(6322)
 	{
 	}
@@ -477,12 +480,12 @@ int32 ReportCrash( LPEXCEPTION_POINTERS ExceptionInfo )
 #if WINVER > 0x502	// Windows Error Reporting is not supported on Windows XP
 	if (GUseCrashReportClient)
 	{
-		ReportCrashUsingCrashReportClient( ExceptionInfo, GErrorMessage, EErrorReportUI::ShowDialog );
+		ReportCrashUsingCrashReportClient( ExceptionInfo, GErrorMessage, EErrorReportUI::ShowDialog, false );
 	}
 	else
 #endif		// WINVER
 	{
-		WriteMinidump( MiniDumpFilenameW, ExceptionInfo );
+		WriteMinidump( MiniDumpFilenameW, ExceptionInfo, false );
 
 #if UE_BUILD_SHIPPING && WITH_EDITOR
 		uint32 dwOpt = 0;
