Distcc compilation fixes for OS X.

diff --git a/Engine/Source/Runtime/Core/Public/Misc/AutomationTest.h b/Engine/Source/Runtime/Core/Public/Misc/AutomationTest.h
--- a/Engine/Source/Runtime/Core/Public/Misc/AutomationTest.h
+++ b/Engine/Source/Runtime/Core/Public/Misc/AutomationTest.h
@@ -1056,7 +1056,7 @@ public: \
 	public: \
 		TClass( const FString& InName ) \
 		:FAutomationTestBase( InName, false ) {\
-			static_assert(TFlags&EAutomationTestFlags::ApplicationContextMask, "AutomationTest has no application flag.  It shouldn't run.  See AutomationTest.h."); \
+			static_assert((TFlags)&EAutomationTestFlags::ApplicationContextMask, "AutomationTest has no application flag.  It shouldn't run.  See AutomationTest.h."); \
 			static_assert(	(((TFlags)&EAutomationTestFlags::FilterMask) == EAutomationTestFlags::SmokeFilter) || \
 							(((TFlags)&EAutomationTestFlags::FilterMask) == EAutomationTestFlags::EngineFilter) || \
 							(((TFlags)&EAutomationTestFlags::FilterMask) == EAutomationTestFlags::ProductFilter) || \
@@ -1083,7 +1083,7 @@ public: \
 	public: \
 		TClass( const FString& InName ) \
 		:FAutomationTestBase( InName, true ) { \
-			static_assert(TFlags&EAutomationTestFlags::ApplicationContextMask, "AutomationTest has no application flag.  It shouldn't run.  See AutomationTest.h."); \
+			static_assert((TFlags)&EAutomationTestFlags::ApplicationContextMask, "AutomationTest has no application flag.  It shouldn't run.  See AutomationTest.h."); \
 			static_assert(	(((TFlags)&EAutomationTestFlags::FilterMask) == EAutomationTestFlags::SmokeFilter) || \
 							(((TFlags)&EAutomationTestFlags::FilterMask) == EAutomationTestFlags::EngineFilter) || \
 							(((TFlags)&EAutomationTestFlags::FilterMask) == EAutomationTestFlags::ProductFilter) || \
@@ -1091,7 +1091,7 @@ public: \
 							(((TFlags)&EAutomationTestFlags::FilterMask) == EAutomationTestFlags::StressFilter), \
 							"All AutomationTests must have exactly 1 filter type specified.  See AutomationTest.h."); \
 		} \
-		virtual uint32 GetTestFlags() const override { return (TFlags & ~(EAutomationTestFlags::SmokeFilter)); } \
+		virtual uint32 GetTestFlags() const override { return ((TFlags) & ~(EAutomationTestFlags::SmokeFilter)); } \
 		virtual bool IsStressTest() const { return true; } \
 		virtual uint32 GetRequiredDeviceNum() const override { return 1; } \
 	protected: \
@@ -1106,7 +1106,7 @@ public: \
 	public: \
 		TClass( const FString& InName ) \
 		:FAutomationTestBase( InName, false ) { \
-			static_assert(TFlags&EAutomationTestFlags::ApplicationContextMask, "AutomationTest has no application flag.  It shouldn't run.  See AutomationTest.h."); \
+			static_assert((TFlags)&EAutomationTestFlags::ApplicationContextMask, "AutomationTest has no application flag.  It shouldn't run.  See AutomationTest.h."); \
 			static_assert(	(((TFlags)&EAutomationTestFlags::FilterMask) == EAutomationTestFlags::SmokeFilter) || \
 							(((TFlags)&EAutomationTestFlags::FilterMask) == EAutomationTestFlags::EngineFilter) || \
 							(((TFlags)&EAutomationTestFlags::FilterMask) == EAutomationTestFlags::ProductFilter) || \
@@ -1114,7 +1114,7 @@ public: \
 							(((TFlags)&EAutomationTestFlags::FilterMask) == EAutomationTestFlags::StressFilter), \
 							"All AutomationTests must have exactly 1 filter type specified.  See AutomationTest.h."); \
 		} \
-		virtual uint32 GetTestFlags() const override { return (TFlags & ~(EAutomationTestFlags::EditorContext | EAutomationTestFlags::CommandletContext | EAutomationTestFlags::SmokeFilter)); } \
+		virtual uint32 GetTestFlags() const override { return ((TFlags) & ~(EAutomationTestFlags::EditorContext | EAutomationTestFlags::CommandletContext | EAutomationTestFlags::SmokeFilter)); } \
 		virtual uint32 GetRequiredDeviceNum() const override { return NumParticipants; } \
 	protected: \
 		virtual void GetTests(TArray<FString>& OutBeautifiedNames, TArray <FString>& OutTestCommands) const override \
diff --git a/Engine/Source/Runtime/Engine/Private/Tests/ExternalToolTest.cpp b/Engine/Source/Runtime/Engine/Private/Tests/ExternalToolTest.cpp
--- a/Engine/Source/Runtime/Engine/Private/Tests/ExternalToolTest.cpp
+++ b/Engine/Source/Runtime/Engine/Private/Tests/ExternalToolTest.cpp
@@ -77,7 +77,7 @@ bool FWaitForProcessToCompleteLatentCommand::Update()
 /**
  * RunExternalToolTest
  */
-IMPLEMENT_COMPLEX_AUTOMATION_TEST(FRunExternalToolTest, "External", EAutomationTestFlags::NonNullRHI | EAutomationTestFlags::EngineFilter)
+IMPLEMENT_COMPLEX_AUTOMATION_TEST(FRunExternalToolTest, "External", EAutomationTestFlags::EditorContext | EAutomationTestFlags::NonNullRHI | EAutomationTestFlags::EngineFilter)
 
 /** 
  * Find all the external too commands
diff --git a/Engine/Source/Runtime/SlateCore/Public/Stats/SlateStats.h b/Engine/Source/Runtime/SlateCore/Public/Stats/SlateStats.h
--- a/Engine/Source/Runtime/SlateCore/Public/Stats/SlateStats.h
+++ b/Engine/Source/Runtime/SlateCore/Public/Stats/SlateStats.h
@@ -9,7 +9,7 @@
 	#define WITH_SLATE_STATS 0
 #endif
 
-#define SLATE_STATS (WITH_SLATE_STATS || WITH_EDITOR)
+#define SLATE_STATS (WITH_SLATE_STATS | WITH_EDITOR)
 
 /** 
  * Predefined detail levels to make choosing level slightly easier.
