Core: Added a basic unit test for the task graph

diff --git a/Engine/Source/Runtime/Core/Private/Tests/Async/TaskGraphTest.cpp b/Engine/Source/Runtime/Core/Private/Tests/Async/TaskGraphTest.cpp
new file mode 100644
--- /dev/null
+++ b/Engine/Source/Runtime/Core/Private/Tests/Async/TaskGraphTest.cpp
@@ -0,0 +1,69 @@
+// Copyright 1998-2015 Epic Games, Inc. All Rights Reserved.
+
+#include "CorePrivatePCH.h"
+#include "AutomationTest.h"
+#include "TaskGraphInterfaces.h"
+
+
+IMPLEMENT_SIMPLE_AUTOMATION_TEST(FTaskGraphTest, "System.Core.Async.TaskGraph", EAutomationTestFlags::ApplicationContextMask | EAutomationTestFlags::EngineFilter)
+
+
+namespace TaskGraphTestTask
+{
+	const FTimespan MaxWaitTime(0, 0, 5);
+	const int64 NumTasks = 10000;
+	int64 CompletedTasks = 0;
+}
+
+
+/**
+ * Implements a test task.
+ */
+class FTaskGraphTestTask
+{
+public:
+
+	void DoTask(ENamedThreads::Type CurrentThread, const FGraphEventRef& MyCompletionGraphEvent)
+	{
+		FPlatformAtomics::InterlockedIncrement(&TaskGraphTestTask::CompletedTasks);
+	}
+
+	ENamedThreads::Type GetDesiredThread()
+	{
+		return ENamedThreads::AnyThread;
+	}
+
+	TStatId GetStatId() const
+	{
+		RETURN_QUICK_DECLARE_CYCLE_STAT(FTaskGraphTestTask, STATGROUP_TaskGraphTasks);
+	}
+
+	static ESubsequentsMode::Type GetSubsequentsMode()
+	{
+		return ESubsequentsMode::FireAndForget; 
+	}
+};
+
+
+bool FTaskGraphTest::RunTest(const FString& Parameters)
+{
+	using namespace TaskGraphTestTask;
+
+	CompletedTasks = 0;
+
+	for (int32 TaskIndex = 0; TaskIndex < TaskGraphTestTask::NumTasks; ++TaskIndex)
+	{
+		TGraphTask<FTaskGraphTestTask>::CreateTask().ConstructAndDispatchWhenReady();
+	}
+
+	FDateTime StartTime = FDateTime::UtcNow();
+
+	while ((CompletedTasks < NumTasks) &&((FDateTime::UtcNow() - StartTime) < MaxWaitTime))
+	{
+		FPlatformProcess::Sleep(0.0f);
+	}
+
+	TestEqual(TEXT("The number of completed tasks must equal the total number of tasks"), CompletedTasks, NumTasks);
+
+	return true;
+}
