Core: Better names for Async unit tests

diff --git a/Engine/Source/Runtime/Core/Private/Tests/Async/AsyncTest.cpp b/Engine/Source/Runtime/Core/Private/Tests/Async/AsyncTest.cpp
--- a/Engine/Source/Runtime/Core/Private/Tests/Async/AsyncTest.cpp
+++ b/Engine/Source/Runtime/Core/Private/Tests/Async/AsyncTest.cpp
@@ -7,14 +7,13 @@
 #include "Future.h"
 
 
-IMPLEMENT_SIMPLE_AUTOMATION_TEST(FAsyncGraphTest, "System.Core.Async.Task Graph", EAutomationTestFlags::EditorContext | EAutomationTestFlags::EngineFilter)
-IMPLEMENT_SIMPLE_AUTOMATION_TEST(FAsyncThreadedTaskTest, "System.Core.Async.Threaded Task", EAutomationTestFlags::EditorContext | EAutomationTestFlags::EngineFilter)
-IMPLEMENT_SIMPLE_AUTOMATION_TEST(FAsyncThreadedPoolTest, "System.Core.Async.Threaded Pool", EAutomationTestFlags::EditorContext | EAutomationTestFlags::EngineFilter)
-IMPLEMENT_SIMPLE_AUTOMATION_TEST(FAsyncVoidTaskTest, "System.Core.Async.Void Task", EAutomationTestFlags::EditorContext | EAutomationTestFlags::EngineFilter)
+IMPLEMENT_SIMPLE_AUTOMATION_TEST(FAsyncGraphTest, "System.Core.Async.Async (Task Graph)", EAutomationTestFlags::EditorContext | EAutomationTestFlags::EngineFilter)
+IMPLEMENT_SIMPLE_AUTOMATION_TEST(FAsyncThreadedTaskTest, "System.Core.Async.Async (Thread)", EAutomationTestFlags::EditorContext | EAutomationTestFlags::EngineFilter)
+IMPLEMENT_SIMPLE_AUTOMATION_TEST(FAsyncThreadedPoolTest, "System.Core.Async.Async (Thread Pool)", EAutomationTestFlags::EditorContext | EAutomationTestFlags::EngineFilter)
+IMPLEMENT_SIMPLE_AUTOMATION_TEST(FAsyncVoidTaskTest, "System.Core.Async.Async (Void)", EAutomationTestFlags::EditorContext | EAutomationTestFlags::EngineFilter)
 
-/**
-* Helper methods used in the test cases
-*/
+
+/** Helper methods used in the test cases. */
 namespace AsyncTestUtils
 {
 	TFunction<int()> Task = []() {
@@ -28,9 +27,8 @@ namespace AsyncTestUtils
 	};
 }
 
-/**
-* Test that task graph tasks return correctly.
-*/
+
+/** Test that task graph tasks return correctly. */
 bool FAsyncGraphTest::RunTest(const FString& Parameters)
 {
 	auto Future = Async(EAsyncExecution::TaskGraph, AsyncTestUtils::Task);
@@ -41,9 +39,8 @@ bool FAsyncGraphTest::RunTest(const FString& Parameters)
 	return true;
 }
 
-/**
-* Test that threaded tasks return correctly.
-*/
+
+/** Test that threaded tasks return correctly. */
 bool FAsyncThreadedTaskTest::RunTest(const FString& Parameters)
 {
 	auto Future = Async(EAsyncExecution::Thread, AsyncTestUtils::Task);
@@ -54,6 +51,7 @@ bool FAsyncThreadedTaskTest::RunTest(const FString& Parameters)
 	return true;
 }
 
+
 /**
 * Test that threaded pool tasks return correctly.
 */
@@ -67,9 +65,8 @@ bool FAsyncThreadedPoolTest::RunTest(const FString& Parameters)
 	return true;
 }
 
-/**
-* Test that void tasks run without errors or warnings.
-*/
+
+/** Test that void tasks run without errors or warnings. */
 bool FAsyncVoidTaskTest::RunTest(const FString& Parameters)
 {
 	// Reset test variable before running
