Expose FProcEnumerator API, and prevent potential bugged use of FProcEnumerator (through assigning a direct call to the constructor, triggering a copy).

diff --git a/Engine/Source/Programs/UnrealSync/Private/GUI.cpp b/Engine/Source/Programs/UnrealSync/Private/GUI.cpp
--- a/Engine/Source/Programs/UnrealSync/Private/GUI.cpp
+++ b/Engine/Source/Programs/UnrealSync/Private/GUI.cpp
@@ -2067,7 +2067,7 @@ private:
 			PossibleNormalizedPaths.Add(MoveTemp(PossibleProcImagePath));
 		}
 
-		auto ProcessEnumerator = FPlatformProcess::FProcEnumerator();
+		FPlatformProcess::FProcEnumerator ProcessEnumerator;
 
 		while (ProcessEnumerator.MoveNext())
 		{
diff --git a/Engine/Source/Runtime/Core/Public/Linux/LinuxPlatformProcess.h b/Engine/Source/Runtime/Core/Public/Linux/LinuxPlatformProcess.h
--- a/Engine/Source/Runtime/Core/Public/Linux/LinuxPlatformProcess.h
+++ b/Engine/Source/Runtime/Core/Public/Linux/LinuxPlatformProcess.h
@@ -190,11 +190,13 @@ struct CORE_API FLinuxPlatformProcess : public FGenericPlatformProcess
 	/**
 	 * Process enumerator.
 	 */
-	class FProcEnumerator
+	class CORE_API FProcEnumerator
 	{
 	public:
 		// Constructor
 		FProcEnumerator();
+		FProcEnumerator(const FProcEnumerator&) = delete;
+		FProcEnumerator& operator=(const FProcEnumerator&) = delete;
 
 		// Destructor
 		~FProcEnumerator();
@@ -216,7 +218,7 @@ struct CORE_API FLinuxPlatformProcess : public FGenericPlatformProcess
 	/**
 	 * Process enumeration info structure.
 	 */
-	struct FProcEnumInfo
+	struct CORE_API FProcEnumInfo
 	{
 		friend FLinuxPlatformProcess::FProcEnumerator::FProcEnumerator();
 
diff --git a/Engine/Source/Runtime/Core/Public/Mac/MacPlatformProcess.h b/Engine/Source/Runtime/Core/Public/Mac/MacPlatformProcess.h
--- a/Engine/Source/Runtime/Core/Public/Mac/MacPlatformProcess.h
+++ b/Engine/Source/Runtime/Core/Public/Mac/MacPlatformProcess.h
@@ -33,11 +33,14 @@ struct CORE_API FMacPlatformProcess : public FGenericPlatformProcess
 	/**
 	 * Process enumerator.
 	 */
-	class FProcEnumerator
+	class CORE_API FProcEnumerator
 	{
 	public:
 		// Constructor
 		FProcEnumerator();
+		FProcEnumerator(const FProcEnumerator&) = delete;
+		FProcEnumerator& operator=(const FProcEnumerator&) = delete;
+
 		// Destructor
 		~FProcEnumerator();
 
@@ -65,7 +68,7 @@ struct CORE_API FMacPlatformProcess : public FGenericPlatformProcess
 	/**
 	 * Process enumeration info structure.
 	 */
-	struct FProcEnumInfo
+	struct CORE_API FProcEnumInfo
 	{
 		friend FProcEnumInfo FMacPlatformProcess::FProcEnumerator::GetCurrent() const;
 
diff --git a/Engine/Source/Runtime/Core/Public/Windows/WindowsPlatformProcess.h b/Engine/Source/Runtime/Core/Public/Windows/WindowsPlatformProcess.h
--- a/Engine/Source/Runtime/Core/Public/Windows/WindowsPlatformProcess.h
+++ b/Engine/Source/Runtime/Core/Public/Windows/WindowsPlatformProcess.h
@@ -62,11 +62,13 @@ struct CORE_API FWindowsPlatformProcess
 	/**
 	 * Process enumerator.
 	 */
-	class FProcEnumerator
+	class CORE_API FProcEnumerator
 	{
 	public:
 		// Constructor
 		FProcEnumerator();
+		FProcEnumerator(const FProcEnumerator&) = delete;
+		FProcEnumerator& operator=(const FProcEnumerator&) = delete;
 
 		// Destructor
 		~FProcEnumerator();
@@ -91,7 +93,7 @@ struct CORE_API FWindowsPlatformProcess
 	/**
 	 * Process enumeration info structure.
 	 */
-	struct FProcEnumInfo
+	struct CORE_API FProcEnumInfo
 	{
 		friend FProcEnumInfo FProcEnumerator::GetCurrent() const;
 	public:
