Fixed WinXP build

diff --git a/Engine/Source/Runtime/Core/Private/Tests/Async/TaskGraphTest.cpp b/Engine/Source/Runtime/Core/Private/Tests/Async/TaskGraphTest.cpp
--- a/Engine/Source/Runtime/Core/Private/Tests/Async/TaskGraphTest.cpp
+++ b/Engine/Source/Runtime/Core/Private/Tests/Async/TaskGraphTest.cpp
@@ -11,8 +11,8 @@ IMPLEMENT_SIMPLE_AUTOMATION_TEST(FTaskGraphTest, "System.Core.Async.TaskGraph",
 namespace TaskGraphTestTask
 {
 	const FTimespan MaxWaitTime(0, 0, 5);
-	const int64 NumTasks = 10000;
-	int64 CompletedTasks = 0;
+	const int32 NumTasks = 10000;
+	int32 CompletedTasks = 0;
 }
 
 
