Fix spelling error in comment

diff --git a/Engine/Source/Runtime/Core/Public/Templates/Sorting.h b/Engine/Source/Runtime/Core/Public/Templates/Sorting.h
--- a/Engine/Source/Runtime/Core/Public/Templates/Sorting.h
+++ b/Engine/Source/Runtime/Core/Public/Templates/Sorting.h
@@ -45,7 +45,7 @@ struct TDereferenceWrapper<T*, PREDICATE_CLASS>
 	TDereferenceWrapper( const PREDICATE_CLASS& InPredicate )
 		: Predicate( InPredicate ) {}
   
-	/** Dereferennce pointers */
+	/** Dereference pointers */
 	FORCEINLINE bool operator()( T* A, T* B ) const 
 	{
 		return Predicate( *A, *B ); 
