#ue4 - Optimized FTransform::Equals(). Check translation first as it's a simpler test and more likely to differ. Scale3D test is now per-component rather than a dot product test versus tolerance.

--------

diff --git a/Engine/Source/Runtime/Core/Public/Math/Transform.h b/Engine/Source/Runtime/Core/Public/Math/Transform.h
--- a/Engine/Source/Runtime/Core/Public/Math/Transform.h
+++ b/Engine/Source/Runtime/Core/Public/Math/Transform.h
@@ -631,7 +631,7 @@ private:
 
 	FORCEINLINE bool Private_Scale3DEquals(const FVector& InScale3D, const float Tolerance = KINDA_SMALL_NUMBER) const
 	{
-		return (Scale3D-InScale3D).SizeSquared() <= FMath::Square(Tolerance);
+		return Scale3D.Equals(InScale3D, Tolerance);
 	}
 
 public:
@@ -678,13 +678,13 @@ public:
 	// Test if all components of the transforms are equal, within a tolerance.
 	inline bool Equals(const FTransform& Other, float Tolerance=KINDA_SMALL_NUMBER) const
 	{
-		return Private_RotationEquals(Other.Rotation, Tolerance) && Private_TranslationEquals(Other.Translation, Tolerance) && Private_Scale3DEquals(Other.Scale3D, Tolerance);
+		return Private_TranslationEquals(Other.Translation, Tolerance) && Private_RotationEquals(Other.Rotation, Tolerance) && Private_Scale3DEquals(Other.Scale3D, Tolerance);
 	}
 
 	// Test if rotation and translation components of the transforms are equal, within a tolerance.
 	inline bool EqualsNoScale(const FTransform& Other, float Tolerance=KINDA_SMALL_NUMBER) const
 	{
-		return Private_RotationEquals(Other.Rotation, Tolerance) && Private_TranslationEquals(Other.Translation, Tolerance);
+		return Private_TranslationEquals(Other.Translation, Tolerance) && Private_RotationEquals(Other.Rotation, Tolerance);
 	}
 
 	/**
diff --git a/Engine/Source/Runtime/Core/Public/Math/TransformVectorized.h b/Engine/Source/Runtime/Core/Public/Math/TransformVectorized.h
--- a/Engine/Source/Runtime/Core/Public/Math/TransformVectorized.h
+++ b/Engine/Source/Runtime/Core/Public/Math/TransformVectorized.h
@@ -662,13 +662,10 @@ private:
 	}
 
 	inline bool Private_Scale3DEquals( const VectorRegister& InScale3D, const ScalarRegister& Tolerance = ScalarRegister(GlobalVectorConstants::KindaSmallNumber)) const
-	{			
-		const VectorRegister ScaleDiff = VectorSubtract(Scale3D, InScale3D);		
-		// d = dot3(ScaleDiff.xyz, ScaleDiff.xyz), VectorRegister( d, d, d, d )
-		const VectorRegister SizeSquared = VectorDot3(ScaleDiff, ScaleDiff);		
-		const VectorRegister VectorToleranceSqaured = VectorMultiply(Tolerance.Value, Tolerance.Value);
-		//  !( (Scale3DX*Scale3D.X + Scale3D.Y*Scale3DY + Scale3DZ*Scale3D.Z) > Tolerance*Tolerance )
-		return !VectorAnyGreaterThan(SizeSquared, VectorToleranceSqaured);
+	{
+		// !( (FMath::Abs(X-V.X) > Tolerance) || (FMath::Abs(Y-V.Y) > Tolerance) || (FMath::Abs(Z-V.Z) > Tolerance) )
+		const VectorRegister ScaleDiff = VectorAbs(VectorSubtract(Scale3D, InScale3D));
+		return !VectorAnyGreaterThan(ScaleDiff, Tolerance.Value);
 	}
 
 public:
@@ -715,14 +712,14 @@ public:
 	inline bool Equals(const FTransform& Other, float Tolerance=KINDA_SMALL_NUMBER) const
 	{
 		const ScalarRegister ToleranceRegister(Tolerance);
-		return Private_RotationEquals(Other.Rotation, ToleranceRegister) && Private_TranslationEquals(Other.Translation, ToleranceRegister) && Private_Scale3DEquals(Other.Scale3D, ToleranceRegister);
+		return Private_TranslationEquals(Other.Translation, ToleranceRegister) && Private_RotationEquals(Other.Rotation, ToleranceRegister) && Private_Scale3DEquals(Other.Scale3D, ToleranceRegister);
 	}
 
 	// Test if rotation and translation components of the transforms are equal, within a tolerance.
 	inline bool EqualsNoScale(const FTransform& Other, float Tolerance=KINDA_SMALL_NUMBER) const
 	{
 		const ScalarRegister ToleranceRegister(Tolerance);
-		return Private_RotationEquals(Other.Rotation, ToleranceRegister) && Private_TranslationEquals(Other.Translation, ToleranceRegister);
+		return Private_TranslationEquals(Other.Translation, ToleranceRegister) && Private_RotationEquals(Other.Rotation, ToleranceRegister);
 	}
 
 	FORCEINLINE static void Multiply(FTransform* OutTransform, const FTransform* A, const FTransform* B);
