Added warning to cooker for resaving packages when they contain FText which has been promoted from FString / FName.

diff --git a/Engine/Source/Runtime/Core/Private/Internationalization/Text.cpp b/Engine/Source/Runtime/Core/Private/Internationalization/Text.cpp
--- a/Engine/Source/Runtime/Core/Private/Internationalization/Text.cpp
+++ b/Engine/Source/Runtime/Core/Private/Internationalization/Text.cpp
@@ -10,10 +10,14 @@
 
 #include "TextHistory.h"
 
-//DEFINE_LOG_CATEGORY(LogText);
+#include "DebugSerializationFlags.h"
 
 //DEFINE_STAT(STAT_TextFormat);
 
+DECLARE_LOG_CATEGORY_EXTERN(LogText, Log, All);
+DEFINE_LOG_CATEGORY(LogText);
+
+
 #define LOCTEXT_NAMESPACE "Core.Text"
 
 FArchive& operator<<( FArchive& Ar, FFormatArgumentData& Value )
@@ -1033,12 +1037,33 @@ CORE_API FArchive& operator<<( FArchive& Ar, FText& Value )
 		}
 	}
 
-	if(Ar.IsPersistent())
+
+#if WITH_EDITOR
+	if (Ar.IsCooking() && Ar.IsSaving() && Ar.IsPersistent() && !(Ar.GetDebugSerializationFlags()&DSF_IsTagExportsArchive))
+	{
+		if (!!(Value.Flags & ETextFlag::ConvertedProperty))
+		{
+			UE_LOG(LogText, Warning, TEXT("Saving FText \"%s\" which has been converted at load time please resave source package %s to avoid determinisitic cook and localization issues."), *Value.ToString(), *Ar.GetArchiveName());
+		}
+		else if (!!(Value.Flags & ETextFlag::InitializedFromString))
+		{
+			UE_LOG(LogText, Warning, TEXT("Saving FText \"%s\" which has been initialized from FString at cook time resave of source package %s may fix issue."), *Value.ToString(), *Ar.GetArchiveName())
+		}
+	}
+#endif
+
+	int32 OriginalFlags = Value.Flags;
+	if(Ar.IsPersistent() && Ar.IsSaving())
 	{
-		Value.Flags &= ~(ETextFlag::ConvertedProperty); // Remove conversion flag before saving.
+		Value.Flags &= ~(ETextFlag::ConvertedProperty|ETextFlag::InitializedFromString); // Remove conversion flag before saving.
 	}
 	Ar << Value.Flags;
 
+	if (Ar.IsSaving())
+	{
+		Value.Flags = OriginalFlags;
+	}
+
 	if( Ar.UE4Ver() >= VER_UE4_FTEXT_HISTORY )
 	{
 		if(Ar.IsSaving())
@@ -1057,63 +1082,63 @@ CORE_API FArchive& operator<<( FArchive& Ar, FText& Value )
 			Ar << HistoryType;
 
 			// Create the history class based on the serialized type
-			switch((ETextHistoryType)HistoryType)
+			switch ((ETextHistoryType)HistoryType)
 			{
 			case ETextHistoryType::Base:
-				{
-					Value.History = MakeShareable(new FTextHistory_Base);
-					break;
-				}
+			{
+										   Value.History = MakeShareable(new FTextHistory_Base);
+										   break;
+			}
 			case ETextHistoryType::NamedFormat:
-				{
-					Value.History = MakeShareable(new FTextHistory_NamedFormat);
-					break;
-				}
+			{
+												  Value.History = MakeShareable(new FTextHistory_NamedFormat);
+												  break;
+			}
 			case ETextHistoryType::OrderedFormat:
-				{
-					Value.History = MakeShareable(new FTextHistory_OrderedFormat);
-					break;
-				}
+			{
+													Value.History = MakeShareable(new FTextHistory_OrderedFormat);
+													break;
+			}
 			case ETextHistoryType::ArgumentFormat:
-				{
-					Value.History = MakeShareable(new FTextHistory_ArgumentDataFormat);
-					break;
-				}
+			{
+													 Value.History = MakeShareable(new FTextHistory_ArgumentDataFormat);
+													 break;
+			}
 			case ETextHistoryType::AsNumber:
-				{
-					Value.History = MakeShareable(new FTextHistory_AsNumber);
-					break;
-				}
+			{
+											   Value.History = MakeShareable(new FTextHistory_AsNumber);
+											   break;
+			}
 			case ETextHistoryType::AsPercent:
-				{
-					Value.History = MakeShareable(new FTextHistory_AsPercent);
-					break;
-				}
+			{
+												Value.History = MakeShareable(new FTextHistory_AsPercent);
+												break;
+			}
 			case ETextHistoryType::AsCurrency:
-				{
-					Value.History = MakeShareable(new FTextHistory_AsCurrency);
-					break;
-				}
+			{
+												 Value.History = MakeShareable(new FTextHistory_AsCurrency);
+												 break;
+			}
 			case ETextHistoryType::AsDate:
-				{
-					Value.History = MakeShareable(new FTextHistory_AsDate);
-					break;
-				}
+			{
+											 Value.History = MakeShareable(new FTextHistory_AsDate);
+											 break;
+			}
 			case ETextHistoryType::AsTime:
-				{
-					Value.History = MakeShareable(new FTextHistory_AsTime);
-					break;
-				}
+			{
+											 Value.History = MakeShareable(new FTextHistory_AsTime);
+											 break;
+			}
 			case ETextHistoryType::AsDateTime:
-				{
-					Value.History = MakeShareable(new FTextHistory_AsDateTime);
-					break;
-				}
+			{
+												 Value.History = MakeShareable(new FTextHistory_AsDateTime);
+												 break;
+			}
 			default:
-				{
-					Value.History.Reset();
-					Value.DisplayString = FText::GetEmpty().DisplayString;
-				}
+			{
+					   Value.History.Reset();
+					   Value.DisplayString = FText::GetEmpty().DisplayString;
+			}
 			}
 		}
 
@@ -1178,6 +1203,7 @@ FText FText::FromName( const FName& Val)
 	{
 		NewText.Flags |= ETextFlag::CultureInvariant;
 	}
+	NewText.Flags |= ETextFlag::InitializedFromString;
 
 	return NewText; 
 }
@@ -1190,6 +1216,7 @@ FText FText::FromString( FString String )
 	{
 		NewText.Flags |= ETextFlag::CultureInvariant;
 	}
+	NewText.Flags |= ETextFlag::InitializedFromString;
 
 	return NewText;
 }
diff --git a/Engine/Source/Runtime/Core/Public/Internationalization/Text.h b/Engine/Source/Runtime/Core/Public/Internationalization/Text.h
--- a/Engine/Source/Runtime/Core/Public/Internationalization/Text.h
+++ b/Engine/Source/Runtime/Core/Public/Internationalization/Text.h
@@ -13,7 +13,6 @@ struct FDateTime;
 
 #define ENABLE_TEXT_ERROR_CHECKING_RESULTS (UE_BUILD_DEBUG | UE_BUILD_DEVELOPMENT | UE_BUILD_TEST )
 
-//DECLARE_LOG_CATEGORY_EXTERN(LogText, Log, All);
 
 //DECLARE_CYCLE_STAT_EXTERN( TEXT("Format Text"), STAT_TextFormat, STATGROUP_Text, );
 
@@ -25,6 +24,7 @@ namespace ETextFlag
 		CultureInvariant = (1 << 1),
 		ConvertedProperty = (1 << 2),
 		Immutable = (1 << 3),
+		InitializedFromString = (1<<4),  // this ftext was initialized using FromString
 	};
 }
 
