Missed file in my last checkin

diff --git a/Engine/Source/Runtime/Core/Public/UObject/DebugSerializationFlags.h b/Engine/Source/Runtime/Core/Public/UObject/DebugSerializationFlags.h
new file mode 100644
--- /dev/null
+++ b/Engine/Source/Runtime/Core/Public/UObject/DebugSerializationFlags.h
@@ -0,0 +1,18 @@
+// Copyright 1998-2015 Epic Games, Inc. All Rights Reserved.
+
+/*=============================================================================
+DebugSerializationFlags.h: 
+Set custom flags on the archive to help track issues while serializing
+=============================================================================*/
+
+#pragma once
+
+// Debug serialization flags
+enum EDebugSerializationFlags
+{
+	/** No special flags */
+	DSF_None =					0x00000000,
+
+	DSF_IgnoreDiff =			0x00000001,
+	DSF_IsTagExportsArchive =	0x00000002,
+};
\ No newline at end of file
