Backed out save package warning.

diff --git a/Engine/Source/Runtime/Core/Private/Internationalization/Text.cpp b/Engine/Source/Runtime/Core/Private/Internationalization/Text.cpp
--- a/Engine/Source/Runtime/Core/Private/Internationalization/Text.cpp
+++ b/Engine/Source/Runtime/Core/Private/Internationalization/Text.cpp
@@ -1039,7 +1039,7 @@ CORE_API FArchive& operator<<( FArchive& Ar, FText& Value )
 
 
 #if WITH_EDITOR
-	if (Ar.IsCooking() && Ar.IsSaving() && Ar.IsPersistent() && !(Ar.GetDebugSerializationFlags()&DSF_IsTagExportsArchive))
+	/*if (Ar.IsCooking() && Ar.IsSaving() && Ar.IsPersistent() && !(Ar.GetDebugSerializationFlags()&DSF_IsTagExportsArchive))
 	{
 		if (!!(Value.Flags & ETextFlag::ConvertedProperty))
 		{
@@ -1049,7 +1049,7 @@ CORE_API FArchive& operator<<( FArchive& Ar, FText& Value )
 		{
 			UE_LOG(LogText, Warning, TEXT("Saving FText \"%s\" which has been initialized from FString at cook time resave of source package %s may fix issue."), *Value.ToString(), *Ar.GetArchiveName())
 		}
-	}
+	}*/
 #endif
 
 	int32 OriginalFlags = Value.Flags;
