UE4 - Fixed smappy spin loop in new console mode task graph scheduling.

diff --git a/Engine/Source/Runtime/Core/Private/Async/TaskGraph.cpp b/Engine/Source/Runtime/Core/Private/Async/TaskGraph.cpp
--- a/Engine/Source/Runtime/Core/Private/Async/TaskGraph.cpp
+++ b/Engine/Source/Runtime/Core/Private/Async/TaskGraph.cpp
@@ -1733,7 +1733,7 @@ public:
 			}
 			else if (GConsoleSpinModeLatched == 2)
 			{
-				FPlatformProcess::Sleep(0.0f);
+				FPlatformProcess::SleepNoStats(0.0f);
 			}
 		}
 	}
