Bunch of minor Linux cleanups and improvements.

- Mostly coming from MatthewL with changes.
- UE-19705 (getting MAC) - reworked to use ioctl
- UE-14633: IsFirstInstance is now a function on Linux
- UE-16061:dummy windows shouldn't be visible in the task bar.
- UE-18217 fixed.
- Unnecessary printf() removed.
- Set output to non-cached if running under debugger.

diff --git a/Engine/Source/Developer/ShaderFormatOpenGL/Private/OpenGLShaderCompiler.cpp b/Engine/Source/Developer/ShaderFormatOpenGL/Private/OpenGLShaderCompiler.cpp
--- a/Engine/Source/Developer/ShaderFormatOpenGL/Private/OpenGLShaderCompiler.cpp
+++ b/Engine/Source/Developer/ShaderFormatOpenGL/Private/OpenGLShaderCompiler.cpp
@@ -284,7 +284,7 @@ static void _PlatformCreateDummyGLWindow(FPlatformOpenGLContext *OutContext)
 	// Create a dummy window.
 	OutContext->hWnd = SDL_CreateWindow(NULL,
 		0, 0, 1, 1,
-		SDL_WINDOW_OPENGL | SDL_WINDOW_BORDERLESS | SDL_WINDOW_HIDDEN);
+		SDL_WINDOW_OPENGL | SDL_WINDOW_BORDERLESS | SDL_WINDOW_HIDDEN | SDL_WINDOW_SKIP_TASKBAR );
 }
 
 static void _PlatformCreateOpenGLContextCore(FPlatformOpenGLContext* OutContext)
diff --git a/Engine/Source/Developer/SlateFileDialogs/Private/SlateFileDlgWindow.cpp b/Engine/Source/Developer/SlateFileDialogs/Private/SlateFileDlgWindow.cpp
--- a/Engine/Source/Developer/SlateFileDialogs/Private/SlateFileDlgWindow.cpp
+++ b/Engine/Source/Developer/SlateFileDialogs/Private/SlateFileDlgWindow.cpp
@@ -224,6 +224,9 @@ FSlateFileDlgWindow::FSlateFileDlgWindow(FSlateFileDialogsStyle *InStyleSet)
 bool FSlateFileDlgWindow::OpenFileDialog(const void* ParentWindowHandle, const FString& DialogTitle, const FString& DefaultPath,
 		const FString& DefaultFile, const FString& FileTypes, uint32 Flags, TArray<FString>& OutFilenames, int32& OutFilterIndex)
 {
+	FString StartDirectory = DefaultPath;
+	TrimStartDirectory(StartDirectory);
+
 	TSharedRef<SWindow> ModalWindow = SNew(SWindow)
 		.SupportsMinimize(false)
 		.SupportsMaximize(false)
@@ -237,7 +240,7 @@ bool FSlateFileDlgWindow::OpenFileDialog(const void* ParentWindowHandle, const F
 	DialogWidget = SNew(SSlateFileOpenDlg)
 		.bMultiSelectEnabled(Flags == 1)
 		.ParentWindow(ModalWindow)
-		.CurrentPath(DefaultPath)
+		.CurrentPath(StartDirectory)
 		.Filters(FileTypes)
 		.WindowTitleText(DialogTitle)
 		.StyleSet(StyleSet);
@@ -267,6 +270,9 @@ bool FSlateFileDlgWindow::OpenDirectoryDialog(const void* ParentWindowHandle, co
 	TArray<FString> TempOut;
 	FString Filters = "";
 
+	FString StartDirectory = DefaultPath;
+	TrimStartDirectory(StartDirectory);
+
 	TSharedRef<SWindow> ModalWindow = SNew(SWindow)
 		.SupportsMinimize(false)
 		.SupportsMaximize(false)
@@ -281,7 +287,7 @@ bool FSlateFileDlgWindow::OpenDirectoryDialog(const void* ParentWindowHandle, co
 		.bMultiSelectEnabled(false)
 		.ParentWindow(ModalWindow)
 		.bDirectoriesOnly(true)
-		.CurrentPath(DefaultPath)
+		.CurrentPath(StartDirectory)
 		.WindowTitleText(DialogTitle)
 		.StyleSet(StyleSet);
 
@@ -311,6 +317,9 @@ bool FSlateFileDlgWindow::SaveFileDialog(const void* ParentWindowHandle, const F
 {
 	int32 DummyIndex;
 
+	FString StartDirectory = DefaultPath;
+	TrimStartDirectory(StartDirectory);
+
 	TSharedRef<SWindow> ModalWindow = SNew(SWindow)
 		.SupportsMinimize(false)
 		.SupportsMaximize(false)
@@ -326,7 +335,7 @@ bool FSlateFileDlgWindow::SaveFileDialog(const void* ParentWindowHandle, const F
 		.ParentWindow(ModalWindow)
 		.bSaveFile(true)
 		.AcceptText(LOCTEXT("SlateFileDialogsSave","Save"))
-		.CurrentPath(DefaultPath)
+		.CurrentPath(StartDirectory)
 		.Filters(FileTypes)
 		.WindowTitleText(DialogTitle)
 		.StyleSet(StyleSet);
@@ -341,7 +350,23 @@ bool FSlateFileDlgWindow::SaveFileDialog(const void* ParentWindowHandle, const F
 	return (DialogWidget->GetResponse() == EResult::Accept && OutFilenames.Num() > 0);
 }
 
+void FSlateFileDlgWindow::TrimStartDirectory(FString &InPath)
+{
+	if (InPath.Len() == 0)
+	{
+		// no path given. nothing to do.
+		return;
+	}
+
+	FPaths::CollapseRelativeDirectories(InPath);
 
+	FString PathPart;
+	FString FileNamePart;
+	FString ExtensionPart;
+
+	FPaths::Split(InPath, PathPart, FileNamePart, ExtensionPart);
+	InPath = PathPart;
+}
 
 //-----------------------------------------------------------------------------
 // custom file dialog widget
diff --git a/Engine/Source/Developer/SlateFileDialogs/Private/SlateFileDlgWindow.h b/Engine/Source/Developer/SlateFileDialogs/Private/SlateFileDlgWindow.h
--- a/Engine/Source/Developer/SlateFileDialogs/Private/SlateFileDlgWindow.h
+++ b/Engine/Source/Developer/SlateFileDialogs/Private/SlateFileDlgWindow.h
@@ -72,6 +72,8 @@ private:
 	FString CurrentDirectory;
 
 	FSlateFileDialogsStyle* StyleSet;
+
+	void TrimStartDirectory(FString &Path);
 };
 
 
diff --git a/Engine/Source/Runtime/Core/Private/Linux/LinuxApplication.cpp b/Engine/Source/Runtime/Core/Private/Linux/LinuxApplication.cpp
--- a/Engine/Source/Runtime/Core/Private/Linux/LinuxApplication.cpp
+++ b/Engine/Source/Runtime/Core/Private/Linux/LinuxApplication.cpp
@@ -1686,7 +1686,6 @@ void FLinuxApplication::GetWindowPositionInEventLoop(SDL_HWindow NativeWindow, i
 
 bool FLinuxApplication::IsMouseAttached() const
 {
-	printf("--- check mouse ---\n");
 	int rc;
 	char Mouse[64] = "/sys/class/input/mouse0";
 	int MouseIdx = strlen(Mouse) - 1;
@@ -1697,11 +1696,9 @@ bool FLinuxApplication::IsMouseAttached() const
 		Mouse[MouseIdx] = '0' + i;
 		if (access(Mouse, F_OK) == 0)
 		{
-			printf("--- check mouse found %d---\n", i);
 			return true;
 		}
 	}
 
-	printf("--- check mouse none found ---\n");
 	return false;
 }
diff --git a/Engine/Source/Runtime/Core/Private/Linux/LinuxPlatformMisc.cpp b/Engine/Source/Runtime/Core/Private/Linux/LinuxPlatformMisc.cpp
--- a/Engine/Source/Runtime/Core/Private/Linux/LinuxPlatformMisc.cpp
+++ b/Engine/Source/Runtime/Core/Private/Linux/LinuxPlatformMisc.cpp
@@ -14,6 +14,11 @@
 #include <fcntl.h>
 #include <signal.h>
 #include <sys/vfs.h>	// statfs()
+#include <sys/ioctl.h>
+
+#include <ifaddrs.h>	// ethernet mac
+#include <net/if.h>
+#include <net/if_arp.h>
 
 #include "ModuleManager.h"
 
@@ -143,7 +148,11 @@ void FLinuxPlatformMisc::PlatformInit()
 	// install a platform-specific signal handler
 	InstallChildExitedSignalHanlder();
 
+	// do not remove the below check for IsFirstInstance() - it is not just for logging, it actually lays the claim to be first
+	bool bFirstInstance = FPlatformProcess::IsFirstInstance();
+
 	UE_LOG(LogInit, Log, TEXT("Linux hardware info:"));
+	UE_LOG(LogInit, Log, TEXT(" - we are %sthe first instance of this executable"), bFirstInstance ? TEXT("") : TEXT("not "));
 	UE_LOG(LogInit, Log, TEXT(" - this process' id (pid) is %d, parent process' id (ppid) is %d"), static_cast< int32 >(getpid()), static_cast< int32 >(getppid()));
 	UE_LOG(LogInit, Log, TEXT(" - we are %srunning under debugger"), IsDebuggerPresent() ? TEXT("") : TEXT("not "));
 	UE_LOG(LogInit, Log, TEXT(" - machine network name is '%s'"), FPlatformProcess::ComputerName());
@@ -179,7 +188,7 @@ void FLinuxPlatformMisc::PlatformInit()
 		PlatformInitMultimedia();
 	}
 
-	if (FPlatformMisc::HasBeenStartedRemotely())
+	if (FPlatformMisc::HasBeenStartedRemotely() || FPlatformMisc::IsDebuggerPresent())
 	{
 		// print output immediately
 		setvbuf(stdout, NULL, _IONBF, 0);
@@ -238,6 +247,8 @@ void FLinuxPlatformMisc::PlatformTearDown()
 		SDL_Quit();
 		GInitializedSDL = false;
 	}
+
+	FPlatformProcess::CeaseBeingFirstInstance();
 }
 
 GenericApplication* FLinuxPlatformMisc::CreateApplication()
@@ -879,48 +890,47 @@ bool FLinuxPlatformMisc::GetDiskTotalAndFreeSpace(const FString& InPath, uint64&
 
 TArray<uint8> FLinuxPlatformMisc::GetMacAddress()
 {
-	char path[64] = "/sys/class/net/eth0";
-	int eth_idx = strlen(path) - 1;
-	char mac_string[64];
-	uint MAC[6] = {0,0,0,0,0,0};
+	struct ifaddrs *ifap, *ifaptr;
 	TArray<uint8> Result;
 
-	strcat(path, "/address");
-
-	// check to see if eth0-eth9 exist. Use the MAC from the first one we get a hit on
-	for (int i=0; i<=9; i++)
+	if (getifaddrs(&ifap) == 0)
 	{
-		path[eth_idx] = '0' + i;
-
-		int EtherInterface = open(path, O_RDONLY);
-		if (EtherInterface != -1)
+		for (ifaptr = ifap; ifaptr != nullptr; ifaptr = (ifaptr)->ifa_next)
 		{
-			ssize_t ReadBytes = read(EtherInterface, mac_string, 64 - 1);
-			close(EtherInterface);
+			struct ifreq ifr;
+
+			strncpy(ifr.ifr_name, ifaptr->ifa_name, IFNAMSIZ-1);
 
-			if (ReadBytes < 0)
+			int Socket = socket(AF_UNIX, SOCK_DGRAM, 0);
+			if (Socket == -1)
 			{
 				continue;
 			}
 
-			mac_string[ReadBytes] = 0;
+			if (ioctl(Socket, SIOCGIFHWADDR, &ifr) == -1)
+			{
+				close(Socket);
+				continue;
+			}
 
-			if (ReadBytes > 0)
+			close(Socket);
+
+			if (ifr.ifr_hwaddr.sa_family != ARPHRD_ETHER)
 			{
-				printf("--- MAC Address: %s\n", mac_string);
-				int rc = sscanf(mac_string,"%x:%x:%x:%x:%x:%x", &MAC[0],&MAC[1],&MAC[2],&MAC[3],&MAC[4],&MAC[5]);
+				continue;
+			}
 
-				if (rc == 6)
-				{
-					for (int j=0; j<6; j++)
-					{
-						Result.Add(MAC[j]);
-					}
+			const uint8 *MAC = (uint8 *) ifr.ifr_hwaddr.sa_data;
 
-					break;
-				}
+			for (int32 i=0; i < 6; i++)
+			{
+				Result.Add(MAC[i]);
 			}
+
+			break;
 		}
+
+		freeifaddrs(ifap);
 	}
 
 	return Result;
@@ -932,7 +942,7 @@ static bool bIsOnBattery = false;
 
 bool FLinuxPlatformMisc::IsRunningOnBattery()
 {
-	char scratch[8];
+	char Scratch[8];
 	FDateTime Time = FDateTime::Now();
 	int64 Seconds = Time.ToUnixTimestamp();
 
@@ -949,12 +959,12 @@ bool FLinuxPlatformMisc::IsRunningOnBattery()
 	if (State != -1)
 	{
 		// found ACAD device. check its state.
-		ssize_t ReadBytes = read(State, scratch, 1);
+		ssize_t ReadBytes = read(State, Scratch, 1);
 		close(State);
 
 		if (ReadBytes > 0)
 		{
-			bIsOnBattery = (scratch[0] == '0');
+			bIsOnBattery = (Scratch[0] == '0');
 		}
 	}
 
diff --git a/Engine/Source/Runtime/Core/Private/Linux/LinuxPlatformProcess.cpp b/Engine/Source/Runtime/Core/Private/Linux/LinuxPlatformProcess.cpp
--- a/Engine/Source/Runtime/Core/Private/Linux/LinuxPlatformProcess.cpp
+++ b/Engine/Source/Runtime/Core/Private/Linux/LinuxPlatformProcess.cpp
@@ -1258,64 +1258,58 @@ FString FLinuxPlatformProcess::FProcEnumInfo::GetName() const
 	return FPaths::GetCleanFilename(GetFullPath());
 }
 
+static int GFileLockDescriptor = -1;
+
 bool FLinuxPlatformProcess::IsFirstInstance()
 {
-#if !(UE_BUILD_SHIPPING && WITH_EDITOR)
+	// set default return if we are unable to access lock file.
+	static bool bIsFirstInstance = false;
+	static bool bNeverFirst = FParse::Param(FCommandLine::Get(), TEXT("neverfirst"));
 
-	if (GFileLockDescriptor != -1)
+#if !(UE_BUILD_SHIPPING && WITH_EDITOR)
+	if (!bIsFirstInstance && !bNeverFirst)	// once we determined that we're first, this can never change until we exit; otherwise, we re-check each time
 	{
-		// attempt to set lock
-		struct flock lock;
+		// create the file if it doesn't exist
+		if (GFileLockDescriptor == -1)
+		{
+			FString LockFileName(TEXT("/tmp/"));
+			FString ExecPath(FPlatformProcess::ExecutableName());
+			ExecPath.ReplaceInline(TEXT("/"), TEXT("-"));
+			// [RCL] 2015-09-20: can run out of filename limits (256 bytes) due to a long path, be conservative and assume 4-char UTF-8 name like e.g. Japanese
+			ExecPath = ExecPath.Right(80);
 
-		lock.l_type = F_WRLCK;
-		lock.l_start = 0;
-		lock.l_whence = SEEK_SET;
-		lock.l_len = 0;
+			LockFileName += ExecPath;
 
-		if (fcntl(GFileLockDescriptor, F_SETLK, &lock) >= 0)
-		{
-			// lock file created and successfully locked by this process.
-			printf("we are the first instance\n");
-			GIsFirstInstance = true;
+			GFileLockDescriptor = open(TCHAR_TO_UTF8(*LockFileName), O_RDWR | O_CREAT, 0666);
 		}
-		else
+
+		if (GFileLockDescriptor != -1)
 		{
-			// we were unable to lock file. so some other process beat us to lock file.
-			printf("we are *NOT* the first instance\n");
-			GIsFirstInstance = false;
+			if (flock(GFileLockDescriptor, LOCK_EX | LOCK_NB) == 0)
+			{
+				// lock file successfully locked by this process - no more checking if we're first!
+				bIsFirstInstance = true;
+			}
+			else
+			{
+				// we were unable to lock file. so some other process beat us to lock file.
+				bIsFirstInstance = false;
+			}
 		}
 	}
-	else
-	{
-		// Should not get here. It means that we were unable to create or open the lock file.
-		// not much we can do other than provide a default return and spit out an error msg.
-		UE_LOG(LogHAL, Warning, TEXT("No process lock file descriptor"));
-		GIsFirstInstance = true;
-	}
-
-	return GIsFirstInstance;
-#elif
-	return true;
 #endif
+	return bIsFirstInstance;
 }
 
-void FLinuxPlatformProcess::ReleaseProcessLockFile()
+void FLinuxPlatformProcess::CeaseBeingFirstInstance()
 {
 #if !(UE_BUILD_SHIPPING && WITH_EDITOR)
-	printf("*** releasing process lock file\n");
-
 	if (GFileLockDescriptor != -1)
 	{
-		// attempt to release lock
-		struct flock lock;
-
-		lock.l_type = F_UNLCK;
-		lock.l_start = 0;
-		lock.l_whence = SEEK_SET;
-		lock.l_len = 0;
-
-		fcntl(GFileLockDescriptor, F_SETLK, &lock);
+		// may fail if we didn't have the lock
+		flock(GFileLockDescriptor, LOCK_UN | LOCK_NB);
+		close(GFileLockDescriptor);
+		GFileLockDescriptor = -1;
 	}
 #endif
 }
-
diff --git a/Engine/Source/Runtime/Core/Private/Misc/Core.cpp b/Engine/Source/Runtime/Core/Private/Misc/Core.cpp
--- a/Engine/Source/Runtime/Core/Private/Misc/Core.cpp
+++ b/Engine/Source/Runtime/Core/Private/Misc/Core.cpp
@@ -211,8 +211,9 @@ uint32					GFrameNumberRenderThread		= 1;
 #if !(UE_BUILD_SHIPPING && WITH_EDITOR)
 // We cannot count on this variable to be accurate in a shipped game, so make sure no code tries to use it
 /** Whether we are the first instance of the game running.													*/
+#if !PLATFORM_LINUX
 bool					GIsFirstInstance				= true;
-int32					GFileLockDescriptor				= -1;
+#endif
 #endif
 /** Threshold for a frame to be considered a hitch (in seconds. */
 float GHitchThreshold = 0.075f;
diff --git a/Engine/Source/Runtime/Core/Public/CoreGlobals.h b/Engine/Source/Runtime/Core/Public/CoreGlobals.h
--- a/Engine/Source/Runtime/Core/Public/CoreGlobals.h
+++ b/Engine/Source/Runtime/Core/Public/CoreGlobals.h
@@ -241,8 +241,12 @@ extern CORE_API uint32 GFrameNumberRenderThread;
 #if !(UE_BUILD_SHIPPING && WITH_EDITOR)
 // We cannot count on this variable to be accurate in a shipped game, so make sure no code tries to use it
 /** Whether we are the first instance of the game running. */
+#if PLATFORM_LINUX
+#define GIsFirstInstance FPlatformProcess::IsFirstInstance()
+#else
 extern CORE_API bool GIsFirstInstance;
-extern CORE_API int32 GFileLockDescriptor;
+#endif
+
 #endif
 
 /** Threshold for a frame to be considered a hitch (in seconds. */
@@ -317,4 +321,4 @@ extern CORE_API TArray<FScriptTraceStackNode> GScriptStack;
  *		construct proper error message.
  */
 CORE_API void EnsureRetrievingVTablePtrDuringCtor(const TCHAR* CtorSignature);
-#endif // WITH_HOT_RELOAD_CTORS
\ No newline at end of file
+#endif // WITH_HOT_RELOAD_CTORS
diff --git a/Engine/Source/Runtime/Core/Public/GenericPlatform/GenericPlatformProcess.h b/Engine/Source/Runtime/Core/Public/GenericPlatform/GenericPlatformProcess.h
--- a/Engine/Source/Runtime/Core/Public/GenericPlatform/GenericPlatformProcess.h
+++ b/Engine/Source/Runtime/Core/Public/GenericPlatform/GenericPlatformProcess.h
@@ -554,11 +554,10 @@ struct CORE_API FGenericPlatformProcess
 	 */
 	static bool Daemonize();
 
-	/** Check to see if we're the first instance */
+	/**
+	 * Checks if we're the first instance. An instance can become first if the previous first instance quits before it.
+	 */
 	static bool IsFirstInstance();
-
-	/** Release process lock file (if any) */
-	static void ReleaseProcessLockFile() { }
 };
 
 
diff --git a/Engine/Source/Runtime/Core/Public/Linux/LinuxPlatformProcess.h b/Engine/Source/Runtime/Core/Public/Linux/LinuxPlatformProcess.h
--- a/Engine/Source/Runtime/Core/Public/Linux/LinuxPlatformProcess.h
+++ b/Engine/Source/Runtime/Core/Public/Linux/LinuxPlatformProcess.h
@@ -283,7 +283,11 @@ struct CORE_API FLinuxPlatformProcess : public FGenericPlatformProcess
 	static void ExploreFolder( const TCHAR* FilePath );
 	static void LaunchFileInDefaultExternalApplication( const TCHAR* FileName, const TCHAR* Parms = NULL, ELaunchVerb::Type Verb = ELaunchVerb::Open );
 	static bool IsFirstInstance();
-	static void ReleaseProcessLockFile();
+
+	/**
+	 * @brief Releases locks that we held for IsFirstInstance check
+	 */
+	static void CeaseBeingFirstInstance();
 };
 
 typedef FLinuxPlatformProcess FPlatformProcess;
diff --git a/Engine/Source/Runtime/Linux/LinuxCommonStartup/Private/LinuxCommonStartup.cpp b/Engine/Source/Runtime/Linux/LinuxCommonStartup/Private/LinuxCommonStartup.cpp
--- a/Engine/Source/Runtime/Linux/LinuxCommonStartup/Private/LinuxCommonStartup.cpp
+++ b/Engine/Source/Runtime/Linux/LinuxCommonStartup/Private/LinuxCommonStartup.cpp
@@ -192,58 +192,6 @@ int CommonLinuxMain(int argc, char *argv[], int (*RealMain)(const TCHAR * Comman
 		GSavedCommandLine += Temp; 	// note: technically it depends on locale
 	}
 
-#if !(UE_BUILD_SHIPPING && WITH_EDITOR)
-	// check to see if we're first instance and try to set file lock.
-	if ( !FParse::Param( *GSavedCommandLine,TEXT("neverfirst")) )
-	{
-		char LockFileName[128];
-		strcpy(LockFileName,"/tmp/");
-
-		int i;
-		char *ptr = argv[0];
-		for (i = strlen(argv[0])-1; i >= 0 && ptr[i] != '/'; i--);
-
-		strncat(LockFileName, &ptr[i+1], 127);
-		strncat(LockFileName, ".lock", 127);
-		GFileLockDescriptor = open(LockFileName, O_RDWR | O_CREAT, 0666);
-
-		if (GFileLockDescriptor != -1)
-		{
-			// attempt to set lock
-			struct flock lock;
-
-			lock.l_type = F_WRLCK;
-			lock.l_start = 0;
-			lock.l_whence = SEEK_SET;
-			lock.l_len = 0;
-
-			if (fcntl(GFileLockDescriptor, F_SETLK, &lock) >= 0)
-			{
-				// lock file created and successfully locked by this process.
-				printf("** we are the first instance\n");
-				GIsFirstInstance = true;
-			}
-			else
-			{
-				// we were unable to lock file. so some other process beat us to lock file.
-				printf("** we are *NOT* the first instance\n");
-				GIsFirstInstance = false;
-			}
-		}
-		else
-		{
-			// Should not get here. It means that we were unable to create or open the lock file.
-			// not much we can do other than provide a default return and spit out an error msg.
-			fprintf(stderr, "*** LaunchLinux - unable to create process lock file '%s'. (%s)\n", LockFileName, strerror(errno));
-			GIsFirstInstance = true;
-		}
-	}
-	else
-	{
-		GIsFirstInstance = false;
-	}
-#endif
-
 	if (!UE_BUILD_SHIPPING)
 	{
 		GAlwaysReportCrash = true;	// set by default and reverse the behavior
@@ -280,13 +228,6 @@ int CommonLinuxMain(int argc, char *argv[], int (*RealMain)(const TCHAR * Comman
 		}
 	}
 
-#if !(UE_BUILD_SHIPPING && WITH_EDITOR)
-	if (GFileLockDescriptor >= 0)
-	{
-		close(GFileLockDescriptor);
-	}
-#endif
-
 	if (ErrorLevel)
 	{
 		printf("Exiting abnormally (error code: %d)\n", ErrorLevel);
diff --git a/Engine/Source/Runtime/OpenGLDrv/Private/Linux/OpenGLLinux.cpp b/Engine/Source/Runtime/OpenGLDrv/Private/Linux/OpenGLLinux.cpp
--- a/Engine/Source/Runtime/OpenGLDrv/Private/Linux/OpenGLLinux.cpp
+++ b/Engine/Source/Runtime/OpenGLDrv/Private/Linux/OpenGLLinux.cpp
@@ -110,11 +110,20 @@ static void _PlatformCreateDummyGLWindow( FPlatformOpenGLContext *OutContext )
 	static bool bInitializedWindowClass = false;
 
 	// Create a dummy window.
-	SDL_HWindow h_wnd = SDL_CreateWindow(	NULL,
+	SDL_HWindow DummyWindow = SDL_CreateWindow(	NULL,
 											0, 0, 1, 1,
-											SDL_WINDOW_OPENGL | SDL_WINDOW_BORDERLESS | SDL_WINDOW_HIDDEN );
+											SDL_WINDOW_OPENGL | SDL_WINDOW_BORDERLESS | SDL_WINDOW_HIDDEN | SDL_WINDOW_SKIP_TASKBAR );
+	if (DummyWindow == nullptr)
+	{
+		UE_LOG(LogLinux, Fatal, TEXT("Cannot create dummy GL window for shared context."));
+		// unreachable
+	}
+	else
+	{
+		SDL_SetWindowTitle(DummyWindow, "UE4 Dummy GL window");
+	}
 
-	OutContext->hWnd					= h_wnd;
+	OutContext->hWnd					= DummyWindow;
 	OutContext->bReleaseWindowOnDestroy	= true;
 }
 
