UE4 - Added stat DumpSpam command and avoid stat messages for no-op adds.

diff --git a/Engine/Source/Runtime/Core/Private/Stats/StatsCommand.cpp b/Engine/Source/Runtime/Core/Private/Stats/StatsCommand.cpp
--- a/Engine/Source/Runtime/Core/Private/Stats/StatsCommand.cpp
+++ b/Engine/Source/Runtime/Core/Private/Stats/StatsCommand.cpp
@@ -1316,6 +1316,65 @@ struct FDumpMultiple
 	}
 };
 
+static struct FDumpSpam* DumpSpam = nullptr;
+
+struct FDumpSpam
+{
+	FStatsThreadState& Stats;
+	TMap<FName, int32> Counts;
+	int32 TotalCount;
+	int32 NumPackets;
+	FDelegateHandle NewRawStatPacketDelegateHandle;
+
+	FDumpSpam()
+		: Stats(FStatsThreadState::GetLocalState())
+		, TotalCount(0)
+		, NumPackets(0)
+	{
+		FThreadStats::EnableRawStats();
+		StatsMasterEnableAdd();
+		NewRawStatPacketDelegateHandle = Stats.NewRawStatPacket.AddRaw(this, &FDumpSpam::NewFrame);
+	}
+
+	~FDumpSpam()
+	{
+		FThreadStats::DisableRawStats();
+		StatsMasterEnableSubtract();
+		UE_LOG(LogStats, Log, TEXT("------------------ %d packets, %d total messages ---------------"), NumPackets, TotalCount);
+
+		Counts.ValueSort(TGreater<int32>());
+
+		for (auto& Pair : Counts)
+		{
+			UE_LOG(LogStats, Log, TEXT("%10d      %s"), Pair.Value, *Pair.Key.ToString());
+		}
+
+		Stats.NewRawStatPacket.Remove(NewRawStatPacketDelegateHandle);
+		DumpSpam = NULL;
+	}
+
+	void NewFrame(const FStatPacket* Packet)
+	{
+		NumPackets++;
+		int32 NumMessages = Packet->StatMessages.Num();
+		TotalCount += NumMessages;
+		for( int32 MessageIndex = 0; MessageIndex < NumMessages; ++MessageIndex )
+		{
+			const FStatMessage& Message = Packet->StatMessages[MessageIndex];
+			FName Name = Message.NameAndInfo.GetRawName();
+			int32* Existing = Counts.Find(Name);
+			if (Existing)
+			{
+				(*Existing)++;
+			}
+			else
+			{
+				Counts.Add(Name, 1);
+			}
+		}
+	}
+};
+
 /** Prints stats help to the specified output device. This is queued to be executed on the game thread. */
 static void PrintStatsHelpToOutputDevice( FOutputDevice& Ar )
 {
@@ -1469,6 +1528,16 @@ static void StatCmd(FString InCmd)
 			DumpMultiple->bSum = true;
 		}
 	}
+	else if( FParse::Command(&Cmd,TEXT("DUMPSPAM")) )
+	{
+		bool bIsStart = FString(Cmd).Find(TEXT("-start")) != INDEX_NONE;
+		bool bIsStop = FString(Cmd).Find(TEXT("-stop")) != INDEX_NONE;
+		delete DumpSpam;
+		if (!bIsStop)
+		{
+			DumpSpam = new FDumpSpam();
+		}
+	}
 	else if( FParse::Command(&Cmd,TEXT("DUMPHITCHES")) )
 	{
 		static bool bToggle = false;
@@ -1688,6 +1757,9 @@ bool DirectStatsCommand(const TCHAR* Cmd, bool bBlockForCompletion /*= false*/,
 		else if( FParse::Command(&TempCmd,TEXT("DUMPSUM")) )
 		{
 		}
+		else if( FParse::Command(&TempCmd,TEXT("DUMPSPAM")) )
+		{
+		}
 		else if( FParse::Command(&TempCmd,TEXT("DUMPHITCHES")) )
 		{
 		}
diff --git a/Engine/Source/Runtime/Core/Public/Stats/Stats2.h b/Engine/Source/Runtime/Core/Public/Stats/Stats2.h
--- a/Engine/Source/Runtime/Core/Public/Stats/Stats2.h
+++ b/Engine/Source/Runtime/Core/Public/Stats/Stats2.h
@@ -1877,19 +1877,23 @@ struct FStat_##StatName\
 }
 #define INC_FLOAT_STAT_BY(Stat, Amount) \
 {\
-	FThreadStats::AddMessage(GET_STATFNAME(Stat), EStatOperation::Add, double(Amount));\
+	if (Amount != 0.0f) \
+		FThreadStats::AddMessage(GET_STATFNAME(Stat), EStatOperation::Add, double(Amount));\
 }
 #define INC_DWORD_STAT_BY(Stat, Amount) \
 {\
-	FThreadStats::AddMessage(GET_STATFNAME(Stat), EStatOperation::Add, int64(Amount));\
+	if (Amount != 0) \
+		FThreadStats::AddMessage(GET_STATFNAME(Stat), EStatOperation::Add, int64(Amount));\
 }
 #define INC_DWORD_STAT_FNAME_BY(StatFName, Amount) \
 {\
-	FThreadStats::AddMessage(StatFName, EStatOperation::Add, int64(Amount));\
+	if (Amount != 0) \
+		FThreadStats::AddMessage(StatFName, EStatOperation::Add, int64(Amount));\
 }
 #define INC_MEMORY_STAT_BY(Stat, Amount) \
 {\
-	FThreadStats::AddMessage(GET_STATFNAME(Stat), EStatOperation::Add, int64(Amount));\
+	if (Amount != 0) \
+		FThreadStats::AddMessage(GET_STATFNAME(Stat), EStatOperation::Add, int64(Amount));\
 }
 #define DEC_DWORD_STAT(Stat) \
 {\
@@ -1897,19 +1901,23 @@ struct FStat_##StatName\
 }
 #define DEC_FLOAT_STAT_BY(Stat,Amount) \
 {\
-	FThreadStats::AddMessage(GET_STATFNAME(Stat), EStatOperation::Subtract, double(Amount));\
+	if (Amount != 0.0f) \
+		FThreadStats::AddMessage(GET_STATFNAME(Stat), EStatOperation::Subtract, double(Amount));\
 }
 #define DEC_DWORD_STAT_BY(Stat,Amount) \
 {\
-	FThreadStats::AddMessage(GET_STATFNAME(Stat), EStatOperation::Subtract, int64(Amount));\
+	if (Amount != 0) \
+		FThreadStats::AddMessage(GET_STATFNAME(Stat), EStatOperation::Subtract, int64(Amount));\
 }
 #define DEC_DWORD_STAT_FNAME_BY(StatFName,Amount) \
 {\
-	FThreadStats::AddMessage(StatFName, EStatOperation::Subtract, int64(Amount));\
+	if (Amount != 0) \
+ 		FThreadStats::AddMessage(StatFName, EStatOperation::Subtract, int64(Amount));\
 }
 #define DEC_MEMORY_STAT_BY(Stat,Amount) \
 {\
-	FThreadStats::AddMessage(GET_STATFNAME(Stat), EStatOperation::Subtract, int64(Amount));\
+	if (Amount != 0) \
+		FThreadStats::AddMessage(GET_STATFNAME(Stat), EStatOperation::Subtract, int64(Amount));\
 }
 #define SET_MEMORY_STAT(Stat,Value) \
 {\
@@ -1943,15 +1951,18 @@ struct FStat_##StatName\
 }
 #define INC_FLOAT_STAT_BY_FName(Stat, Amount) \
 {\
-	FThreadStats::AddMessage(Stat, EStatOperation::Add, double(Amount));\
+	if (Amount != 0.0f) \
+		FThreadStats::AddMessage(Stat, EStatOperation::Add, double(Amount));\
 }
 #define INC_DWORD_STAT_BY_FName(Stat, Amount) \
 {\
-	FThreadStats::AddMessage(Stat, EStatOperation::Add, int64(Amount));\
+	if (Amount != 0) \
+		FThreadStats::AddMessage(Stat, EStatOperation::Add, int64(Amount));\
 }
 #define INC_MEMORY_STAT_BY_FName(Stat, Amount) \
 {\
-	FThreadStats::AddMessage(Stat, EStatOperation::Add, int64(Amount));\
+	if (Amount != 0) \
+		FThreadStats::AddMessage(Stat, EStatOperation::Add, int64(Amount));\
 }
 #define DEC_DWORD_STAT_FName(Stat) \
 {\
@@ -1959,15 +1970,18 @@ struct FStat_##StatName\
 }
 #define DEC_FLOAT_STAT_BY_FName(Stat,Amount) \
 {\
-	FThreadStats::AddMessage(Stat, EStatOperation::Subtract, double(Amount));\
+	if (Amount != 0.0f) \
+		FThreadStats::AddMessage(Stat, EStatOperation::Subtract, double(Amount));\
 }
 #define DEC_DWORD_STAT_BY_FName(Stat,Amount) \
 {\
-	FThreadStats::AddMessage(Stat, EStatOperation::Subtract, int64(Amount));\
+	if (Amount != 0) \
+		FThreadStats::AddMessage(Stat, EStatOperation::Subtract, int64(Amount));\
 }
 #define DEC_MEMORY_STAT_BY_FName(Stat,Amount) \
 {\
-	FThreadStats::AddMessage(Stat, EStatOperation::Subtract, int64(Amount));\
+	if (Amount != 0) \
+		FThreadStats::AddMessage(Stat, EStatOperation::Subtract, int64(Amount));\
 }
 #define SET_MEMORY_STAT_FName(Stat,Value) \
 {\
