Added some missing PLATFORM_SUPPORTS_STACK_SYMBOLS defines

Some platforms that implement stack walking weren't setting this to true, and there was also no default for platforms that don't support stack walking.

This define is only used when building with bUseMallocProfiler turned on.

diff --git a/Engine/Source/Developer/CrashDebugHelper/Private/Windows/WindowsPlatformStackWalkExt.h b/Engine/Source/Developer/CrashDebugHelper/Private/Windows/WindowsPlatformStackWalkExt.h
--- a/Engine/Source/Developer/CrashDebugHelper/Private/Windows/WindowsPlatformStackWalkExt.h
+++ b/Engine/Source/Developer/CrashDebugHelper/Private/Windows/WindowsPlatformStackWalkExt.h
@@ -2,9 +2,6 @@
 
 #pragma once
 
-#undef PLATFORM_SUPPORTS_STACK_SYMBOLS
-#define PLATFORM_SUPPORTS_STACK_SYMBOLS 1
-
 class FCrashInfo;
 class FCrashModuleInfo;
 
diff --git a/Engine/Source/Runtime/Core/Public/Android/AndroidPlatform.h b/Engine/Source/Runtime/Core/Public/Android/AndroidPlatform.h
--- a/Engine/Source/Runtime/Core/Public/Android/AndroidPlatform.h
+++ b/Engine/Source/Runtime/Core/Public/Android/AndroidPlatform.h
@@ -43,6 +43,7 @@ typedef FAndroidTypes FPlatformTypes;
 #define PLATFORM_USES_ES2							1
 #define PLATFORM_HAS_BSD_SOCKET_FEATURE_IOCTL		1
 #define PLATFORM_HAS_TOUCH_MAIN_SCREEN				1
+#define PLATFORM_SUPPORTS_STACK_SYMBOLS				1
 
 // Function type macros.
 #define VARARGS													/* Functions with variable arguments */
diff --git a/Engine/Source/Runtime/Core/Public/Android/AndroidPlatformStackWalk.h b/Engine/Source/Runtime/Core/Public/Android/AndroidPlatformStackWalk.h
--- a/Engine/Source/Runtime/Core/Public/Android/AndroidPlatformStackWalk.h
+++ b/Engine/Source/Runtime/Core/Public/Android/AndroidPlatformStackWalk.h
@@ -8,9 +8,6 @@
 #pragma once
 #include "GenericPlatform/GenericPlatformStackWalk.h"
 
-#undef PLATFORM_SUPPORTS_STACK_SYMBOLS
-#define PLATFORM_SUPPORTS_STACK_SYMBOLS 1
-
 /**
 * Android platform stack walking
 */
diff --git a/Engine/Source/Runtime/Core/Public/Apple/ApplePlatformStackWalk.h b/Engine/Source/Runtime/Core/Public/Apple/ApplePlatformStackWalk.h
--- a/Engine/Source/Runtime/Core/Public/Apple/ApplePlatformStackWalk.h
+++ b/Engine/Source/Runtime/Core/Public/Apple/ApplePlatformStackWalk.h
@@ -7,10 +7,6 @@
 #pragma once
 #include "GenericPlatform/GenericPlatformStackWalk.h"
 
-#undef PLATFORM_SUPPORTS_STACK_SYMBOLS
-#define PLATFORM_SUPPORTS_STACK_SYMBOLS 1
-
-
 /**
  * Apple platform implementation of the misc OS functions
  */
diff --git a/Engine/Source/Runtime/Core/Public/HAL/Platform.h b/Engine/Source/Runtime/Core/Public/HAL/Platform.h
--- a/Engine/Source/Runtime/Core/Public/HAL/Platform.h
+++ b/Engine/Source/Runtime/Core/Public/HAL/Platform.h
@@ -274,6 +274,10 @@
 	#define PLATFORM_HAS_TOUCH_MAIN_SCREEN		0
 #endif
 
+#ifndef PLATFORM_SUPPORTS_STACK_SYMBOLS
+	#define PLATFORM_SUPPORTS_STACK_SYMBOLS 0
+#endif
+
 #ifndef PLATFORM_HAS_128BIT_ATOMICS
 	#define PLATFORM_HAS_128BIT_ATOMICS 0
 #endif
diff --git a/Engine/Source/Runtime/Core/Public/HTML5/HTML5Platform.h b/Engine/Source/Runtime/Core/Public/HTML5/HTML5Platform.h
--- a/Engine/Source/Runtime/Core/Public/HTML5/HTML5Platform.h
+++ b/Engine/Source/Runtime/Core/Public/HTML5/HTML5Platform.h
@@ -51,6 +51,7 @@ typedef FHTML5Types FPlatformTypes;
 #define PLATFORM_SUPPORTS_TBB						0
 #define PLATFORM_ENABLE_VECTORINTRINSICS			0
 #define PLATFORM_USES_ES2							1
+#define PLATFORM_SUPPORTS_STACK_SYMBOLS				1
 
 // Function type macros.
 #if PLATFORM_HTML5_WIN32
diff --git a/Engine/Source/Runtime/Core/Public/HTML5/HTML5PlatformStackWalk.h b/Engine/Source/Runtime/Core/Public/HTML5/HTML5PlatformStackWalk.h
--- a/Engine/Source/Runtime/Core/Public/HTML5/HTML5PlatformStackWalk.h
+++ b/Engine/Source/Runtime/Core/Public/HTML5/HTML5PlatformStackWalk.h
@@ -9,7 +9,7 @@
 #include "GenericPlatform/GenericPlatformStackWalk.h"
 
 /**
-* Android platform stack walking
+* HTML5 platform stack walking
 */
 struct CORE_API FHTML5PlatformStackWalk : public FGenericPlatformStackWalk
 {
diff --git a/Engine/Source/Runtime/Core/Public/IOS/IOSPlatform.h b/Engine/Source/Runtime/Core/Public/IOS/IOSPlatform.h
--- a/Engine/Source/Runtime/Core/Public/IOS/IOSPlatform.h
+++ b/Engine/Source/Runtime/Core/Public/IOS/IOSPlatform.h
@@ -41,6 +41,7 @@ typedef FIOSPlatformTypes FPlatformTypes;
 #define PLATFORM_ALLOW_NULL_RHI							1
 #define PLATFORM_HAS_TOUCH_MAIN_SCREEN					1
 #define PLATFORM_ENABLE_VECTORINTRINSICS_NEON			1
+#define PLATFORM_SUPPORTS_STACK_SYMBOLS					1
 
 // @todo iOS: temporarily use Ansi allocator as wxWidgets cause problems with MallocTBB
 #define FORCE_ANSI_ALLOCATOR 1
diff --git a/Engine/Source/Runtime/Core/Public/Linux/LinuxPlatform.h b/Engine/Source/Runtime/Core/Public/Linux/LinuxPlatform.h
--- a/Engine/Source/Runtime/Core/Public/Linux/LinuxPlatform.h
+++ b/Engine/Source/Runtime/Core/Public/Linux/LinuxPlatform.h
@@ -44,6 +44,7 @@ typedef FLinuxPlatformTypes FPlatformTypes;
 #define PLATFORM_HAS_NO_EPROCLIM					1
 #define PLATFORM_HAS_BSD_SOCKET_FEATURE_IOCTL		1
 #define PLATFORM_HAS_BSD_IPV6_SOCKETS				1
+#define PLATFORM_SUPPORTS_STACK_SYMBOLS				1
 
 // SOCK_CLOEXEC is available on Linux since 2.6.27
 #if LINUX_VERSION_CODE >= KERNEL_VERSION(2,6,27)
diff --git a/Engine/Source/Runtime/Core/Public/Linux/LinuxPlatformStackWalk.h b/Engine/Source/Runtime/Core/Public/Linux/LinuxPlatformStackWalk.h
--- a/Engine/Source/Runtime/Core/Public/Linux/LinuxPlatformStackWalk.h
+++ b/Engine/Source/Runtime/Core/Public/Linux/LinuxPlatformStackWalk.h
@@ -8,9 +8,6 @@
 #pragma once
 #include "GenericPlatform/GenericPlatformStackWalk.h"
 
-#undef PLATFORM_SUPPORTS_STACK_SYMBOLS
-#define PLATFORM_SUPPORTS_STACK_SYMBOLS 1
-
 struct CORE_API FLinuxPlatformStackWalk : public FGenericPlatformStackWalk
 {
 	static void ProgramCounterToSymbolInfo( uint64 ProgramCounter, FProgramCounterSymbolInfo& out_SymbolInfo );
diff --git a/Engine/Source/Runtime/Core/Public/Mac/MacPlatform.h b/Engine/Source/Runtime/Core/Public/Mac/MacPlatform.h
--- a/Engine/Source/Runtime/Core/Public/Mac/MacPlatform.h
+++ b/Engine/Source/Runtime/Core/Public/Mac/MacPlatform.h
@@ -38,6 +38,7 @@ typedef FMacPlatformTypes FPlatformTypes;
 //#define PLATFORM_USE_PTHREADS						1
 #define PLATFORM_MAX_FILEPATH_LENGTH				MAX_PATH
 #define PLATFORM_SUPPORTS_TBB						1
+#define PLATFORM_SUPPORTS_STACK_SYMBOLS				1
 
 // Function type macros.
 #define VARARGS														/* Functions with variable arguments */
diff --git a/Engine/Source/Runtime/Core/Public/Windows/WIndowsPlatform.h b/Engine/Source/Runtime/Core/Public/Windows/WIndowsPlatform.h
--- a/Engine/Source/Runtime/Core/Public/Windows/WIndowsPlatform.h
+++ b/Engine/Source/Runtime/Core/Public/Windows/WIndowsPlatform.h
@@ -77,6 +77,7 @@ typedef FWindowsPlatformTypes FPlatformTypes;
 #define PLATFORM_RHITHREAD_DEFAULT_BYPASS					1
 #define PLATFORM_CAN_TOGGLE_RHITHREAD_IN_SHIPPING			1
 
+#define PLATFORM_SUPPORTS_STACK_SYMBOLS						1
 
 // Intrinsics for 128-bit atomics on Windows platform requires Windows 8 or higher (WINVER>0x0602)
 // http://msdn.microsoft.com/en-us/library/windows/desktop/hh972640.aspx
diff --git a/Engine/Source/Runtime/Core/Public/Windows/WindowsPlatformStackWalk.h b/Engine/Source/Runtime/Core/Public/Windows/WindowsPlatformStackWalk.h
--- a/Engine/Source/Runtime/Core/Public/Windows/WindowsPlatformStackWalk.h
+++ b/Engine/Source/Runtime/Core/Public/Windows/WindowsPlatformStackWalk.h
@@ -4,10 +4,6 @@
 
 #include "GenericPlatform/GenericPlatformStackWalk.h"
 
-#undef PLATFORM_SUPPORTS_STACK_SYMBOLS
-#define PLATFORM_SUPPORTS_STACK_SYMBOLS 1
-
-
 /**
  * Windows implementation of the stack walking.
  **/
