Minimized FText allocations when generating text at runtime

FText used to always perform two allocations, one for the shared display string pointer, and another for the text history.

This change allows text that is generated at runtime (such as via FText::AsNumber) to allocate its string and history together in a single allocation. Profiling shows that this almost negates the cost of the text history, without having to remove the text history (which is essential for FText serialization).

This means that we never have to take the hit of a second allocation for text which is generated at runtime and never serialized. Should text that was generated at runtime be serialized, then it will call TGeneratedTextData::PersistText to perform the allocation of the shared display string pointer, and allow the text to be saved, gathered, and localized.

This change also removes some allocations from FTextHistory_Base and FFormatArgumentValue, makes the text history types fully movable, and allows you to move the arguments passed to FText::Format if you know you no longer need your own copy.

diff --git a/Engine/Source/Runtime/Core/Private/Internationalization/ICUText.cpp b/Engine/Source/Runtime/Core/Private/Internationalization/ICUText.cpp
--- a/Engine/Source/Runtime/Core/Private/Internationalization/ICUText.cpp
+++ b/Engine/Source/Runtime/Core/Private/Internationalization/ICUText.cpp
@@ -4,6 +4,7 @@
 
 #if UE_ENABLE_ICU
 #include "Text.h"
+#include "TextData.h"
 #include "TextHistory.h"
 #include "ICUCulture.h"
 
@@ -44,7 +45,7 @@ FText FText::AsDate(const FDateTime& DateTime, const EDateTimeStyle::Type DateSt
 	FString NativeString;
 	ICUUtilities::ConvertString(FormattedString, NativeString);
 
-	return FText::CreateChronologicalText(MoveTemp(NativeString), MakeShareable(new FTextHistory_AsDate(DateTime, DateStyle, TimeZone, TargetCulture)));
+	return FText::CreateChronologicalText(MakeShareable(new TGeneratedTextData<FTextHistory_AsDate>(MoveTemp(NativeString), FTextHistory_AsDate(DateTime, DateStyle, TimeZone, TargetCulture))));
 }
 
 FText FText::AsTime(const FDateTime& DateTime, const EDateTimeStyle::Type TimeStyle, const FString& TimeZone, const FCulturePtr& TargetCulture)
@@ -64,7 +65,7 @@ FText FText::AsTime(const FDateTime& DateTime, const EDateTimeStyle::Type TimeSt
 	FString NativeString;
 	ICUUtilities::ConvertString(FormattedString, NativeString);
 
-	return FText::CreateChronologicalText(MoveTemp(NativeString), MakeShareable(new FTextHistory_AsTime(DateTime, TimeStyle, TimeZone, TargetCulture)));
+	return FText::CreateChronologicalText(MakeShareable(new TGeneratedTextData<FTextHistory_AsTime>(MoveTemp(NativeString), FTextHistory_AsTime(DateTime, TimeStyle, TimeZone, TargetCulture))));
 }
 
 FText FText::AsTimespan(const FTimespan& Timespan, const FCulturePtr& TargetCulture)
@@ -108,7 +109,7 @@ FText FText::AsDateTime(const FDateTime& DateTime, const EDateTimeStyle::Type Da
 	FString NativeString;
 	ICUUtilities::ConvertString(FormattedString, NativeString);
 
-	return FText::CreateChronologicalText(MoveTemp(NativeString), MakeShareable(new FTextHistory_AsDateTime(DateTime, DateStyle, TimeStyle, TimeZone, TargetCulture)));
+	return FText::CreateChronologicalText(MakeShareable(new TGeneratedTextData<FTextHistory_AsDateTime>(MoveTemp(NativeString), FTextHistory_AsDateTime(DateTime, DateStyle, TimeStyle, TimeZone, TargetCulture))));
 }
 
 FText FText::AsMemory(SIZE_T NumBytes, const FNumberFormattingOptions* const Options, const FCulturePtr& TargetCulture)
@@ -143,12 +144,12 @@ int32 FText::CompareTo( const FText& Other, const ETextComparisonLevel::Type Com
 
 	// Create an iterator for 'this' so that we can interface with ICU
 	UCharIterator DisplayStringICUIterator;
-	FICUTextCharacterIterator DisplayStringIterator(&DisplayString.Get());
+	FICUTextCharacterIterator DisplayStringIterator(&TextData->GetDisplayString());
 	uiter_setCharacterIterator(&DisplayStringICUIterator, &DisplayStringIterator);
 
 	// Create an iterator for 'Other' so that we can interface with ICU
 	UCharIterator OtherDisplayStringICUIterator;
-	FICUTextCharacterIterator OtherDisplayStringIterator(&Other.DisplayString.Get());
+	FICUTextCharacterIterator OtherDisplayStringIterator(&Other.TextData->GetDisplayString());
 	uiter_setCharacterIterator(&OtherDisplayStringICUIterator, &OtherDisplayStringIterator);
 
 	UErrorCode ICUStatus = U_ZERO_ERROR;
@@ -185,12 +186,12 @@ public:
 	{
 		// Create an iterator for 'A' so that we can interface with ICU
 		UCharIterator ADisplayStringICUIterator;
-		FICUTextCharacterIterator ADisplayStringIterator(&A.DisplayString.Get());
+		FICUTextCharacterIterator ADisplayStringIterator(&A.TextData->GetDisplayString());
 		uiter_setCharacterIterator(&ADisplayStringICUIterator, &ADisplayStringIterator);
 
 		// Create an iterator for 'B' so that we can interface with ICU
 		UCharIterator BDisplayStringICUIterator;
-		FICUTextCharacterIterator BDisplayStringIterator(&B.DisplayString.Get());
+		FICUTextCharacterIterator BDisplayStringIterator(&B.TextData->GetDisplayString());
 		uiter_setCharacterIterator(&BDisplayStringICUIterator, &BDisplayStringIterator);
 
 		UErrorCode ICUStatus = U_ZERO_ERROR;
diff --git a/Engine/Source/Runtime/Core/Private/Internationalization/ICUText.h b/Engine/Source/Runtime/Core/Private/Internationalization/ICUText.h
--- a/Engine/Source/Runtime/Core/Private/Internationalization/ICUText.h
+++ b/Engine/Source/Runtime/Core/Private/Internationalization/ICUText.h
@@ -4,6 +4,7 @@
 
 #include "ICUUtilities.h"
 #include "ICUCulture.h"
+#include "TextData.h"
 #include "TextHistory.h"
 #include "FastDecimalFormat.h"
 #include "unicode/utypes.h"
@@ -36,7 +37,7 @@ FText FText::AsNumberTemplate(T1 Val, const FNumberFormattingOptions* const Opti
 	ICUUtilities::ConvertString(FormattedString, NativeString);
 #endif // EXPERIMENTAL_TEXT_FAST_DECIMAL_FORMAT
 
-	return FText::CreateNumericalText(MoveTemp(NativeString), MakeShareable(new FTextHistory_AsNumber(Val, Options, TargetCulture)));
+	return FText::CreateNumericalText(MakeShareable(new TGeneratedTextData<FTextHistory_AsNumber>(MoveTemp(NativeString), FTextHistory_AsNumber(Val, Options, TargetCulture))));
 }
 
 template<typename T1, typename T2>
@@ -62,7 +63,7 @@ FText FText::AsCurrencyTemplate(T1 Val, const FString& CurrencyCode, const FNumb
 	ICUUtilities::ConvertString(FormattedString, NativeString);
 #endif // EXPERIMENTAL_TEXT_FAST_DECIMAL_FORMAT
 
-	return FText::CreateNumericalText(MoveTemp(NativeString), MakeShareable(new FTextHistory_AsCurrency(Val, CurrencyCode, Options, TargetCulture)));
+	return FText::CreateNumericalText(MakeShareable(new TGeneratedTextData<FTextHistory_AsCurrency>(MoveTemp(NativeString), FTextHistory_AsCurrency(Val, CurrencyCode, Options, TargetCulture))));
 }
 
 template<typename T1, typename T2>
@@ -88,7 +89,7 @@ FText FText::AsPercentTemplate(T1 Val, const FNumberFormattingOptions* const Opt
 	ICUUtilities::ConvertString(FormattedString, NativeString);
 #endif // EXPERIMENTAL_TEXT_FAST_DECIMAL_FORMAT
 
-	return FText::CreateNumericalText(MoveTemp(NativeString), MakeShareable(new FTextHistory_AsPercent(Val, Options, TargetCulture)));
+	return FText::CreateNumericalText(MakeShareable(new TGeneratedTextData<FTextHistory_AsPercent>(MoveTemp(NativeString), FTextHistory_AsPercent(Val, Options, TargetCulture))));
 }
 
 #undef EXPERIMENTAL_TEXT_FAST_DECIMAL_FORMAT
diff --git a/Engine/Source/Runtime/Core/Private/Internationalization/LegacyText.cpp b/Engine/Source/Runtime/Core/Private/Internationalization/LegacyText.cpp
--- a/Engine/Source/Runtime/Core/Private/Internationalization/LegacyText.cpp
+++ b/Engine/Source/Runtime/Core/Private/Internationalization/LegacyText.cpp
@@ -4,6 +4,7 @@
 
 #if !UE_ENABLE_ICU
 #include "Text.h"
+#include "TextData.h"
 
 bool FText::IsWhitespace( const TCHAR Char )
 {
@@ -63,17 +64,17 @@ FText FText::AsMemory( SIZE_T NumBytes, const FNumberFormattingOptions* const Op
 
 int32 FText::CompareTo( const FText& Other, const ETextComparisonLevel::Type ComparisonLevel ) const
 {
-	return FCString::Strcmp( *DisplayString.Get(), *Other.DisplayString.Get() );
+	return FCString::Strcmp( *TextData->GetDisplayString(), *Other.TextData->GetDisplayString() );
 }
 
 int32 FText::CompareToCaseIgnored( const FText& Other ) const
 {
-	return FCString::Stricmp( *DisplayString.Get(), *Other.DisplayString.Get() );
+	return FCString::Stricmp( *TextData->GetDisplayString(), *Other.TextData->GetDisplayString() );
 }
 
 bool FText::EqualTo( const FText& Other, const ETextComparisonLevel::Type ComparisonLevel ) const
 {
-	return FCString::Strcmp( *DisplayString.Get(), *Other.DisplayString.Get() ) == 0;
+	return FCString::Strcmp( *TextData->GetDisplayString(), *Other.TextData->GetDisplayString() ) == 0;
 }
 
 bool FText::EqualToCaseIgnored( const FText& Other ) const
diff --git a/Engine/Source/Runtime/Core/Private/Internationalization/LegacyText.h b/Engine/Source/Runtime/Core/Private/Internationalization/LegacyText.h
--- a/Engine/Source/Runtime/Core/Private/Internationalization/LegacyText.h
+++ b/Engine/Source/Runtime/Core/Private/Internationalization/LegacyText.h
@@ -2,13 +2,15 @@
 
 #pragma once
 
+#include "TextData.h"
+
 template<typename T1, typename T2>
 FText FText::AsNumberTemplate(T1 Val, const FNumberFormattingOptions* const Options, const FCulturePtr& TargetCulture)
 {
 	FInternationalization& I18N = FInternationalization::Get();
 	checkf(I18N.IsInitialized() == true, TEXT("FInternationalization is not initialized. An FText formatting method was likely used in static object initialization - this is not supported."));
 	const FCulture& Culture = TargetCulture.IsValid() ? *TargetCulture : *I18N.GetCurrentCulture();
-	return FText::CreateNumericalText( Culture.NumberFormattingRule.AsNumber(Val), nullptr );
+	return FText::CreateNumericalText(MakeShareable(new TGeneratedTextData<FTextHistory_AsNumber>(Culture.NumberFormattingRule.AsNumber(Val), FTextHistory_AsNumber(Val, Options, TargetCulture))));
 }
 
 template<typename T1, typename T2>
@@ -17,7 +19,7 @@ FText FText::AsCurrencyTemplate(T1 Val, const FString& CurrencyCode, const FNumb
 	FInternationalization& I18N = FInternationalization::Get();
 	checkf(I18N.IsInitialized() == true, TEXT("FInternationalization is not initialized. An FText formatting method was likely used in static object initialization - this is not supported."));
 	const FCulture& Culture = TargetCulture.IsValid() ? *TargetCulture : *I18N.GetCurrentCulture();
-	return FText::CreateNumericalText( Culture.NumberFormattingRule.AsCurrency(Val), nullptr );
+	return FText::CreateNumericalText(MakeShareable(new TGeneratedTextData<FTextHistory_AsCurrency>(Culture.NumberFormattingRule.AsCurrency(Val), FTextHistory_AsCurrency(Val, CurrencyCode, Options, TargetCulture))));
 }
 
 template<typename T1, typename T2>
@@ -26,5 +28,5 @@ FText FText::AsPercentTemplate(T1 Val, const FNumberFormattingOptions* const Opt
 	FInternationalization& I18N = FInternationalization::Get();
 	checkf(I18N.IsInitialized() == true, TEXT("FInternationalization is not initialized. An FText formatting method was likely used in static object initialization - this is not supported."));
 	const FCulture& Culture = TargetCulture.IsValid() ? *TargetCulture : *I18N.GetCurrentCulture();
-	return FText::CreateNumericalText( Culture.NumberFormattingRule.AsPercent(Val), nullptr );
+	return FText::CreateNumericalText(MakeShareable(new TGeneratedTextData<FTextHistory_AsPercent>(Culture.NumberFormattingRule.AsPercent(Val), FTextHistory_AsPercent(Val, Options, TargetCulture))));
 }
\ No newline at end of file
diff --git a/Engine/Source/Runtime/Core/Private/Internationalization/Text.cpp b/Engine/Source/Runtime/Core/Private/Internationalization/Text.cpp
--- a/Engine/Source/Runtime/Core/Private/Internationalization/Text.cpp
+++ b/Engine/Source/Runtime/Core/Private/Internationalization/Text.cpp
@@ -8,6 +8,7 @@
 #include "LegacyText.h"
 #endif
 
+#include "TextData.h"
 #include "TextHistory.h"
 
 #include "DebugSerializationFlags.h"
@@ -20,13 +21,6 @@ DEFINE_LOG_CATEGORY(LogText);
 
 #define LOCTEXT_NAMESPACE "Core.Text"
 
-FArchive& operator<<( FArchive& Ar, FFormatArgumentData& Value )
-{
-	Ar << Value.ArgumentName;
-	Ar << Value.ArgumentValue;
-	return Ar;
-}
-
 bool FTextInspector::ShouldGatherForLocalization(const FText& Text)
 {
 	return Text.ShouldGatherForLocalization();
@@ -34,171 +28,56 @@ bool FTextInspector::ShouldGatherForLocalization(const FText& Text)
 
 TOptional<FString> FTextInspector::GetNamespace(const FText& Text)
 {
-	FString Namespace;
-	FString Key;
-	const bool WasNamespaceAndKeyFound = FTextLocalizationManager::Get().FindNamespaceAndKeyFromDisplayString(Text.DisplayString, Namespace, Key);
-	return WasNamespaceAndKeyFound ? TOptional<FString>(Namespace) : TOptional<FString>();
+	FTextDisplayStringPtr LocalizedString = Text.TextData->GetLocalizedString();
+	if (LocalizedString.IsValid())
+	{
+		FString Namespace;
+		FString Key;
+		if (FTextLocalizationManager::Get().FindNamespaceAndKeyFromDisplayString(LocalizedString.ToSharedRef(), Namespace, Key))
+		{
+			return Namespace;
+		}
+	}
+	return TOptional<FString>();
 }
 
 TOptional<FString> FTextInspector::GetKey(const FText& Text)
 {
-	FString Namespace;
-	FString Key;
-	const bool WasNamespaceAndKeyFound = FTextLocalizationManager::Get().FindNamespaceAndKeyFromDisplayString(Text.DisplayString, Namespace, Key);
-	return WasNamespaceAndKeyFound ? TOptional<FString>(Key) : TOptional<FString>();
+	FTextDisplayStringPtr LocalizedString = Text.TextData->GetLocalizedString();
+	if (LocalizedString.IsValid())
+	{
+		FString Namespace;
+		FString Key;
+		if (FTextLocalizationManager::Get().FindNamespaceAndKeyFromDisplayString(LocalizedString.ToSharedRef(), Namespace, Key))
+		{
+			return Key;
+		}
+	}
+	return TOptional<FString>();
 }
 
 const FString* FTextInspector::GetSourceString(const FText& Text)
 {
-	return Text.GetSourceString().Get();
+	return &Text.GetSourceString();
 }
 
 const FString& FTextInspector::GetDisplayString(const FText& Text)
 {
-	return Text.DisplayString.Get();
+	return Text.TextData->GetDisplayString();
 }
 
 const FTextDisplayStringRef FTextInspector::GetSharedDisplayString(const FText& Text)
 {
-	return Text.DisplayString;
+	// todo: calling PersistText here probably isn't the right thing to do, however it avoids having to make an external API change at this point
+	Text.TextData->PersistText();
+	return Text.TextData->GetLocalizedString().ToSharedRef();
 }
 
-int32 FTextInspector::GetFlags(const FText& Text)
+uint32 FTextInspector::GetFlags(const FText& Text)
 {
 	return Text.Flags;
 }
 
-FFormatArgumentValue::FFormatArgumentValue()
-	: Type(EFormatArgumentType::Int)
-{
-	IntValue = 0;
-};
-
-
-FFormatArgumentValue::FFormatArgumentValue(const int Value)
-	: Type(EFormatArgumentType::Int)
-{
-	IntValue = Value;
-}
-
-FFormatArgumentValue::FFormatArgumentValue(const unsigned int Value)
-	: Type(EFormatArgumentType::UInt)
-{
-	UIntValue = Value;
-}
-
-FFormatArgumentValue::FFormatArgumentValue(const int64 Value)
-	: Type(EFormatArgumentType::Int)
-{
-	IntValue = Value;
-}
-
-FFormatArgumentValue::FFormatArgumentValue(const uint64 Value)
-	: Type(EFormatArgumentType::UInt)
-{
-	UIntValue = Value;
-}
-
-FFormatArgumentValue::FFormatArgumentValue(const float Value)
-	: Type(EFormatArgumentType::Float)
-{
-	FloatValue = Value;
-}
-
-FFormatArgumentValue::FFormatArgumentValue(const double Value)
-	: Type(EFormatArgumentType::Double)
-{
-	DoubleValue = Value;
-}
-
-FFormatArgumentValue::FFormatArgumentValue(const FText& Value)
-	: Type(EFormatArgumentType::Text)
-{
-	TextValue = new FText(Value);
-}
-
-FFormatArgumentValue::FFormatArgumentValue(const FFormatArgumentValue& Source)
-	: Type(Source.Type)
-{
-	switch(Type)
-	{
-	case EFormatArgumentType::Int:
-		{
-			IntValue = Source.IntValue;
-		}
-		break;
-	case EFormatArgumentType::UInt:
-		{
-			UIntValue = Source.UIntValue;
-		}
-		break;
-	case EFormatArgumentType::Float:
-		{
-			FloatValue = Source.FloatValue;
-		}
-	case EFormatArgumentType::Double:
-		{
-			DoubleValue = Source.DoubleValue;
-		}
-		break;
-	case EFormatArgumentType::Text:
-		{
-			TextValue = new FText(*Source.TextValue);
-		}
-		break;
-	}
-}
-
-FFormatArgumentValue::~FFormatArgumentValue()
-{
-	if(Type == EFormatArgumentType::Text)
-	{
-		delete TextValue;
-	}
-}
-
-FArchive& operator<<( FArchive& Ar, FFormatArgumentValue& Value )
-{
-	int8 TypeAsInt8 = Value.Type;
-	Ar << TypeAsInt8;
-	Value.Type = (EFormatArgumentType::Type)TypeAsInt8;
-
-	switch(Value.Type)
-	{
-	case EFormatArgumentType::Double:
-		{
-			Ar << Value.DoubleValue;
-			break;
-		}
-	case EFormatArgumentType::Float:
-		{
-			Ar << Value.FloatValue;
-			break;
-		}
-	case EFormatArgumentType::Int:
-		{
-			Ar << Value.IntValue;
-			break;
-		}
-	case EFormatArgumentType::UInt:
-		{
-			Ar << Value.UIntValue;
-			break;
-		}
-	case EFormatArgumentType::Text:
-		{
-			if(Ar.IsLoading())
-			{
-				Value.TextValue = new FText;
-			}
-			Ar << *Value.TextValue;
-			break;
-		}
-	}
-	
-	return Ar;
-}
-
 // These default values have been duplicated to the KismetTextLibrary functions for Blueprints. Please replicate any changes there!
 FNumberFormattingOptions::FNumberFormattingOptions()
 	: UseGrouping(true)
@@ -276,13 +155,13 @@ const FText FText::UnEscapedCloseBraceOutsideOfArgumentBlock = LOCTEXT("Error_Un
 const FText FText::SerializationFailureError = LOCTEXT("Error_SerializationFailure", "ERR: Transient text cannot be serialized \"{0}\".");
 
 FText::FText()
-	: DisplayString( GetEmpty().DisplayString )
+	: TextData(GetEmpty().TextData)
 	, Flags(0)
 {
 }
 
 FText::FText( EInitToEmptyString )
-	: DisplayString( new FString() )
+	: TextData(new TLocalizedTextData<FTextHistory_Base>(MakeShareable(new FString())))
 	, Flags(0)
 {
 }
@@ -294,65 +173,72 @@ FText::FText( EInitToEmptyString )
 	FText& FText::operator=(FText&& Other) = default;
 #else
 FText::FText(const FText& Source)
-	: DisplayString(Source.DisplayString)
-	, History(Source.History)
+	: TextData(Source.TextData)
 	, Flags(Source.Flags)
 {
 }
 
 FText::FText(FText&& Source)
-	: DisplayString(MoveTemp(Source.DisplayString))
-	, History(MoveTemp(Source.History))
-	, Flags(MoveTemp(Source.Flags))
+	: TextData(MoveTemp(Source.TextData))
+	, Flags(Source.Flags)
 {
 }
 
 FText& FText::operator=(const FText& Source)
 {
-	DisplayString = Source.DisplayString;
-	History = Source.History;
-	Flags = Source.Flags;
-
+	if (this != &Source)
+	{
+		TextData = Source.TextData;
+		Flags = Source.Flags;
+	}
 	return *this;
 }
 
 FText& FText::operator=(FText&& Source)
 {
-	DisplayString = MoveTemp(Source.DisplayString);
-	History = MoveTemp(Source.History);
-	Flags = MoveTemp(Source.Flags);
-
+	if (this != &Source)
+	{
+		TextData = MoveTemp(Source.TextData);
+		Flags = Source.Flags;
+	}
 	return *this;
 }
 #endif
 
-FText::FText( FString InSourceString, TSharedPtr<class FTextHistory, ESPMode::ThreadSafe> InHistory )
-	: DisplayString( new FString( MoveTemp(InSourceString) ))
-	, History( MoveTemp(InHistory) )
+FText::FText( TSharedRef<ITextData, ESPMode::ThreadSafe> InTextData )
+	: TextData(MoveTemp(InTextData))
 	, Flags(0)
 {
-	if (!History.IsValid())
-	{
-		History = MakeShareable(new FTextHistory_Base(DisplayString));
-	}
 }
 
-FText::FText( FString InSourceString, FString InNamespace, FString InKey, int32 InFlags )
-	: DisplayString( FTextLocalizationManager::Get().GetDisplayString(InNamespace, InKey, &InSourceString) )
+FText::FText( FString InSourceString )
+	: TextData(new TGeneratedTextData<FTextHistory_Base>(InSourceString))
+	, Flags(0)
+{
+	TextData->SetTextHistory(FTextHistory_Base(MoveTemp(InSourceString)));
+}
+
+FText::FText( FString InSourceString, FString InNamespace, FString InKey, uint32 InFlags )
+	: TextData(new TLocalizedTextData<FTextHistory_Base>(FTextLocalizationManager::Get().GetDisplayString(InNamespace, InKey, &InSourceString)))
 	, Flags(InFlags)
 {
-	History = MakeShareable(new FTextHistory_Base(InSourceString));
+	TextData->SetTextHistory(FTextHistory_Base(MoveTemp(InSourceString)));
 }
 
+bool FText::IsEmpty() const
+{
+	return TextData->GetDisplayString().IsEmpty();
+}
 
 bool FText::IsEmptyOrWhitespace() const
 {
-	if (DisplayString.Get().IsEmpty())
+	const FString& DisplayString = TextData->GetDisplayString();
+	if (DisplayString.IsEmpty())
 	{
 		return true;
 	}
 
-	for( const TCHAR Character : DisplayString.Get() )
+	for( const TCHAR Character : DisplayString )
 	{
 		if (!IsWhitespace(Character))
 		{
@@ -381,17 +267,17 @@ FText FText::TrimPreceding( const FText& InText )
 		TrimmedString = TrimmedString.Right( TrimmedString.Len() - StartPos );
 	}
 
-	FText NewText = FText( MoveTemp( TrimmedString ), nullptr );
+	FText NewText = FText( MoveTemp( TrimmedString ) );
 
 	if (!GIsEditor)
 	{
-		if( (InText.Flags & (1 << ETextFlag::CultureInvariant)) != 0 )
+		if( (NewText.Flags & ETextFlag::CultureInvariant) != 0 )
 		{
-			NewText.Flags = NewText.Flags | ETextFlag::Transient;
+			NewText.Flags |= ETextFlag::Transient;
 		}
 		else
 		{
-			NewText.Flags = NewText.Flags | ETextFlag::CultureInvariant;
+			NewText.Flags |= ETextFlag::CultureInvariant;
 		}
 	}
 
@@ -416,17 +302,17 @@ FText FText::TrimTrailing( const FText& InText )
 		TrimmedString = TrimmedString.Left( EndPos + 1 );
 	}
 
-	FText NewText = FText( MoveTemp ( TrimmedString ), nullptr );
+	FText NewText = FText( MoveTemp ( TrimmedString ) );
 
 	if (!GIsEditor)
 	{
-		if( (InText.Flags & (1 << ETextFlag::CultureInvariant)) != 0 )
+		if( (NewText.Flags & ETextFlag::CultureInvariant) != 0 )
 		{
-			NewText.Flags = NewText.Flags & ETextFlag::Transient;
+			NewText.Flags |= ETextFlag::Transient;
 		}
 		else
 		{
-			NewText.Flags = NewText.Flags & ETextFlag::CultureInvariant;
+			NewText.Flags |= ETextFlag::CultureInvariant;
 		}
 	}
 
@@ -463,17 +349,17 @@ FText FText::TrimPrecedingAndTrailing( const FText& InText )
 		TrimmedString = TrimmedString.Mid( StartPos, Len );
 	}
 
-	FText NewText = FText( MoveTemp( TrimmedString ), nullptr );
+	FText NewText = FText( MoveTemp( TrimmedString ) );
 
 	if (!GIsEditor)
 	{
-		if( (InText.Flags & (1 << ETextFlag::CultureInvariant)) != 0 )
+		if( (NewText.Flags & ETextFlag::CultureInvariant) != 0 )
 		{
-			NewText.Flags = NewText.Flags | ETextFlag::Transient;
+			NewText.Flags |= ETextFlag::Transient;
 		}
 		else
 		{
-			NewText.Flags = NewText.Flags | ETextFlag::CultureInvariant;
+			NewText.Flags |= ETextFlag::CultureInvariant;
 		}
 	}
 
@@ -536,15 +422,15 @@ public:
 
 	static int32 EstimateArgumentValueLength(const FFormatArgumentValue& ArgumentValue)
 	{
-		switch(ArgumentValue.Type)
+		switch(ArgumentValue.GetType())
 		{
 		case EFormatArgumentType::Text:
-			return ArgumentValue.TextValue->ToString().Len();
+			return ArgumentValue.GetTextValue().ToString().Len();
 		case EFormatArgumentType::Int:
 		case EFormatArgumentType::UInt:
 		case EFormatArgumentType::Float:
 		case EFormatArgumentType::Double:
-			return 16;
+			return 20;
 		default:
 			break;
 		}
@@ -677,7 +563,7 @@ public:
 		}
 	}
 
-	static FText Format(const FText& Pattern, const int32 EstimatedArgumentValuesLength, FGetArgumentValue GetArgumentValue, const TSharedPtr<FTextHistory, ESPMode::ThreadSafe>& InHistory, bool bInRebuildText, bool bInRebuildAsSource)
+	static FString Format(const FText& Pattern, const int32 EstimatedArgumentValuesLength, FGetArgumentValue GetArgumentValue, bool bInRebuildText, bool bInRebuildAsSource)
 	{
 		checkf(FInternationalization::Get().IsInitialized() == true, TEXT("FInternationalization is not initialized. An FText formatting method was likely used in static object initialization - this is not supported."));
 		//SCOPE_CYCLE_COUNTER( STAT_TextFormat );
@@ -787,37 +673,39 @@ public:
 								if( PossibleArgumentValue )
 								{
 									const FFormatArgumentValue& ArgumentValue = *PossibleArgumentValue;
-									switch(ArgumentValue.Type)
+									switch(ArgumentValue.GetType())
 									{
 									case EFormatArgumentType::Text:
 										{
+											const FText& TextValue = ArgumentValue.GetTextValue();
+
 											// When doing a rebuild, all FText arguments need to be rebuilt during the Format
 											if( bInRebuildText )
 											{
-												ArgumentValue.TextValue->Rebuild();
+												TextValue.Rebuild();
 											}
 
-											ResultString += bInRebuildAsSource ? ArgumentValue.TextValue->BuildSourceString() : ArgumentValue.TextValue->ToString();
+											ResultString += bInRebuildAsSource ? TextValue.BuildSourceString() : TextValue.ToString();
 										}
 										break;
 									case EFormatArgumentType::Int:
 										{
-											ResultString += FText::AsNumber(ArgumentValue.IntValue).ToString();
+											ResultString += FText::AsNumber(ArgumentValue.GetIntValue()).ToString();
 										}
 										break;
 									case EFormatArgumentType::UInt:
 										{
-											ResultString += FText::AsNumber(ArgumentValue.UIntValue).ToString();
+											ResultString += FText::AsNumber(ArgumentValue.GetUIntValue()).ToString();
 										}
 										break;
 									case EFormatArgumentType::Float:
 										{
-											ResultString += FText::AsNumber(ArgumentValue.FloatValue).ToString();
+											ResultString += FText::AsNumber(ArgumentValue.GetFloatValue()).ToString();
 										}
 										break;
 									case EFormatArgumentType::Double:
 										{
-											ResultString += FText::AsNumber(ArgumentValue.DoubleValue).ToString();
+											ResultString += FText::AsNumber(ArgumentValue.GetDoubleValue()).ToString();
 										}
 										break;
 									}
@@ -844,12 +732,7 @@ public:
 			}
 		}
 
-		FText Result = FText(MoveTemp(ResultString), InHistory);
-		if (!GIsEditor)
-		{
-			Result.Flags = Result.Flags | ETextFlag::Transient;
-		}
-		return Result;
+		return ResultString;
 	}
 };
 
@@ -889,7 +772,14 @@ FText FText::FormatInternal(const FText& Pattern, const FFormatNamedArguments& A
 		return Arguments.Find(ArgumentName);
 	};
 
-	return FTextFormatHelper::Format(Pattern, EstimatedArgumentValuesLength, FTextFormatHelper::FGetArgumentValue::CreateLambda(GetArgumentValue), MakeShareable(new FTextHistory_NamedFormat(Pattern, Arguments)), bInRebuildText, bInRebuildAsSource);
+	FString ResultString = FTextFormatHelper::Format(Pattern, EstimatedArgumentValuesLength, FTextFormatHelper::FGetArgumentValue::CreateLambda(GetArgumentValue), bInRebuildText, bInRebuildAsSource);
+	
+	FText Result = FText(MakeShareable(new TGeneratedTextData<FTextHistory_NamedFormat>(MoveTemp(ResultString), FTextHistory_NamedFormat(Pattern, Arguments))));
+	if (!GIsEditor)
+	{
+		Result.Flags |= ETextFlag::Transient;
+	}
+	return Result;
 }
 
 FText FText::FormatInternal(const FText& Pattern, const FFormatOrderedArguments& Arguments, bool bInRebuildText, bool bInRebuildAsSource)
@@ -917,7 +807,14 @@ FText FText::FormatInternal(const FText& Pattern, const FFormatOrderedArguments&
 		return ArgumentIndex != INDEX_NONE && ArgumentIndex < Arguments.Num() ? &(Arguments[ArgumentIndex]) : nullptr;
 	};
 
-	return FTextFormatHelper::Format(Pattern, EstimatedArgumentValuesLength, FTextFormatHelper::FGetArgumentValue::CreateLambda(GetArgumentValue), MakeShareable(new FTextHistory_OrderedFormat(Pattern, Arguments)), bInRebuildText, bInRebuildAsSource);
+	FString ResultString = FTextFormatHelper::Format(Pattern, EstimatedArgumentValuesLength, FTextFormatHelper::FGetArgumentValue::CreateLambda(GetArgumentValue), bInRebuildText, bInRebuildAsSource);
+
+	FText Result = FText(MakeShareable(new TGeneratedTextData<FTextHistory_OrderedFormat>(MoveTemp(ResultString), FTextHistory_OrderedFormat(Pattern, Arguments))));
+	if (!GIsEditor)
+	{
+		Result.Flags |= ETextFlag::Transient;
+	}
+	return Result;
 }
 
 FText FText::FormatInternal(const FText& Pattern, const TArray< struct FFormatArgumentData > Arguments, bool bInRebuildText, bool bInRebuildAsSource)
@@ -938,7 +835,14 @@ FText FText::FormatInternal(const FText& Pattern, const TArray< struct FFormatAr
 		return FormatNamedArguments.Find(ArgumentName);
 	};
 
-	return FTextFormatHelper::Format(Pattern, EstimatedArgumentValuesLength, FTextFormatHelper::FGetArgumentValue::CreateLambda(GetArgumentValue), MakeShareable(new FTextHistory_ArgumentDataFormat(Pattern, Arguments)), bInRebuildText, bInRebuildAsSource);
+	FString ResultString = FTextFormatHelper::Format(Pattern, EstimatedArgumentValuesLength, FTextFormatHelper::FGetArgumentValue::CreateLambda(GetArgumentValue), bInRebuildText, bInRebuildAsSource);
+
+	FText Result = FText(MakeShareable(new TGeneratedTextData<FTextHistory_ArgumentDataFormat>(MoveTemp(ResultString), FTextHistory_ArgumentDataFormat(Pattern, Arguments))));
+	if (!GIsEditor)
+	{
+		Result.Flags |= ETextFlag::Transient;
+	}
+	return Result;
 }
 
 /**
@@ -1021,7 +925,7 @@ CORE_API FArchive& operator<<( FArchive& Ar, FText& Value )
 		FString SourceStringToImplantIntoHistory;
 		Ar << SourceStringToImplantIntoHistory;
 
-		Value.History = MakeShareable(new FTextHistory_Base(SourceStringToImplantIntoHistory));
+		FTextDisplayStringPtr DisplayString;
 
 		// Namespaces and keys are no longer stored in the FText, we need to read them in and discard
 		if( Ar.UE4Ver() >= VER_UE4_ADDED_NAMESPACE_AND_KEY_DATA_TO_FTEXT )
@@ -1033,10 +937,16 @@ CORE_API FArchive& operator<<( FArchive& Ar, FText& Value )
 			Ar << Key;
 
 			// Get the DisplayString using the namespace, key, and source string.
-			Value.DisplayString = FTextLocalizationManager::Get().GetDisplayString(Namespace, Key, &SourceStringToImplantIntoHistory);
+			DisplayString = FTextLocalizationManager::Get().GetDisplayString(Namespace, Key, &SourceStringToImplantIntoHistory);
+		}
+		else
+		{
+			DisplayString = MakeShareable(new FString());
 		}
-	}
 
+		check(DisplayString.IsValid());
+		Value.TextData = MakeShareable(new TLocalizedTextData<FTextHistory_Base>(DisplayString.ToSharedRef(), FTextHistory_Base(MoveTemp(SourceStringToImplantIntoHistory))));
+	}
 
 #if WITH_EDITOR
 	/*if (Ar.IsCooking() && Ar.IsSaving() && Ar.IsPersistent() && !(Ar.GetDebugSerializationFlags()&DSF_IsTagExportsArchive))
@@ -1052,10 +962,15 @@ CORE_API FArchive& operator<<( FArchive& Ar, FText& Value )
 	}*/
 #endif
 
-	int32 OriginalFlags = Value.Flags;
-	if(Ar.IsPersistent() && Ar.IsSaving())
+	const int32 OriginalFlags = Value.Flags;
+
+	if(Ar.IsSaving())
 	{
-		Value.Flags &= ~(ETextFlag::ConvertedProperty|ETextFlag::InitializedFromString); // Remove conversion flag before saving.
+		Value.TextData->PersistText(); // We always need to do this when saving so that we can save the history correctly
+		if(Ar.IsPersistent())
+		{
+			Value.Flags &= ~(ETextFlag::ConvertedProperty); // Remove conversion flag before saving.
+		}
 	}
 	Ar << Value.Flags;
 
@@ -1066,10 +981,14 @@ CORE_API FArchive& operator<<( FArchive& Ar, FText& Value )
 
 	if( Ar.UE4Ver() >= VER_UE4_FTEXT_HISTORY )
 	{
+		bool bSerializeHistory = true;
+
 		if(Ar.IsSaving())
 		{
-			// If there is no history, mark it, otherwise the history will serialize it's type
-			if(!Value.History.IsValid())
+			// Skip the history for empty texts
+			bSerializeHistory = !Value.IsEmpty();
+
+			if(!bSerializeHistory)
 			{
 				int8 NoHistory = INDEX_NONE;
 				Ar << NoHistory;
@@ -1085,67 +1004,68 @@ CORE_API FArchive& operator<<( FArchive& Ar, FText& Value )
 			switch ((ETextHistoryType)HistoryType)
 			{
 			case ETextHistoryType::Base:
-			{
-										   Value.History = MakeShareable(new FTextHistory_Base);
-										   break;
-			}
+				{
+					Value.TextData = MakeShareable(new TLocalizedTextData<FTextHistory_Base>());
+					break;
+				}
 			case ETextHistoryType::NamedFormat:
-			{
-												  Value.History = MakeShareable(new FTextHistory_NamedFormat);
-												  break;
-			}
+				{
+					Value.TextData = MakeShareable(new TLocalizedTextData<FTextHistory_NamedFormat>());
+					break;
+				}
 			case ETextHistoryType::OrderedFormat:
-			{
-													Value.History = MakeShareable(new FTextHistory_OrderedFormat);
-													break;
-			}
+				{
+					Value.TextData = MakeShareable(new TLocalizedTextData<FTextHistory_OrderedFormat>());
+					break;
+				}
 			case ETextHistoryType::ArgumentFormat:
-			{
-													 Value.History = MakeShareable(new FTextHistory_ArgumentDataFormat);
-													 break;
-			}
+				{
+					Value.TextData = MakeShareable(new TLocalizedTextData<FTextHistory_ArgumentDataFormat>());
+					break;
+				}
 			case ETextHistoryType::AsNumber:
-			{
-											   Value.History = MakeShareable(new FTextHistory_AsNumber);
-											   break;
-			}
+				{
+					Value.TextData = MakeShareable(new TLocalizedTextData<FTextHistory_AsNumber>());
+					break;
+				}
 			case ETextHistoryType::AsPercent:
-			{
-												Value.History = MakeShareable(new FTextHistory_AsPercent);
-												break;
-			}
+				{
+					Value.TextData = MakeShareable(new TLocalizedTextData<FTextHistory_AsPercent>());
+					break;
+				}
 			case ETextHistoryType::AsCurrency:
-			{
-												 Value.History = MakeShareable(new FTextHistory_AsCurrency);
-												 break;
-			}
+				{
+					Value.TextData = MakeShareable(new TLocalizedTextData<FTextHistory_AsCurrency>());
+					break;
+				}
 			case ETextHistoryType::AsDate:
-			{
-											 Value.History = MakeShareable(new FTextHistory_AsDate);
-											 break;
-			}
+				{
+					Value.TextData = MakeShareable(new TLocalizedTextData<FTextHistory_AsDate>());
+					break;
+				}
 			case ETextHistoryType::AsTime:
-			{
-											 Value.History = MakeShareable(new FTextHistory_AsTime);
-											 break;
-			}
+				{
+					Value.TextData = MakeShareable(new TLocalizedTextData<FTextHistory_AsTime>());
+					break;
+				}
 			case ETextHistoryType::AsDateTime:
-			{
-												 Value.History = MakeShareable(new FTextHistory_AsDateTime);
-												 break;
-			}
+				{
+					Value.TextData = MakeShareable(new TLocalizedTextData<FTextHistory_AsDateTime>());
+					break;
+				}
 			default:
-			{
-					   Value.History.Reset();
-					   Value.DisplayString = FText::GetEmpty().DisplayString;
-			}
+				{
+					bSerializeHistory = false;
+					Value.TextData = FText::GetEmpty().TextData;
+				}
 			}
 		}
 
-		if(Value.History.IsValid())
+		if(bSerializeHistory)
 		{
-			Value.History->Serialize(Ar);
-			Value.History->SerializeForDisplayString(Ar, Value.DisplayString);
+			FTextHistory& MutableTextHistory = Value.TextData->GetMutableTextHistory();
+			MutableTextHistory.Serialize(Ar);
+			MutableTextHistory.SerializeForDisplayString(Ar, Value.TextData->GetMutableLocalizedString());
 		}
 	}
 
@@ -1171,13 +1091,13 @@ CORE_API FArchive& operator<<( FArchive& Ar, FText& Value )
 #if WITH_EDITOR
 FText FText::ChangeKey( FString Namespace, FString Key, const FText& Text )
 {
-	return FText( *Text.History->GetSourceString().Get(), MoveTemp( Namespace ), MoveTemp( Key ) );
+	return FText( *Text.TextData->GetTextHistory().GetSourceString(), MoveTemp( Namespace ), MoveTemp( Key ), Text.Flags );
 }
 #endif
 
-FText FText::CreateNumericalText(FString InSourceString, TSharedPtr<class FTextHistory, ESPMode::ThreadSafe> InHistory)
+FText FText::CreateNumericalText(TSharedRef<ITextData, ESPMode::ThreadSafe> InTextData)
 {
-	FText NewText = FText( MoveTemp( InSourceString ), MoveTemp( InHistory ) );
+	FText NewText = FText(MoveTemp(InTextData));
 	if (!GIsEditor)
 	{
 		NewText.Flags |= ETextFlag::Transient;
@@ -1185,9 +1105,9 @@ FText FText::CreateNumericalText(FString InSourceString, TSharedPtr<class FTextH
 	return NewText;
 }
 
-FText FText::CreateChronologicalText(FString InSourceString, TSharedPtr<class FTextHistory, ESPMode::ThreadSafe> InHistory)
+FText FText::CreateChronologicalText(TSharedRef<ITextData, ESPMode::ThreadSafe> InTextData)
 {
-	FText NewText = FText( MoveTemp( InSourceString ), MoveTemp( InHistory ) );
+	FText NewText = FText(MoveTemp(InTextData));
 	if (!GIsEditor)
 	{
 		NewText.Flags |= ETextFlag::Transient;
@@ -1197,20 +1117,12 @@ FText FText::CreateChronologicalText(FString InSourceString, TSharedPtr<class FT
 
 FText FText::FromName( const FName& Val) 
 {
-	FText NewText = FText( Val.ToString(), nullptr );
-
-	if (!GIsEditor)
-	{
-		NewText.Flags |= ETextFlag::CultureInvariant;
-	}
-	NewText.Flags |= ETextFlag::InitializedFromString;
-
-	return NewText; 
+	return FText::FromString(Val.ToString());
 }
 
 FText FText::FromString( FString String )
 {
-	FText NewText = FText( MoveTemp(String), nullptr );
+	FText NewText = FText( MoveTemp(String) );
 
 	if (!GIsEditor)
 	{
@@ -1223,7 +1135,7 @@ FText FText::FromString( FString String )
 
 FText FText::AsCultureInvariant( FString String )
 {
-	FText NewText = FText( String, nullptr );
+	FText NewText = FText( String );
 	NewText.Flags |= ETextFlag::CultureInvariant;
 
 	return NewText;
@@ -1241,31 +1153,45 @@ const FString& FText::ToString() const
 {
 	Rebuild();
 
-	return DisplayString.Get();
+	return TextData->GetDisplayString();
 }
 
 FString FText::BuildSourceString() const
 {
-	if(History.IsValid())
-	{
-		return History->ToText(true).ToString();
-	}
-	return DisplayString.Get();
+	return TextData->GetTextHistory().ToText(true).ToString();
+}
+
+bool FText::IsNumeric() const
+{
+	return TextData->GetDisplayString().IsNumeric();
 }
 
 void FText::Rebuild() const
 {
-	if(History.IsValid())
+	FTextHistory& MutableTextHistory = TextData->GetMutableTextHistory();
+	if (MutableTextHistory.IsOutOfDate())
 	{
-		History->Rebuild(DisplayString);
+		// Need to persist the text before the rebuild so that we have a valid localized string pointer
+		TextData->PersistText();
+		MutableTextHistory.Rebuild(TextData->GetLocalizedString().ToSharedRef());
 	}
 }
 
+bool FText::IsTransient() const
+{
+	return (Flags & ETextFlag::Transient) != 0;
+}
+
+bool FText::IsCultureInvariant() const
+{
+	return (Flags & ETextFlag::CultureInvariant) != 0;
+}
+
 bool FText::ShouldGatherForLocalization() const
 {
-	auto SourceString = GetSourceString();
+	const FString& SourceString = GetSourceString();
 
-	auto IsAllWhitespace = [](FString& String) -> bool
+	auto IsAllWhitespace = [](const FString& String) -> bool
 	{
 		for(int32 i = 0; i < String.Len(); ++i)
 		{
@@ -1277,46 +1203,85 @@ bool FText::ShouldGatherForLocalization() const
 		return true;
 	};
 
-	return !((Flags & ETextFlag::CultureInvariant) || (Flags & ETextFlag::Transient)) && SourceString.IsValid() && !SourceString->IsEmpty() && !IsAllWhitespace(*SourceString);
+	return !((Flags & ETextFlag::CultureInvariant) || (Flags & ETextFlag::Transient)) && !SourceString.IsEmpty() && !IsAllWhitespace(SourceString);
 }
 
-TSharedPtr< FString, ESPMode::ThreadSafe > FText::GetSourceString() const
+const FString& FText::GetSourceString() const
 {
-	if(History.IsValid())
+	const FString* SourceString = TextData->GetTextHistory().GetSourceString();
+	if(SourceString)
 	{
-		TSharedPtr< FString, ESPMode::ThreadSafe > SourceString = History->GetSourceString();
-		if(SourceString.IsValid())
-		{
-			return SourceString;
-		}
+		return *SourceString;
 	}
 
-	return DisplayString;
+	return TextData->GetDisplayString();
 }
 
 bool FText::IdenticalTo( const FText& Other ) const
 {
-	// If both instances point to the same string, then both instances are considered identical.
+	// If both instances point to the same data or localized string, then both instances are considered identical.
 	// This is fast as it skips a lexical compare, however it can also return false for two instances that have identical strings, but in different pointers.
 	// For instance, this method will return false for two FText objects created from FText::FromString("Wooble") as they each have unique, non-shared instances.
-	return DisplayString == Other.DisplayString;
+	return TextData == Other.TextData || TextData->GetLocalizedString() == Other.TextData->GetLocalizedString();
 }
 
 void FText::GetSourceTextsFromFormatHistory(TArray<FText>& OutSourceTexts) const
 {
-	History->GetSourceTextsFromFormatHistory(*(this), OutSourceTexts);
+	TextData->GetTextHistory().GetSourceTextsFromFormatHistory(*this, OutSourceTexts);
+}
+
+FArchive& operator<<(FArchive& Ar, FFormatArgumentValue& Value)
+{
+	int8 TypeAsInt8 = Value.Type;
+	Ar << TypeAsInt8;
+	Value.Type = (EFormatArgumentType::Type)TypeAsInt8;
+
+	switch(Value.Type)
+	{
+	case EFormatArgumentType::Double:
+		{
+			Ar << Value.DoubleValue;
+			break;
+		}
+	case EFormatArgumentType::Float:
+		{
+			Ar << Value.FloatValue;
+			break;
+		}
+	case EFormatArgumentType::Int:
+		{
+			Ar << Value.IntValue;
+			break;
+		}
+	case EFormatArgumentType::UInt:
+		{
+			Ar << Value.UIntValue;
+			break;
+		}
+	case EFormatArgumentType::Text:
+		{
+			if(Ar.IsLoading())
+			{
+				Value.TextValue = FText();
+			}
+			Ar << Value.TextValue.GetValue();
+			break;
+		}
+	}
+	
+	return Ar;
 }
 
 FTextSnapshot::FTextSnapshot()
-	: DisplayStringPtr()
+	: TextDataPtr()
 	, HistoryRevision(INDEX_NONE)
 	, Flags(0)
 {
 }
 
 FTextSnapshot::FTextSnapshot(const FText& InText)
-	: DisplayStringPtr(InText.DisplayString)
-	, HistoryRevision(InText.History.IsValid() ? InText.History->Revision : INDEX_NONE)
+	: TextDataPtr(InText.TextData)
+	, HistoryRevision(InText.TextData->GetTextHistory().Revision)
 	, Flags(InText.Flags)
 {
 }
@@ -1327,10 +1292,10 @@ bool FTextSnapshot::IdenticalTo(const FText& InText) const
 	// (this usually happens when ToString() is called)
 	InText.Rebuild();
 
-	const int32 InHistoryRevision = InText.History.IsValid() ? InText.History->Revision : INDEX_NONE;
-	return DisplayStringPtr == InText.DisplayString
-		&& HistoryRevision == InHistoryRevision
-		&& Flags == InText.Flags;
+	TSharedPtr<ITextData, ESPMode::ThreadSafe> TextDataPin = TextDataPtr.Pin();
+	return TextDataPin == InText.TextData 
+		&& HistoryRevision == InText.TextData->GetTextHistory().Revision
+		&& Flags == InText.Flags;;
 }
 
 bool FTextSnapshot::IsDisplayStringEqualTo(const FText& InText) const
@@ -1341,20 +1306,22 @@ bool FTextSnapshot::IsDisplayStringEqualTo(const FText& InText) const
 
 	// We have to assume that the display string has changed if the history of the text has changed
 	// (due to a culture change), as we no longer have the old display string to compare against
-	const int32 InHistoryRevision = InText.History.IsValid() ? InText.History->Revision : INDEX_NONE;
-	return HistoryRevision == InHistoryRevision && DisplayStringPtr.IsValid() && DisplayStringPtr->Equals(InText.ToString(), ESearchCase::CaseSensitive);
+	TSharedPtr<ITextData, ESPMode::ThreadSafe> TextDataPin = TextDataPtr.Pin();
+	return HistoryRevision == InText.TextData->GetTextHistory().Revision 
+		&& TextDataPin.IsValid() && TextDataPin->GetDisplayString().Equals(InText.ToString(), ESearchCase::CaseSensitive);
 }
 
 FScopedTextIdentityPreserver::FScopedTextIdentityPreserver(FText& InTextToPersist)
-	: HadFoundNamespaceAndKey(false)
-	, Flags(InTextToPersist.Flags)
-	, TextToPersist(InTextToPersist)
+	: TextToPersist(InTextToPersist)
+	, HadFoundNamespaceAndKey(false)
+	, Flags(TextToPersist.Flags)
 {
 	// Empty display strings can't have a namespace or key.
-	if (GIsEditor && !InTextToPersist.DisplayString->IsEmpty())
+	if (GIsEditor && !TextToPersist.TextData->GetDisplayString().IsEmpty())
 	{
 		// Save off namespace and key to be restored later.
-		HadFoundNamespaceAndKey = FTextLocalizationManager::Get().FindNamespaceAndKeyFromDisplayString(InTextToPersist.DisplayString, Namespace, Key);
+		TextToPersist.TextData->PersistText();
+		HadFoundNamespaceAndKey = FTextLocalizationManager::Get().FindNamespaceAndKeyFromDisplayString(TextToPersist.TextData->GetLocalizedString().ToSharedRef(), Namespace, Key);
 	}
 }
 
@@ -1371,14 +1338,11 @@ FScopedTextIdentityPreserver::~FScopedTextIdentityPreserver()
 		// Without a source string, we can't possibly preserve the identity. If the text we're preserving identity for can't possibly have an identity anymore, this class shouldn't be used on this text.
 		check(SourceString);
 
-		// Make the history so that it can attempt find and serialize the identity for this source string when the history itself is serialized.
-		TextToPersist.History = MakeShareable(new FTextHistory_Base(*SourceString));
-
 		// Create/update the display string instance for this identity in the text localization manager...
 		const FTextDisplayStringRef DisplayString = FTextLocalizationManager::Get().GetDisplayString(Namespace, Key, SourceString);
 
-		// ... and set it so that the text's history's serialization will properly map this display string instance back to the identity we intended to preserve and restore.
-		TextToPersist.DisplayString = DisplayString;
+		// ... and update the data on the text instance
+		TextToPersist.TextData = MakeShareable(new TLocalizedTextData<FTextHistory_Base>(MoveTemp(DisplayString), FTextHistory_Base(*SourceString)));
 	}
 }
 
diff --git a/Engine/Source/Runtime/Core/Private/Internationalization/TextData.h b/Engine/Source/Runtime/Core/Private/Internationalization/TextData.h
new file mode 100644
--- /dev/null
+++ b/Engine/Source/Runtime/Core/Private/Internationalization/TextData.h
@@ -0,0 +1,150 @@
+﻿// Copyright 1998-2015 Epic Games, Inc. All Rights Reserved.
+
+#pragma once
+
+#include "ITextData.h"
+#include "TextHistory.h"
+
+/** 
+ * Implementation of ITextData that stores the common data and functionality.
+ */
+template <typename THistoryType>
+class TTextData : public ITextData
+{
+public:
+	TTextData()
+		: LocalizedString()
+		, History()
+	{
+	}
+
+	explicit TTextData(FTextDisplayStringPtr InLocalizedString)
+		: LocalizedString(MoveTemp(InLocalizedString))
+		, History()
+	{
+	}
+
+	TTextData(FTextDisplayStringPtr InLocalizedString, THistoryType InHistory)
+		: LocalizedString(MoveTemp(InLocalizedString))
+		, History(MoveTemp(InHistory))
+	{
+	}
+
+	virtual ~TTextData()
+	{
+	}
+
+	virtual FTextDisplayStringPtr GetLocalizedString() const override
+	{
+		return LocalizedString;
+	}
+
+	virtual FTextDisplayStringPtr& GetMutableLocalizedString() override
+	{
+		return LocalizedString;
+	}
+
+	virtual const FTextHistory& GetTextHistory() const override
+	{
+		return History;
+	}
+
+	virtual FTextHistory& GetMutableTextHistory() override
+	{
+		return History;
+	}
+
+protected:
+	FTextDisplayStringPtr LocalizedString;
+	THistoryType History;
+};
+
+
+/** 
+ * Implementation of ITextData optimized to track localized text retrieved from the text localization manager, or (re)generated via persistent text history.
+ */
+template <typename THistoryType>
+class TLocalizedTextData : public TTextData<THistoryType>
+{
+public:
+	TLocalizedTextData()
+	{
+	}
+
+	explicit TLocalizedTextData(FTextDisplayStringRef InLocalizedString)
+		: TTextData<THistoryType>(InLocalizedString)
+	{
+	}
+
+	TLocalizedTextData(FTextDisplayStringRef InLocalizedString, THistoryType InHistory)
+		: TTextData<THistoryType>(InLocalizedString, MoveTemp(InHistory))
+	{
+	}
+
+	virtual ~TLocalizedTextData()
+	{
+	}
+
+	virtual const FString& GetDisplayString() const override
+	{
+		return *this->LocalizedString;
+	}
+
+	virtual void PersistText() override
+	{
+	}
+};
+
+
+/**
+ * Implementation of ITextData optimized to track text that was generated at runtime.
+ * This data will avoid heap allocating its localized string until we know that it needs to be persisted.
+ */
+template <typename THistoryType>
+class TGeneratedTextData : public TTextData<THistoryType>
+{
+public:
+	TGeneratedTextData()
+	{
+	}
+
+	explicit TGeneratedTextData(FString InDisplayString)
+		: TTextData<THistoryType>(nullptr)
+		, DisplayString(MoveTemp(InDisplayString))
+	{
+	}
+
+	TGeneratedTextData(FString InDisplayString, THistoryType InHistory)
+		: TTextData<THistoryType>(nullptr, MoveTemp(InHistory))
+		, DisplayString(MoveTemp(InDisplayString))
+	{
+	}
+
+	virtual ~TGeneratedTextData()
+	{
+	}
+
+	virtual const FString& GetDisplayString() const override
+	{
+		return (this->LocalizedString.IsValid()) ? *this->LocalizedString : DisplayString;
+	}
+
+	virtual void PersistText() override
+	{
+		if (!this->LocalizedString.IsValid())
+		{
+			FScopeLock ScopeLock(&PersistTextCS);
+
+			// Check the pointer again in case another thread beat us to it
+			if (!this->LocalizedString.IsValid())
+			{
+				// We copy (rather than move) DisplayString here, as other threads may currently be accessing it
+				this->LocalizedString = MakeShareable(new FString(DisplayString));
+			}
+		}
+	}
+    
+protected:
+	FString DisplayString;
+	FCriticalSection PersistTextCS;
+};
diff --git a/Engine/Source/Runtime/Core/Private/Internationalization/TextHistory.cpp b/Engine/Source/Runtime/Core/Private/Internationalization/TextHistory.cpp
--- a/Engine/Source/Runtime/Core/Private/Internationalization/TextHistory.cpp
+++ b/Engine/Source/Runtime/Core/Private/Internationalization/TextHistory.cpp
@@ -13,7 +13,20 @@
 FTextHistory::FTextHistory()
 	: Revision(FTextLocalizationManager::Get().GetTextRevision())
 {
+}
+
+FTextHistory::FTextHistory(FTextHistory&& Other)
+	: Revision(MoveTemp(Other.Revision))
+{
+}
 
+FTextHistory& FTextHistory::operator=(FTextHistory&& Other)
+{
+	if (this != &Other)
+	{
+		Revision = Other.Revision;
+	}
+	return *this;
 }
 
 bool FTextHistory::IsOutOfDate()
@@ -21,7 +34,7 @@ bool FTextHistory::IsOutOfDate()
 	return Revision < FTextLocalizationManager::Get().GetTextRevision();
 }
 
-TSharedPtr< FString, ESPMode::ThreadSafe > FTextHistory::GetSourceString() const
+const FString* FTextHistory::GetSourceString() const
 {
 	return NULL;
 }
@@ -32,7 +45,7 @@ void FTextHistory::GetSourceTextsFromFormatHistory(FText Text, TArray<FText>& Ou
 	OutSourceTexts.Add(Text);
 }
 
-void FTextHistory::SerializeForDisplayString(FArchive& Ar, FTextDisplayStringRef& InOutDisplayString)
+void FTextHistory::SerializeForDisplayString(FArchive& Ar, FTextDisplayStringPtr& InOutDisplayString)
 {
 	if(Ar.IsLoading())
 	{
@@ -64,26 +77,31 @@ void FTextHistory::Rebuild(TSharedRef< FString, ESPMode::ThreadSafe > InDisplayS
 // FTextHistory_Base
 
 FTextHistory_Base::FTextHistory_Base(FString InSourceString)
-	: SourceString(new FString( MoveTemp( InSourceString ) ))
+	: SourceString(MoveTemp(InSourceString))
 {
-
 }
 
-FTextHistory_Base::FTextHistory_Base(TSharedPtr< FString, ESPMode::ThreadSafe > InSourceString)
-	: SourceString(InSourceString)
+FTextHistory_Base::FTextHistory_Base(FTextHistory_Base&& Other)
+	: FTextHistory(MoveTemp(Other))
+	, SourceString(MoveTemp(Other.SourceString))
 {
+}
 
+FTextHistory_Base& FTextHistory_Base::operator=(FTextHistory_Base&& Other)
+{
+	FTextHistory::operator=(MoveTemp(Other));
+	if (this != &Other)
+	{
+		SourceString = MoveTemp(Other.SourceString);
+	}
+	return *this;
 }
 
 FText FTextHistory_Base::ToText(bool bInAsSource) const
 {
 	if(bInAsSource)
 	{
-		if(SourceString.IsValid())
-		{
-			return FText::FromString(*SourceString.Get());
-		}
-		return FText::GetEmpty();
+		return FText::FromString(SourceString);
 	}
 
 	// This should never be called except to build as source
@@ -92,9 +110,9 @@ FText FTextHistory_Base::ToText(bool bInAsSource) const
 	return FText::GetEmpty();
 }
 
-TSharedPtr< FString, ESPMode::ThreadSafe > FTextHistory_Base::GetSourceString() const
+const FString* FTextHistory_Base::GetSourceString() const
 {
-	return SourceString;
+	return &SourceString;
 }
 
 void FTextHistory_Base::Serialize( FArchive& Ar )
@@ -107,7 +125,7 @@ void FTextHistory_Base::Serialize( FArchive& Ar )
 	}
 }
 
-void FTextHistory_Base::SerializeForDisplayString(FArchive& Ar, FTextDisplayStringRef& InOutDisplayString)
+void FTextHistory_Base::SerializeForDisplayString(FArchive& Ar, FTextDisplayStringPtr& InOutDisplayString)
 {
 	if(Ar.IsLoading())
 	{
@@ -116,33 +134,28 @@ void FTextHistory_Base::SerializeForDisplayString(FArchive& Ar, FTextDisplayStri
 
 		FString Namespace;
 		FString Key;
-		FString SourceStringRaw;
 
 		Ar << Namespace;
 		Ar << Key;
-		Ar << SourceStringRaw;
-
-		// If there was a SourceString, store it in the member variable
-		if(!SourceStringRaw.IsEmpty())
-		{
-			SourceString = MakeShareable(new FString(SourceStringRaw));
-		}
+		Ar << SourceString;
 
 		// Using the deserialized namespace and key, find the DisplayString.
-		InOutDisplayString = FTextLocalizationManager::Get().GetDisplayString(Namespace, Key, SourceString.Get());
+		InOutDisplayString = FTextLocalizationManager::Get().GetDisplayString(Namespace, Key, &SourceString);
 	}
 	else if(Ar.IsSaving())
 	{
+		check(InOutDisplayString.IsValid());
+
 		FString Namespace;
 		FString Key;
 
-		const bool FoundNamespaceAndKey = FTextLocalizationManager::Get().FindNamespaceAndKeyFromDisplayString(InOutDisplayString, Namespace, Key);
+		const bool FoundNamespaceAndKey = FTextLocalizationManager::Get().FindNamespaceAndKeyFromDisplayString(InOutDisplayString.ToSharedRef(), Namespace, Key);
 
 		// If this has no namespace or key, attempt to give it a GUID for a key and register it.
 		if (!FoundNamespaceAndKey && GIsEditor && (Ar.IsPersistent() && !Ar.HasAnyPortFlags(PPF_Duplicate)))
 		{
 			Key = FGuid::NewGuid().ToString();
-			if (!FTextLocalizationManager::Get().AddDisplayString(InOutDisplayString, Namespace, Key))
+			if (!FTextLocalizationManager::Get().AddDisplayString(InOutDisplayString.ToSharedRef(), Namespace, Key))
 			{
 				// Could not add display string, reset namespace and key.
 				Namespace.Empty();
@@ -156,27 +169,36 @@ void FTextHistory_Base::SerializeForDisplayString(FArchive& Ar, FTextDisplayStri
 		// Serialize the Key
 		Ar << Key;
 
-		// Serialize the SourceString, or an empty string if there is none
-		if( SourceString.IsValid() )
-		{
-			Ar << *(SourceString);
-		}
-		else
-		{
-			FString Empty;
-			Ar << Empty;
-		}
+		// Serialize the SourceString
+		Ar << SourceString;
 	}
 }
 
 ///////////////////////////////////////
 // FTextHistory_NamedFormat
 
-FTextHistory_NamedFormat::FTextHistory_NamedFormat(const FText& InSourceText, const FFormatNamedArguments& InArguments)
-	: SourceText(InSourceText)
-	, Arguments(InArguments)
+FTextHistory_NamedFormat::FTextHistory_NamedFormat(FText InSourceText, FFormatNamedArguments InArguments)
+	: SourceText(MoveTemp(InSourceText))
+	, Arguments(MoveTemp(InArguments))
 {
+}
 
+FTextHistory_NamedFormat::FTextHistory_NamedFormat(FTextHistory_NamedFormat&& Other)
+	: FTextHistory(MoveTemp(Other))
+	, SourceText(MoveTemp(Other.SourceText))
+	, Arguments(MoveTemp(Other.Arguments))
+{
+}
+
+FTextHistory_NamedFormat& FTextHistory_NamedFormat::operator=(FTextHistory_NamedFormat&& Other)
+{
+	FTextHistory::operator=(MoveTemp(Other));
+	if (this != &Other)
+	{
+		SourceText = MoveTemp(Other.SourceText);
+		Arguments = MoveTemp(Other.Arguments);
+	}
+	return *this;
 }
 
 FText FTextHistory_NamedFormat::ToText(bool bInAsSource) const
@@ -204,13 +226,11 @@ void FTextHistory_NamedFormat::GetSourceTextsFromFormatHistory(FText, TArray<FTe
 	for (auto It = Arguments.CreateConstIterator(); It; ++It)
 	{
 		const FFormatArgumentValue& ArgumentValue = It.Value();
-		if (ArgumentValue.Type == EFormatArgumentType::Text)
+		if (ArgumentValue.GetType() == EFormatArgumentType::Text)
 		{
-			if (ArgumentValue.TextValue && ArgumentValue.TextValue->History.IsValid())
-			{
-				// Search any text arguments for source text
-				ArgumentValue.TextValue->History->GetSourceTextsFromFormatHistory(*(ArgumentValue.TextValue), OutSourceTexts);
-			}
+			// Search any text arguments for source text
+			const FText& TextValue = ArgumentValue.GetTextValue();
+			TextValue.TextData->GetTextHistory().GetSourceTextsFromFormatHistory(TextValue, OutSourceTexts);
 		}
 	}
 }
@@ -218,11 +238,28 @@ void FTextHistory_NamedFormat::GetSourceTextsFromFormatHistory(FText, TArray<FTe
 ///////////////////////////////////////
 // FTextHistory_OrderedFormat
 
-FTextHistory_OrderedFormat::FTextHistory_OrderedFormat(const FText& InSourceText, const FFormatOrderedArguments& InArguments)
-	: SourceText(InSourceText)
-	, Arguments(InArguments)
+FTextHistory_OrderedFormat::FTextHistory_OrderedFormat(FText InSourceText, FFormatOrderedArguments InArguments)
+	: SourceText(MoveTemp(InSourceText))
+	, Arguments(MoveTemp(InArguments))
+{
+}
+
+FTextHistory_OrderedFormat::FTextHistory_OrderedFormat(FTextHistory_OrderedFormat&& Other)
+	: FTextHistory(MoveTemp(Other))
+	, SourceText(MoveTemp(Other.SourceText))
+	, Arguments(MoveTemp(Other.Arguments))
 {
+}
 
+FTextHistory_OrderedFormat& FTextHistory_OrderedFormat::operator=(FTextHistory_OrderedFormat&& Other)
+{
+	FTextHistory::operator=(MoveTemp(Other));
+	if (this != &Other)
+	{
+		SourceText = MoveTemp(Other.SourceText);
+		Arguments = MoveTemp(Other.Arguments);
+	}
+	return *this;
 }
 
 FText FTextHistory_OrderedFormat::ToText(bool bInAsSource) const
@@ -250,13 +287,11 @@ void FTextHistory_OrderedFormat::GetSourceTextsFromFormatHistory(FText, TArray<F
 	for (auto It = Arguments.CreateConstIterator(); It; ++It)
 	{
 		const FFormatArgumentValue& ArgumentValue = *It;
-		if (ArgumentValue.Type == EFormatArgumentType::Text)
+		if (ArgumentValue.GetType() == EFormatArgumentType::Text)
 		{
-			if (ArgumentValue.TextValue && ArgumentValue.TextValue->History.IsValid())
-			{
-				// Search any text arguments for source text
-				ArgumentValue.TextValue->History->GetSourceTextsFromFormatHistory(*(ArgumentValue.TextValue), OutSourceTexts);
-			}
+			// Search any text arguments for source text
+			const FText& TextValue = ArgumentValue.GetTextValue();
+			TextValue.TextData->GetTextHistory().GetSourceTextsFromFormatHistory(TextValue, OutSourceTexts);
 		}
 	}
 }
@@ -264,11 +299,28 @@ void FTextHistory_OrderedFormat::GetSourceTextsFromFormatHistory(FText, TArray<F
 ///////////////////////////////////////
 // FTextHistory_ArgumentDataFormat
 
-FTextHistory_ArgumentDataFormat::FTextHistory_ArgumentDataFormat(const FText& InSourceText, const TArray< struct FFormatArgumentData >& InArguments)
-	: SourceText(InSourceText)
-	, Arguments(InArguments)
+FTextHistory_ArgumentDataFormat::FTextHistory_ArgumentDataFormat(FText InSourceText, TArray<FFormatArgumentData> InArguments)
+	: SourceText(MoveTemp(InSourceText))
+	, Arguments(MoveTemp(InArguments))
 {
+}
 
+FTextHistory_ArgumentDataFormat::FTextHistory_ArgumentDataFormat(FTextHistory_ArgumentDataFormat&& Other)
+	: FTextHistory(MoveTemp(Other))
+	, SourceText(MoveTemp(Other.SourceText))
+	, Arguments(MoveTemp(Other.Arguments))
+{
+}
+
+FTextHistory_ArgumentDataFormat& FTextHistory_ArgumentDataFormat::operator=(FTextHistory_ArgumentDataFormat&& Other)
+{
+	FTextHistory::operator=(MoveTemp(Other));
+	if (this != &Other)
+	{
+		SourceText = MoveTemp(Other.SourceText);
+		Arguments = MoveTemp(Other.Arguments);
+	}
+	return *this;
 }
 
 FText FTextHistory_ArgumentDataFormat::ToText(bool bInAsSource) const
@@ -296,13 +348,11 @@ void FTextHistory_ArgumentDataFormat::GetSourceTextsFromFormatHistory(FText, TAr
 	for (int32 x = 0; x < Arguments.Num(); ++x)
 	{
 		const FFormatArgumentValue& ArgumentValue = Arguments[x].ArgumentValue;
-		if (ArgumentValue.Type == EFormatArgumentType::Text)
+		if (ArgumentValue.GetType() == EFormatArgumentType::Text)
 		{
-			if (ArgumentValue.TextValue && ArgumentValue.TextValue->History.IsValid())
-			{
-				// Search any text arguments for source text
-				ArgumentValue.TextValue->History->GetSourceTextsFromFormatHistory(*(ArgumentValue.TextValue), OutBaseTexts);
-			}
+			// Search any text arguments for source text
+			const FText& TextValue = ArgumentValue.GetTextValue();
+			TextValue.TextData->GetTextHistory().GetSourceTextsFromFormatHistory(TextValue, OutBaseTexts);
 		}
 	}
 }
@@ -310,43 +360,59 @@ void FTextHistory_ArgumentDataFormat::GetSourceTextsFromFormatHistory(FText, TAr
 ///////////////////////////////////////
 // FTextHistory_FormatNumber
 
-FTextHistory_FormatNumber::FTextHistory_FormatNumber()
-	: FormatOptions(nullptr)
-{
-}
-
-FTextHistory_FormatNumber::FTextHistory_FormatNumber(const FFormatArgumentValue& InSourceValue, const FNumberFormattingOptions* const InFormatOptions, const FCulturePtr InTargetCulture)
-	: SourceValue(InSourceValue)
-	, FormatOptions(nullptr)
-	, TargetCulture(InTargetCulture)
+FTextHistory_FormatNumber::FTextHistory_FormatNumber(FFormatArgumentValue InSourceValue, const FNumberFormattingOptions* const InFormatOptions, FCulturePtr InTargetCulture)
+	: SourceValue(MoveTemp(InSourceValue))
+	, FormatOptions()
+	, TargetCulture(MoveTemp(InTargetCulture))
 {
 	if(InFormatOptions)
 	{
-		FormatOptions = new FNumberFormattingOptions;
-		(*FormatOptions) = *InFormatOptions;
+		FormatOptions = *InFormatOptions;
 	}
 }
 
-FTextHistory_FormatNumber::~FTextHistory_FormatNumber()
+FTextHistory_FormatNumber::FTextHistory_FormatNumber(FTextHistory_FormatNumber&& Other)
+	: FTextHistory(MoveTemp(Other))
+	, SourceValue(MoveTemp(Other.SourceValue))
+	, FormatOptions(MoveTemp(Other.FormatOptions))
+	, TargetCulture(MoveTemp(Other.TargetCulture))
+{
+}
+
+FTextHistory_FormatNumber& FTextHistory_FormatNumber::operator=(FTextHistory_FormatNumber&& Other)
 {
-	delete FormatOptions;
+	FTextHistory::operator=(MoveTemp(Other));
+	if (this != &Other)
+	{
+		SourceValue = MoveTemp(Other.SourceValue);
+		FormatOptions = MoveTemp(Other.FormatOptions);
+		TargetCulture = MoveTemp(Other.TargetCulture);
+	}
+	return *this;
 }
 
 void FTextHistory_FormatNumber::Serialize(FArchive& Ar)
 {
 	Ar << SourceValue;
 
-	bool bHasFormatOptions = FormatOptions != nullptr;
+	bool bHasFormatOptions = FormatOptions.IsSet();
 	Ar << bHasFormatOptions;
 
-	if(bHasFormatOptions)
+	if(Ar.IsLoading())
 	{
-		if(Ar.IsLoading())
+		if(bHasFormatOptions)
+		{
+			FormatOptions = FNumberFormattingOptions();
+		}
+		else
 		{
-			FormatOptions = new FNumberFormattingOptions;
+			FormatOptions.Reset();
 		}
-		CA_SUPPRESS(6011)
-		Ar << *FormatOptions;
+	}
+	if(bHasFormatOptions)
+	{
+		check(FormatOptions.IsSet());
+		Ar << FormatOptions.GetValue();
 	}
 
 	if(Ar.IsSaving())
@@ -369,32 +435,44 @@ void FTextHistory_FormatNumber::Serialize(FArchive& Ar)
 ///////////////////////////////////////
 // FTextHistory_AsNumber
 
-FTextHistory_AsNumber::FTextHistory_AsNumber(const FFormatArgumentValue& InSourceValue, const FNumberFormattingOptions* const InFormatOptions, const FCulturePtr InTargetCulture)
-	: FTextHistory_FormatNumber(InSourceValue, InFormatOptions, InTargetCulture)
+FTextHistory_AsNumber::FTextHistory_AsNumber(FFormatArgumentValue InSourceValue, const FNumberFormattingOptions* const InFormatOptions, FCulturePtr InTargetCulture)
+	: FTextHistory_FormatNumber(MoveTemp(InSourceValue), InFormatOptions, MoveTemp(InTargetCulture))
+{
+}
+
+FTextHistory_AsNumber::FTextHistory_AsNumber(FTextHistory_AsNumber&& Other)
+	: FTextHistory_FormatNumber(MoveTemp(Other))
 {
 }
 
+FTextHistory_AsNumber& FTextHistory_AsNumber::operator=(FTextHistory_AsNumber&& Other)
+{
+	FTextHistory_FormatNumber::operator=(MoveTemp(Other));
+	return *this;
+}
+
 FText FTextHistory_AsNumber::ToText(bool bInAsSource) const
 {
 	TSharedPtr< FCulture, ESPMode::ThreadSafe > CurrentCulture = bInAsSource? FInternationalization::Get().GetInvariantCulture() : TargetCulture;
 
-	switch(SourceValue.Type)
+	const FNumberFormattingOptions* FormatOptionsPtr = (FormatOptions.IsSet()) ? &FormatOptions.GetValue() : nullptr;
+	switch(SourceValue.GetType())
 	{
 	case EFormatArgumentType::UInt:
 		{
-			return FText::AsNumber(SourceValue.UIntValue, FormatOptions, CurrentCulture);
+			return FText::AsNumber(SourceValue.GetUIntValue(), FormatOptionsPtr, CurrentCulture);
 		}
 	case EFormatArgumentType::Int:
 		{
-			return FText::AsNumber(SourceValue.IntValue, FormatOptions, CurrentCulture);
+			return FText::AsNumber(SourceValue.GetIntValue(), FormatOptionsPtr, CurrentCulture);
 		}
 	case EFormatArgumentType::Float:
 		{
-			return FText::AsNumber(SourceValue.FloatValue, FormatOptions, CurrentCulture);
+			return FText::AsNumber(SourceValue.GetFloatValue(), FormatOptionsPtr, CurrentCulture);
 		}
 	case EFormatArgumentType::Double:
 		{
-			return FText::AsNumber(SourceValue.DoubleValue, FormatOptions, CurrentCulture);
+			return FText::AsNumber(SourceValue.GetDoubleValue(), FormatOptionsPtr, CurrentCulture);
 		}
 	default:
 		{
@@ -419,24 +497,36 @@ void FTextHistory_AsNumber::Serialize( FArchive& Ar )
 ///////////////////////////////////////
 // FTextHistory_AsPercent
 
-FTextHistory_AsPercent::FTextHistory_AsPercent(const FFormatArgumentValue& InSourceValue, const FNumberFormattingOptions* const InFormatOptions, const FCulturePtr InTargetCulture)
-	: FTextHistory_FormatNumber(InSourceValue, InFormatOptions, InTargetCulture)
+FTextHistory_AsPercent::FTextHistory_AsPercent(FFormatArgumentValue InSourceValue, const FNumberFormattingOptions* const InFormatOptions, FCulturePtr InTargetCulture)
+	: FTextHistory_FormatNumber(MoveTemp(InSourceValue), InFormatOptions, MoveTemp(InTargetCulture))
+{
+}
+
+FTextHistory_AsPercent::FTextHistory_AsPercent(FTextHistory_AsPercent&& Other)
+	: FTextHistory_FormatNumber(MoveTemp(Other))
+{
+}
+
+FTextHistory_AsPercent& FTextHistory_AsPercent::operator=(FTextHistory_AsPercent&& Other)
 {
+	FTextHistory_FormatNumber::operator=(MoveTemp(Other));
+	return *this;
 }
 
 FText FTextHistory_AsPercent::ToText(bool bInAsSource) const
 {
 	TSharedPtr< FCulture, ESPMode::ThreadSafe > CurrentCulture = bInAsSource? FInternationalization::Get().GetInvariantCulture() : TargetCulture;
 
-	switch(SourceValue.Type)
+	const FNumberFormattingOptions* FormatOptionsPtr = (FormatOptions.IsSet()) ? &FormatOptions.GetValue() : nullptr;
+	switch(SourceValue.GetType())
 	{
 	case EFormatArgumentType::Float:
 		{
-			return FText::AsPercent(SourceValue.FloatValue, FormatOptions, CurrentCulture);
+			return FText::AsPercent(SourceValue.GetFloatValue(), FormatOptionsPtr, CurrentCulture);
 		}
 	case EFormatArgumentType::Double:
 		{
-			return FText::AsPercent(SourceValue.DoubleValue, FormatOptions, CurrentCulture);
+			return FText::AsPercent(SourceValue.GetDoubleValue(), FormatOptionsPtr, CurrentCulture);
 		}
 	default:
 		{
@@ -461,33 +551,50 @@ void FTextHistory_AsPercent::Serialize( FArchive& Ar )
 ///////////////////////////////////////
 // FTextHistory_AsCurrency
 
-FTextHistory_AsCurrency::FTextHistory_AsCurrency(const FFormatArgumentValue& InSourceValue, const FString& InCurrencyCode, const FNumberFormattingOptions* const InFormatOptions, const FCulturePtr InTargetCulture)
-: FTextHistory_FormatNumber(InSourceValue, InFormatOptions, InTargetCulture)
-, CurrencyCode(InCurrencyCode)
+FTextHistory_AsCurrency::FTextHistory_AsCurrency(FFormatArgumentValue InSourceValue, FString InCurrencyCode, const FNumberFormattingOptions* const InFormatOptions, FCulturePtr InTargetCulture)
+	: FTextHistory_FormatNumber(MoveTemp(InSourceValue), InFormatOptions, MoveTemp(InTargetCulture))
+	, CurrencyCode(MoveTemp(InCurrencyCode))
 {
 }
 
+FTextHistory_AsCurrency::FTextHistory_AsCurrency(FTextHistory_AsCurrency&& Other)
+	: FTextHistory_FormatNumber(MoveTemp(Other))
+	, CurrencyCode(MoveTemp(Other.CurrencyCode))
+{
+}
+
+FTextHistory_AsCurrency& FTextHistory_AsCurrency::operator=(FTextHistory_AsCurrency&& Other)
+{
+	FTextHistory_FormatNumber::operator=(MoveTemp(Other));
+	if (this != &Other)
+	{
+		CurrencyCode = MoveTemp(Other.CurrencyCode);
+	}
+	return *this;
+}
+
 FText FTextHistory_AsCurrency::ToText(bool bInAsSource) const
 {
 	TSharedPtr< FCulture, ESPMode::ThreadSafe > CurrentCulture = bInAsSource? FInternationalization::Get().GetInvariantCulture() : TargetCulture;
 
-	switch(SourceValue.Type)
+	const FNumberFormattingOptions* FormatOptionsPtr = (FormatOptions.IsSet()) ? &FormatOptions.GetValue() : nullptr;
+	switch(SourceValue.GetType())
 	{
 	case EFormatArgumentType::UInt:
 		{
-			return FText::AsCurrency(SourceValue.UIntValue, CurrencyCode, FormatOptions, CurrentCulture);
+			return FText::AsCurrency(SourceValue.GetUIntValue(), CurrencyCode, FormatOptionsPtr, CurrentCulture);
 		}
 	case EFormatArgumentType::Int:
 		{
-			return FText::AsCurrency(SourceValue.IntValue, CurrencyCode, FormatOptions, CurrentCulture);
+			return FText::AsCurrency(SourceValue.GetIntValue(), CurrencyCode, FormatOptionsPtr, CurrentCulture);
 		}
 	case EFormatArgumentType::Float:
 		{
-			return FText::AsCurrency(SourceValue.FloatValue, CurrencyCode, FormatOptions, CurrentCulture);
+			return FText::AsCurrency(SourceValue.GetFloatValue(), CurrencyCode, FormatOptionsPtr, CurrentCulture);
 		}
 	case EFormatArgumentType::Double:
 		{
-			return FText::AsCurrency(SourceValue.DoubleValue, CurrencyCode, FormatOptions, CurrentCulture);
+			return FText::AsCurrency(SourceValue.GetDoubleValue(), CurrencyCode, FormatOptionsPtr, CurrentCulture);
 		}
 	default:
 		{
@@ -517,12 +624,34 @@ void FTextHistory_AsCurrency::Serialize( FArchive& Ar )
 ///////////////////////////////////////
 // FTextHistory_AsDate
 
-FTextHistory_AsDate::FTextHistory_AsDate(const FDateTime& InSourceDateTime, const EDateTimeStyle::Type InDateStyle, const FString& InTimeZone, const FCulturePtr InTargetCulture)
-	: SourceDateTime(InSourceDateTime)
+FTextHistory_AsDate::FTextHistory_AsDate(FDateTime InSourceDateTime, const EDateTimeStyle::Type InDateStyle, FString InTimeZone, FCulturePtr InTargetCulture)
+	: SourceDateTime(MoveTemp(InSourceDateTime))
 	, DateStyle(InDateStyle)
-	, TimeZone(InTimeZone)
-	, TargetCulture(InTargetCulture)
+	, TimeZone(MoveTemp(InTimeZone))
+	, TargetCulture(MoveTemp(InTargetCulture))
+{
+}
+
+FTextHistory_AsDate::FTextHistory_AsDate(FTextHistory_AsDate&& Other)
+	: FTextHistory(MoveTemp(Other))
+	, SourceDateTime(MoveTemp(Other.SourceDateTime))
+	, DateStyle(Other.DateStyle)
+	, TimeZone(MoveTemp(Other.TimeZone))
+	, TargetCulture(MoveTemp(Other.TargetCulture))
+{
+}
+
+FTextHistory_AsDate& FTextHistory_AsDate::operator=(FTextHistory_AsDate&& Other)
 {
+	FTextHistory::operator=(MoveTemp(Other));
+	if (this != &Other)
+	{
+		SourceDateTime = MoveTemp(Other.SourceDateTime);
+		DateStyle = Other.DateStyle;
+		TimeZone = MoveTemp(Other.TimeZone);
+		TargetCulture = MoveTemp(Other.TargetCulture);
+	}
+	return *this;
 }
 
 void FTextHistory_AsDate::Serialize(FArchive& Ar)
@@ -570,14 +699,36 @@ FText FTextHistory_AsDate::ToText(bool bInAsSource) const
 ///////////////////////////////////////
 // FTextHistory_AsTime
 
-FTextHistory_AsTime::FTextHistory_AsTime(const FDateTime& InSourceDateTime, const EDateTimeStyle::Type InTimeStyle, const FString& InTimeZone, const FCulturePtr InTargetCulture)
-	: SourceDateTime(InSourceDateTime)
+FTextHistory_AsTime::FTextHistory_AsTime(FDateTime InSourceDateTime, const EDateTimeStyle::Type InTimeStyle, FString InTimeZone, FCulturePtr InTargetCulture)
+	: SourceDateTime(MoveTemp(InSourceDateTime))
 	, TimeStyle(InTimeStyle)
-	, TimeZone(InTimeZone)
-	, TargetCulture(InTargetCulture)
+	, TimeZone(MoveTemp(InTimeZone))
+	, TargetCulture(MoveTemp(InTargetCulture))
+{
+}
+
+FTextHistory_AsTime::FTextHistory_AsTime(FTextHistory_AsTime&& Other)
+	: FTextHistory(MoveTemp(Other))
+	, SourceDateTime(MoveTemp(Other.SourceDateTime))
+	, TimeStyle(Other.TimeStyle)
+	, TimeZone(MoveTemp(Other.TimeZone))
+	, TargetCulture(MoveTemp(Other.TargetCulture))
 {
 }
 
+FTextHistory_AsTime& FTextHistory_AsTime::operator=(FTextHistory_AsTime&& Other)
+{
+	FTextHistory::operator=(MoveTemp(Other));
+	if (this != &Other)
+	{
+		SourceDateTime = MoveTemp(Other.SourceDateTime);
+		TimeStyle = Other.TimeStyle;
+		TimeZone = MoveTemp(Other.TimeZone);
+		TargetCulture = MoveTemp(Other.TargetCulture);
+	}
+	return *this;
+}
+
 void FTextHistory_AsTime::Serialize(FArchive& Ar)
 {
 	if(Ar.IsSaving())
@@ -621,15 +772,39 @@ FText FTextHistory_AsTime::ToText(bool bInAsSource) const
 ///////////////////////////////////////
 // FTextHistory_AsDateTime
 
-FTextHistory_AsDateTime::FTextHistory_AsDateTime(const FDateTime& InSourceDateTime, const EDateTimeStyle::Type InDateStyle, const EDateTimeStyle::Type InTimeStyle, const FString& InTimeZone, const FCulturePtr InTargetCulture)
-	: SourceDateTime(InSourceDateTime)
+FTextHistory_AsDateTime::FTextHistory_AsDateTime(FDateTime InSourceDateTime, const EDateTimeStyle::Type InDateStyle, const EDateTimeStyle::Type InTimeStyle, FString InTimeZone, FCulturePtr InTargetCulture)
+	: SourceDateTime(MoveTemp(InSourceDateTime))
 	, DateStyle(InDateStyle)
 	, TimeStyle(InTimeStyle)
-	, TimeZone(InTimeZone)
-	, TargetCulture(InTargetCulture)
+	, TimeZone(MoveTemp(InTimeZone))
+	, TargetCulture(MoveTemp(InTargetCulture))
 {
 }
 
+FTextHistory_AsDateTime::FTextHistory_AsDateTime(FTextHistory_AsDateTime&& Other)
+	: FTextHistory(MoveTemp(Other))
+	, SourceDateTime(MoveTemp(Other.SourceDateTime))
+	, DateStyle(Other.DateStyle)
+	, TimeStyle(Other.TimeStyle)
+	, TimeZone(MoveTemp(Other.TimeZone))
+	, TargetCulture(MoveTemp(Other.TargetCulture))
+{
+}
+
+FTextHistory_AsDateTime& FTextHistory_AsDateTime::operator=(FTextHistory_AsDateTime&& Other)
+{
+	FTextHistory::operator=(MoveTemp(Other));
+	if (this != &Other)
+	{
+		SourceDateTime = MoveTemp(Other.SourceDateTime);
+		DateStyle = Other.DateStyle;
+		TimeStyle = Other.TimeStyle;
+		TimeZone = MoveTemp(Other.TimeZone);
+		TargetCulture = MoveTemp(Other.TargetCulture);
+	}
+	return *this;
+}
+
 void FTextHistory_AsDateTime::Serialize(FArchive& Ar)
 {
 	if(Ar.IsSaving())
diff --git a/Engine/Source/Runtime/Core/Private/Internationalization/TextHistory.h b/Engine/Source/Runtime/Core/Private/Internationalization/TextHistory.h
--- a/Engine/Source/Runtime/Core/Private/Internationalization/TextHistory.h
+++ b/Engine/Source/Runtime/Core/Private/Internationalization/TextHistory.h
@@ -25,7 +25,11 @@ public:
 
 	FTextHistory();
 
-	virtual ~FTextHistory() {};
+	virtual ~FTextHistory() {}
+
+	/** Allow moving */
+	FTextHistory(FTextHistory&& Other);
+	FTextHistory& operator=(FTextHistory&& Other);
 
 	/** Rebuilds the FText from the hierarchical history, the result should be in the current locale */
 	virtual FText ToText(bool bInAsSource) const = 0;
@@ -34,13 +38,13 @@ public:
 	virtual void Serialize(FArchive& Ar) = 0;
 
 	/** Serializes data needed to get the FText's DisplayString */
-	virtual void SerializeForDisplayString(FArchive& Ar, FTextDisplayStringRef& InOutDisplayString);
+	virtual void SerializeForDisplayString(FArchive& Ar, FTextDisplayStringPtr& InOutDisplayString);
 
 	/** Returns TRUE if the Revision is out of date */
 	virtual bool IsOutOfDate();
 
 	/** Returns the source string managed by the history (if any). */
-	virtual TSharedPtr< FString, ESPMode::ThreadSafe > GetSourceString() const;
+	virtual const FString* GetSourceString() const;
 
 	/** Trace the history of this Text until we find the base Texts it was comprised from */
 	virtual void GetSourceTextsFromFormatHistory(FText Text, TArray<FText>& OutSourceTexts) const;
@@ -51,36 +55,51 @@ public:
 protected:
 	/** Revision index of this history, rebuilds when it is out of sync with the FTextLocalizationManager */
 	int32 Revision;
+
+private:
+	/** Disallow copying */
+	FTextHistory(const FTextHistory&);
+	FTextHistory& operator=(FTextHistory&);
 };
 
 /** No complexity to it, just holds the source string. */
 class CORE_API FTextHistory_Base : public FTextHistory
 {
 public:
-	FTextHistory_Base() {};
-	FTextHistory_Base(FString InSourceString);
-	FTextHistory_Base(TSharedPtr< FString, ESPMode::ThreadSafe > InSourceString);
+	FTextHistory_Base() {}
+	explicit FTextHistory_Base(FString InSourceString);
 
+	/** Allow moving */
+	FTextHistory_Base(FTextHistory_Base&& Other);
+	FTextHistory_Base& operator=(FTextHistory_Base&& Other);
 
 	//~ Begin FTextHistory Interface
 	virtual FText ToText(bool bInAsSource) const override;
 	virtual void Serialize(FArchive& Ar) override;
-	virtual void SerializeForDisplayString(FArchive& Ar, FTextDisplayStringRef& InOutDisplayString) override;
+	virtual void SerializeForDisplayString(FArchive& Ar, FTextDisplayStringPtr& InOutDisplayString) override;
 	virtual bool IsOutOfDate() override { return false; }
-	virtual TSharedPtr< FString, ESPMode::ThreadSafe > GetSourceString() const override;
+	virtual const FString* GetSourceString() const override;
 	//~ End FTextHistory Interface
 
 private:
+	/** Disallow copying */
+	FTextHistory_Base(const FTextHistory_Base&);
+	FTextHistory_Base& operator=(FTextHistory_Base&);
+
 	/** The source string for an FText */
-	FTextDisplayStringPtr SourceString;
+	FString SourceString;
 };
 
 /** Handles history for FText::Format when passing named arguments */
 class CORE_API FTextHistory_NamedFormat : public FTextHistory
 {
 public:
-	FTextHistory_NamedFormat() {};
-	FTextHistory_NamedFormat(const FText& InSourceText, const FFormatNamedArguments& InArguments);
+	FTextHistory_NamedFormat() {}
+	FTextHistory_NamedFormat(FText InSourceText, FFormatNamedArguments InArguments);
+
+	/** Allow moving */
+	FTextHistory_NamedFormat(FTextHistory_NamedFormat&& Other);
+	FTextHistory_NamedFormat& operator=(FTextHistory_NamedFormat&& Other);
 
 	//~ Begin FTextHistory Interface
 	virtual FText ToText(bool bInAsSource) const override;
@@ -89,6 +108,10 @@ public:
 	//~ End FTextHistory Interface
 
 private:
+	/** Disallow copying */
+	FTextHistory_NamedFormat(const FTextHistory_NamedFormat&);
+	FTextHistory_NamedFormat& operator=(FTextHistory_NamedFormat&);
+
 	/** The pattern used to format the text */
 	FText SourceText;
 	/** Arguments to replace in the pattern string */
@@ -99,8 +122,12 @@ private:
 class CORE_API FTextHistory_OrderedFormat : public FTextHistory
 {
 public:
-	FTextHistory_OrderedFormat() {};
-	FTextHistory_OrderedFormat(const FText& InSourceText, const FFormatOrderedArguments& InArguments);
+	FTextHistory_OrderedFormat() {}
+	FTextHistory_OrderedFormat(FText InSourceText, FFormatOrderedArguments InArguments);
+
+	/** Allow moving */
+	FTextHistory_OrderedFormat(FTextHistory_OrderedFormat&& Other);
+	FTextHistory_OrderedFormat& operator=(FTextHistory_OrderedFormat&& Other);
 
 	//~ Begin FTextHistory Interface
 	virtual FText ToText(bool bInAsSource) const override;
@@ -109,6 +136,10 @@ public:
 	//~ End FTextHistory Interface
 
 private:
+	/** Disallow copying */
+	FTextHistory_OrderedFormat(const FTextHistory_OrderedFormat&);
+	FTextHistory_OrderedFormat& operator=(FTextHistory_OrderedFormat&);
+
 	/** The pattern used to format the text */
 	FText SourceText;
 	/** Arguments to replace in the pattern string */
@@ -119,8 +150,12 @@ private:
 class CORE_API FTextHistory_ArgumentDataFormat : public FTextHistory
 {
 public:
-	FTextHistory_ArgumentDataFormat() {};
-	FTextHistory_ArgumentDataFormat(const FText& InSourceText, const TArray< struct FFormatArgumentData >& InArguments);
+	FTextHistory_ArgumentDataFormat() {}
+	FTextHistory_ArgumentDataFormat(FText InSourceText, TArray<FFormatArgumentData> InArguments);
+
+	/** Allow moving */
+	FTextHistory_ArgumentDataFormat(FTextHistory_ArgumentDataFormat&& Other);
+	FTextHistory_ArgumentDataFormat& operator=(FTextHistory_ArgumentDataFormat&& Other);
 
 	//~ Begin FTextHistory Interface
 	virtual FText ToText(bool bInAsSource) const override;
@@ -129,65 +164,99 @@ public:
 	//~ End FTextHistory Interface
 
 private:
+	/** Disallow copying */
+	FTextHistory_ArgumentDataFormat(const FTextHistory_ArgumentDataFormat&);
+	FTextHistory_ArgumentDataFormat& operator=(FTextHistory_ArgumentDataFormat&);
+
 	/** The pattern used to format the text */
 	FText SourceText;
 	/** Arguments to replace in the pattern string */
-	TArray< struct FFormatArgumentData > Arguments;
+	TArray<FFormatArgumentData> Arguments;
 };
 
 /** Base class for managing formatting FText's from: AsNumber, AsPercent, and AsCurrency. Manages data serialization of these history events */
 class CORE_API FTextHistory_FormatNumber : public FTextHistory
 {
 public:
-	FTextHistory_FormatNumber();
-	FTextHistory_FormatNumber(const FFormatArgumentValue& InSourceValue, const FNumberFormattingOptions* const InFormatOptions, const FCulturePtr InTargetCulture);
+	FTextHistory_FormatNumber() {}
+	FTextHistory_FormatNumber(FFormatArgumentValue InSourceValue, const FNumberFormattingOptions* const InFormatOptions, FCulturePtr InTargetCulture);
 
-	~FTextHistory_FormatNumber();
+	/** Allow moving */
+	FTextHistory_FormatNumber(FTextHistory_FormatNumber&& Other);
+	FTextHistory_FormatNumber& operator=(FTextHistory_FormatNumber&& Other);
 
 	//~ Begin FTextHistory Interface
 	virtual void Serialize(FArchive& Ar) override;
-	//~ End FTextHistory Interface
+	//~ End FTextHistory interface
+
 protected:
 	/** The source value to format from */
 	FFormatArgumentValue SourceValue;
-	/** All the formatting options available to format using. This can be NULL */
-	FNumberFormattingOptions* FormatOptions;
+	/** All the formatting options available to format using. This can be empty. */
+	TOptional<FNumberFormattingOptions> FormatOptions;
 	/** The culture to format using */
 	FCulturePtr TargetCulture;
+
+private:
+	/** Disallow copying */
+	FTextHistory_FormatNumber(const FTextHistory_FormatNumber&);
+	FTextHistory_FormatNumber& operator=(FTextHistory_FormatNumber&);
 };
 
 /**  Handles history for formatting using AsNumber */
 class CORE_API FTextHistory_AsNumber : public FTextHistory_FormatNumber
 {
 public:
-	FTextHistory_AsNumber() {};
-	FTextHistory_AsNumber(const FFormatArgumentValue& InSourceValue, const FNumberFormattingOptions* const InFormatOptions, const FCulturePtr InTargetCulture);
+	FTextHistory_AsNumber() {}
+	FTextHistory_AsNumber(FFormatArgumentValue InSourceValue, const FNumberFormattingOptions* const InFormatOptions, FCulturePtr InTargetCulture);
+
+	/** Allow moving */
+	FTextHistory_AsNumber(FTextHistory_AsNumber&& Other);
+	FTextHistory_AsNumber& operator=(FTextHistory_AsNumber&& Other);
 
 	//~ Begin FTextHistory Interface
 	virtual FText ToText(bool bInAsSource) const override;
 	virtual void Serialize(FArchive& Ar) override;
-	//~ End FTextHistory Interface
+	//~ End FTextHistory interface
+
+private:
+	/** Disallow copying */
+	FTextHistory_AsNumber(const FTextHistory_AsNumber&);
+	FTextHistory_AsNumber& operator=(FTextHistory_AsNumber&);
 };
 
 /**  Handles history for formatting using AsPercent */
 class CORE_API FTextHistory_AsPercent : public FTextHistory_FormatNumber
 {
 public:
-	FTextHistory_AsPercent() {};
-	FTextHistory_AsPercent(const FFormatArgumentValue& InSourceValue, const FNumberFormattingOptions* const InFormatOptions, const FCulturePtr InTargetCulture);
+	FTextHistory_AsPercent() {}
+	FTextHistory_AsPercent(FFormatArgumentValue InSourceValue, const FNumberFormattingOptions* const InFormatOptions, FCulturePtr InTargetCulture);
+
+	/** Allow moving */
+	FTextHistory_AsPercent(FTextHistory_AsPercent&& Other);
+	FTextHistory_AsPercent& operator=(FTextHistory_AsPercent&& Other);
 
 	//~ Begin FTextHistory Interface
 	virtual FText ToText(bool bInAsSource) const override;
 	virtual void Serialize(FArchive& Ar) override;
-	//~ End FTextHistory Interface
+	//~ End FTextHistory interface
+
+private:
+	/** Disallow copying */
+	FTextHistory_AsPercent(const FTextHistory_AsPercent&);
+	FTextHistory_AsPercent& operator=(FTextHistory_AsPercent&);
 };
 
 /**  Handles history for formatting using AsCurrency */
 class CORE_API FTextHistory_AsCurrency : public FTextHistory_FormatNumber
 {
 public:
-	FTextHistory_AsCurrency() {};
-	FTextHistory_AsCurrency(const FFormatArgumentValue& InSourceValue, const FString& CurrencyCode, const FNumberFormattingOptions* const InFormatOptions, const FCulturePtr InTargetCulture);
+	FTextHistory_AsCurrency() {}
+	FTextHistory_AsCurrency(FFormatArgumentValue InSourceValue, FString InCurrencyCode, const FNumberFormattingOptions* const InFormatOptions, FCulturePtr InTargetCulture);
+
+	/** Allow moving */
+	FTextHistory_AsCurrency(FTextHistory_AsCurrency&& Other);
+	FTextHistory_AsCurrency& operator=(FTextHistory_AsCurrency&& Other);
 
 	//~ Begin FTextHistory Interface
 	virtual FText ToText(bool bInAsSource) const override;
@@ -195,6 +264,10 @@ public:
 	//~ End FTextHistory Interface
 
 private:
+	/** Disallow copying */
+	FTextHistory_AsCurrency(const FTextHistory_AsCurrency&);
+	FTextHistory_AsCurrency& operator=(FTextHistory_AsCurrency&);
+
 	/** The currency used to format the number. */
 	FString CurrencyCode;
 };
@@ -203,8 +276,12 @@ private:
 class CORE_API FTextHistory_AsDate : public FTextHistory
 {
 public:
-	FTextHistory_AsDate() {};
-	FTextHistory_AsDate(const FDateTime& InSourceDateTime, const EDateTimeStyle::Type InDateStyle, const FString& InTimeZone, const FCulturePtr InTargetCulture);
+	FTextHistory_AsDate() {}
+	FTextHistory_AsDate(FDateTime InSourceDateTime, const EDateTimeStyle::Type InDateStyle, FString InTimeZone, FCulturePtr InTargetCulture);
+
+	/** Allow moving */
+	FTextHistory_AsDate(FTextHistory_AsDate&& Other);
+	FTextHistory_AsDate& operator=(FTextHistory_AsDate&& Other);
 
 	//~ Begin FTextHistory Interface
 	virtual FText ToText(bool bInAsSource) const override;
@@ -212,6 +289,10 @@ public:
 	//~ End FTextHistory Interface
 
 private:
+	/** Disallow copying */
+	FTextHistory_AsDate(const FTextHistory_AsDate&);
+	FTextHistory_AsDate& operator=(FTextHistory_AsDate&);
+
 	/** The source date structure to format */
 	FDateTime SourceDateTime;
 	/** Style to format the date using */
@@ -226,8 +307,12 @@ private:
 class CORE_API FTextHistory_AsTime : public FTextHistory
 {
 public:
-	FTextHistory_AsTime() {};
-	FTextHistory_AsTime(const FDateTime& InSourceDateTime, const EDateTimeStyle::Type InTimeStyle, const FString& InTimeZone, const FCulturePtr InTargetCulture);
+	FTextHistory_AsTime() {}
+	FTextHistory_AsTime(FDateTime InSourceDateTime, const EDateTimeStyle::Type InTimeStyle, FString InTimeZone, FCulturePtr InTargetCulture);
+
+	/** Allow moving */
+	FTextHistory_AsTime(FTextHistory_AsTime&& Other);
+	FTextHistory_AsTime& operator=(FTextHistory_AsTime&& Other);
 
 	//~ Begin FTextHistory Interface
 	virtual FText ToText(bool bInAsSource) const override;
@@ -235,6 +320,10 @@ public:
 	//~ End FTextHistory Interface
 
 private:
+	/** Disallow copying */
+	FTextHistory_AsTime(const FTextHistory_AsTime&);
+	FTextHistory_AsTime& operator=(FTextHistory_AsTime&);
+
 	/** The source time structure to format */
 	FDateTime SourceDateTime;
 	/** Style to format the time using */
@@ -249,8 +338,12 @@ private:
 class CORE_API FTextHistory_AsDateTime : public FTextHistory
 {
 public:
-	FTextHistory_AsDateTime() {};
-	FTextHistory_AsDateTime(const FDateTime& InSourceDateTime, const EDateTimeStyle::Type InDateStyle, const EDateTimeStyle::Type InTimeStyle, const FString& InTimeZone, const FCulturePtr InTargetCulture);
+	FTextHistory_AsDateTime() {}
+	FTextHistory_AsDateTime(FDateTime InSourceDateTime, const EDateTimeStyle::Type InDateStyle, const EDateTimeStyle::Type InTimeStyle, FString InTimeZone, FCulturePtr InTargetCulture);
+
+	/** Allow moving */
+	FTextHistory_AsDateTime(FTextHistory_AsDateTime&& Other);
+	FTextHistory_AsDateTime& operator=(FTextHistory_AsDateTime&& Other);
 
 	//~ Begin FTextHistory Interface
 	virtual FText ToText(bool bInAsSource) const override;
@@ -258,6 +351,10 @@ public:
 	//~ End FTextHistory Interfaces
 
 private:
+	/** Disallow copying */
+	FTextHistory_AsDateTime(const FTextHistory_AsDateTime&);
+	FTextHistory_AsDateTime& operator=(FTextHistory_AsDateTime&);
+
 	/** The source date and time structure to format */
 	FDateTime SourceDateTime;
 	/** Style to format the date using */
diff --git a/Engine/Source/Runtime/Core/Public/Internationalization/ITextData.h b/Engine/Source/Runtime/Core/Public/Internationalization/ITextData.h
new file mode 100644
--- /dev/null
+++ b/Engine/Source/Runtime/Core/Public/Internationalization/ITextData.h
@@ -0,0 +1,62 @@
+﻿// Copyright 1998-2015 Epic Games, Inc. All Rights Reserved.
+
+#pragma once
+
+#include "TextLocalizationManager.h"
+
+class FTextHistory;
+
+/** 
+ * Interface to the internal data for an FText.
+ * Various derived types are optimized to reduce memory allocation overhead.
+ */
+class ITextData
+{
+public:
+	virtual ~ITextData()
+	{
+	}
+
+	/**
+	 * Get the string to use for display purposes.
+	 * This may have come from the localization manager, or may been generated at runtime (eg, via FText::AsNumber).
+	 */
+	virtual const FString& GetDisplayString() const = 0;
+
+	/** 
+	 * Get the string pointer that was retrieved from the text localization manager.
+	 * Text that was generated at runtime by default will not have one of these, and you must call Persist() to generate one.
+	 */
+	virtual FTextDisplayStringPtr GetLocalizedString() const = 0;
+
+	/** 
+	 * Get a mutable reference to the localized string associated with this text (used when loading/saving text).
+	 */
+	virtual FTextDisplayStringPtr& GetMutableLocalizedString() = 0;
+
+	/**
+	 * Get the history associated with this text.
+	 */
+	virtual const FTextHistory& GetTextHistory() const = 0;
+
+	/**
+	 * Get a mutable reference to the history associated with this text (used when loading/saving text).
+	 */
+	virtual FTextHistory& GetMutableTextHistory() = 0;
+
+	/**
+	 * Persist this text so that it can be stored in the localization manager.
+	 */
+	virtual void PersistText() = 0;
+
+	/**
+	 * Assign a new history object to this instance.
+	 * @note: There is no RTTI on these types, so it is your responsibility to make sure that the history object you're assigning is of the correct type!
+	 */
+	template <typename THistoryType>
+	void SetTextHistory(THistoryType&& InHistory)
+	{
+		// Need to cast to avoid slicing data
+		static_cast<THistoryType&>(GetMutableTextHistory()) = Forward<THistoryType>(InHistory);
+	}
+};
diff --git a/Engine/Source/Runtime/Core/Public/Internationalization/Text.h b/Engine/Source/Runtime/Core/Public/Internationalization/Text.h
--- a/Engine/Source/Runtime/Core/Public/Internationalization/Text.h
+++ b/Engine/Source/Runtime/Core/Public/Internationalization/Text.h
@@ -4,6 +4,7 @@
 #include "Internationalization/CulturePointer.h"
 #include "Templates/SharedPointer.h"
 #include "UObject/NameTypes.h"
+#include "ITextData.h"
 #include "TextLocalizationManager.h"
 #include "Optional.h"
 #include "UniquePtr.h"
@@ -11,6 +12,8 @@
 struct FTimespan;
 struct FDateTime;
 
+class FTextHistory;
+
 #define ENABLE_TEXT_ERROR_CHECKING_RESULTS (UE_BUILD_DEBUG | UE_BUILD_DEVELOPMENT | UE_BUILD_TEST )
 
 
@@ -67,33 +70,7 @@ namespace EFormatArgumentType
 	};
 }
 
-struct CORE_API FFormatArgumentValue
-{
-	EFormatArgumentType::Type Type;
-	union
-	{
-		int64 IntValue;
-		uint64 UIntValue;
-		float FloatValue;
-		double DoubleValue;
-		FText* TextValue;
-	};
-
-	FFormatArgumentValue();
-
-	FFormatArgumentValue( const int Value );
-	FFormatArgumentValue( const unsigned int Value );
-	FFormatArgumentValue( const int64 Value );
-	FFormatArgumentValue( const uint64 Value );
-	FFormatArgumentValue( const float Value );
-	FFormatArgumentValue( const double Value );
-	FFormatArgumentValue( const FText& Value );
-	FFormatArgumentValue( const FFormatArgumentValue& Source );
-	~FFormatArgumentValue();
-
-	friend FArchive& operator<<( FArchive& Ar, FFormatArgumentValue& Value );
-};
-
+class FFormatArgumentValue;
 typedef TMap<FString, FFormatArgumentValue> FFormatNamedArguments;
 typedef TArray<FFormatArgumentValue> FFormatOrderedArguments;
 
@@ -264,10 +241,7 @@ public:
 	/** Deep build of the source string for this FText, climbing the history hierarchy */
 	FString BuildSourceString() const;
 
-	bool IsNumeric() const
-	{
-		return DisplayString.Get().IsNumeric();
-	}
+	bool IsNumeric() const;
 
 	int32 CompareTo( const FText& Other, const ETextComparisonLevel::Type ComparisonLevel = ETextComparisonLevel::Default ) const;
 	int32 CompareToCaseIgnored( const FText& Other ) const;
@@ -300,10 +274,7 @@ public:
 #endif
 	};
 
-	bool IsEmpty() const
-	{
-		return DisplayString.Get().IsEmpty();
-	}
+	bool IsEmpty() const;
 
 	bool IsEmptyOrWhitespace() const;
 
@@ -352,14 +323,7 @@ public:
 	 * @return a formatted FText
 	 */
 	template < typename... TArguments >
-	static FText FormatNamed( const FText& Fmt, TArguments&&... Args )
-	{
-		static_assert( sizeof...( TArguments ) % 2 == 0, "FormatNamed requires an even number of Name <-> Value pairs" );
-
-		FFormatNamedArguments FormatArguments;
-		FormatNamed_Util( FormatArguments, Forward< TArguments >( Args )... );
-		return FormatInternal( Fmt, FormatArguments, false, false );
-	}
+	static FText FormatNamed( const FText& Fmt, TArguments&&... Args );
 
 	/**
 	 * FormatOrdered allows you to pass a variadic list of types to use for formatting in order desired
@@ -369,42 +333,8 @@ public:
 	 * @return a formatted FText
 	 */
 	template < typename... TArguments >
-	static FText FormatOrdered( const FText& Fmt, TArguments&&... Args )
-	{
-		FFormatOrderedArguments FormatArguments;
-		FormatOrdered_Util( FormatArguments, Forward< TArguments >( Args )... );
-		return FormatInternal( Fmt, FormatArguments, false, false );
-	}
-
-private:
-	template < typename TName, typename TValue, typename... TArguments >
-	static void FormatNamed_Util( OUT FFormatNamedArguments& Result, TName&& Name, TValue&& Value, TArguments&&... Args )
-	{
-		FormatNamed_Util( Result, Name, Value );
-		FormatNamed_Util( Result, Forward< TArguments >( Args )... );
-	}
-
-	template < typename TName, typename TValue >
-	static void FormatNamed_Util( OUT FFormatNamedArguments& Result, TName&& Name, TValue&& Value )
-	{
-		Result.Emplace( Name, Value );
-	}
-
-	template < typename TValue, typename... TArguments >
-	static void FormatOrdered_Util( OUT FFormatOrderedArguments& Result, TValue&& Value, TArguments&&... Args )
-	{
-		FormatOrdered_Util( Result, Value );
-		FormatOrdered_Util( Result, Forward< TArguments >( Args )... );
-	}
-
-	template < typename TValue >
-	static void FormatOrdered_Util( OUT FFormatOrderedArguments& Result, TValue&& Value )
-	{
-		Result.Emplace( Value );
-	}
-
-public:
-#endif
+	static FText FormatOrdered( const FText& Fmt, TArguments&&... Args );
+#endif // PLATFORM_COMPILER_HAS_VARIADIC_TEMPLATES
 
 	static void SetEnableErrorCheckingResults(bool bEnable){bEnableErrorCheckingResults=bEnable;}
 	static bool GetEnableErrorCheckingResults(){return bEnableErrorCheckingResults;}
@@ -412,8 +342,8 @@ public:
 	static void SetSuppressWarnings(bool bSuppress){ bSuppressWarnings = bSuppress; }
 	static bool GetSuppressWarnings(){ return bSuppressWarnings; }
 
-	bool IsTransient() const { return (Flags & ETextFlag::Transient) != 0; }
-	bool IsCultureInvariant() const { return (Flags & ETextFlag::CultureInvariant) != 0; }
+	bool IsTransient() const;
+	bool IsCultureInvariant() const;
 
 	bool ShouldGatherForLocalization() const;
 
@@ -423,9 +353,11 @@ private:
 	enum class EInitToEmptyString : uint8 { Value };
 	explicit FText( EInitToEmptyString );
 
-	FText( FString InSourceString, TSharedPtr<class FTextHistory, ESPMode::ThreadSafe> InHistory );
+	explicit FText( TSharedRef<ITextData, ESPMode::ThreadSafe> InTextData );
 
-	FText( FString InSourceString, FString InNamespace, FString InKey, int32 InFlags=0 );
+	explicit FText( FString InSourceString );
+
+	FText( FString InSourceString, FString InNamespace, FString InKey, uint32 InFlags=0 );
 
 	friend CORE_API FArchive& operator<<( FArchive& Ar, FText& Value );
 
@@ -439,15 +371,15 @@ private:
 	/**
 	 * Generate an FText for a string formatted numerically.
 	 */
-	static FText CreateNumericalText(FString InSourceString, TSharedPtr<class FTextHistory, ESPMode::ThreadSafe> InHistory);
+	static FText CreateNumericalText(TSharedRef<ITextData, ESPMode::ThreadSafe> InTextData);
 
 	/**
 	 * Generate an FText for a string formatted from a date/time.
 	 */
-	static FText CreateChronologicalText(FString InSourceString, TSharedPtr<class FTextHistory, ESPMode::ThreadSafe> InHistory);
+	static FText CreateChronologicalText(TSharedRef<ITextData, ESPMode::ThreadSafe> InTextData);
 
 	/** Returns the source string of the FText */
-	TSharedPtr< FString, ESPMode::ThreadSafe > GetSourceString() const;
+	const FString& GetSourceString() const;
 
 	/** Rebuilds the FText under the current culture if needed */
 	void Rebuild() const;
@@ -465,14 +397,11 @@ private:
 	static FText AsPercentTemplate(T1 Val, const FNumberFormattingOptions* const Options = NULL, const FCulturePtr& TargetCulture = NULL);
 
 private:
-	/** The visible display string for this FText */
-	FTextDisplayStringRef DisplayString;
-
-	/** The FText's history, to allow it to rebuild under a new culture */
-	TSharedPtr<class FTextHistory, ESPMode::ThreadSafe> History;
+	/** The internal shared data for this FText */
+	TSharedRef<ITextData, ESPMode::ThreadSafe> TextData;
 
 	/** Flags with various information on what sort of FText this is */
-	int32 Flags;
+	uint32 Flags;
 
 	static bool bEnableErrorCheckingResults;
 	static bool bSuppressWarnings;
@@ -503,6 +432,170 @@ public:
 	friend class FScopedTextIdentityPreserver;
 };
 
+class CORE_API FFormatArgumentValue
+{
+public:
+	FFormatArgumentValue()
+		: Type(EFormatArgumentType::Int)
+	{
+		IntValue = 0;
+	}
+
+	FFormatArgumentValue(const int32 Value)
+		: Type(EFormatArgumentType::Int)
+	{
+		IntValue = Value;
+	}
+
+	FFormatArgumentValue(const uint32 Value)
+		: Type(EFormatArgumentType::UInt)
+	{
+		UIntValue = Value;
+	}
+
+	FFormatArgumentValue(const int64 Value)
+		: Type(EFormatArgumentType::Int)
+	{
+		IntValue = Value;
+	}
+
+	FFormatArgumentValue(const uint64 Value)
+		: Type(EFormatArgumentType::UInt)
+	{
+		UIntValue = Value;
+	}
+
+	FFormatArgumentValue(const float Value)
+		: Type(EFormatArgumentType::Float)
+	{
+		FloatValue = Value;
+	}
+
+	FFormatArgumentValue(const double Value)
+		: Type(EFormatArgumentType::Double)
+	{
+		DoubleValue = Value;
+	}
+
+	FFormatArgumentValue(FText Value)
+		: Type(EFormatArgumentType::Text)
+		, TextValue(MoveTemp(Value))
+	{
+	}
+
+	friend FArchive& operator<<(FArchive& Ar, FFormatArgumentValue& Value);
+
+	FORCEINLINE EFormatArgumentType::Type GetType() const
+	{
+		return Type;
+	}
+
+	FORCEINLINE int64 GetIntValue() const
+	{
+		check(Type == EFormatArgumentType::Int);
+		return IntValue;
+	}
+
+	FORCEINLINE uint64 GetUIntValue() const
+	{
+		check(Type == EFormatArgumentType::UInt);
+		return UIntValue;
+	}
+
+	FORCEINLINE float GetFloatValue() const
+	{
+		check(Type == EFormatArgumentType::Float);
+		return FloatValue;
+	}
+
+	FORCEINLINE double GetDoubleValue() const
+	{
+		check(Type == EFormatArgumentType::Double);
+		return DoubleValue;
+	}
+
+	FORCEINLINE const FText& GetTextValue() const
+	{
+		check(Type == EFormatArgumentType::Text);
+		return TextValue.GetValue();
+	}
+
+private:
+	EFormatArgumentType::Type Type;
+	union
+	{
+		int64 IntValue;
+		uint64 UIntValue;
+		float FloatValue;
+		double DoubleValue;
+	};
+	TOptional<FText> TextValue;
+};
+
+struct FFormatArgumentData
+{
+	FText ArgumentName;
+	FText ArgumentValue;
+
+	friend inline FArchive& operator<<( FArchive& Ar, FFormatArgumentData& Value )
+	{
+		Ar << Value.ArgumentName;
+		Ar << Value.ArgumentValue;
+		return Ar;
+	}
+};
+
+#if PLATFORM_COMPILER_HAS_VARIADIC_TEMPLATES
+namespace TextFormatUtil
+{
+
+template < typename TName, typename TValue >
+void FormatNamed( OUT FFormatNamedArguments& Result, TName&& Name, TValue&& Value )
+{
+	Result.Emplace( Forward< TName >( Name ), Forward< TValue >( Value ) );
+}
+
+template < typename TName, typename TValue, typename... TArguments >
+void FormatNamed( OUT FFormatNamedArguments& Result, TName&& Name, TValue&& Value, TArguments&&... Args )
+{
+	FormatNamed( Result, Forward< TName >( Name ), Forward< TValue >( Value ) );
+	FormatNamed( Result, Forward< TArguments >( Args )... );
+}
+
+template < typename TValue >
+void FormatOrdered( OUT FFormatOrderedArguments& Result, TValue&& Value )
+{
+	Result.Emplace( Forward< TValue >( Value ) );
+}
+
+template < typename TValue, typename... TArguments >
+void FormatOrdered( OUT FFormatOrderedArguments& Result, TValue&& Value, TArguments&&... Args )
+{
+	FormatOrdered( Result, Forward< TValue >( Value ) );
+	FormatOrdered( Result, Forward< TArguments >( Args )... );
+}
+
+} // namespace TextFormatUtil
+
+template < typename... TArguments >
+FText FText::FormatNamed( const FText& Fmt, TArguments&&... Args )
+{
+	static_assert( sizeof...( TArguments ) % 2 == 0, "FormatNamed requires an even number of Name <-> Value pairs" );
+
+	FFormatNamedArguments FormatArguments;
+	TextFormatUtil::FormatNamed( FormatArguments, Forward< TArguments >( Args )... );
+	return FormatInternal( Fmt, FormatArguments, false, false );
+}
+
+template < typename... TArguments >
+FText FText::FormatOrdered( const FText& Fmt, TArguments&&... Args )
+{
+	FFormatOrderedArguments FormatArguments;
+	TextFormatUtil::FormatOrdered( FormatArguments, Forward< TArguments >( Args )... );
+	return FormatInternal( Fmt, FormatArguments, false, false );
+}
+#endif // PLATFORM_COMPILER_HAS_VARIADIC_TEMPLATES
+
 /** A snapshot of an FText at a point in time that can be used to detect changes in the FText, including live-culture changes */
 class CORE_API FTextSnapshot
 {
@@ -518,14 +611,14 @@ public:
 	bool IsDisplayStringEqualTo(const FText& InText) const;
 
 private:
-	/** A pointer to the visible display string for the FText we took a snapshot of (used for an efficient pointer compare) */
-	FTextDisplayStringPtr DisplayStringPtr;
+	/** A pointer to the text data for the FText that we took a snapshot of (used for an efficient pointer compare) */
+	TWeakPtr<ITextData, ESPMode::ThreadSafe> TextDataPtr;
 
 	/** Revision index of the history of the FText we took a snapshot of, or INDEX_NONE if there was no history */
 	int32 HistoryRevision;
 
 	/** Flags with various information on what sort of FText we took a snapshot of */
-	int32 Flags;
+	uint32 Flags;
 };
 
 class CORE_API FTextInspector
@@ -541,15 +634,7 @@ public:
 	static const FString* GetSourceString(const FText& Text);
 	static const FString& GetDisplayString(const FText& Text);
 	static const FTextDisplayStringRef GetSharedDisplayString(const FText& Text);
-	static int32 GetFlags(const FText& Text);
-};
-
-struct FFormatArgumentData
-{
-	FText ArgumentName;
-	FText ArgumentValue;
-
-	friend FArchive& operator<<( FArchive& Ar, FFormatArgumentData& Value );
+	static uint32 GetFlags(const FText& Text);
 };
 
 class CORE_API FTextBuilder
@@ -660,11 +745,11 @@ public:
 	~FScopedTextIdentityPreserver();
 
 private:
+	FText& TextToPersist;
 	bool HadFoundNamespaceAndKey;
 	FString Namespace;
 	FString Key;
-	int32 Flags;
-	FText& TextToPersist;
+	uint32 Flags;
 };
 
 /**
diff --git a/Engine/Source/Runtime/CoreUObject/Private/UObject/Class.cpp b/Engine/Source/Runtime/CoreUObject/Private/UObject/Class.cpp
--- a/Engine/Source/Runtime/CoreUObject/Private/UObject/Class.cpp
+++ b/Engine/Source/Runtime/CoreUObject/Private/UObject/Class.cpp
@@ -1006,6 +1006,7 @@ void UStruct::SerializeTaggedProperties(FArchive& Ar, uint8* Data, UStruct* Defa
 				FString str;
 				Ar << str;
 				FText Text = FText::FromString(str);
+				Text.TextData->PersistText();
 				Text.Flags |= ETextFlag::ConvertedProperty;
 				CastChecked<UTextProperty>(Property)->SetPropertyValue_InContainer(Data, Text, Tag.ArrayIndex);
 				AdvanceProperty = true;
@@ -1185,6 +1186,7 @@ void UStruct::SerializeTaggedProperties(FArchive& Ar, uint8* Data, UStruct* Defa
 						FString str;
 						Ar << str;
 						FText Text = FText::FromString(str);
+						Text.TextData->PersistText();
 						Text.Flags |= ETextFlag::ConvertedProperty;
 						CastChecked<UTextProperty>(ArrayProperty->Inner)->SetPropertyValue(ScriptArrayHelper.GetRawPtr(i), Text);
 						AdvanceProperty = true;
@@ -1210,6 +1212,7 @@ void UStruct::SerializeTaggedProperties(FArchive& Ar, uint8* Data, UStruct* Defa
 						FName Name;
 						Ar << Name;
 						FText Text = FText::FromName(Name);
+						Text.TextData->PersistText();
 						Text.Flags |= ETextFlag::ConvertedProperty;
 						CastChecked<UTextProperty>(ArrayProperty->Inner)->SetPropertyValue(ScriptArrayHelper.GetRawPtr(i), Text);
 						AdvanceProperty = true;
