Minimized FText allocations when generating text at runtime

This allows you to move the arguments passed to FText::Format if you know you no longer need your own copy (our internal helper Format functions now do this with their temporary data).

I'd previously claimed this was submitted as part of CL# 2702379, but found it as part of a different CL.

diff --git a/Engine/Source/Runtime/Core/Private/Internationalization/Text.cpp b/Engine/Source/Runtime/Core/Private/Internationalization/Text.cpp
--- a/Engine/Source/Runtime/Core/Private/Internationalization/Text.cpp
+++ b/Engine/Source/Runtime/Core/Private/Internationalization/Text.cpp
@@ -366,38 +366,38 @@ FText FText::TrimPrecedingAndTrailing( const FText& InText )
 	return NewText;
 }
 
-FText FText::Format(const FText& Fmt,const FText& v1)
+FText FText::Format(FText Fmt, FText v1)
 {
 	FFormatOrderedArguments Arguments;
-	Arguments.Add(v1);
-	return FText::Format(Fmt, Arguments);
+	Arguments.Add(MoveTemp(v1));
+	return FText::Format(MoveTemp(Fmt), MoveTemp(Arguments));
 }
 
-FText FText::Format(const FText& Fmt,const FText& v1,const FText& v2)
+FText FText::Format(FText Fmt, FText v1, FText v2)
 {
 	FFormatOrderedArguments Arguments;
-	Arguments.Add(v1);
-	Arguments.Add(v2);
-	return FText::Format(Fmt, Arguments);
+	Arguments.Add(MoveTemp(v1));
+	Arguments.Add(MoveTemp(v2));
+	return FText::Format(MoveTemp(Fmt), MoveTemp(Arguments));
 }
 
-FText FText::Format(const FText& Fmt,const FText& v1,const FText& v2,const FText& v3)
+FText FText::Format(FText Fmt, FText v1, FText v2, FText v3)
 {
 	FFormatOrderedArguments Arguments;
-	Arguments.Add(v1);
-	Arguments.Add(v2);
-	Arguments.Add(v3);
-	return FText::Format(Fmt, Arguments);
+	Arguments.Add(MoveTemp(v1));
+	Arguments.Add(MoveTemp(v2));
+	Arguments.Add(MoveTemp(v3));
+	return FText::Format(MoveTemp(Fmt), MoveTemp(Arguments));
 }
 
-FText FText::Format(const FText& Fmt,const FText& v1,const FText& v2,const FText& v3,const FText& v4)
+FText FText::Format(FText Fmt, FText v1, FText v2, FText v3, FText v4)
 {
 	FFormatOrderedArguments Arguments;
-	Arguments.Add(v1);
-	Arguments.Add(v2);
-	Arguments.Add(v3);
-	Arguments.Add(v4);
-	return FText::Format(Fmt, Arguments);
+	Arguments.Add(MoveTemp(v1));
+	Arguments.Add(MoveTemp(v2));
+	Arguments.Add(MoveTemp(v3));
+	Arguments.Add(MoveTemp(v4));
+	return FText::Format(MoveTemp(Fmt), MoveTemp(Arguments));
 }
 
 class FTextFormatHelper
@@ -741,22 +741,22 @@ void FText::GetFormatPatternParameters(const FText& Pattern, TArray<FString>& Pa
 	FTextFormatHelper::EnumerateParameters(Pattern, ParameterNames);
 }
 
-FText FText::Format(const FText& Pattern, const FFormatNamedArguments& Arguments)
+FText FText::Format(FText Pattern, FFormatNamedArguments Arguments)
 {
-	return FormatInternal(Pattern, Arguments, false, false);
+	return FormatInternal(MoveTemp(Pattern), MoveTemp(Arguments), false, false);
 }
 
-FText FText::Format(const FText& Pattern, const FFormatOrderedArguments& Arguments)
+FText FText::Format(FText Pattern, FFormatOrderedArguments Arguments)
 {
-	return FormatInternal(Pattern, Arguments, false, false);
+	return FormatInternal(MoveTemp(Pattern), MoveTemp(Arguments), false, false);
 }
 
-FText FText::Format(const FText& Pattern, const TArray< FFormatArgumentData > InArguments)
+FText FText::Format(FText Pattern, TArray< FFormatArgumentData > InArguments)
 {
-	return FormatInternal(Pattern, InArguments, false, false);
+	return FormatInternal(MoveTemp(Pattern), MoveTemp(InArguments), false, false);
 }
 
-FText FText::FormatInternal(const FText& Pattern, const FFormatNamedArguments& Arguments, bool bInRebuildText, bool bInRebuildAsSource)
+FText FText::FormatInternal(FText Pattern, FFormatNamedArguments Arguments, bool bInRebuildText, bool bInRebuildAsSource)
 {
 	checkf(FInternationalization::Get().IsInitialized() == true, TEXT("FInternationalization is not initialized. An FText formatting method was likely used in static object initialization - this is not supported."));
 	//SCOPE_CYCLE_COUNTER( STAT_TextFormat );
@@ -774,7 +774,7 @@ FText FText::FormatInternal(const FText& Pattern, const FFormatNamedArguments& A
 
 	FString ResultString = FTextFormatHelper::Format(Pattern, EstimatedArgumentValuesLength, FTextFormatHelper::FGetArgumentValue::CreateLambda(GetArgumentValue), bInRebuildText, bInRebuildAsSource);
 	
-	FText Result = FText(MakeShareable(new TGeneratedTextData<FTextHistory_NamedFormat>(MoveTemp(ResultString), FTextHistory_NamedFormat(Pattern, Arguments))));
+	FText Result = FText(MakeShareable(new TGeneratedTextData<FTextHistory_NamedFormat>(MoveTemp(ResultString), FTextHistory_NamedFormat(MoveTemp(Pattern), MoveTemp(Arguments)))));
 	if (!GIsEditor)
 	{
 		Result.Flags |= ETextFlag::Transient;
@@ -782,7 +782,7 @@ FText FText::FormatInternal(const FText& Pattern, const FFormatNamedArguments& A
 	return Result;
 }
 
-FText FText::FormatInternal(const FText& Pattern, const FFormatOrderedArguments& Arguments, bool bInRebuildText, bool bInRebuildAsSource)
+FText FText::FormatInternal(FText Pattern, FFormatOrderedArguments Arguments, bool bInRebuildText, bool bInRebuildAsSource)
 {
 	checkf(FInternationalization::Get().IsInitialized() == true, TEXT("FInternationalization is not initialized. An FText formatting method was likely used in static object initialization - this is not supported."));
 	//SCOPE_CYCLE_COUNTER( STAT_TextFormat );
@@ -809,7 +809,7 @@ FText FText::FormatInternal(const FText& Pattern, const FFormatOrderedArguments&
 
 	FString ResultString = FTextFormatHelper::Format(Pattern, EstimatedArgumentValuesLength, FTextFormatHelper::FGetArgumentValue::CreateLambda(GetArgumentValue), bInRebuildText, bInRebuildAsSource);
 
-	FText Result = FText(MakeShareable(new TGeneratedTextData<FTextHistory_OrderedFormat>(MoveTemp(ResultString), FTextHistory_OrderedFormat(Pattern, Arguments))));
+	FText Result = FText(MakeShareable(new TGeneratedTextData<FTextHistory_OrderedFormat>(MoveTemp(ResultString), FTextHistory_OrderedFormat(MoveTemp(Pattern), MoveTemp(Arguments)))));
 	if (!GIsEditor)
 	{
 		Result.Flags |= ETextFlag::Transient;
@@ -817,7 +817,7 @@ FText FText::FormatInternal(const FText& Pattern, const FFormatOrderedArguments&
 	return Result;
 }
 
-FText FText::FormatInternal(const FText& Pattern, const TArray< struct FFormatArgumentData > Arguments, bool bInRebuildText, bool bInRebuildAsSource)
+FText FText::FormatInternal(FText Pattern, TArray< struct FFormatArgumentData > Arguments, bool bInRebuildText, bool bInRebuildAsSource)
 {
 	checkf(FInternationalization::Get().IsInitialized() == true, TEXT("FInternationalization is not initialized. An FText formatting method was likely used in static object initialization - this is not supported."));
 	//SCOPE_CYCLE_COUNTER( STAT_TextFormat );
@@ -837,7 +837,7 @@ FText FText::FormatInternal(const FText& Pattern, const TArray< struct FFormatAr
 
 	FString ResultString = FTextFormatHelper::Format(Pattern, EstimatedArgumentValuesLength, FTextFormatHelper::FGetArgumentValue::CreateLambda(GetArgumentValue), bInRebuildText, bInRebuildAsSource);
 
-	FText Result = FText(MakeShareable(new TGeneratedTextData<FTextHistory_ArgumentDataFormat>(MoveTemp(ResultString), FTextHistory_ArgumentDataFormat(Pattern, Arguments))));
+	FText Result = FText(MakeShareable(new TGeneratedTextData<FTextHistory_ArgumentDataFormat>(MoveTemp(ResultString), FTextHistory_ArgumentDataFormat(MoveTemp(Pattern), MoveTemp(Arguments)))));
 	if (!GIsEditor)
 	{
 		Result.Flags |= ETextFlag::Transient;
diff --git a/Engine/Source/Runtime/Core/Public/Internationalization/Text.h b/Engine/Source/Runtime/Core/Public/Internationalization/Text.h
--- a/Engine/Source/Runtime/Core/Public/Internationalization/Text.h
+++ b/Engine/Source/Runtime/Core/Public/Internationalization/Text.h
@@ -303,14 +303,14 @@ public:
 
 	static void GetFormatPatternParameters(const FText& Pattern, TArray<FString>& ParameterNames);
 
-	static FText Format(const FText& Pattern, const FFormatNamedArguments& Arguments);
-	static FText Format(const FText& Pattern, const FFormatOrderedArguments& Arguments);
-	static FText Format(const FText& Pattern, const TArray< struct FFormatArgumentData > InArguments);
+	static FText Format(FText Pattern, FFormatNamedArguments Arguments);
+	static FText Format(FText Pattern, FFormatOrderedArguments Arguments);
+	static FText Format(FText Pattern, TArray< struct FFormatArgumentData > InArguments);
 
-	static FText Format(const FText& Fmt,const FText& v1);
-	static FText Format(const FText& Fmt,const FText& v1,const FText& v2);
-	static FText Format(const FText& Fmt,const FText& v1,const FText& v2,const FText& v3);
-	static FText Format(const FText& Fmt,const FText& v1,const FText& v2,const FText& v3,const FText& v4);
+	static FText Format(FText Fmt, FText v1);
+	static FText Format(FText Fmt, FText v1, FText v2);
+	static FText Format(FText Fmt, FText v1, FText v2, FText v3);
+	static FText Format(FText Fmt, FText v1, FText v2, FText v3, FText v4);
 
 #if PLATFORM_COMPILER_HAS_VARIADIC_TEMPLATES
 	/**
@@ -323,7 +323,7 @@ public:
 	 * @return a formatted FText
 	 */
 	template < typename... TArguments >
-	static FText FormatNamed( const FText& Fmt, TArguments&&... Args );
+	static FText FormatNamed( FText Fmt, TArguments&&... Args );
 
 	/**
 	 * FormatOrdered allows you to pass a variadic list of types to use for formatting in order desired
@@ -333,7 +333,7 @@ public:
 	 * @return a formatted FText
 	 */
 	template < typename... TArguments >
-	static FText FormatOrdered( const FText& Fmt, TArguments&&... Args );
+	static FText FormatOrdered( FText Fmt, TArguments&&... Args );
 #endif // PLATFORM_COMPILER_HAS_VARIADIC_TEMPLATES
 
 	static void SetEnableErrorCheckingResults(bool bEnable){bEnableErrorCheckingResults=bEnable;}
@@ -384,9 +384,9 @@ private:
 	/** Rebuilds the FText under the current culture if needed */
 	void Rebuild() const;
 
-	static FText FormatInternal(const FText& Pattern, const FFormatNamedArguments& Arguments, bool bInRebuildText, bool bInRebuildAsSource);
-	static FText FormatInternal(const FText& Pattern, const FFormatOrderedArguments& Arguments, bool bInRebuildText, bool bInRebuildAsSource);
-	static FText FormatInternal(const FText& Pattern, const TArray< struct FFormatArgumentData > InArguments, bool bInRebuildText, bool bInRebuildAsSource);
+	static FText FormatInternal(FText Pattern, FFormatNamedArguments Arguments, bool bInRebuildText, bool bInRebuildAsSource);
+	static FText FormatInternal(FText Pattern, FFormatOrderedArguments Arguments, bool bInRebuildText, bool bInRebuildAsSource);
+	static FText FormatInternal(FText Pattern, TArray< struct FFormatArgumentData > InArguments, bool bInRebuildText, bool bInRebuildAsSource);
 
 private:
 	template<typename T1, typename T2>
@@ -578,21 +578,21 @@ void FormatOrdered( OUT FFormatOrderedArguments& Result, TValue&& Value, TArgume
 } // namespace TextFormatUtil
 
 template < typename... TArguments >
-FText FText::FormatNamed( const FText& Fmt, TArguments&&... Args )
+FText FText::FormatNamed( FText Fmt, TArguments&&... Args )
 {
 	static_assert( sizeof...( TArguments ) % 2 == 0, "FormatNamed requires an even number of Name <-> Value pairs" );
 
 	FFormatNamedArguments FormatArguments;
 	TextFormatUtil::FormatNamed( FormatArguments, Forward< TArguments >( Args )... );
-	return FormatInternal( Fmt, FormatArguments, false, false );
+	return FormatInternal( MoveTemp( Fmt ), MoveTemp( FormatArguments ), false, false );
 }
 
 template < typename... TArguments >
-FText FText::FormatOrdered( const FText& Fmt, TArguments&&... Args )
+FText FText::FormatOrdered( FText Fmt, TArguments&&... Args )
 {
 	FFormatOrderedArguments FormatArguments;
 	TextFormatUtil::FormatOrdered( FormatArguments, Forward< TArguments >( Args )... );
-	return FormatInternal( Fmt, FormatArguments, false, false );
+	return FormatInternal( MoveTemp( Fmt ), MoveTemp( FormatArguments ), false, false );
 }
 #endif // PLATFORM_COMPILER_HAS_VARIADIC_TEMPLATES
 
