Changed verify to an ensure
This is a non-fatal error

diff --git a/Engine/Source/Runtime/Core/Public/Misc/FeedbackContext.h b/Engine/Source/Runtime/Core/Public/Misc/FeedbackContext.h
--- a/Engine/Source/Runtime/Core/Public/Misc/FeedbackContext.h
+++ b/Engine/Source/Runtime/Core/Public/Misc/FeedbackContext.h
@@ -269,7 +269,7 @@ public:
 		CompletedWork += CurrentFrameScope;
 
 		const float WorkRemaining = TotalAmountOfWork - CompletedWork;
-		verifySlow(ExpectedWorkThisFrame <= WorkRemaining);
+		ensureMsgf(ExpectedWorkThisFrame <= WorkRemaining, TEXT("Work overflow in slow task. Please revise call-site to account for entire progress range."));
 		CurrentFrameScope = FMath::Min(WorkRemaining, ExpectedWorkThisFrame);
 
 		if (!bCreatedDialog && OpenDialogThreshold.IsSet() && static_cast<float>(FPlatformTime::Seconds() - StartTime) > OpenDialogThreshold.GetValue())
