Added the concept of an 'arbitrary multiplier' unit

diff --git a/Engine/Source/Runtime/Core/Private/Math/UnitConversion.cpp b/Engine/Source/Runtime/Core/Private/Math/UnitConversion.cpp
--- a/Engine/Source/Runtime/Core/Private/Math/UnitConversion.cpp
+++ b/Engine/Source/Runtime/Core/Private/Math/UnitConversion.cpp
@@ -71,6 +71,8 @@ FParseCandidate ParseCandidates[] = {
 	{ TEXT("Days"),					EUnit::Days },					{ TEXT("dy"),		EUnit::Days },
 	{ TEXT("Months"),				EUnit::Months },				{ TEXT("mth"),		EUnit::Months },
 	{ TEXT("Years"),				EUnit::Years },					{ TEXT("yr"),		EUnit::Years },
+
+	{ TEXT("times"),				EUnit::Multiplier },			{ TEXT("x"),	EUnit::Multiplier },			{ TEXT("multiplier"),		EUnit::Multiplier },
 };
 
 /** Static array of display strings that directly map to EUnit enumerations */
@@ -97,6 +99,8 @@ const TCHAR* const DisplayStrings[] = {
 	TEXT("lm"),
 
 	TEXT("ms"), TEXT("s"), TEXT("min"), TEXT("hr"), TEXT("dy"), TEXT("mth"), TEXT("yr"),
+
+	TEXT("x"),
 };
 
 const EUnitType UnitTypes[] = {
@@ -122,6 +126,8 @@ const EUnitType UnitTypes[] = {
 	EUnitType::LuminousFlux,
 
 	EUnitType::Time,		EUnitType::Time,		EUnitType::Time,		EUnitType::Time,		EUnitType::Time,		EUnitType::Time,		EUnitType::Time,
+
+	EUnitType::Arbitrary,
 };
 
 
diff --git a/Engine/Source/Runtime/Core/Public/Math/UnitConversion.h b/Engine/Source/Runtime/Core/Public/Math/UnitConversion.h
--- a/Engine/Source/Runtime/Core/Public/Math/UnitConversion.h
+++ b/Engine/Source/Runtime/Core/Public/Math/UnitConversion.h
@@ -44,6 +44,9 @@ enum class EUnit
 	/** Time units */
 	Milliseconds, Seconds, Minutes, Hours, Days, Months, Years,
 
+	/** Arbitrary multiplier */
+	Multiplier,
+
 	/** Symbolic entry, not specifyable on meta data */
 	Unspecified
 };
@@ -51,7 +54,7 @@ enum class EUnit
 /** Enumeration that specifies particular classes of unit */
 enum class EUnitType
 {
-	Distance, Angle, Speed, Temperature, Mass, Force, Frequency, DataSize, LuminousFlux, Time,
+	Distance, Angle, Speed, Temperature, Mass, Force, Frequency, DataSize, LuminousFlux, Time, Arbitrary,
 
 	// Symbolic entry - do not use directly
 	NumberOf,
