Add LIKELY/UNLIKELY macros for static branch prediction hinting.

(Merging CL 2704512 to the main branch)

diff --git a/Engine/Source/Runtime/Core/Public/HAL/Platform.h b/Engine/Source/Runtime/Core/Public/HAL/Platform.h
--- a/Engine/Source/Runtime/Core/Public/HAL/Platform.h
+++ b/Engine/Source/Runtime/Core/Public/HAL/Platform.h
@@ -335,6 +335,23 @@
 	#define ASSUME(...) 
 #endif
 
+/** Branch prediction hints */
+#ifndef LIKELY						/* Hints compiler that expression is likely to be true, much softer than ASSUME - allows (penalized by worse performance) expression to be false */
+	#if defined(__clang__) || defined(__GNUC__)
+		#define LIKELY(x)			__builtin_expect((x), 1)
+	#else
+		#define LIKELY(x)			(x)
+	#endif
+#endif
+
+#ifndef UNLIKELY					/* Hints compiler that expression is unlikely to be true, allows (penalized by worse performance) expression to be true */
+	#if defined(__clang__) || defined(__GNUC__)
+		#define UNLIKELY(x)			__builtin_expect((x), 0)
+	#else
+		#define UNLIKELY(x)			(x)
+	#endif
+#endif
+
 // Optimization macros (uses __pragma to enable inside a #define).
 #ifndef PRAGMA_DISABLE_OPTIMIZATION_ACTUAL
 	#define PRAGMA_DISABLE_OPTIMIZATION_ACTUAL
