Support multiple gamepad controllers on Android including identifying Samsung EI-GP20 and input plugins
Adds ResetGamepadAssignments, ResetGamepadAssignmentToController, and IsControllerAssignedToGamepad to PlatformMisc and BP nodes
#ue4
#android

diff --git a/Engine/Build/Android/Java/src/com/epicgames/ue4/GameActivity.java b/Engine/Build/Android/Java/src/com/epicgames/ue4/GameActivity.java
--- a/Engine/Build/Android/Java/src/com/epicgames/ue4/GameActivity.java
+++ b/Engine/Build/Android/Java/src/com/epicgames/ue4/GameActivity.java
@@ -32,6 +32,7 @@ import android.content.pm.PackageManager.NameNotFoundException;
 import android.media.AudioManager;
 import android.util.DisplayMetrics;
 
+import android.view.InputDevice;
 import android.view.Gravity;
 import android.view.MotionEvent;
 import android.view.View;
@@ -289,6 +290,27 @@ public class GameActivity extends NativeActivity
 		}
 */
 
+/*
+		// log a list of input devices for debugging
+		{
+			int[] deviceIds = InputDevice.getDeviceIds();
+			for (int deviceIndex=0; deviceIndex < deviceIds.length; deviceIndex++)
+			{
+				InputDevice inputDevice = InputDevice.getDevice(deviceIds[deviceIndex]);
+				Log.debug("Device index " + deviceIndex + ": (deviceId=" + inputDevice.getId() + 
+				", controllerNumber=" + inputDevice.getControllerNumber() + ", sources=" + String.format("%08x", inputDevice.getSources()) +
+				", vendorId=" + String.format("%04x", inputDevice.getVendorId()) + ", productId=" + String.format("%04x", inputDevice.getProductId()) + 
+				", descriptor=" + inputDevice.getDescriptor() +	", deviceName=" + inputDevice.getName() + ")");
+
+				// is it a joystick?
+				if ((inputDevice.getSources() & InputDevice.SOURCE_JOYSTICK) != 0)
+				{
+					Log.debug("Gamepad detected: (deviceIndex=" + deviceIndex + ", deviceId=" + inputDevice.getId() + ", deviceName=" + inputDevice.getName() + ")");
+				}
+			}
+		}
+*/
+
 		// tell Android that we want volume controls to change the media volume, aka music
 		setVolumeControlStream(AudioManager.STREAM_MUSIC);
 		
@@ -1262,6 +1284,78 @@ public class GameActivity extends NativeActivity
 		}
 	}
 
+	private static class DeviceInfoData {
+		public final int vendorId;
+		public final int productId;
+		public final String name;
+
+		DeviceInfoData(int vid, int pid, String inName)
+		{
+			vendorId = vid;
+			productId = pid;
+			name = inName;
+		}
+
+		boolean IsMatch(int vid, int pid)
+		{
+			return (vendorId == vid && productId == pid);
+		}
+	}
+
+	// List of vendor/product ids
+	private static final DeviceInfoData[] DeviceInfoList = {
+		new DeviceInfoData(0x04e8, 0xa000, "Samsung Game Pad EI-GP20"),
+		new DeviceInfoData(0x0955, 0x7203, "NVIDIA Corporation NVIDIA Controller v01.01"),
+		new DeviceInfoData(0x0955, 0x7210, "NVIDIA Corporation NVIDIA Controller v01.03"),
+		new DeviceInfoData(0x1949, 0x0404, "Amazon Fire TV Remote"),
+		new DeviceInfoData(0x1949, 0x0406, "Amazon Fire Game Controller")
+	};
+
+	public class InputDeviceInfo {
+		public int deviceId;
+		public int vendorId;
+		public int productId;
+		public int controllerId;
+		public String name;
+		public String descriptor;
+
+		InputDeviceInfo(int did, int vid, int pid, int cid, String inName, String inDescriptor)
+		{
+			deviceId = did;
+			vendorId = vid;
+			productId = pid;
+			controllerId = cid;
+			name = inName;
+			descriptor = inDescriptor;
+		}
+	}
+
+	public InputDeviceInfo AndroidThunkJava_GetInputDeviceInfo(int deviceId)
+	{
+		int[] deviceIds = InputDevice.getDeviceIds();
+		for (int deviceIndex=0; deviceIndex < deviceIds.length; deviceIndex++)
+		{
+			InputDevice inputDevice = InputDevice.getDevice(deviceIds[deviceIndex]);
+			if (inputDevice.getId() == deviceId)
+			{
+				// note: inputDevice.getName may not return a proper name so check vendor and product id first
+				int vendorId = inputDevice.getVendorId();
+				int productId = inputDevice.getProductId();
+				for (DeviceInfoData deviceInfo : DeviceInfoList)
+				{
+					if (deviceInfo.IsMatch(vendorId, productId))
+					{
+						return new InputDeviceInfo(deviceId, vendorId, productId, inputDevice.getControllerNumber(), deviceInfo.name, inputDevice.getDescriptor());
+					}
+				}
+
+				// use device name as fallback (may be generic like "Bluetooth HID" so not always useful)
+				return new InputDeviceInfo(deviceId, vendorId, productId, inputDevice.getControllerNumber(), inputDevice.getName(), inputDevice.getDescriptor());
+			}
+		}
+		return new InputDeviceInfo(deviceId, 0, 0, -1, "Unknown", "Unknown");
+	}
+
 	public native boolean nativeIsShippingBuild();
 	public native void nativeSetGlobalActivity();
 	public native void nativeSetWindowInfo(boolean bIsPortrait, int DepthBufferPreference);
diff --git a/Engine/Source/Runtime/Core/Private/Android/AndroidApplication.cpp b/Engine/Source/Runtime/Core/Private/Android/AndroidApplication.cpp
--- a/Engine/Source/Runtime/Core/Private/Android/AndroidApplication.cpp
+++ b/Engine/Source/Runtime/Core/Private/Android/AndroidApplication.cpp
@@ -4,6 +4,7 @@
 #include "AndroidApplication.h"
 #include "AndroidInputInterface.h"
 #include "AndroidWindow.h"
+#include "IInputDeviceModule.h"
 
 DEFINE_LOG_CATEGORY_STATIC(LogAndroidApplication, Log, All);
 
@@ -17,6 +18,7 @@ FAndroidApplication* FAndroidApplication::CreateAndroidApplication()
 FAndroidApplication::FAndroidApplication()
 	: GenericApplication( NULL )
 	, InputInterface( FAndroidInputInterface::Create( MessageHandler ) )
+	, bHasLoadedInputPlugins(false)
 {
 }
 
@@ -26,8 +28,29 @@ void FAndroidApplication::SetMessageHandler( const TSharedRef< FGenericApplicati
 	InputInterface->SetMessageHandler( MessageHandler );
 }
 
+void FAndroidApplication::AddExternalInputDevice(TSharedPtr<IInputDevice> InputDevice)
+{
+	if (InputDevice.IsValid())
+	{
+		InputInterface->AddExternalInputDevice(InputDevice);
+	}
+}
+
 void FAndroidApplication::PollGameDeviceState( const float TimeDelta )
 {
+	// initialize any externally-implemented input devices (we delay load initialize the array so any plugins have had time to load)
+	if (!bHasLoadedInputPlugins)
+	{
+		TArray<IInputDeviceModule*> PluginImplementations = IModularFeatures::Get().GetModularFeatureImplementations<IInputDeviceModule>(IInputDeviceModule::GetModularFeatureName());
+		for (auto InputPluginIt = PluginImplementations.CreateIterator(); InputPluginIt; ++InputPluginIt)
+		{
+			TSharedPtr<IInputDevice> Device = (*InputPluginIt)->CreateInputDevice(MessageHandler);
+			AddExternalInputDevice(Device);
+		}
+
+		bHasLoadedInputPlugins = true;
+	}
+
 	// Poll game device state and send new events
 	InputInterface->Tick( TimeDelta );
 	InputInterface->SendControllerEvents();
diff --git a/Engine/Source/Runtime/Core/Private/Android/AndroidInputInterface.cpp b/Engine/Source/Runtime/Core/Private/Android/AndroidInputInterface.cpp
--- a/Engine/Source/Runtime/Core/Private/Android/AndroidInputInterface.cpp
+++ b/Engine/Source/Runtime/Core/Private/Android/AndroidInputInterface.cpp
@@ -2,6 +2,8 @@
 
 #include "CorePrivatePCH.h"
 #include "AndroidInputInterface.h"
+//#include "AndroidInputDeviceMappings.h"
+#include "IInputDevice.h"
 #include "GenericApplicationMessageHandler.h"
 #include <android/input.h>
 
@@ -9,6 +11,11 @@
 TArray<TouchInput> FAndroidInputInterface::TouchInputStack = TArray<TouchInput>();
 FCriticalSection FAndroidInputInterface::TouchInputCriticalSection;
 
+FAndroidGamepadDeviceMapping FAndroidInputInterface::DeviceMapping[MAX_NUM_CONTROLLERS];
+
+bool FAndroidInputInterface::VibeIsOn;
+FForceFeedbackValues FAndroidInputInterface::VibeValues;
+
 FAndroidControllerData FAndroidInputInterface::OldControllerData[MAX_NUM_CONTROLLERS];
 FAndroidControllerData FAndroidInputInterface::NewControllerData[MAX_NUM_CONTROLLERS];
 
@@ -36,8 +43,6 @@ FAndroidInputInterface::~FAndroidInputInterface()
 namespace AndroidKeyNames
 {
 	const FGamepadKeyNames::Type Android_Back("Android_Back");
-	const FGamepadKeyNames::Type Android_Volume_Up("Android_Volume_Up");
-	const FGamepadKeyNames::Type Android_Volume_Down("Android_Volume_Down");
 	const FGamepadKeyNames::Type Android_Menu("Android_Menu");
 }
 
@@ -61,31 +66,76 @@ FAndroidInputInterface::FAndroidInputInterface( const TSharedRef< FGenericApplic
 	ButtonMapping[14] = FGamepadKeyNames::DPadLeft;
 	ButtonMapping[15] = FGamepadKeyNames::DPadRight;
 	ButtonMapping[16] = AndroidKeyNames::Android_Back;  // Technically just an alias for SpecialLeft
-	ButtonMapping[17] = AndroidKeyNames::Android_Volume_Up;
-	ButtonMapping[18] = AndroidKeyNames::Android_Volume_Down;
-	ButtonMapping[19] = AndroidKeyNames::Android_Menu;  // Technically just an alias for SpecialRightNewControllerData[deviceId].ButtonStates[6] = buttonDown; 
+	ButtonMapping[17] = AndroidKeyNames::Android_Menu;  // Technically just an alias for SpecialRight
 
 	InitialButtonRepeatDelay = 0.2f;
 	ButtonRepeatDelay = 0.1f;
+
+	VibeIsOn = false;
+
+	ResetGamepadAssignments();
+}
+
+void FAndroidInputInterface::ResetGamepadAssignments()
+{
+	for (int32 DeviceIndex = 0; DeviceIndex < MAX_NUM_CONTROLLERS; DeviceIndex++)
+	{
+		DeviceMapping[DeviceIndex].DeviceInfo.DeviceId = 0;
+		DeviceMapping[DeviceIndex].DeviceState = MappingState::Unassigned;
+	}
+}
+
+void FAndroidInputInterface::ResetGamepadAssignmentToController(int32 ControllerId)
+{
+	if (ControllerId < 0 || ControllerId >= MAX_NUM_CONTROLLERS)
+		return;
+
+	DeviceMapping[ControllerId].DeviceInfo.DeviceId = 0;
+	DeviceMapping[ControllerId].DeviceState = MappingState::Unassigned;
+}
+
+bool FAndroidInputInterface::IsControllerAssignedToGamepad(int32 ControllerId)
+{
+	if (ControllerId < 0 || ControllerId >= MAX_NUM_CONTROLLERS)
+		return false;
+
+	return (DeviceMapping[ControllerId].DeviceState == MappingState::Valid);
 }
 
 void FAndroidInputInterface::SetMessageHandler( const TSharedRef< FGenericApplicationMessageHandler >& InMessageHandler )
 {
 	MessageHandler = InMessageHandler;
+
+	for (auto DeviceIt = ExternalInputDevices.CreateIterator(); DeviceIt; ++DeviceIt)
+	{
+		(*DeviceIt)->SetMessageHandler(InMessageHandler);
+	}
 }
 
-void FAndroidInputInterface::Tick(float DeltaTime)
+void FAndroidInputInterface::AddExternalInputDevice(TSharedPtr<IInputDevice> InputDevice)
 {
+	if (InputDevice.IsValid())
+	{
+		ExternalInputDevices.Add(InputDevice);
+	}
+}
 
+extern bool AndroidThunkCpp_GetInputDeviceInfo(int32 deviceId, FAndroidInputDeviceInfo &results);
 
+void FAndroidInputInterface::Tick(float DeltaTime)
+{
+	for (auto DeviceIt = ExternalInputDevices.CreateIterator(); DeviceIt; ++DeviceIt)
+	{
+		(*DeviceIt)->Tick(DeltaTime);
+	}
 }
 
 void FAndroidInputInterface::SetForceFeedbackChannelValue(int32 ControllerId, FForceFeedbackChannelType ChannelType, float Value)
 {
-	// For now, force the device to 0
-	// Should use Java to enumerate number of controllers and assign device ID
-	// to controller number
-	ControllerId = 0;
+	for (auto DeviceIt = ExternalInputDevices.CreateIterator(); DeviceIt; ++DeviceIt)
+	{
+		(*DeviceIt)->SetChannelValue(ControllerId, ChannelType, Value);
+	}
 
 	// Note: only one motor on Android at the moment, but remember all the settings
 	// update will look at combination of all values to pick state
@@ -94,19 +144,19 @@ void FAndroidInputInterface::SetForceFeedbackChannelValue(int32 ControllerId, FF
 	switch (ChannelType)
 	{
 		case FForceFeedbackChannelType::LEFT_LARGE:
-			NewControllerData[ControllerId].VibeValues.LeftLarge = Value;
+			VibeValues.LeftLarge = Value;
 			break;
 
 		case FForceFeedbackChannelType::LEFT_SMALL:
-			NewControllerData[ControllerId].VibeValues.LeftSmall = Value;
+			VibeValues.LeftSmall = Value;
 			break;
 
 		case FForceFeedbackChannelType::RIGHT_LARGE:
-			NewControllerData[ControllerId].VibeValues.RightLarge = Value;
+			VibeValues.RightLarge = Value;
 			break;
 
 		case FForceFeedbackChannelType::RIGHT_SMALL:
-			NewControllerData[ControllerId].VibeValues.RightSmall = Value;
+			VibeValues.RightSmall = Value;
 			break;
 
 		default:
@@ -115,41 +165,41 @@ void FAndroidInputInterface::SetForceFeedbackChannelValue(int32 ControllerId, FF
 	}
 
 	// Update with the latest values (wait for SendControllerEvents later?)
-	UpdateVibeMotors(NewControllerData[ControllerId]);
+	UpdateVibeMotors();
 }
 
 void FAndroidInputInterface::SetForceFeedbackChannelValues(int32 ControllerId, const FForceFeedbackValues &Values)
 {
-	// For now, force the device to 0
-	// Should use Java to enumerate number of controllers and assign device ID
-	// to controller number
-	ControllerId = 0;
+	for (auto DeviceIt = ExternalInputDevices.CreateIterator(); DeviceIt; ++DeviceIt)
+	{
+		(*DeviceIt)->SetChannelValues(ControllerId, Values);
+	}
 
 	// Note: only one motor on Android at the moment, but remember all the settings
 	// update will look at combination of all values to pick state
 
-	NewControllerData[ControllerId].VibeValues = Values;
+	VibeValues = Values;
 
 	// Update with the latest values (wait for SendControllerEvents later?)
-	UpdateVibeMotors(NewControllerData[ControllerId]);
+	UpdateVibeMotors();
 }
 
 extern void AndroidThunkCpp_Vibrate(int32 Duration);
 
-void FAndroidInputInterface::UpdateVibeMotors(FAndroidControllerData &State)
+void FAndroidInputInterface::UpdateVibeMotors()
 {
 	// Use largest vibration state as value
-	float MaxLeft = State.VibeValues.LeftLarge > State.VibeValues.LeftSmall ? State.VibeValues.LeftLarge : State.VibeValues.LeftSmall;
-	float MaxRight = State.VibeValues.RightLarge > State.VibeValues.RightSmall ? State.VibeValues.RightLarge : State.VibeValues.RightSmall;
+	float MaxLeft = VibeValues.LeftLarge > VibeValues.LeftSmall ? VibeValues.LeftLarge : VibeValues.LeftSmall;
+	float MaxRight = VibeValues.RightLarge > VibeValues.RightSmall ? VibeValues.RightLarge : VibeValues.RightSmall;
 	float Value = MaxLeft > MaxRight ? MaxLeft : MaxRight;
 
-	if (State.VibeIsOn)
+	if (VibeIsOn)
 	{
 		// Turn it off if below threshold
 		if (Value < 0.3f)
 		{
 			AndroidThunkCpp_Vibrate(0);
-			State.VibeIsOn = false;
+			VibeIsOn = false;
 		}
 	}
 	else {
@@ -157,7 +207,7 @@ void FAndroidInputInterface::UpdateVibeMotors(FAndroidControllerData &State)
 		{
 			// Turn it on for 10 seconds (or until below threshold)
 			AndroidThunkCpp_Vibrate(10000);
-			State.VibeIsOn = true;
+			VibeIsOn = true;
 		}
 	}
 }
@@ -620,6 +670,94 @@ void FAndroidInputInterface::SendControllerEvents()
 {
 	FScopeLock Lock(&TouchInputCriticalSection);
 
+	// Check for gamepads needing validation
+	for (int32 DeviceIndex = 0; DeviceIndex < MAX_NUM_CONTROLLERS; DeviceIndex++)
+	{
+		FAndroidGamepadDeviceMapping& CurrentDevice = DeviceMapping[DeviceIndex];
+
+		if (CurrentDevice.DeviceState == MappingState::ToValidate)
+		{
+			// Query for the device type from Java side
+			if (AndroidThunkCpp_GetInputDeviceInfo(CurrentDevice.DeviceInfo.DeviceId, CurrentDevice.DeviceInfo))
+			{
+				// It is possible this is actually a previously assigned controller if it disconnected and reconnected (device ID can change)
+				int32 FoundMatch = -1;
+				for (int32 DeviceScan = 0; DeviceScan < MAX_NUM_CONTROLLERS; DeviceScan++)
+				{
+					if (DeviceMapping[DeviceScan].DeviceState != MappingState::Valid)
+						continue;
+
+					if (DeviceMapping[DeviceScan].DeviceInfo.Descriptor.Equals(CurrentDevice.DeviceInfo.Descriptor))
+					{
+						FoundMatch = DeviceScan;
+						break;
+					}
+				}
+
+				// Deal with new controller
+				if (FoundMatch == -1)
+				{
+					CurrentDevice.DeviceState = MappingState::Valid;
+
+					// Generic mappings
+					CurrentDevice.bSupportsHat = false;
+					CurrentDevice.bRightStickZRZ = true;
+					CurrentDevice.bRightStickRXRY = false;
+
+					// Use device name to decide on mapping scheme
+					if (CurrentDevice.DeviceInfo.Name.StartsWith(TEXT("Amazon")))
+					{
+						if (CurrentDevice.DeviceInfo.Name.StartsWith(TEXT("Amazon Fire Game Controller")))
+						{
+							CurrentDevice.bSupportsHat = true;
+						}
+						else if (CurrentDevice.DeviceInfo.Name.StartsWith(TEXT("Amazon Fire TV Remote")))
+						{
+						}
+						else
+						{
+						}
+					}
+					else if (CurrentDevice.DeviceInfo.Name.StartsWith(TEXT("NVIDIA Corporation NVIDIA Controller")))
+					{
+						CurrentDevice.bSupportsHat = true;
+					}
+					else if (CurrentDevice.DeviceInfo.Name.StartsWith(TEXT("Samsung Game Pad EI-GP20")))
+					{
+						CurrentDevice.bSupportsHat = true;
+						CurrentDevice.bRightStickZRZ = false;
+						CurrentDevice.bRightStickRXRY = true;
+					}
+
+					FPlatformMisc::LowLevelOutputDebugStringf(TEXT("Assigned new gamepad controller %d: DeviceId=%d, ControllerId=%d, DeviceName=%s, Descriptor=%s"),
+						DeviceIndex, CurrentDevice.DeviceInfo.DeviceId, CurrentDevice.DeviceInfo.ControllerId, *CurrentDevice.DeviceInfo.Name, *CurrentDevice.DeviceInfo.Descriptor);
+					continue;
+				}
+				else
+				{
+					// Already assigned this controller so reconnect it
+					DeviceMapping[FoundMatch].DeviceInfo.DeviceId = CurrentDevice.DeviceInfo.DeviceId;
+					CurrentDevice.DeviceInfo.DeviceId = 0;
+					CurrentDevice.DeviceState = MappingState::Unassigned;
+
+					// Transfer state back to this controller
+					NewControllerData[FoundMatch] = NewControllerData[DeviceIndex];
+					NewControllerData[FoundMatch].DeviceId = FoundMatch;
+					OldControllerData[FoundMatch].DeviceId = FoundMatch;
+
+					FPlatformMisc::LowLevelOutputDebugStringf(TEXT("Reconnected gamepad controller %d: DeviceId=%d, ControllerId=%d, DeviceName=%s, Descriptor=%s"),
+						FoundMatch, DeviceMapping[FoundMatch].DeviceInfo.DeviceId, CurrentDevice.DeviceInfo.ControllerId, *CurrentDevice.DeviceInfo.Name, *CurrentDevice.DeviceInfo.Descriptor);
+				}
+			}
+			else
+			{
+				// Did not find match so clear the assignment
+				FPlatformMisc::LowLevelOutputDebugStringf(TEXT("Failed to assign gamepad controller %d: DeviceId=%d"), DeviceIndex, CurrentDevice.DeviceInfo.DeviceId);
+			}
+
+		}
+	}
+
 	for(int i = 0; i < FAndroidInputInterface::TouchInputStack.Num(); ++i)
 	{
 		TouchInput Touch = FAndroidInputInterface::TouchInputStack[i];
@@ -642,6 +780,12 @@ void FAndroidInputInterface::SendControllerEvents()
 	// Extract differences in new and old states and send messages
 	for (int32 ControllerIndex = 0; ControllerIndex < MAX_NUM_CONTROLLERS; ControllerIndex++)
 	{
+		// Skip unassigned or invalid controllers (treat first one as special case)
+		if (ControllerIndex > 0 && (DeviceMapping[ControllerIndex].DeviceState !=  MappingState::Valid))
+		{
+			continue;
+		}
+
 		FAndroidControllerData& OldControllerState = OldControllerData[ControllerIndex];
 		FAndroidControllerData& NewControllerState = NewControllerData[ControllerIndex];
 
@@ -763,6 +907,11 @@ void FAndroidInputInterface::SendControllerEvents()
 	FAndroidInputInterface::TouchInputStack.Empty(0);
 
 	FAndroidInputInterface::MotionDataStack.Empty();
+
+	for (auto DeviceIt = ExternalInputDevices.CreateIterator(); DeviceIt; ++DeviceIt)
+	{
+		(*DeviceIt)->SendControllerEvents();
+	}
 }
 
 void FAndroidInputInterface::QueueTouchInput(TArray<TouchInput> InTouchEvents)
@@ -772,26 +921,127 @@ void FAndroidInputInterface::QueueTouchInput(TArray<TouchInput> InTouchEvents)
 	FAndroidInputInterface::TouchInputStack.Append(InTouchEvents);
 }
 
+int32 FAndroidInputInterface::GetControllerIndex(int32 deviceId)
+{
+	// Treat non-positive device ids special (always controller 0)
+	if (deviceId < 1)
+		return 0;
+
+	// Look for this deviceId in controllers discovered
+	int32 UnassignedIndex = -1;
+	for (int32 ControllerIndex = 0; ControllerIndex < MAX_NUM_CONTROLLERS; ControllerIndex++)
+	{
+		if (DeviceMapping[ControllerIndex].DeviceState == MappingState::Unassigned)
+		{
+			if (UnassignedIndex == -1)
+				UnassignedIndex = ControllerIndex;
+			continue;
+		}
+
+		if (DeviceMapping[ControllerIndex].DeviceInfo.DeviceId == deviceId)
+		{
+			return ControllerIndex;
+		}
+	}
+
+	// Haven't seen this one before, make sure there is room for a new one
+	if (UnassignedIndex == -1)
+		return -1;
+
+	// Register it
+	DeviceMapping[UnassignedIndex].DeviceInfo.DeviceId = deviceId;
+	OldControllerData[UnassignedIndex].DeviceId = UnassignedIndex;
+	NewControllerData[UnassignedIndex].DeviceId = UnassignedIndex;
+
+	// Mark it for validation later
+	DeviceMapping[UnassignedIndex].DeviceState = MappingState::ToValidate;
+
+	return UnassignedIndex;
+}
+
 void FAndroidInputInterface::JoystickAxisEvent(int32 deviceId, int32 axisId, float axisValue)
 {
 	FScopeLock Lock(&TouchInputCriticalSection);
 
-	// For now, force the device to 0
-	// Should use Java to enumerate number of controllers and assign device ID
-	// to controller number
-	deviceId = 0;
+	// Get the controller index matching deviceId (if there is one)
+	deviceId = GetControllerIndex(deviceId);
+	if (deviceId == -1)
+		return;
 
-	// Apply a small dead zone to the analog sticks
-	const float deadZone = 0.2f;
+	// Deal with left stick and triggers (generic)
 	switch (axisId)
 	{
-		// Also, invert Y and RZ to match what the engine expects
-		case AMOTION_EVENT_AXIS_X:        NewControllerData[deviceId].LXAnalog =  axisValue; break;
-		case AMOTION_EVENT_AXIS_Y:        NewControllerData[deviceId].LYAnalog = -axisValue; break;
-		case AMOTION_EVENT_AXIS_Z:        NewControllerData[deviceId].RXAnalog =  axisValue; break;
-		case AMOTION_EVENT_AXIS_RZ:       NewControllerData[deviceId].RYAnalog = -axisValue; break;
-		case AMOTION_EVENT_AXIS_LTRIGGER: NewControllerData[deviceId].LTAnalog =  axisValue; break;
-		case AMOTION_EVENT_AXIS_RTRIGGER: NewControllerData[deviceId].RTAnalog =  axisValue; break;
+		case AMOTION_EVENT_AXIS_X:			NewControllerData[deviceId].LXAnalog =  axisValue; return;
+		case AMOTION_EVENT_AXIS_Y:			NewControllerData[deviceId].LYAnalog = -axisValue; return;
+		case AMOTION_EVENT_AXIS_LTRIGGER:	NewControllerData[deviceId].LTAnalog =  axisValue; return;
+		case AMOTION_EVENT_AXIS_RTRIGGER:	NewControllerData[deviceId].RTAnalog =  axisValue; return;
+	}
+
+	// Deal with right stick Z/RZ events
+	if (DeviceMapping[deviceId].bRightStickZRZ)
+	{
+		switch (axisId)
+		{
+			case AMOTION_EVENT_AXIS_Z:		NewControllerData[deviceId].RXAnalog =  axisValue; return;
+			case AMOTION_EVENT_AXIS_RZ:		NewControllerData[deviceId].RYAnalog = -axisValue; return;
+		}
+	}
+
+	// Deal with right stick RX/RY events
+	if (DeviceMapping[deviceId].bRightStickRXRY)
+	{
+		switch (axisId)
+		{
+			case AMOTION_EVENT_AXIS_RX:		NewControllerData[deviceId].RXAnalog =  axisValue; return;
+			case AMOTION_EVENT_AXIS_RY:		NewControllerData[deviceId].RYAnalog = -axisValue; return;
+		}
+	}
+
+	// Deal with hat (convert to DPAD buttons)
+	if (DeviceMapping[deviceId].bSupportsHat)
+	{
+		// Apply a small dead zone to hats
+		const float deadZone = 0.2f;
+
+		switch (axisId)
+		{
+			case AMOTION_EVENT_AXIS_HAT_X:
+				// AMOTION_EVENT_AXIS_HAT_X translates to KEYCODE_DPAD_LEFT and AKEYCODE_DPAD_RIGHT
+				if (axisValue > deadZone)
+				{
+					NewControllerData[deviceId].ButtonStates[14] = false;	// DPAD_LEFT released
+					NewControllerData[deviceId].ButtonStates[15] = true;	// DPAD_RIGHT pressed
+				}
+				else if (axisValue < -deadZone)
+				{
+					NewControllerData[deviceId].ButtonStates[14] = true;	// DPAD_LEFT pressed
+					NewControllerData[deviceId].ButtonStates[15] = false;	// DPAD_RIGHT released
+				}
+				else
+				{
+					NewControllerData[deviceId].ButtonStates[14] = false;	// DPAD_LEFT released
+					NewControllerData[deviceId].ButtonStates[15] = false;	// DPAD_RIGHT released
+				}
+				return;
+			case AMOTION_EVENT_AXIS_HAT_Y:
+				// AMOTION_EVENT_AXIS_HAT_Y translates to KEYCODE_DPAD_UP and AKEYCODE_DPAD_DOWN
+				if (axisValue > deadZone)
+				{
+					NewControllerData[deviceId].ButtonStates[12] = false;	// DPAD_UP released
+					NewControllerData[deviceId].ButtonStates[13] = true;	// DPAD_DOWN pressed
+				}
+				else if (axisValue < -deadZone)
+				{
+					NewControllerData[deviceId].ButtonStates[12] = true;	// DPAD_UP pressed
+					NewControllerData[deviceId].ButtonStates[13] = false;	// DPAD_DOWN released
+				}
+				else
+				{
+					NewControllerData[deviceId].ButtonStates[12] = false;	// DPAD_UP released
+					NewControllerData[deviceId].ButtonStates[13] = false;	// DPAD_DOWN released
+				}
+				return;
+		}
 	}
 }
 
@@ -799,10 +1049,10 @@ void FAndroidInputInterface::JoystickButtonEvent(int32 deviceId, int32 buttonId,
 {
 	FScopeLock Lock(&TouchInputCriticalSection);
 
-	// For now, force the device to 0
-	// Should use Java to enumerate number of controllers and assign device ID
-	// to controller number
-	deviceId = 0;
+	// Get the controller index matching deviceId (if there is one)
+	deviceId = GetControllerIndex(deviceId);
+	if (deviceId == -1)
+		return;
 
 	switch (buttonId)
 	{
@@ -814,7 +1064,7 @@ void FAndroidInputInterface::JoystickButtonEvent(int32 deviceId, int32 buttonId,
 		case AKEYCODE_BUTTON_L1:     NewControllerData[deviceId].ButtonStates[4] = buttonDown; break;
 		case AKEYCODE_BUTTON_R1:     NewControllerData[deviceId].ButtonStates[5] = buttonDown; break;
 		case AKEYCODE_BUTTON_START:
-		case AKEYCODE_MENU:          NewControllerData[deviceId].ButtonStates[6] = buttonDown; NewControllerData[deviceId].ButtonStates[19] = buttonDown;  break;
+		case AKEYCODE_MENU:          NewControllerData[deviceId].ButtonStates[6] = buttonDown; NewControllerData[deviceId].ButtonStates[17] = buttonDown;  break;
 		case AKEYCODE_BUTTON_SELECT: 
 		case AKEYCODE_BACK:          NewControllerData[deviceId].ButtonStates[7] = buttonDown; NewControllerData[deviceId].ButtonStates[16] = buttonDown;  break;
 		case AKEYCODE_BUTTON_THUMBL: NewControllerData[deviceId].ButtonStates[8] = buttonDown; break;
diff --git a/Engine/Source/Runtime/Core/Private/Android/AndroidMisc.cpp b/Engine/Source/Runtime/Core/Private/Android/AndroidMisc.cpp
--- a/Engine/Source/Runtime/Core/Private/Android/AndroidMisc.cpp
+++ b/Engine/Source/Runtime/Core/Private/Android/AndroidMisc.cpp
@@ -774,6 +774,20 @@ void FAndroidMisc::SetVolumeButtonsHandledBySystem(bool enabled)
 	VolumeButtonsHandledBySystem = enabled;
 }
 
+void FAndroidMisc::ResetGamepadAssignments()
+{
+	FAndroidInputInterface::ResetGamepadAssignments();
+}
+
+void FAndroidMisc::ResetGamepadAssignmentToController(int32 ControllerId)
+{
+	FAndroidInputInterface::ResetGamepadAssignmentToController(ControllerId);
+}
+
+bool FAndroidMisc::IsControllerAssignedToGamepad(int32 ControllerId)
+{
+	return FAndroidInputInterface::IsControllerAssignedToGamepad(ControllerId);
+}
 
 int32 FAndroidMisc::GetAndroidBuildVersion()
 {
diff --git a/Engine/Source/Runtime/Core/Public/Android/AndroidApplication.h b/Engine/Source/Runtime/Core/Public/Android/AndroidApplication.h
--- a/Engine/Source/Runtime/Core/Public/Android/AndroidApplication.h
+++ b/Engine/Source/Runtime/Core/Public/Android/AndroidApplication.h
@@ -61,7 +61,9 @@ public:
 	virtual IInputInterface* GetInputInterface() override;
 
 	virtual TSharedRef< FGenericWindow > MakeWindow() override;
-	
+
+	virtual void AddExternalInputDevice(TSharedPtr<class IInputDevice> InputDevice);
+
 	void InitializeWindow( const TSharedRef< FGenericWindow >& InWindow, const TSharedRef< FGenericWindowDefinition >& InDefinition, const TSharedPtr< FGenericWindow >& InParent, const bool bShowImmediately );
 
 	static void OnWindowSizeChanged();
@@ -74,6 +76,7 @@ private:
 private:
 
 	TSharedPtr< class FAndroidInputInterface > InputInterface;
+	bool bHasLoadedInputPlugins;
 
 	TArray< TSharedRef< FAndroidWindow > > Windows;
 
diff --git a/Engine/Source/Runtime/Core/Public/Android/AndroidInputInterface.h b/Engine/Source/Runtime/Core/Public/Android/AndroidInputInterface.h
--- a/Engine/Source/Runtime/Core/Public/Android/AndroidInputInterface.h
+++ b/Engine/Source/Runtime/Core/Public/Android/AndroidInputInterface.h
@@ -77,6 +77,40 @@ enum TouchType
 	TouchEnded,
 };
 
+enum MappingState
+{
+	Unassigned,
+	ToValidate,
+	Valid
+};
+
+struct FAndroidInputDeviceInfo {
+	int32 DeviceId;
+	int32 VendorId;
+	int32 ProductId;
+	int32 ControllerId;
+	FString Name;
+	FString Descriptor;
+};
+
+struct FAndroidGamepadDeviceMapping
+{
+	// Information for this mapped device
+	FAndroidInputDeviceInfo DeviceInfo;
+
+	// State of mapping
+	MappingState DeviceState;
+
+	// Device supports hat as dpad
+	bool bSupportsHat;
+
+	// Right stick on Z/RZ
+	bool bRightStickZRZ;
+
+	// Right stick on RX/RY
+	bool bRightStickRXRY;
+};
+
 struct TouchInput
 {
 	int Handle;
@@ -85,8 +119,8 @@ struct TouchInput
 	FVector2D Position;
 };
 
-#define MAX_NUM_CONTROLLERS					1  // Only one for now
-#define MAX_NUM_CONTROLLER_BUTTONS			20
+#define MAX_NUM_CONTROLLERS					8  // reasonable limit for now
+#define MAX_NUM_CONTROLLER_BUTTONS			18
 #define MAX_DEFERRED_MESSAGE_QUEUE_SIZE		128
 
 struct FAndroidControllerData
@@ -105,10 +139,6 @@ struct FAndroidControllerData
 	float RYAnalog;
 	float LTAnalog;
 	float RTAnalog;
-
-	/** Vibration settings */
-	bool VibeIsOn;
-	FForceFeedbackValues VibeValues;
 };
 
 enum FAndroidMessageType
@@ -162,6 +192,10 @@ public:
 
 	static void QueueTouchInput(TArray<TouchInput> InTouchEvents);
 
+	static void ResetGamepadAssignments();
+	static void ResetGamepadAssignmentToController(int32 ControllerId);
+	static bool IsControllerAssignedToGamepad(int32 ControllerId);
+
 	static void JoystickAxisEvent(int32 deviceId, int32 axisId, float axisValue);
 	static void JoystickButtonEvent(int32 deviceId, int32 buttonId, bool buttonDown);
 
@@ -176,6 +210,8 @@ public:
 	virtual void SetForceFeedbackChannelValues(int32 ControllerId, const FForceFeedbackValues &values) override;
 	virtual void SetLightColor(int32 ControllerId, FColor Color) override {}
 
+	virtual void AddExternalInputDevice(TSharedPtr<class IInputDevice> InputDevice);
+
 private:
 
 	FAndroidInputInterface( const TSharedRef< FGenericApplicationMessageHandler >& InMessageHandler );
@@ -183,8 +219,11 @@ private:
 
 private:
 
+	/** Get controller index corresponding to deviceId (assigns and queries type if new) */
+	static int32 GetControllerIndex(int32 deviceId);
+
 	/** Push Vibration changes to the controllers */
-	void UpdateVibeMotors(FAndroidControllerData &State);
+	void UpdateVibeMotors();
 
 	struct MotionData
 	{
@@ -199,6 +238,12 @@ private:
 
 	static TArray<TouchInput> TouchInputStack;
 
+	/** Vibration settings */
+	static bool VibeIsOn;
+	static FForceFeedbackValues VibeValues;
+
+	static FAndroidGamepadDeviceMapping DeviceMapping[MAX_NUM_CONTROLLERS];
+
 	static FAndroidControllerData OldControllerData[MAX_NUM_CONTROLLERS];
 	static FAndroidControllerData NewControllerData[MAX_NUM_CONTROLLERS];
 
@@ -214,4 +259,7 @@ private:
 	static TArray<MotionData> MotionDataStack;
 
 	TSharedRef< FGenericApplicationMessageHandler > MessageHandler;
+
+	/** List of input devices implemented in external modules. */
+	TArray<TSharedPtr<class IInputDevice>> ExternalInputDevices;
 };
diff --git a/Engine/Source/Runtime/Core/Public/Android/AndroidMisc.h b/Engine/Source/Runtime/Core/Public/Android/AndroidMisc.h
--- a/Engine/Source/Runtime/Core/Public/Android/AndroidMisc.h
+++ b/Engine/Source/Runtime/Core/Public/Android/AndroidMisc.h
@@ -46,6 +46,9 @@ struct CORE_API FAndroidMisc : public FGenericPlatformMisc
 	static const TCHAR* GetDefaultDeviceProfileName() { return TEXT("Android_Default"); }
 	static bool GetVolumeButtonsHandledBySystem();
 	static void SetVolumeButtonsHandledBySystem(bool enabled);
+	static void ResetGamepadAssignments();
+	static void ResetGamepadAssignmentToController(int32 ControllerId);
+	static bool IsControllerAssignedToGamepad(int32 ControllerId);
 
 	/** @return Memory representing a true type or open type font provided by the platform as a default font for unreal to consume; empty array if the default font failed to load. */
 	static TArray<uint8> GetSystemFontBytes();
diff --git a/Engine/Source/Runtime/Core/Public/GenericPlatform/GenericPlatformMisc.h b/Engine/Source/Runtime/Core/Public/GenericPlatform/GenericPlatformMisc.h
--- a/Engine/Source/Runtime/Core/Public/GenericPlatform/GenericPlatformMisc.h
+++ b/Engine/Source/Runtime/Core/Public/GenericPlatform/GenericPlatformMisc.h
@@ -760,6 +760,26 @@ public:
 		return true;
 	}
 
+	/*
+	 * Resets the gamepad to player controller id assignments
+	 */
+	static void ResetGamepadAssignments()
+	{}
+
+	/*
+	* Resets the gamepad assignment to player controller id
+	*/
+	static void ResetGamepadAssignmentToController(int32 ControllerId)
+	{}
+
+	/*
+	 * Returns true if controller id assigned to a gamepad
+	 */
+	static bool IsControllerAssignedToGamepad(int32 ControllerId)
+	{
+		return true;
+	}
+
 	/*
 	 * Set whether the volume buttons are handled by the system
 	 */
diff --git a/Engine/Source/Runtime/Engine/Classes/Kismet/KismetSystemLibrary.h b/Engine/Source/Runtime/Engine/Classes/Kismet/KismetSystemLibrary.h
--- a/Engine/Source/Runtime/Engine/Classes/Kismet/KismetSystemLibrary.h
+++ b/Engine/Source/Runtime/Engine/Classes/Kismet/KismetSystemLibrary.h
@@ -1437,6 +1437,24 @@ class ENGINE_API UKismetSystemLibrary : public UBlueprintFunctionLibrary
 	UFUNCTION(BlueprintPure, Category = "Utilities|Platform")
 	static bool GetVolumeButtonsHandledBySystem();
 
+	/**
+	 * Resets the gamepad to player controller id assignments (Android only)
+	 */
+	UFUNCTION(BlueprintCallable, Category = "Utilities|Platform")
+	static void ResetGamepadAssignments();
+
+	/*
+	 * Resets the gamepad assignment to player controller id (Android only)
+	 */
+	UFUNCTION(BlueprintCallable, Category = "Utilities|Platform")
+	static void ResetGamepadAssignmentToController(int32 ControllerId);
+
+	/**
+	 * Returns true if controller id assigned to a gamepad (Android only)
+	 */
+	UFUNCTION(BlueprintPure, Category = "Utilities|Platform")
+	static bool IsControllerAssignedToGamepad(int32 ControllerId);
+
 	/**
 	 * Sets the state of the transition message rendered by the viewport. (The blue text displayed when the game is paused and so forth.)
 	 *
diff --git a/Engine/Source/Runtime/Engine/Private/KismetSystemLibrary.cpp b/Engine/Source/Runtime/Engine/Private/KismetSystemLibrary.cpp
--- a/Engine/Source/Runtime/Engine/Private/KismetSystemLibrary.cpp
+++ b/Engine/Source/Runtime/Engine/Private/KismetSystemLibrary.cpp
@@ -3148,6 +3148,21 @@ bool UKismetSystemLibrary::GetVolumeButtonsHandledBySystem()
 	return FPlatformMisc::GetVolumeButtonsHandledBySystem();
 }
 
+void UKismetSystemLibrary::ResetGamepadAssignments()
+{
+	FPlatformMisc::ResetGamepadAssignments();
+}
+
+void UKismetSystemLibrary::ResetGamepadAssignmentToController(int32 ControllerId)
+{
+	FPlatformMisc::ResetGamepadAssignmentToController(ControllerId);
+}
+
+bool UKismetSystemLibrary::IsControllerAssignedToGamepad(int32 ControllerId)
+{
+	return FPlatformMisc::IsControllerAssignedToGamepad(ControllerId);
+}
+
 void UKismetSystemLibrary::SetSupressViewportTransitionMessage(UObject* WorldContextObject, bool bState)
 {
 	UWorld* World = GEngine->GetWorldFromContextObject(WorldContextObject);
diff --git a/Engine/Source/Runtime/Launch/Private/Android/AndroidJNI.cpp b/Engine/Source/Runtime/Launch/Private/Android/AndroidJNI.cpp
--- a/Engine/Source/Runtime/Launch/Private/Android/AndroidJNI.cpp
+++ b/Engine/Source/Runtime/Launch/Private/Android/AndroidJNI.cpp
@@ -6,6 +6,7 @@
 
 #include "AndroidJNI.h"
 #include "AndroidApplication.h"
+#include "AndroidInputInterface.h"
 #include <Android/asset_manager.h>
 #include <Android/asset_manager_jni.h>
 
@@ -40,7 +41,9 @@ if (Id == 0) \
 void FJavaWrapper::FindClassesAndMethods(JNIEnv* Env)
 {
 	bool bIsOptional = false;
-	GameActivityClassID = FindClass(Env, "com/epicgames/ue4/GameActivity", bIsOptional);
+	jclass localGameActivityClass = FindClass(Env, "com/epicgames/ue4/GameActivity", bIsOptional);
+	GameActivityClassID = (jclass)Env->NewGlobalRef(localGameActivityClass);
+	Env->DeleteLocalRef(localGameActivityClass);
 	AndroidThunkJava_ShowConsoleWindow = FindMethod(Env, GameActivityClassID, "AndroidThunkJava_ShowConsoleWindow", "(Ljava/lang/String;)V", bIsOptional);
 	AndroidThunkJava_ShowVirtualKeyboardInput = FindMethod(Env, GameActivityClassID, "AndroidThunkJava_ShowVirtualKeyboardInput", "(ILjava/lang/String;Ljava/lang/String;)V", bIsOptional);
 	AndroidThunkJava_LaunchURL = FindMethod(Env, GameActivityClassID, "AndroidThunkJava_LaunchURL", "(Ljava/lang/String;)V", bIsOptional);
@@ -54,6 +57,17 @@ void FJavaWrapper::FindClassesAndMethods(JNIEnv* Env)
 	AndroidThunkJava_InitHMDs = FindMethod(Env, GameActivityClassID, "AndroidThunkJava_InitHMDs", "()V", bIsOptional);
 	AndroidThunkJava_IsGearVRApplication = FindMethod(Env, GameActivityClassID, "AndroidThunkJava_IsGearVRApplication", "()Z", bIsOptional);
 	AndroidThunkJava_DismissSplashScreen = FindMethod(Env, GameActivityClassID, "AndroidThunkJava_DismissSplashScreen", "()V", bIsOptional);
+	AndroidThunkJava_GetInputDeviceInfo = FindMethod(Env, GameActivityClassID, "AndroidThunkJava_GetInputDeviceInfo", "(I)Lcom/epicgames/ue4/GameActivity$InputDeviceInfo;", bIsOptional);
+
+	// get field IDs for InputDeviceInfo class members
+	jclass localInputDeviceInfoClass = FindClass(Env, "com/epicgames/ue4/GameActivity$InputDeviceInfo", bIsOptional);
+	InputDeviceInfoClass = (jclass)Env->NewGlobalRef(localInputDeviceInfoClass);
+	Env->DeleteLocalRef(localInputDeviceInfoClass);
+	InputDeviceInfo_VendorId = FJavaWrapper::FindField(Env, InputDeviceInfoClass, "vendorId", "I", bIsOptional);
+	InputDeviceInfo_ProductId = FJavaWrapper::FindField(Env, InputDeviceInfoClass, "productId", "I", bIsOptional);
+	InputDeviceInfo_ControllerId = FJavaWrapper::FindField(Env, InputDeviceInfoClass, "controllerId", "I", bIsOptional);
+	InputDeviceInfo_Name = FJavaWrapper::FindField(Env, InputDeviceInfoClass, "name", "Ljava/lang/String;", bIsOptional);
+	InputDeviceInfo_Descriptor = FJavaWrapper::FindField(Env, InputDeviceInfoClass, "descriptor", "Ljava/lang/String;", bIsOptional);
 
 	// the rest are optional
 	bIsOptional = true;
@@ -66,6 +80,9 @@ void FJavaWrapper::FindClassesAndMethods(JNIEnv* Env)
 	AndroidThunkJava_CloseAdBanner = FindMethod(Env, GoogleServicesClassID, "AndroidThunkJava_CloseAdBanner", "()V", bIsOptional);
 
 	// In app purchase functionality
+	jclass localStringClass = Env->FindClass("java/lang/String");
+	JavaStringClass = (jclass)Env->NewGlobalRef(localStringClass);
+	Env->DeleteLocalRef(localStringClass);
 	AndroidThunkJava_IapSetupService = FindMethod(Env, GoogleServicesClassID, "AndroidThunkJava_IapSetupService", "(Ljava/lang/String;)V", bIsOptional);
 	AndroidThunkJava_IapQueryInAppPurchases = FindMethod(Env, GoogleServicesClassID, "AndroidThunkJava_IapQueryInAppPurchases", "([Ljava/lang/String;[Z)Z", bIsOptional);
 	AndroidThunkJava_IapBeginPurchase = FindMethod(Env, GoogleServicesClassID, "AndroidThunkJava_IapBeginPurchase", "(Ljava/lang/String;Z)Z", bIsOptional);
@@ -129,6 +146,21 @@ jobject FJavaWrapper::CallObjectMethod(JNIEnv* Env, jobject Object, jmethodID Me
 	return Return;
 }
 
+int32 FJavaWrapper::CallIntMethod(JNIEnv* Env, jobject Object, jmethodID Method, ...)
+{
+	if (Method == NULL || Object == NULL)
+	{
+		return false;
+	}
+
+	va_list Args;
+	va_start(Args, Method);
+	jint Return = Env->CallIntMethodV(Object, Method, Args);
+	va_end(Args);
+
+	return (int32)Return;
+}
+
 bool FJavaWrapper::CallBooleanMethod(JNIEnv* Env, jobject Object, jmethodID Method, ...)
 {
 	if (Method == NULL || Object == NULL)
@@ -160,6 +192,14 @@ jmethodID FJavaWrapper::AndroidThunkJava_KeepScreenOn;
 jmethodID FJavaWrapper::AndroidThunkJava_InitHMDs;
 jmethodID FJavaWrapper::AndroidThunkJava_IsGearVRApplication;
 jmethodID FJavaWrapper::AndroidThunkJava_DismissSplashScreen;
+jmethodID FJavaWrapper::AndroidThunkJava_GetInputDeviceInfo;
+
+jclass FJavaWrapper::InputDeviceInfoClass;
+jfieldID FJavaWrapper::InputDeviceInfo_VendorId;
+jfieldID FJavaWrapper::InputDeviceInfo_ProductId;
+jfieldID FJavaWrapper::InputDeviceInfo_ControllerId;
+jfieldID FJavaWrapper::InputDeviceInfo_Name;
+jfieldID FJavaWrapper::InputDeviceInfo_Descriptor;
 
 jclass FJavaWrapper::GoogleServicesClassID;
 jobject FJavaWrapper::GoogleServicesThis;
@@ -167,6 +207,8 @@ jmethodID FJavaWrapper::AndroidThunkJava_ResetAchievements;
 jmethodID FJavaWrapper::AndroidThunkJava_ShowAdBanner;
 jmethodID FJavaWrapper::AndroidThunkJava_HideAdBanner;
 jmethodID FJavaWrapper::AndroidThunkJava_CloseAdBanner;
+
+jclass FJavaWrapper::JavaStringClass;
 jmethodID FJavaWrapper::AndroidThunkJava_IapSetupService;
 jmethodID FJavaWrapper::AndroidThunkJava_IapQueryInAppPurchases;
 jmethodID FJavaWrapper::AndroidThunkJava_IapBeginPurchase;
@@ -248,6 +290,47 @@ void AndroidThunkCpp_DismissSplashScreen()
 	}
 }
 
+bool AndroidThunkCpp_GetInputDeviceInfo(int32 deviceId, FAndroidInputDeviceInfo &results)
+{
+	if (JNIEnv* Env = FAndroidApplication::GetJavaEnv())
+	{
+		jobject deviceInfo = (jobject)Env->CallObjectMethod(FJavaWrapper::GameActivityThis, FJavaWrapper::AndroidThunkJava_GetInputDeviceInfo, deviceId);
+		bool bIsOptional = false;
+		results.DeviceId = deviceId;
+		results.VendorId = (int32)Env->GetIntField(deviceInfo, FJavaWrapper::InputDeviceInfo_VendorId);
+		results.ProductId = (int32)Env->GetIntField(deviceInfo, FJavaWrapper::InputDeviceInfo_ProductId);
+		results.ControllerId = (int32)Env->GetIntField(deviceInfo, FJavaWrapper::InputDeviceInfo_ControllerId);
+
+		jstring jsName = (jstring)Env->GetObjectField(deviceInfo, FJavaWrapper::InputDeviceInfo_Name);
+		CHECK_JNI_RESULT(jsName);
+		const char * nativeName = Env->GetStringUTFChars(jsName, 0);
+		results.Name = FString(nativeName);
+		Env->ReleaseStringUTFChars(jsName, nativeName);
+		Env->DeleteLocalRef(jsName);
+
+		jstring jsDescriptor = (jstring)Env->GetObjectField(deviceInfo, FJavaWrapper::InputDeviceInfo_Descriptor);
+		CHECK_JNI_RESULT(jsDescriptor);
+		const char * nativeDescriptor = Env->GetStringUTFChars(jsDescriptor, 0);
+		results.Descriptor = FString(nativeDescriptor);
+		Env->ReleaseStringUTFChars(jsDescriptor, nativeDescriptor);
+		Env->DeleteLocalRef(jsDescriptor);
+		
+		// release references
+		Env->DeleteLocalRef(deviceInfo);
+
+		return true;
+	}
+
+	// failed
+	results.DeviceId = deviceId;
+	results.VendorId = 0;
+	results.ProductId = 0;
+	results.ControllerId = -1;
+	results.Name = FString("Unknown");
+	results.Descriptor = FString("Unknown");
+	return false;
+}
+
 void AndroidThunkCpp_ShowConsoleWindow()
 {
 	if (JNIEnv* Env = FAndroidApplication::GetJavaEnv())
@@ -380,8 +463,9 @@ jobject AndroidJNI_GetJavaAssetManager()
 	{
 		if (JNIEnv* Env = FAndroidApplication::GetJavaEnv())
 		{
-			GJavaAssetManager = FJavaWrapper::CallObjectMethod(Env, FJavaWrapper::GameActivityThis, FJavaWrapper::AndroidThunkJava_GetAssetManager);
-			Env->NewGlobalRef(GJavaAssetManager);
+			jobject local = FJavaWrapper::CallObjectMethod(Env, FJavaWrapper::GameActivityThis, FJavaWrapper::AndroidThunkJava_GetAssetManager);
+			GJavaAssetManager = (jobject)Env->NewGlobalRef(local);
+			Env->DeleteLocalRef(local);
 		}
 	}
 	return GJavaAssetManager;
@@ -445,10 +529,8 @@ bool AndroidThunkCpp_Iap_QueryInAppPurchases(const TArray<FString>& ProductIDs,
 
 	if (JNIEnv* Env = FAndroidApplication::GetJavaEnv())
 	{
-		static jclass StringClass = Env->FindClass("java/lang/String");
-
 		// Populate some java types with the provided product information
-		jobjectArray ProductIDArray = (jobjectArray)Env->NewObjectArray(ProductIDs.Num(), StringClass, NULL);
+		jobjectArray ProductIDArray = (jobjectArray)Env->NewObjectArray(ProductIDs.Num(), FJavaWrapper::JavaStringClass, NULL);
 		jbooleanArray ConsumeArray = (jbooleanArray)Env->NewBooleanArray(ProductIDs.Num());
 
 		jboolean* ConsumeArrayValues = Env->GetBooleanArrayElements(ConsumeArray, 0);
@@ -542,6 +624,9 @@ JNIEXPORT jint JNI_OnLoad(JavaVM* InJavaVM, void* InReserved)
 
 	// then release...
 	Env->ReleaseStringUTFChars(pathString, nativePathString);
+	Env->DeleteLocalRef(pathString);
+	Env->DeleteLocalRef(externalStoragePath);
+	Env->DeleteLocalRef(EnvClass);
 	FPlatformMisc::LowLevelOutputDebugStringf(TEXT("Path found as '%s'\n"), *GFilePathBase);
 
 	// Get the system font directory
@@ -549,6 +634,7 @@ JNIEXPORT jint JNI_OnLoad(JavaVM* InJavaVM, void* InReserved)
 	const char * nativeFontPathString = Env->GetStringUTFChars(fontPath, 0);
 	GFontPathBase = FString(nativeFontPathString);
 	Env->ReleaseStringUTFChars(fontPath, nativeFontPathString);
+	Env->DeleteLocalRef(fontPath);
 	FPlatformMisc::LowLevelOutputDebugStringf(TEXT("Font Path found as '%s'\n"), *GFontPathBase);
 
 	// Wire up to core delegates, so core code can call out to Java
@@ -570,20 +656,19 @@ extern "C" void Java_com_epicgames_ue4_GameActivity_nativeSetGlobalActivity(JNIE
 	if (!FJavaWrapper::GameActivityThis)
 	{
 		FJavaWrapper::GameActivityThis = jenv->NewGlobalRef(thiz);
-		FAndroidApplication::InitializeJavaEnv(GJavaVM, JNI_CURRENT_VERSION, FJavaWrapper::GameActivityThis);
 		if (!FJavaWrapper::GameActivityThis)
 		{
 			FPlatformMisc::LowLevelOutputDebugString(L"Error setting the global GameActivity activity");
 			check(false);
 		}
 
+		// This call is only to set the correct GameActivityThis
+		FAndroidApplication::InitializeJavaEnv(GJavaVM, JNI_CURRENT_VERSION, FJavaWrapper::GameActivityThis);
+
 		// @todo split GooglePlay, this needs to be passed in to this function
 		FJavaWrapper::GoogleServicesThis = FJavaWrapper::GameActivityThis;
 		// FJavaWrapper::GoogleServicesThis = jenv->NewGlobalRef(googleServices);
 
-		bool bIsOptional = false;
-		FJavaWrapper::GameActivityClassID = FJavaWrapper::FindClass(jenv, "com/epicgames/ue4/GameActivity", bIsOptional);
-
 		// Next we check to see if the OBB file is in the APK
 		jmethodID isOBBInAPKMethod = jenv->GetStaticMethodID(FJavaWrapper::GameActivityClassID, "isOBBInAPK", "()Z");
 		GOBBinAPK = (bool)jenv->CallStaticBooleanMethod(FJavaWrapper::GameActivityClassID, isOBBInAPKMethod, nullptr);
@@ -600,6 +685,9 @@ extern "C" void Java_com_epicgames_ue4_GameActivity_nativeSetGlobalActivity(JNIE
 
 		// then release...
 		jenv->ReleaseStringUTFChars(externalFilesPathString, nativeExternalFilesPathString);
+		jenv->DeleteLocalRef(externalFilesPathString);
+		jenv->DeleteLocalRef(externalFilesDirPath);
+		jenv->DeleteLocalRef(ContextClass);
 		FPlatformMisc::LowLevelOutputDebugStringf(TEXT("ExternalFilePath found as '%s'\n"), *GExternalFilePath);
 	}
 }
diff --git a/Engine/Source/Runtime/Launch/Private/Android/AndroidJNI.h b/Engine/Source/Runtime/Launch/Private/Android/AndroidJNI.h
--- a/Engine/Source/Runtime/Launch/Private/Android/AndroidJNI.h
+++ b/Engine/Source/Runtime/Launch/Private/Android/AndroidJNI.h
@@ -30,6 +30,19 @@ public:
 	static jmethodID AndroidThunkJava_InitHMDs;
 	static jmethodID AndroidThunkJava_IsGearVRApplication;
 	static jmethodID AndroidThunkJava_DismissSplashScreen;
+	static jmethodID AndroidThunkJava_GetCurrentVolume;
+	static jmethodID AndroidThunkJava_GetMaximumVolume;
+	static jmethodID AndroidThunkJava_SetVolume;
+	static jmethodID AndroidThunkJava_AdjustVolume;
+	static jmethodID AndroidThunkJava_GetInputDeviceInfo;
+	
+	// InputDeviceInfo member field ids
+	static jclass InputDeviceInfoClass;
+	static jfieldID InputDeviceInfo_VendorId;
+	static jfieldID InputDeviceInfo_ProductId;
+	static jfieldID InputDeviceInfo_ControllerId;
+	static jfieldID InputDeviceInfo_Name;
+	static jfieldID InputDeviceInfo_Descriptor;
 
 	// IDs related to google play services
 	static jclass GoogleServicesClassID;
@@ -40,6 +53,7 @@ public:
 	static jmethodID AndroidThunkJava_CloseAdBanner;
 
 	// In app purchase functionality
+	static jclass JavaStringClass;
 	static jmethodID AndroidThunkJava_IapSetupService;
 	static jmethodID AndroidThunkJava_IapQueryInAppPurchases;
 	static jmethodID AndroidThunkJava_IapBeginPurchase;
@@ -60,6 +74,7 @@ public:
 
 	static void CallVoidMethod(JNIEnv* Env, jobject Object, jmethodID Method, ...);
 	static jobject CallObjectMethod(JNIEnv* Env, jobject Object, jmethodID Method, ...);
+	static int32 CallIntMethod(JNIEnv* Env, jobject Object, jmethodID Method, ...);
 	static bool CallBooleanMethod(JNIEnv* Env, jobject Object, jmethodID Method, ...);
 
 	// Delegate that can be registered to that is called when an activity is finished
diff --git a/Engine/Source/Runtime/Launch/Private/Android/LaunchAndroid.cpp b/Engine/Source/Runtime/Launch/Private/Android/LaunchAndroid.cpp
--- a/Engine/Source/Runtime/Launch/Private/Android/LaunchAndroid.cpp
+++ b/Engine/Source/Runtime/Launch/Private/Android/LaunchAndroid.cpp
@@ -32,27 +32,18 @@ static const int32_t AxisList[] =
 	AMOTION_EVENT_AXIS_X,
     AMOTION_EVENT_AXIS_Y,
 	AMOTION_EVENT_AXIS_Z,
-    AMOTION_EVENT_AXIS_RZ,
-    AMOTION_EVENT_AXIS_LTRIGGER,
-    AMOTION_EVENT_AXIS_RTRIGGER,
-	AMOTION_EVENT_AXIS_GAS,
-    AMOTION_EVENT_AXIS_BRAKE,
+	AMOTION_EVENT_AXIS_RX,
+	AMOTION_EVENT_AXIS_RY,
+	AMOTION_EVENT_AXIS_RZ,
 
 	//These are DPAD analogs
-    //AMOTION_EVENT_AXIS_HAT_X,
-    //AMOTION_EVENT_AXIS_HAT_Y,
+	AMOTION_EVENT_AXIS_HAT_X,
+	AMOTION_EVENT_AXIS_HAT_Y,
 };
 
 // map of all supported keycodes
 static TSet<uint32> MappedKeyCodes;
 
-// List of gamepad keycodes to ignore
-static const uint32 IgnoredGamepadKeyCodesList[] =
-{
-	AKEYCODE_VOLUME_UP,
-	AKEYCODE_VOLUME_DOWN
-};
-
 // List of desired gamepad keycodes
 static const uint32 ValidGamepadKeyCodesList[] =
 {
@@ -77,9 +68,6 @@ static const uint32 ValidGamepadKeyCodesList[] =
 	AKEYCODE_DPAD_RIGHT
 };
 
-// map of gamepad keycodes that should be ignored
-static TSet<uint32> IgnoredGamepadKeyCodes;
-
 // map of gamepad keycodes that should be passed forward
 static TSet<uint32> ValidGamepadKeyCodes;
 
@@ -232,7 +220,7 @@ int32 AndroidMain(struct android_app* state)
 	}
 
 	// setup joystick support
-	// r19 is the first NDK to include AMotionEvenet_getAxisValue in the headers
+	// r19 is the first NDK to include AMotionEvent_getAxisValue in the headers
 	// However, it has existed in the so since Honeycomb, query for the symbol
 	// to determine whether to try controller support
 	{
@@ -262,12 +250,6 @@ int32 AndroidMain(struct android_app* state)
 		MappedKeyCodes.Add(KeyCodes[i]);
 	}
 
-	const int IgnoredGamepadKeyCodeCount = sizeof(IgnoredGamepadKeyCodesList)/sizeof(uint32);
-	for (int i = 0; i < IgnoredGamepadKeyCodeCount; ++i)
-	{
-		IgnoredGamepadKeyCodes.Add(IgnoredGamepadKeyCodesList[i]);
-	}
-
 	const int ValidGamepadKeyCodeCount = sizeof(ValidGamepadKeyCodesList)/sizeof(uint32);
 	for (int i = 0; i < ValidGamepadKeyCodeCount; ++i)
 	{
@@ -451,7 +433,10 @@ static int32_t HandleInputCB(struct android_app* app, AInputEvent* event)
 //	FPlatformMisc::LowLevelOutputDebugStringf(L"INPUT - type: %x, action: %x, source: %x, keycode: %x, buttons: %x", AInputEvent_getType(event), 
 //		AMotionEvent_getAction(event), AInputEvent_getSource(event), AKeyEvent_getKeyCode(event), AMotionEvent_getButtonState(event));
 
-	if (AInputEvent_getType(event) == AINPUT_EVENT_TYPE_MOTION)
+	int32 EventType = AInputEvent_getType(event);
+	int32 EventSource = AInputEvent_getSource(event);
+
+	if (EventType == AINPUT_EVENT_TYPE_MOTION)
 	{
 		int action = AMotionEvent_getAction(event);
 		int actionType = action & AMOTION_EVENT_ACTION_MASK;
@@ -459,18 +444,29 @@ static int32_t HandleInputCB(struct android_app* app, AInputEvent* event)
 		bool isActionTargeted = (actionType == AMOTION_EVENT_ACTION_POINTER_DOWN || actionType == AMOTION_EVENT_ACTION_POINTER_UP);
 
 		// trap Joystick events first, with fallthrough if there is no joystick support
-		if (((AInputEvent_getSource(event) & AINPUT_SOURCE_CLASS_JOYSTICK) != 0) &&
-				(GetAxes != NULL) &&
-				(actionType == AMOTION_EVENT_ACTION_MOVE))
+		if (((EventSource & AINPUT_SOURCE_CLASS_JOYSTICK) == AINPUT_SOURCE_CLASS_JOYSTICK) &&
+			(GetAxes != NULL) &&
+			(actionType == AMOTION_EVENT_ACTION_MOVE))
 		{
 			const int axisCount = sizeof(AxisList)/sizeof(int32_t);
-			int device = AInputEvent_getDeviceId(event);
+			int32 device = AInputEvent_getDeviceId(event);
 
+			// poll all the axes and forward to update controller state
 			for (int axis = 0; axis < axisCount; axis++)
 			{
 				float val = GetAxes( event, AxisList[axis], 0);
-				FAndroidInputInterface::JoystickAxisEvent( device, AxisList[axis], val);
+				FAndroidInputInterface::JoystickAxisEvent(device, AxisList[axis], val);
 			}
+
+			// handle L/R trigger and Brake/Gas special (all in 0..1 range)
+			// LTRIGGER will either be LTRIGGER or BRAKE, whichever is larger
+			// RTRIGGER will either be RTRIGGER or GAS, whichever is larger
+			float ltrigger = GetAxes(event, AMOTION_EVENT_AXIS_LTRIGGER, 0);
+			float rtrigger = GetAxes(event, AMOTION_EVENT_AXIS_RTRIGGER, 0);
+			float brake = GetAxes(event, AMOTION_EVENT_AXIS_BRAKE, 0);
+			float gas = GetAxes(event, AMOTION_EVENT_AXIS_GAS, 0);
+			FAndroidInputInterface::JoystickAxisEvent(device, AMOTION_EVENT_AXIS_LTRIGGER, ltrigger > brake ? ltrigger : brake);
+			FAndroidInputInterface::JoystickAxisEvent(device, AMOTION_EVENT_AXIS_RTRIGGER, rtrigger > gas ? rtrigger : gas);
 		}
 		else
 		{
@@ -556,10 +552,10 @@ static int32_t HandleInputCB(struct android_app* app, AInputEvent* event)
 					float y = FMath::Min<float>(AMotionEvent_getY(event, i) / Height, 1.f);
 					y *= (ScreenRect.Bottom - 1);
 
-					UE_LOG(LogAndroid, Verbose, TEXT("Received motion event from pointer %u (id %d) action %d: (%.2f, %.2f)"), i, action, AMotionEvent_getPointerId(event,i), x, y);
+					UE_LOG(LogAndroid, Verbose, TEXT("Received motion event from index %u (id %d) action %d: (%.2f, %.2f)"), i, pointerId, action, x, y);
 
 					TouchInput TouchMessage;
-					TouchMessage.Handle = AMotionEvent_getPointerId(event, i);
+					TouchMessage.Handle = pointerId;
 					TouchMessage.Type = type;
 					TouchMessage.Position = FVector2D(x, y);
 					TouchMessage.LastPosition = FVector2D(x, y);		//@todo android: AMotionEvent_getHistoricalRawX
@@ -585,23 +581,26 @@ static int32_t HandleInputCB(struct android_app* app, AInputEvent* event)
 
 		return 1;
 	}
-	else if (AInputEvent_getType(event) == AINPUT_EVENT_TYPE_KEY)
+
+	if (EventType == AINPUT_EVENT_TYPE_KEY)
 	{
 		int keyCode = AKeyEvent_getKeyCode(event);
 
 		FPlatformMisc::LowLevelOutputDebugStringf(L"Received keycode: %d", keyCode);
 
-		//Trap Joystick events first, with fallthrough if there is no joystick support
-		if (((AInputEvent_getSource(event) & (AINPUT_SOURCE_GAMEPAD | AINPUT_SOURCE_DPAD)) != 0) && (GetAxes != NULL) && ValidGamepadKeyCodes.Contains(keyCode))
+		//Trap codes handled as possible gamepad events
+		if (ValidGamepadKeyCodes.Contains(keyCode))
 		{
-			if (IgnoredGamepadKeyCodes.Contains(keyCode))
+			//Only pass on the device id if really a gamepad, joystick or dpad (allows menu and back to be treated as gamepad events)
+			int32 device = 0;
+			if ( (((EventSource & AINPUT_SOURCE_JOYSTICK) == AINPUT_SOURCE_JOYSTICK) && (GetAxes != NULL)) ||
+				 ((EventSource & AINPUT_SOURCE_GAMEPAD) == AINPUT_SOURCE_GAMEPAD) ||
+				 ((EventSource & AINPUT_SOURCE_DPAD) == AINPUT_SOURCE_DPAD) )
 			{
-				return 0;
+				device = AInputEvent_getDeviceId(event);
 			}
-
-			int device = AInputEvent_getDeviceId(event);
 			bool down = AKeyEvent_getAction(event) != AKEY_EVENT_ACTION_UP;
-			FAndroidInputInterface::JoystickButtonEvent( device, keyCode, down);
+			FAndroidInputInterface::JoystickButtonEvent(device, keyCode, down);
 			FPlatformMisc::LowLevelOutputDebugStringf(L"Received gamepad button: %d", keyCode);
 		}
 		else
