Linux: fix reporting memory stats (OR-7349, UE-9538).

(Merging CL 2705697 to the main branch).

diff --git a/Engine/Source/Programs/TestPAL/Private/Main.cpp b/Engine/Source/Programs/TestPAL/Private/Main.cpp
--- a/Engine/Source/Programs/TestPAL/Private/Main.cpp
+++ b/Engine/Source/Programs/TestPAL/Private/Main.cpp
@@ -246,6 +246,8 @@ int32 SysInfoTest(const TCHAR* CommandLine)
 	FString OSInstanceGuid = FPlatformMisc::GetOperatingSystemId();
 	UE_LOG(LogTestPAL, Display, TEXT("  FPlatformMisc::GetOperatingSystemId() = %s"), *OSInstanceGuid);
 
+	FPlatformMemory::DumpStats(*GLog);
+
 	FEngineLoop::AppPreExit();
 	FEngineLoop::AppExit();
 	return 0;
diff --git a/Engine/Source/Runtime/Core/Private/Linux/LinuxPlatformMemory.cpp b/Engine/Source/Runtime/Core/Private/Linux/LinuxPlatformMemory.cpp
--- a/Engine/Source/Runtime/Core/Private/Linux/LinuxPlatformMemory.cpp
+++ b/Engine/Source/Runtime/Core/Private/Linux/LinuxPlatformMemory.cpp
@@ -123,11 +123,145 @@ void FLinuxPlatformMemory::BinnedFreeToOS( void* Ptr )
 	return free(Ptr);
 }
 
+namespace LinuxPlatformMemory
+{
+	/**
+	 * @brief Returns value in bytes from a status line
+	 * @param Line in format "Blah:  10000 kB" - needs to be writable as it will modify it
+	 * @return value in bytes (10240000, i.e. 10000 * 1024 for the above example)
+	 */
+	uint64 GetBytesFromStatusLine(char * Line)
+	{
+		check(Line);
+		int Len = strlen(Line);
+
+		// Len should be long enough to hold at least " kB\n"
+		const int kSuffixLength = 4;	// " kB\n"
+		if (Len <= kSuffixLength)
+		{
+			return 0;
+		}
+
+		// let's check that this is indeed "kB"
+		char * Suffix = &Line[Len - kSuffixLength];
+		if (strcmp(Suffix, " kB\n") != 0)
+		{
+			// Linux the kernel changed the format, huh?
+			return 0;
+		}
+
+		// kill the kB
+		*Suffix = 0;
+
+		// find the beginning of the number
+		for (const char * NumberBegin = Suffix; NumberBegin >= Line; --NumberBegin)
+		{
+			if (*NumberBegin == ' ')
+			{
+				return static_cast< uint64 >(atol(NumberBegin + 1)) * 1024ULL;
+			}
+		}
+
+		// we were unable to find whitespace in front of the number
+		return 0;
+	}
+}
+
 FPlatformMemoryStats FLinuxPlatformMemory::GetStats()
 {
-	FPlatformMemoryStats MemoryStats;
+	FPlatformMemoryStats MemoryStats;	// will init from constants
+
+	// open to all kind of overflows, thanks to Linux approach of exposing system stats via /proc and lack of proper C API
+	// And no, sysinfo() isn't useful for this (cannot get the same value for MemAvailable through it for example).
+
+	if (FILE* FileGlobalMemStats = fopen("/proc/meminfo", "r"))
+	{
+		int FieldsSetSuccessfully = 0;
+		SIZE_T MemFree = 0, Cached = 0;
+		do
+		{
+			char LineBuffer[256] = {0};
+			char *Line = fgets(LineBuffer, ARRAY_COUNT(LineBuffer), FileGlobalMemStats);
+			if (Line == nullptr)
+			{
+				break;	// eof or an error
+			}
+
+			// if we have MemAvailable, favor that (see http://git.kernel.org/cgit/linux/kernel/git/torvalds/linux.git/commit/?id=34e431b0ae398fc54ea69ff85ec700722c9da773)
+			if (strstr(Line, "MemAvailable:") == Line)
+			{
+				MemoryStats.AvailablePhysical = LinuxPlatformMemory::GetBytesFromStatusLine(Line);
+				++FieldsSetSuccessfully;
+			}
+			else if (strstr(Line, "SwapFree:") == Line)
+			{
+				MemoryStats.AvailableVirtual = LinuxPlatformMemory::GetBytesFromStatusLine(Line);
+				++FieldsSetSuccessfully;
+			}
+			else if (strstr(Line, "MemFree:") == Line)
+			{
+				MemFree = LinuxPlatformMemory::GetBytesFromStatusLine(Line);
+				++FieldsSetSuccessfully;
+			}
+			else if (strstr(Line, "Cached:") == Line)
+			{
+				Cached = LinuxPlatformMemory::GetBytesFromStatusLine(Line);
+				++FieldsSetSuccessfully;
+			}
+		}
+		while(FieldsSetSuccessfully < 4);
+
+		// if we didn't have MemAvailable (kernels < 3.14 or CentOS 6.x), use free + cached as a (bad) approximation
+		if (MemoryStats.AvailablePhysical == 0)
+		{
+			MemoryStats.AvailablePhysical = FMath::Min(MemFree + Cached, MemoryStats.TotalPhysical);
+		}
+
+		fclose(FileGlobalMemStats);
+	}
+
+	// again /proc "API" :/
+	if (FILE* ProcMemStats = fopen("/proc/self/status", "r"))
+	{
+		int FieldsSetSuccessfully = 0;
+		do
+		{
+			char LineBuffer[256] = {0};
+			char *Line = fgets(LineBuffer, ARRAY_COUNT(LineBuffer), ProcMemStats);
+			if (Line == nullptr)
+			{
+				break;	// eof or an error
+			}
+
+			if (strstr(Line, "VmPeak:") == Line)
+			{
+				MemoryStats.PeakUsedVirtual = LinuxPlatformMemory::GetBytesFromStatusLine(Line);
+				++FieldsSetSuccessfully;
+			}
+			else if (strstr(Line, "VmSize:") == Line)
+			{
+				MemoryStats.UsedVirtual = LinuxPlatformMemory::GetBytesFromStatusLine(Line);
+				++FieldsSetSuccessfully;
+			}
+			else if (strstr(Line, "VmHWM:") == Line)
+			{
+				MemoryStats.PeakUsedPhysical = LinuxPlatformMemory::GetBytesFromStatusLine(Line);
+				++FieldsSetSuccessfully;
+			}
+			else if (strstr(Line, "VmRSS:") == Line)
+			{
+				MemoryStats.UsedPhysical = LinuxPlatformMemory::GetBytesFromStatusLine(Line);
+				++FieldsSetSuccessfully;
+			}
+		}
+		while(FieldsSetSuccessfully < 4);
+
+		fclose(ProcMemStats);
+	}
 
-	// @todo
+	// sanitize stats as sometimes peak < used for some reason
+	MemoryStats.PeakUsedVirtual = FMath::Max(MemoryStats.PeakUsedVirtual, MemoryStats.UsedVirtual);
+	MemoryStats.PeakUsedPhysical = FMath::Max(MemoryStats.PeakUsedPhysical, MemoryStats.UsedPhysical);
 
 	return MemoryStats;
 }
@@ -141,13 +275,16 @@ const FPlatformMemoryConstants& FLinuxPlatformMemory::GetConstants()
 		// Gather platform memory stats.
 		struct sysinfo SysInfo;
 		unsigned long long MaxPhysicalRAMBytes = 0;
+		unsigned long long MaxVirtualRAMBytes = 0;
 
 		if (0 == sysinfo(&SysInfo))
 		{
 			MaxPhysicalRAMBytes = static_cast< unsigned long long >( SysInfo.mem_unit ) * static_cast< unsigned long long >( SysInfo.totalram );
+			MaxVirtualRAMBytes = static_cast< unsigned long long >( SysInfo.mem_unit ) * static_cast< unsigned long long >( SysInfo.totalswap );
 		}
 
 		MemoryConstants.TotalPhysical = MaxPhysicalRAMBytes;
+		MemoryConstants.TotalVirtual = MaxVirtualRAMBytes;
 		MemoryConstants.TotalPhysicalGB = (MemoryConstants.TotalPhysical + 1024 * 1024 * 1024 - 1) / 1024 / 1024 / 1024;
 		MemoryConstants.PageSize = sysconf(_SC_PAGESIZE);
 	}
