Merging using Orion->UE4

CL# 2706133
Fat-free TArray and FString ranged-for iteration.

diff --git a/Engine/Source/Runtime/Core/Public/Containers/Array.h b/Engine/Source/Runtime/Core/Public/Containers/Array.h
--- a/Engine/Source/Runtime/Core/Public/Containers/Array.h
+++ b/Engine/Source/Runtime/Core/Public/Containers/Array.h
@@ -131,6 +131,32 @@ FORCEINLINE TIndexedContainerIterator<ContainerType, ElementType, IndexType> ope
 	return RHS + Offset;
 }
 
+template <typename ElementType>
+struct TDereferencingPointerIterator
+{
+	explicit TDereferencingPointerIterator(ElementType** InPtr)
+		: Ptr(InPtr)
+	{
+	}
+
+	FORCEINLINE ElementType& operator*() const
+	{
+		return **Ptr;
+	}
+
+	FORCEINLINE TDereferencingPointerIterator& operator++()
+	{
+		++Ptr;
+		return *this;
+	}
+
+private:
+	ElementType** Ptr;
+
+	FORCEINLINE friend bool operator==(const TDereferencingPointerIterator& Lhs, const TDereferencingPointerIterator& Rhs) { return Lhs.Ptr == Rhs.Ptr; }
+	FORCEINLINE friend bool operator!=(const TDereferencingPointerIterator& Lhs, const TDereferencingPointerIterator& Rhs) { return Lhs.Ptr != Rhs.Ptr; }
+};
+
 /**
  * Base dynamic array.
  * An untyped data array; mirrors a TArray's members, but doesn't need an exact C++ type for its elements.
@@ -2222,10 +2248,10 @@ private:
 	 * DO NOT USE DIRECTLY
 	 * STL-like iterators to enable range-based for loop support.
 	 */
-	FORCEINLINE friend TIterator      begin(      TArray& Array) { return TIterator     (Array); }
-	FORCEINLINE friend TConstIterator begin(const TArray& Array) { return TConstIterator(Array); }
-	FORCEINLINE friend TIterator      end  (      TArray& Array) { return TIterator     (Array, Array.Num()); }
-	FORCEINLINE friend TConstIterator end  (const TArray& Array) { return TConstIterator(Array, Array.Num()); }
+	FORCEINLINE friend       ElementType* begin(      TArray& Array) { return Array.GetData(); }
+	FORCEINLINE friend const ElementType* begin(const TArray& Array) { return Array.GetData(); }
+	FORCEINLINE friend       ElementType* end  (      TArray& Array) { return Array.GetData() + Array.Num(); }
+	FORCEINLINE friend const ElementType* end  (const TArray& Array) { return Array.GetData() + Array.Num(); }
 
 public:
 	/**
@@ -3373,10 +3399,10 @@ private:
 	 * DO NOT USE DIRECTLY
 	 * STL-like iterators to enable range-based for loop support.
 	 */
-	FORCEINLINE friend TIterator      begin(      TIndirectArray& IndirectArray) { return TIterator     (IndirectArray); }
-	FORCEINLINE friend TConstIterator begin(const TIndirectArray& IndirectArray) { return TConstIterator(IndirectArray); }
-	FORCEINLINE friend TIterator      end  (      TIndirectArray& IndirectArray) { return TIterator     (IndirectArray, IndirectArray.Array.Num()); }
-	FORCEINLINE friend TConstIterator end  (const TIndirectArray& IndirectArray) { return TConstIterator(IndirectArray, IndirectArray.Array.Num()); }
+	FORCEINLINE friend TDereferencingPointerIterator<      ElementType> begin(      TIndirectArray& IndirectArray) { return TDereferencingPointerIterator<      ElementType>(IndirectArray.GetData()); }
+	FORCEINLINE friend TDereferencingPointerIterator<const ElementType> begin(const TIndirectArray& IndirectArray) { return TDereferencingPointerIterator<const ElementType>(IndirectArray.GetData()); }
+	FORCEINLINE friend TDereferencingPointerIterator<      ElementType> end  (      TIndirectArray& IndirectArray) { return TDereferencingPointerIterator<      ElementType>(IndirectArray.GetData() + IndirectArray.Num()); }
+	FORCEINLINE friend TDereferencingPointerIterator<const ElementType> end  (const TIndirectArray& IndirectArray) { return TDereferencingPointerIterator<const ElementType>(IndirectArray.GetData() + IndirectArray.Num()); }
 
 	InternalArrayType Array;
 };
diff --git a/Engine/Source/Runtime/Core/Public/Containers/UnrealString.h b/Engine/Source/Runtime/Core/Public/Containers/UnrealString.h
--- a/Engine/Source/Runtime/Core/Public/Containers/UnrealString.h
+++ b/Engine/Source/Runtime/Core/Public/Containers/UnrealString.h
@@ -240,10 +240,10 @@ private:
 	 * DO NOT USE DIRECTLY
 	 * STL-like iterators to enable range-based for loop support.
 	 */
-	FORCEINLINE friend TIterator      begin(      FString& Str) { return begin(Str.Data); }
-	FORCEINLINE friend TConstIterator begin(const FString& Str) { return begin(Str.Data); }
-	FORCEINLINE friend TIterator      end  (      FString& Str) { TIterator      Result = end(Str.Data); if (Str.Data.Num()) { --Result; } return Result; }
-	FORCEINLINE friend TConstIterator end  (const FString& Str) { TConstIterator Result = end(Str.Data); if (Str.Data.Num()) { --Result; } return Result; }
+	FORCEINLINE friend       TCHAR* begin(      FString& Str) { return begin(Str.Data); }
+	FORCEINLINE friend const TCHAR* begin(const FString& Str) { return begin(Str.Data); }
+	FORCEINLINE friend       TCHAR* end  (      FString& Str) {       TCHAR* Result = end(Str.Data); if (Str.Data.Num()) { --Result; } return Result; }
+	FORCEINLINE friend const TCHAR* end  (const FString& Str) { const TCHAR* Result = end(Str.Data); if (Str.Data.Num()) { --Result; } return Result; }
 
 public:
 	FORCEINLINE uint32 GetAllocatedSize() const
