UE4 - Improved the console performance of FBitReader.

diff --git a/Engine/Source/Runtime/Core/Private/Serialization/Bitstreams.cpp b/Engine/Source/Runtime/Core/Private/Serialization/Bitstreams.cpp
--- a/Engine/Source/Runtime/Core/Private/Serialization/Bitstreams.cpp
+++ b/Engine/Source/Runtime/Core/Private/Serialization/Bitstreams.cpp
@@ -11,10 +11,11 @@ static const uint8 GShift[8]={0x01,0x02,0x04,0x08,0x10,0x20,0x40,0x80};
 static const uint8 GMask [8]={0x00,0x01,0x03,0x07,0x0f,0x1f,0x3f,0x7f};
 
 // Optimized arbitrary bit range memory copy routine.
+
 void appBitsCpy( uint8* Dest, int32 DestBit, uint8* Src, int32 SrcBit, int32 BitCount )
 {
 	if( BitCount==0 ) return;
-	
+
 	// Special case - always at least one bit to copy,
 	// a maximum of 2 bytes to read, 2 to write - only touch bytes that are actually used.
 	if( BitCount <= 8 ) 
@@ -328,130 +329,12 @@ void FBitReader::AppendTo( TArray<uint8> &DestBuffer )
 {
 	DestBuffer.Append(Buffer);
 }
-void FBitReader::EatByteAlign()
-{
-	// Skip over remaining bits in current byte
-	Pos = (Pos+7) & (~0x07);
-	if ( Pos > Num )
-	{
-		UE_LOG( LogNetSerialization, Error, TEXT( "FBitReader::EatByteAlign: Pos > Num" ) );
-		SetOverflowed();
-	}
-}
-void FBitReader::SerializeBits( void* Dest, int64 LengthBits )
-{
-	FMemory::Memzero( Dest, (LengthBits+7)>>3 );
 
-	if ( IsError() )
-	{
-		return;
-	}
-
-	if( Pos+LengthBits<=Num )
-	{
-		//for( int32 i=0; i<LengthBits; i++,Pos++ )
-		//	if( Buffer(Pos>>3) & GShift[Pos&7] )
-		//		((uint8*)Dest)[i>>3] |= GShift[i&7];
-		if( LengthBits == 1 )
-		{
-			if( Buffer[Pos>>3] & GShift[Pos&7] )
-				((uint8*)Dest)[0] |= 0x01;
-			Pos++;
-		}
-		else
-		{
-			appBitsCpy((uint8*)Dest, 0, Buffer.GetData(), Pos, LengthBits);
-			Pos += LengthBits;
-		}
-	}
-	else 
-	{
-		SetOverflowed();
-		UE_LOG( LogNetSerialization, Error, TEXT( "FBitReader::SerializeBits: Pos + LengthBits > Num" ) );
-	}
-}
-void FBitReader::SerializeInt( uint32& OutValue, uint32 ValueMax )
-{
-	if ( IsError() )
-	{
-		return;
-	}
-
-	uint32 Value=0; // Use local variable to avoid Load-Hit-Store
-	int64 LocalPos = Pos;
-	const int64 LocalNum = Num;
-
-	for( uint32 Mask=1; Value+Mask<ValueMax && Mask; Mask*=2,LocalPos++ )
-	{
-		if( LocalPos>=LocalNum )
-		{
-			SetOverflowed();
-			UE_LOG( LogNetSerialization, Error, TEXT( "FBitReader::SerializeInt: LocalPos >= LocalNum" ) );
-			break;
-		}
-		if( Buffer[LocalPos>>3] & GShift[LocalPos&7] )
-		{
-			Value |= Mask;
-		}
-	}
-	// Now write back
-	Pos = LocalPos;
-	OutValue = Value;
-}
-uint32 FBitReader::ReadInt( uint32 Max )
-{
-	uint32 Value=0;
-	SerializeInt( Value, Max );
-	return Value;
-}
-uint8 FBitReader::ReadBit()
-{
-	uint8 Bit=0;
-	SerializeBits( &Bit, 1 );
-	return Bit;
-}
-void FBitReader::Serialize( void* Dest, int64 LengthBytes )
-{
-	SerializeBits( Dest, LengthBytes*8 );
-}
-uint8* FBitReader::GetData()
-{
-	return Buffer.GetData();
-}
-uint8* FBitReader::GetDataPosChecked()
-{
-	check(Pos % 8 == 0);
-	return &Buffer[Pos >> 3];
-}
-uint32 FBitReader::GetBytesLeft()
-{
-	return ((Num - Pos) + 7) >> 3;
-}
-uint32 FBitReader::GetBitsLeft()
-{
-	return (Num - Pos);
-}
-bool FBitReader::AtEnd()
-{
-	return ArIsError || Pos>=Num;
-}
 void FBitReader::SetOverflowed()
 {
 	UE_LOG( LogNetSerialization, Error, TEXT( "FBitReader::SetOverflowed() called" ) );
 	ArIsError = 1;
 }
-int64 FBitReader::GetNumBytes()
-{
-	return (Num+7)>>3;
-}
-int64 FBitReader::GetNumBits()
-{
-	return Num;
-}
-int64 FBitReader::GetPosBits()
-{
-	return Pos;
-}
 
 /*-----------------------------------------------------------------------------
 	FBitReader.
diff --git a/Engine/Source/Runtime/Core/Public/Serialization/Bitstreams.h b/Engine/Source/Runtime/Core/Public/Serialization/Bitstreams.h
--- a/Engine/Source/Runtime/Core/Public/Serialization/Bitstreams.h
+++ b/Engine/Source/Runtime/Core/Public/Serialization/Bitstreams.h
@@ -3,6 +3,9 @@
 #pragma once
 
 
+CORE_API void appBitsCpy( uint8* Dest, int32 DestBit, uint8* Src, int32 SrcBit, int32 BitCount );
+
+
 /*-----------------------------------------------------------------------------
 	FBitWriter.
 -----------------------------------------------------------------------------*/
@@ -188,30 +191,163 @@ public:
 
 	FBitReader( uint8* Src = nullptr, int64 CountBits = 0 );
 	void SetData( FBitReader& Src, int64 CountBits );
-	void SerializeBits( void* Dest, int64 LengthBits );
-	void SerializeInt( uint32& Value, uint32 Max );
-	uint32 ReadInt( uint32 Max );
-	uint8 ReadBit();
-	void Serialize( void* Dest, int64 LengthBytes );
-	uint8* GetData();
-	uint8* GetDataPosChecked();
-	uint32 GetBytesLeft();
-	uint32 GetBitsLeft();
-	bool AtEnd();
+	FORCEINLINE_DEBUGGABLE void SerializeBits( void* Dest, int64 LengthBits )
+	{
+		if ( IsError() || Pos+LengthBits > Num)
+		{
+			if (!IsError())
+			{
+				SetOverflowed();
+				//UE_LOG( LogNetSerialization, Error, TEXT( "FBitReader::SerializeBits: Pos + LengthBits > Num" ) );
+			}
+			FMemory::Memzero( Dest, (LengthBits+7)>>3 );
+			return;
+		}
+		//for( int32 i=0; i<LengthBits; i++,Pos++ )
+		//	if( Buffer(Pos>>3) & GShift[Pos&7] )
+		//		((uint8*)Dest)[i>>3] |= GShift[i&7];
+		if( LengthBits == 1 )
+		{
+			((uint8*)Dest)[0] = 0;
+			if( Buffer[Pos>>3] & Shift(Pos&7) )
+				((uint8*)Dest)[0] |= 0x01;
+			Pos++;
+		}
+		else if (LengthBits != 0)
+		{
+			((uint8*)Dest)[((LengthBits+7)>>3) - 1] = 0;
+			appBitsCpy((uint8*)Dest, 0, Buffer.GetData(), Pos, LengthBits);
+			Pos += LengthBits;
+		}
+	}
+
+	FORCEINLINE_DEBUGGABLE void SerializeInt( uint32& OutValue, uint32 ValueMax )
+	{
+		if ( IsError() )
+		{
+			return;
+		}
+
+		uint32 Value=0; // Use local variable to avoid Load-Hit-Store
+		int64 LocalPos = Pos;
+		const int64 LocalNum = Num;
+
+		for( uint32 Mask=1; Value+Mask<ValueMax && Mask; Mask*=2,LocalPos++ )
+		{
+			if( LocalPos>=LocalNum )
+			{
+				SetOverflowed();
+				//UE_LOG( LogNetSerialization, Error, TEXT( "FBitReader::SerializeInt: LocalPos >= LocalNum" ) );
+				break;
+			}
+			if( Buffer[LocalPos>>3] & Shift(LocalPos&7) )
+			{
+				Value |= Mask;
+			}
+		}
+		// Now write back
+		Pos = LocalPos;
+		OutValue = Value;
+	}
+
+	FORCEINLINE_DEBUGGABLE uint32 ReadInt( uint32 Max )
+	{
+		uint32 Value=0;
+		SerializeInt( Value, Max );
+		return Value;
+	}
+
+	FORCEINLINE_DEBUGGABLE uint8 ReadBit()
+	{
+		uint8 Bit=0;
+		//SerializeBits( &Bit, 1 );
+		if ( !IsError() )
+		{
+			int64 LocalPos = Pos;
+			const int64 LocalNum = Num;
+			if (LocalPos >= LocalNum)
+			{
+				SetOverflowed();
+				//UE_LOG( LogNetSerialization, Error, TEXT( "FBitReader::SerializeInt: LocalPos >= LocalNum" ) );
+			}
+			else
+			{
+				Bit = !!(Buffer[LocalPos>>3] & Shift(LocalPos&7));
+				Pos++;
+			}
+		}
+		return Bit;
+	}
+
+	FORCEINLINE_DEBUGGABLE void Serialize( void* Dest, int64 LengthBytes )
+	{
+		SerializeBits( Dest, LengthBytes*8 );
+	}
+
+	FORCEINLINE_DEBUGGABLE uint8* GetData()
+	{
+		return Buffer.GetData();
+	}
+
+	FORCEINLINE_DEBUGGABLE uint8* GetDataPosChecked()
+	{
+		check(Pos % 8 == 0);
+		return &Buffer[Pos >> 3];
+	}
+
+	FORCEINLINE_DEBUGGABLE uint32 GetBytesLeft()
+	{
+		return ((Num - Pos) + 7) >> 3;
+	}
+	FORCEINLINE_DEBUGGABLE uint32 GetBitsLeft()
+	{
+		return (Num - Pos);
+	}
+	FORCEINLINE_DEBUGGABLE bool AtEnd()
+	{
+		return ArIsError || Pos>=Num;
+	}
+	FORCEINLINE_DEBUGGABLE int64 GetNumBytes()
+	{
+		return (Num+7)>>3;
+	}
+	FORCEINLINE_DEBUGGABLE int64 GetNumBits()
+	{
+		return Num;
+	}
+	FORCEINLINE_DEBUGGABLE int64 GetPosBits()
+	{
+		return Pos;
+	}
+	FORCEINLINE_DEBUGGABLE void EatByteAlign()
+	{
+		// Skip over remaining bits in current byte
+		Pos = (Pos+7) & (~0x07);
+		if ( Pos > Num )
+		{
+			//UE_LOG( LogNetSerialization, Error, TEXT( "FBitReader::EatByteAlign: Pos > Num" ) );
+			SetOverflowed();
+		}
+	}
+
 	void SetOverflowed();
-	int64 GetNumBytes();
-	int64 GetNumBits();
-	int64 GetPosBits();
 	void AppendDataFromChecked( FBitReader& Src );
 	void AppendDataFromChecked( uint8* Src, uint32 NumBits );
 	void AppendTo( TArray<uint8> &Buffer );
-	void EatByteAlign();
 
 protected:
 
 	TArray<uint8> Buffer;
 	int64 Num;
 	int64 Pos;
+
+private:
+
+	FORCEINLINE uint8 Shift(uint8 Cnt)
+	{
+		return 1<<Cnt;
+	}
+
 };
 
 
