Merging using UE4-To-UE4-Releases-4.10
Change generic IsControllerAssignedToGamepad to return true only for controller 0
#ue4

diff --git a/Engine/Source/Runtime/Core/Public/GenericPlatform/GenericPlatformMisc.h b/Engine/Source/Runtime/Core/Public/GenericPlatform/GenericPlatformMisc.h
--- a/Engine/Source/Runtime/Core/Public/GenericPlatform/GenericPlatformMisc.h
+++ b/Engine/Source/Runtime/Core/Public/GenericPlatform/GenericPlatformMisc.h
@@ -777,7 +777,7 @@ public:
 	 */
 	static bool IsControllerAssignedToGamepad(int32 ControllerId)
 	{
-		return true;
+		return (ControllerId == 0);
 	}
 
 	/*
