Fixed FTextSnapshot reporting all transient text instances as dirty

FTextSnapshot was changed to use a TWeakPtr, however this meant that for text that was generated each frame, the FTextSnapshot would always report that it was out-of-date (even if the generated text was the same as the previous frame) as it no longer had a valid reference point to compare against.

FTextSnapshot has been changed to once again use a TSharedPtr, which will keep the transient instance around long enough for it to be compared with the text from the next frame.

diff --git a/Engine/Source/Runtime/Core/Private/Internationalization/Text.cpp b/Engine/Source/Runtime/Core/Private/Internationalization/Text.cpp
--- a/Engine/Source/Runtime/Core/Private/Internationalization/Text.cpp
+++ b/Engine/Source/Runtime/Core/Private/Internationalization/Text.cpp
@@ -1292,8 +1292,7 @@ bool FTextSnapshot::IdenticalTo(const FText& InText) const
 	// (this usually happens when ToString() is called)
 	InText.Rebuild();
 
-	TSharedPtr<ITextData, ESPMode::ThreadSafe> TextDataPin = TextDataPtr.Pin();
-	return TextDataPin == InText.TextData 
+	return TextDataPtr == InText.TextData 
 		&& HistoryRevision == InText.TextData->GetTextHistory().Revision
 		&& Flags == InText.Flags;;
 }
@@ -1306,9 +1305,8 @@ bool FTextSnapshot::IsDisplayStringEqualTo(const FText& InText) const
 
 	// We have to assume that the display string has changed if the history of the text has changed
 	// (due to a culture change), as we no longer have the old display string to compare against
-	TSharedPtr<ITextData, ESPMode::ThreadSafe> TextDataPin = TextDataPtr.Pin();
 	return HistoryRevision == InText.TextData->GetTextHistory().Revision 
-		&& TextDataPin.IsValid() && TextDataPin->GetDisplayString().Equals(InText.ToString(), ESearchCase::CaseSensitive);
+		&& TextDataPtr.IsValid() && TextDataPtr->GetDisplayString().Equals(InText.ToString(), ESearchCase::CaseSensitive);
 }
 
 FScopedTextIdentityPreserver::FScopedTextIdentityPreserver(FText& InTextToPersist)
diff --git a/Engine/Source/Runtime/Core/Public/Internationalization/Text.h b/Engine/Source/Runtime/Core/Public/Internationalization/Text.h
--- a/Engine/Source/Runtime/Core/Public/Internationalization/Text.h
+++ b/Engine/Source/Runtime/Core/Public/Internationalization/Text.h
@@ -612,7 +612,7 @@ public:
 
 private:
 	/** A pointer to the text data for the FText that we took a snapshot of (used for an efficient pointer compare) */
-	TWeakPtr<ITextData, ESPMode::ThreadSafe> TextDataPtr;
+	TSharedPtr<ITextData, ESPMode::ThreadSafe> TextDataPtr;
 
 	/** Revision index of the history of the FText we took a snapshot of, or INDEX_NONE if there was no history */
 	int32 HistoryRevision;
