Merging CL 2708073 - Fixed a Mac-specific crash at exit that occured if GLog was initialized before the game thread was created

diff --git a/Engine/Source/Runtime/Core/Private/Mac/CocoaThread.cpp b/Engine/Source/Runtime/Core/Private/Mac/CocoaThread.cpp
--- a/Engine/Source/Runtime/Core/Private/Mac/CocoaThread.cpp
+++ b/Engine/Source/Runtime/Core/Private/Mac/CocoaThread.cpp
@@ -356,6 +356,11 @@ FCocoaRunLoopSource* FCocoaRunLoopSource::GameRunLoopSource = nullptr;
 	// Register game run loop source
 	FCocoaRunLoopSource::RegisterGameRunLoop([GameRunLoop getCFRunLoop]);
 	
+	if (GLog)
+	{
+		GLog->SetCurrentThreadAsMasterThread();
+	}
+
 	[super main];
 	
 	// We have exited the game thread, so any UE4 code running now should treat the Main thread
