Merging using UE4-To-UE4-Releases-4.10
Map Samsung gamepad L1 and R1 buttons to triggers
#ue4
#android

diff --git a/Engine/Source/Runtime/Core/Private/Android/AndroidInputInterface.cpp b/Engine/Source/Runtime/Core/Private/Android/AndroidInputInterface.cpp
--- a/Engine/Source/Runtime/Core/Private/Android/AndroidInputInterface.cpp
+++ b/Engine/Source/Runtime/Core/Private/Android/AndroidInputInterface.cpp
@@ -701,6 +701,7 @@ void FAndroidInputInterface::SendControllerEvents()
 
 					// Generic mappings
 					CurrentDevice.bSupportsHat = false;
+					CurrentDevice.bMapL1R1ToTriggers = false;
 					CurrentDevice.bRightStickZRZ = true;
 					CurrentDevice.bRightStickRXRY = false;
 
@@ -725,6 +726,7 @@ void FAndroidInputInterface::SendControllerEvents()
 					else if (CurrentDevice.DeviceInfo.Name.StartsWith(TEXT("Samsung Game Pad EI-GP20")))
 					{
 						CurrentDevice.bSupportsHat = true;
+						CurrentDevice.bMapL1R1ToTriggers = true;
 						CurrentDevice.bRightStickZRZ = false;
 						CurrentDevice.bRightStickRXRY = true;
 					}
@@ -1061,8 +1063,18 @@ void FAndroidInputInterface::JoystickButtonEvent(int32 deviceId, int32 buttonId,
 		case AKEYCODE_BUTTON_B:      NewControllerData[deviceId].ButtonStates[1] = buttonDown; break;
 		case AKEYCODE_BUTTON_X:      NewControllerData[deviceId].ButtonStates[2] = buttonDown; break;
 		case AKEYCODE_BUTTON_Y:      NewControllerData[deviceId].ButtonStates[3] = buttonDown; break;
-		case AKEYCODE_BUTTON_L1:     NewControllerData[deviceId].ButtonStates[4] = buttonDown; break;
-		case AKEYCODE_BUTTON_R1:     NewControllerData[deviceId].ButtonStates[5] = buttonDown; break;
+		case AKEYCODE_BUTTON_L1:     NewControllerData[deviceId].ButtonStates[4] = buttonDown;
+									 if (DeviceMapping[deviceId].bMapL1R1ToTriggers)
+									 {
+										 NewControllerData[deviceId].ButtonStates[10] = buttonDown;
+									 }
+									 break;
+		case AKEYCODE_BUTTON_R1:     NewControllerData[deviceId].ButtonStates[5] = buttonDown;
+									 if (DeviceMapping[deviceId].bMapL1R1ToTriggers)
+									 {
+										 NewControllerData[deviceId].ButtonStates[11] = buttonDown;
+									 }
+									 break;
 		case AKEYCODE_BUTTON_START:
 		case AKEYCODE_MENU:          NewControllerData[deviceId].ButtonStates[6] = buttonDown; NewControllerData[deviceId].ButtonStates[17] = buttonDown;  break;
 		case AKEYCODE_BUTTON_SELECT: 
diff --git a/Engine/Source/Runtime/Core/Public/Android/AndroidInputInterface.h b/Engine/Source/Runtime/Core/Public/Android/AndroidInputInterface.h
--- a/Engine/Source/Runtime/Core/Public/Android/AndroidInputInterface.h
+++ b/Engine/Source/Runtime/Core/Public/Android/AndroidInputInterface.h
@@ -104,6 +104,9 @@ struct FAndroidGamepadDeviceMapping
 	// Device supports hat as dpad
 	bool bSupportsHat;
 
+	// Map L1 and R1 to LTRIGGER and RTRIGGER
+	bool bMapL1R1ToTriggers;
+
 	// Right stick on Z/RZ
 	bool bRightStickZRZ;
 
