Core: Prevented FPaths::MakeStandardFilename from crashing when passed a 1 character string
[INTEGRATE] Change 2708644 by [REDACTED_PERSON][REDACTED_HANDLE] on 2015/09/28 19:06:15

diff --git a/Engine/Source/Runtime/Core/Private/Misc/Paths.cpp b/Engine/Source/Runtime/Core/Private/Misc/Paths.cpp
--- a/Engine/Source/Runtime/Core/Private/Misc/Paths.cpp
+++ b/Engine/Source/Runtime/Core/Private/Misc/Paths.cpp
@@ -680,7 +680,7 @@ void FPaths::MakeStandardFilename(FString& InPath)
 	// UNC (windows) network path
 	bool bCannotBeStandardized = InPath.StartsWith(TEXT("\\\\"), ESearchCase::CaseSensitive);
 	// windows drive letter path that doesn't start with base dir
-	bCannotBeStandardized |= (InPath[1] == ':' && !WithSlashes.StartsWith(RootDirectory));
+	bCannotBeStandardized |= ((InPath.Len() > 1) && (InPath[1] == ':') && !WithSlashes.StartsWith(RootDirectory));
 	// Unix style absolute path that doesn't start with base dir
 	bCannotBeStandardized |= (WithSlashes.GetCharArray()[0] == '/' && !WithSlashes.StartsWith(RootDirectory));
 
