Do not get OS memory stats each tick unless we're collecting them.

- Windows already does that since CL 2460429 by GilG, this is costly on other platforms (e.g. Linux) too.

(Merging CL 2705709 to the main branch).

diff --git a/Engine/Source/Runtime/Core/Private/GenericPlatform/GenericPlatformMemory.cpp b/Engine/Source/Runtime/Core/Private/GenericPlatform/GenericPlatformMemory.cpp
--- a/Engine/Source/Runtime/Core/Private/GenericPlatform/GenericPlatformMemory.cpp
+++ b/Engine/Source/Runtime/Core/Private/GenericPlatform/GenericPlatformMemory.cpp
@@ -135,19 +135,25 @@ uint32 FGenericPlatformMemory::GetPhysicalGBRam()
 
 void FGenericPlatformMemory::UpdateStats()
 {
-	FPlatformMemoryStats MemoryStats = FPlatformMemory::GetStats();
-
-	SET_MEMORY_STAT(STAT_TotalPhysical,MemoryStats.TotalPhysical);
-	SET_MEMORY_STAT(STAT_TotalVirtual,MemoryStats.TotalVirtual);
-	SET_MEMORY_STAT(STAT_PageSize,MemoryStats.PageSize);
-	SET_MEMORY_STAT(STAT_TotalPhysicalGB,MemoryStats.TotalPhysicalGB);
-
-	SET_MEMORY_STAT(STAT_AvailablePhysical,MemoryStats.AvailablePhysical);
-	SET_MEMORY_STAT(STAT_AvailableVirtual,MemoryStats.AvailableVirtual);
-	SET_MEMORY_STAT(STAT_UsedPhysical,MemoryStats.UsedPhysical);
-	SET_MEMORY_STAT(STAT_PeakUsedPhysical,MemoryStats.PeakUsedPhysical);
-	SET_MEMORY_STAT(STAT_UsedVirtual,MemoryStats.UsedVirtual);
-	SET_MEMORY_STAT(STAT_PeakUsedVirtual,MemoryStats.PeakUsedVirtual);
+	// avoid getting OS data (costly on Linux, Windows - see CL 2460429) if we aren't collecting stats
+#if STATS
+	if (FThreadStats::IsCollectingData(GET_STATID(STAT_TotalPhysical)))
+	{
+		FPlatformMemoryStats MemoryStats = FPlatformMemory::GetStats();
+
+		SET_MEMORY_STAT(STAT_TotalPhysical,MemoryStats.TotalPhysical);
+		SET_MEMORY_STAT(STAT_TotalVirtual,MemoryStats.TotalVirtual);
+		SET_MEMORY_STAT(STAT_PageSize,MemoryStats.PageSize);
+		SET_MEMORY_STAT(STAT_TotalPhysicalGB,MemoryStats.TotalPhysicalGB);
+
+		SET_MEMORY_STAT(STAT_AvailablePhysical,MemoryStats.AvailablePhysical);
+		SET_MEMORY_STAT(STAT_AvailableVirtual,MemoryStats.AvailableVirtual);
+		SET_MEMORY_STAT(STAT_UsedPhysical,MemoryStats.UsedPhysical);
+		SET_MEMORY_STAT(STAT_PeakUsedPhysical,MemoryStats.PeakUsedPhysical);
+		SET_MEMORY_STAT(STAT_UsedVirtual,MemoryStats.UsedVirtual);
+		SET_MEMORY_STAT(STAT_PeakUsedVirtual,MemoryStats.PeakUsedVirtual);
+	}
+#endif	// STATS
 }
 
 void* FGenericPlatformMemory::BinnedAllocFromOS( SIZE_T Size )
