Fixed some FText instances not rebuilding correctly when the culture is changed

FTextHistory has an IsOutOfDate function that FTextHistory_Base was overriding to prevent any text rebuilding from happening. This meant that the recent changes to FText::Rebuild had stopped any FTextHistory_Base instances from ever reporting themselves as dirty, and so they didn't update when the culture was changed.

This change separates the concerns of a history being "out of date" (ie, not at the head revision of the localisation manager), and "allowing a text rebuild" (now handled by FTextHistory::CanRebuildText). This fixes the issue where some texts would fail to update when the culture was changed.

diff --git a/Engine/Source/Runtime/Core/Private/Internationalization/TextHistory.cpp b/Engine/Source/Runtime/Core/Private/Internationalization/TextHistory.cpp
--- a/Engine/Source/Runtime/Core/Private/Internationalization/TextHistory.cpp
+++ b/Engine/Source/Runtime/Core/Private/Internationalization/TextHistory.cpp
@@ -29,7 +29,7 @@ FTextHistory& FTextHistory::operator=(FTextHistory&& Other)
 	return *this;
 }
 
-bool FTextHistory::IsOutOfDate()
+bool FTextHistory::IsOutOfDate() const
 {
 	return Revision < FTextLocalizationManager::Get().GetTextRevision();
 }
@@ -62,12 +62,13 @@ void FTextHistory::SerializeForDisplayString(FArchive& Ar, FTextDisplayStringPtr
 void FTextHistory::Rebuild(TSharedRef< FString, ESPMode::ThreadSafe > InDisplayString)
 {
 	const bool bIsOutOfDate = IsOutOfDate();
+	const bool bCanRebuildText = CanRebuildText();
 
 	// FTextHistory_Base will never report being out-of-date, but we need to keep the history revision in sync
 	// with the head culture regardless so that FTextSnapshot::IdenticalTo still works correctly
 	Revision = FTextLocalizationManager::Get().GetTextRevision();
 
-	if(bIsOutOfDate)
+	if(bIsOutOfDate && bCanRebuildText)
 	{
 		InDisplayString.Get() = FTextInspector::GetDisplayString(ToText(false));
 	}
diff --git a/Engine/Source/Runtime/Core/Private/Internationalization/TextHistory.h b/Engine/Source/Runtime/Core/Private/Internationalization/TextHistory.h
--- a/Engine/Source/Runtime/Core/Private/Internationalization/TextHistory.h
+++ b/Engine/Source/Runtime/Core/Private/Internationalization/TextHistory.h
@@ -41,7 +41,7 @@ public:
 	virtual void SerializeForDisplayString(FArchive& Ar, FTextDisplayStringPtr& InOutDisplayString);
 
 	/** Returns TRUE if the Revision is out of date */
-	virtual bool IsOutOfDate();
+	virtual bool IsOutOfDate() const;
 
 	/** Returns the source string managed by the history (if any). */
 	virtual const FString* GetSourceString() const;
@@ -53,6 +53,9 @@ public:
 	void Rebuild(TSharedRef< FString, ESPMode::ThreadSafe > InDisplayString);
 
 protected:
+	/** Returns true if this kind of text history is able to rebuild its text */
+	virtual bool CanRebuildText() { return true; }
+
 	/** Revision index of this history, rebuilds when it is out of sync with the FTextLocalizationManager */
 	int32 Revision;
 
@@ -77,10 +80,14 @@ public:
 	virtual FText ToText(bool bInAsSource) const override;
 	virtual void Serialize(FArchive& Ar) override;
 	virtual void SerializeForDisplayString(FArchive& Ar, FTextDisplayStringPtr& InOutDisplayString) override;
-	virtual bool IsOutOfDate() override { return false; }
 	virtual const FString* GetSourceString() const override;
 	//~ End FTextHistory Interface
 
+protected:
+	//~ Begin FTextHistory Interface
+	virtual bool CanRebuildText() { return false; }
+	//~ End FTextHistory Interface
+
 private:
 	/** Disallow copying */
 	FTextHistory_Base(const FTextHistory_Base&);
