Cleaned up some of the logic and comments in FTextHistory::Rebuild after CL# 2709156

diff --git a/Engine/Source/Runtime/Core/Private/Internationalization/TextHistory.cpp b/Engine/Source/Runtime/Core/Private/Internationalization/TextHistory.cpp
--- a/Engine/Source/Runtime/Core/Private/Internationalization/TextHistory.cpp
+++ b/Engine/Source/Runtime/Core/Private/Internationalization/TextHistory.cpp
@@ -62,15 +62,17 @@ void FTextHistory::SerializeForDisplayString(FArchive& Ar, FTextDisplayStringPtr
 void FTextHistory::Rebuild(TSharedRef< FString, ESPMode::ThreadSafe > InDisplayString)
 {
 	const bool bIsOutOfDate = IsOutOfDate();
-	const bool bCanRebuildText = CanRebuildText();
-
-	// FTextHistory_Base will never report being out-of-date, but we need to keep the history revision in sync
-	// with the head culture regardless so that FTextSnapshot::IdenticalTo still works correctly
-	Revision = FTextLocalizationManager::Get().GetTextRevision();
-
-	if(bIsOutOfDate && bCanRebuildText)
+	if(bIsOutOfDate)
 	{
-		InDisplayString.Get() = FTextInspector::GetDisplayString(ToText(false));
+		// FTextHistory_Base will never report being able to rebuild its text, but we need to keep the history 
+		// revision in sync with the head culture so that FTextSnapshot::IdenticalTo still works correctly
+		Revision = FTextLocalizationManager::Get().GetTextRevision();
+
+		const bool bCanRebuildText = CanRebuildText();
+		if(bCanRebuildText)
+		{
+			InDisplayString.Get() = FTextInspector::GetDisplayString(ToText(false));
+		}
 	}
 }
 
