VS2015 now the default toolchain for Windows
Fallback to 2013 is automatic if 2015 is not installed or a platform SDK is detected that doesn't support the new IDE
Manual override is available via -2013

diff --git a/Engine/Build/BatchFiles/GenerateProjectFiles.bat b/Engine/Build/BatchFiles/GenerateProjectFiles.bat
--- a/Engine/Build/BatchFiles/GenerateProjectFiles.bat
+++ b/Engine/Build/BatchFiles/GenerateProjectFiles.bat
@@ -40,6 +40,17 @@ if exist "%ProgramFiles(x86)%\MSBuild\12.0\bin\MSBuild.exe" (
 	goto ReadyToCompile
 )
 
+rem ## Check for Visual Studio 2015
+
+pushd %~dp0
+call GetVSComnToolsPath 14
+popd
+
+if "%VsComnToolsPath%" == "" goto NoVisualStudio2013Environment
+call "%VsComnToolsPath%/../../VC/bin/x86_amd64/vcvarsx86_amd64.bat" >NUL
+goto ReadyToCompile
+
+:NoVisualStudio2015Environment
 rem ## Check for Visual Studio 2013
 
 pushd %~dp0
@@ -109,7 +120,7 @@ goto Exit
 
 :Error_NoVisualStudioEnvironment
 echo.
-echo GenerateProjectFiles ERROR: We couldn't find a valid installation of Visual Studio.  This program requires either Visual Studio 2013 or Visual Studio 2012.  Please check that you have Visual Studio installed, then verify that the HKEY_LOCAL_MACHINE\SOFTWARE\Wow6432Node\Microsoft\VisualStudio\12.0\InstallDir (or HKEY_LOCAL_MACHINE\SOFTWARE\Microsoft\VisualStudio\12.0\InstallDir on 32-bit machines) registry value is set.  Visual Studio configures this value when it is installed, and this program expects it to be set to the '\Common7\IDE\' sub-folder under a valid Visual Studio installation directory.
+echo GenerateProjectFiles ERROR: We couldn't find a valid installation of Visual Studio.  This program requires Visual Studio 2015 or Visual Studio 2013 or Visual Studio 2012.  Please check that you have Visual Studio installed, then verify that the HKEY_LOCAL_MACHINE\SOFTWARE\Wow6432Node\Microsoft\VisualStudio\14.0\InstallDir (or HKEY_LOCAL_MACHINE\SOFTWARE\Microsoft\VisualStudio\14.0\InstallDir on 32-bit machines) registry value is set.  Visual Studio configures this value when it is installed, and this program expects it to be set to the '\Common7\IDE\' sub-folder under a valid Visual Studio installation directory.
 echo.
 pause
 goto Exit
diff --git a/Engine/Build/BatchFiles/RunUAT.bat b/Engine/Build/BatchFiles/RunUAT.bat
--- a/Engine/Build/BatchFiles/RunUAT.bat
+++ b/Engine/Build/BatchFiles/RunUAT.bat
@@ -29,7 +29,19 @@ if not exist Source\Programs\AutomationToolLauncher\AutomationToolLauncher.cspro
 rem ## Check to see if we're already running under a Visual Studio environment shell
 if not "%INCLUDE%" == "" if not "%LIB%" == "" goto ReadyToCompile
 
+rem ## Check for Visual Studio 2015
+
+pushd %~dp0
+call GetVSComnToolsPath 14
+popd
+
+if "%VsComnToolsPath%" == "" goto NoVisualStudio2015Environment
+call "%VsComnToolsPath%/../../VC/bin/x86_amd64/vcvarsx86_amd64.bat" >NUL
+goto ReadyToCompile
+
+
 rem ## Check for Visual Studio 2013
+:NoVisualStudio2015Environment
 
 pushd %~dp0
 call GetVSComnToolsPath 12
@@ -87,7 +99,7 @@ set RUNUAT_EXITCODE=1
 goto Exit_Failure
 
 :Error_NoVisualStudioEnvironment
-echo RunUAT.bat ERROR: A valid version of Visual Studio 2013 or Visual Studio 2012 does not appear to be installed.
+echo RunUAT.bat ERROR: A valid version of Visual Studio 2015 or Visual Studio 2013 or Visual Studio 2012 does not appear to be installed.
 set RUNUAT_EXITCODE=1
 goto Exit_Failure
 
diff --git a/Engine/Source/Programs/AutomationTool/AutomationUtils/CommandUtils.cs b/Engine/Source/Programs/AutomationTool/AutomationUtils/CommandUtils.cs
--- a/Engine/Source/Programs/AutomationTool/AutomationUtils/CommandUtils.cs
+++ b/Engine/Source/Programs/AutomationTool/AutomationUtils/CommandUtils.cs
@@ -2281,7 +2281,11 @@ namespace AutomationTool
 			}
 
 			string SignToolName = null;
-			if (WindowsPlatform.Compiler == WindowsCompiler.VisualStudio2013)
+			if (WindowsPlatform.Compiler == WindowsCompiler.VisualStudio2015)
+			{
+				SignToolName = "C:/Program Files (x86)/Windows Kits/8.1/bin/x86/SignTool.exe";
+			}
+			else if (WindowsPlatform.Compiler == WindowsCompiler.VisualStudio2013)
 			{
 				SignToolName = "C:/Program Files (x86)/Windows Kits/8.1/bin/x86/SignTool.exe";
 			}
@@ -2459,7 +2463,11 @@ namespace AutomationTool
 		public static void SignListFilesIfEXEOrDLL(string FilesToSign)
 		{
 			string SignToolName = null;
-			if (WindowsPlatform.Compiler == WindowsCompiler.VisualStudio2013)
+			if (WindowsPlatform.Compiler == WindowsCompiler.VisualStudio2015)
+			{
+				SignToolName = "C:/Program Files (x86)/Windows Kits/8.1/bin/x86/SignTool.exe";
+			}
+			else if (WindowsPlatform.Compiler == WindowsCompiler.VisualStudio2013)
 			{
 				SignToolName = "C:/Program Files (x86)/Windows Kits/8.1/bin/x86/SignTool.exe";
 			}
diff --git a/Engine/Source/Programs/AutomationTool/AutomationUtils/WindowsHostPlatform.cs b/Engine/Source/Programs/AutomationTool/AutomationUtils/WindowsHostPlatform.cs
--- a/Engine/Source/Programs/AutomationTool/AutomationUtils/WindowsHostPlatform.cs
+++ b/Engine/Source/Programs/AutomationTool/AutomationUtils/WindowsHostPlatform.cs
@@ -46,7 +46,17 @@ namespace AutomationTool
 			WindowsSDKDir = "";
 			BaseVSToolPath = WindowsPlatform.GetVSComnToolsPath();
 
-			if (WindowsPlatform.Compiler == WindowsCompiler.VisualStudio2013)
+			if (WindowsPlatform.Compiler == WindowsCompiler.VisualStudio2015)
+			{
+				WindowsSDKDir = FindWindowsSDKInstallationFolder( "v8.1" );
+
+				if (string.IsNullOrEmpty(BaseVSToolPath))
+				{
+					Log.TraceError("Visual Studio 2015 must be installed in order to build this target.");
+					return false;
+				}
+			}
+			else if (WindowsPlatform.Compiler == WindowsCompiler.VisualStudio2013)
 			{
 				WindowsSDKDir = FindWindowsSDKInstallationFolder( "v8.1" );
 
diff --git a/Engine/Source/Programs/UnrealBuildTool/System/ExternalExecution.cs b/Engine/Source/Programs/UnrealBuildTool/System/ExternalExecution.cs
--- a/Engine/Source/Programs/UnrealBuildTool/System/ExternalExecution.cs
+++ b/Engine/Source/Programs/UnrealBuildTool/System/ExternalExecution.cs
@@ -731,10 +731,10 @@ namespace UnrealBuildTool
 							UBTArguments.Append(" -noxge");
 						}
 
-						// Propagate command-line option to enable the newer compiler on Windows platform
-						if (UnrealBuildTool.CommandLineContains("-2015"))
+						// Propagate command-line option to switch back to old 2013 toolchain
+						if (UnrealBuildTool.CommandLineContains("-2013"))
 						{
-							UBTArguments.Append(" -2015");
+							UBTArguments.Append(" -2013");
 						}
 
 						if (RunExternalExecutable(UnrealBuildTool.GetUBTPath(), UBTArguments.ToString()) != 0)
diff --git a/Engine/Source/Programs/UnrealBuildTool/System/ProjectFileGenerator.cs b/Engine/Source/Programs/UnrealBuildTool/System/ProjectFileGenerator.cs
--- a/Engine/Source/Programs/UnrealBuildTool/System/ProjectFileGenerator.cs
+++ b/Engine/Source/Programs/UnrealBuildTool/System/ProjectFileGenerator.cs
@@ -2016,7 +2016,7 @@ namespace UnrealBuildTool
                 var AutomationToolDir = DirectoryReference.Combine(UnrealBuildTool.EngineSourceDirectory, "Programs", "AutomationTool");
                 new XDocument(
                     new XElement(NS + "Project",
-                        new XAttribute("ToolsVersion", "4.0"),
+                        new XAttribute("ToolsVersion", VCProjectFileGenerator.ProjectFileToolVersionString),
                         new XAttribute("DefaultTargets", "Build"),
                         new XElement(NS + "ItemGroup",
                             from AutomationProject in AutomationProjectFiles
diff --git a/Engine/Source/Programs/UnrealBuildTool/System/VCProjectFileGenerator.cs b/Engine/Source/Programs/UnrealBuildTool/System/VCProjectFileGenerator.cs
--- a/Engine/Source/Programs/UnrealBuildTool/System/VCProjectFileGenerator.cs
+++ b/Engine/Source/Programs/UnrealBuildTool/System/VCProjectFileGenerator.cs
@@ -259,18 +259,36 @@ namespace UnrealBuildTool
 					// Don't worry about platforms that we're missing SDKs for
 					if (BuildPlatform.HasRequiredSDKsInstalled() == SDKStatus.Valid)
 					{
-						// Make sure Visual Studio 2013 project files will work...
-						if (ProjectFileFormat == VCProjectFileFormat.VisualStudio2013)
+						// ...but only if the user didn't override this via the command-line.
+						if (!UnrealBuildTool.CommandLineContains("-2015") && !UnrealBuildTool.CommandLineContains("-2013") && !UnrealBuildTool.CommandLineContains("-2012"))
 						{
-							// ...but only if the user didn't override this via the command-line.
-							if (!UnrealBuildTool.CommandLineContains("-2013"))
+							VCProjectFileFormat ProposedFormat = ProjectFileFormat;
+
+							// Visual Studio 2015 is not supported by PS4 VSI
+							if( SupportedPlatform == UnrealTargetPlatform.PS4 )
 							{
-								// Visual Studio 2013 is not supported by Xbox One debugger add-in yet
-								if (SupportedPlatform == UnrealTargetPlatform.XboxOne)
-								{
-									Log.TraceInformation("Forcing Visual Studio 2012 projects for Xbox One compatibility (use '-2013' to override.)");
-									ProjectFileFormat = VCProjectFileFormat.VisualStudio2012;
-								}
+								Log.TraceInformation("Forcing Visual Studio max version to 2013 projects for PS4 compatibility (use '-2015' to override.)");
+								ProposedFormat = VCProjectFileFormat.VisualStudio2013;
+							}
+
+							// Visual Studio 2015 is not supported by the Android debugger we currently furnish
+							if (SupportedPlatform == UnrealTargetPlatform.Android)
+							{
+								Log.TraceInformation("Forcing Visual Studio max version to 2013 projects for Android compatibility (use '-2015' to override.)");
+								ProposedFormat = VCProjectFileFormat.VisualStudio2013;
+							}
+
+							// Visual Studio 2015 is not supported by Xbox One debugger add-in yet
+							if (SupportedPlatform == UnrealTargetPlatform.XboxOne)
+							{
+								Log.TraceInformation("Forcing Visual Studio max version to 2012 projects for Xbox One compatibility (use '-2015' to override.)");
+								ProposedFormat = VCProjectFileFormat.VisualStudio2012;
+							}
+
+							// Reduce the Visual Studio version to the max supported by each platform we plan to include.
+							if (ProposedFormat < ProjectFileFormat)
+							{
+								ProjectFileFormat = ProposedFormat;
 							}
 						}
 					}
diff --git a/Engine/Source/Programs/UnrealBuildTool/Windows/UEBuildWindows.cs b/Engine/Source/Programs/UnrealBuildTool/Windows/UEBuildWindows.cs
--- a/Engine/Source/Programs/UnrealBuildTool/Windows/UEBuildWindows.cs
+++ b/Engine/Source/Programs/UnrealBuildTool/Windows/UEBuildWindows.cs
@@ -58,20 +58,18 @@ namespace UnrealBuildTool
 				}
 
 				// Second, default based on what's installed, from newest to oldest
-				// @todo UWP: Currently we default to VS 2013, even if VS 2015 is installed, until the engine is fully updated to
-				// support compiling with VS 2015.  You can override this with the "-2015" command-line switch.
-				// 				else if (!String.IsNullOrEmpty(WindowsPlatform.GetVSComnToolsPath(WindowsCompiler.VisualStudio2015)))
-				// 				{
-				// 					CachedCompiler = WindowsCompiler.VisualStudio2015;
-				// 				}
+ 				else if (!String.IsNullOrEmpty(WindowsPlatform.GetVSComnToolsPath(WindowsCompiler.VisualStudio2015)))
+ 				{
+ 					CachedCompiler = WindowsCompiler.VisualStudio2015;
+ 				}
 				else if (!String.IsNullOrEmpty(WindowsPlatform.GetVSComnToolsPath(WindowsCompiler.VisualStudio2013)))
 				{
 					CachedCompiler = WindowsCompiler.VisualStudio2013;
 				}
 				else
 				{
-					// Finally assume 2013 is installed to defer errors somewhere else like VCToolChain
-					CachedCompiler = WindowsCompiler.VisualStudio2013;
+					// Finally assume 2015 is installed to defer errors somewhere else like VCToolChain
+					CachedCompiler = WindowsCompiler.VisualStudio2015;
 				}
 
 				return CachedCompiler.Value;
diff --git a/Engine/Source/Runtime/Core/Private/ProfilingDebugging/Windows/VSPerfExternalProfiler.cpp b/Engine/Source/Runtime/Core/Private/ProfilingDebugging/Windows/VSPerfExternalProfiler.cpp
--- a/Engine/Source/Runtime/Core/Private/ProfilingDebugging/Windows/VSPerfExternalProfiler.cpp
+++ b/Engine/Source/Runtime/Core/Private/ProfilingDebugging/Windows/VSPerfExternalProfiler.cpp
@@ -88,8 +88,10 @@ public:
 		// Try to load the VSPerf DLL
 		// NOTE: VSPerfXXX.dll is installed into /Windows/System32 when Visual Studio is installed.  The XXX is the version number of
 		// Visual Studio.  For example, for Visual Studio 2013, the file name is VSPerf120.dll.
-#if _MSC_VER >= 1800
-		DLLHandle = FPlatformProcess::GetDllHandle( TEXT( "VSPerf120.dll" ) );	// Visual Studio 2013
+#if _MSC_VER >= 1900
+		DLLHandle = FPlatformProcess::GetDllHandle(TEXT("VSPerf140.dll"));	// Visual Studio 2015
+#elif _MSC_VER >= 1800
+		DLLHandle = FPlatformProcess::GetDllHandle(TEXT("VSPerf120.dll"));	// Visual Studio 2013
 #elif _MSC_VER >= 1700
 		DLLHandle = FPlatformProcess::GetDllHandle( TEXT( "VSPerf110.dll" ) );	// Visual Studio 2012
 #else
diff --git a/Engine/Source/Runtime/Core/Public/Windows/WindowsPlatformCompilerSetup.h b/Engine/Source/Runtime/Core/Public/Windows/WindowsPlatformCompilerSetup.h
--- a/Engine/Source/Runtime/Core/Public/Windows/WindowsPlatformCompilerSetup.h
+++ b/Engine/Source/Runtime/Core/Public/Windows/WindowsPlatformCompilerSetup.h
@@ -25,7 +25,7 @@ static_assert(_MSC_VER >= 1800, "Visual Studio 2013 or later is required to comp
 #pragma warning (error: 4100 4101 4102 4103 4104 4105 4106 4107 4108 4109 4110 4111 4112 4113 4114 4115 4116 4117 4118 4119 4120 4121 4122 4123 4124 4125 4126 4127 4128 4129 4130 4131 4132 4133 4134 4135 4136 4137 4138 4139 4140 4141 4142 4143 4144 4145 4146 4147 4148 4149 4150 4151 4152 4153 4154 4155 4156 4157 4158 4159 4160 4161 4162 4163 4164      4166 4167 4168 4169 4170 4171 4172 4173 4174 4175 4176 4177 4178 4179 4180 4181 4182 4183 4184 4185 4186 4187 4188 4189 4190 4191 4192 4193 4194 4195 4196 4197 4198 4199)
 #pragma warning (error: 4200 4201 4202 4203 4204 4205 4206 4207 4208 4209 4210 4211 4212 4213 4214 4215 4216 4217 4218 4219 4220 4221 4222 4223 4224 4225 4226 4227 4228 4229 4230 4231 4232 4233 4234 4235 4236 4237 4238 4239 4240 4241      4243 4244 4245 4246 4247 4248 4249 4250 4251 4252 4253      4255 4256 4257 4258 4259 4260 4261 4262 4263 4264           4267 4268 4269 4270 4271 4272 4273 4274 4275 4276 4277 4278 4279 4280 4281 4282 4283 4284 4285 4286 4287 4288 4289 4290 4291 4292 4293 4294 4295      4297 4298 4299)
 #pragma warning (error: 4300 4301 4302 4303 4304      4306 4307 4308 4309 4310           4313 4314 4315 4316 4317 4318 4319 4320 4321 4322 4323 4324 4325 4326 4327 4328 4329 4330 4331 4332 4333 4334 4335 4336 4337 4338      4340 4341      4343 4344      4346 4347 4348 4349           4352 4353 4354 4355 4356 4357 4358 4359 4360 4361 4362 4363 4364      4366 4367 4368 4369           4372      4374 4375 4376 4377 4378 4379 4380 4381 4382 4383 4384 4385 4386 4387      4389 4390 4391 4392 4393 4394 4395 4396 4397 4398 4399)
-#pragma warning (error: 4400 4401 4402 4403 4404 4405 4406 4407 4408 4409 4410 4411      4413 4414 4415 4416 4417 4418 4419 4420 4421 4422 4423 4424 4425 4426 4427 4428 4429 4430 4431 4432 4433 4434      4436      4438 4439 4440 4441 4442 4443      4445 4446 4447 4448 4449 4450 4451 4452 4453 4454 4455      4457           4460 4461 4462      4464 4465 4466 4467 4468 4469 4470           4473 4474 4475 4476 4477 4478 4479 4480      4482 4483 4484 4485 4486 4487 4488 4489 4490 4491 4492 4493 4494 4495 4496 4497 4498 4499)
+#pragma warning (error: 4400 4401 4402 4403 4404 4405 4406 4407 4408 4409 4410 4411      4413 4414 4415 4416 4417 4418 4419 4420 4421 4422 4423 4424 4425 4426 4427 4428 4429 4430 4431 4432 4433 4434      4436      4438 4439 4440 4441 4442 4443      4445 4446 4447 4448 4449 4450 4451 4452 4453 4454 4455                     4460 4461 4462      4464 4465 4466 4467 4468 4469 4470           4473 4474 4475 4476 4477 4478 4479 4480      4482 4483 4484 4485 4486 4487 4488 4489 4490 4491 4492 4493 4494 4495 4496 4497 4498 4499)
 #pragma warning (error: 4500 4501 4502 4503 4504 4505 4506 4507 4508 4509      4511 4512 4513 4514 4515 4516 4517 4518 4519 4520 4521 4522 4523 4524 4525 4526 4527 4528 4529 4530 4531 4532 4533 4534 4535 4536 4537 4538 4539 4540 4541 4542 4543 4544 4545 4546                4550 4551 4552 4553 4554      4556 4557 4558 4559 4560 4561 4562 4563 4564 4565 4566 4567 4568 4569 4570      4572 4573      4575 4576 4577 4578 4579 4580 4581 4582 4583 4584 4585 4586 4587 4588 4589 4590 4591 4592 4593 4594 4595 4596 4597 4598 4599)
 #pragma warning (error: 4600 4601 4602 4603 4604 4605 4606 4607      4609 4610 4611 4612 4613 4614 4615 4616 4617 4618      4620 4621 4622 4623 4624 4625 4626 4627 4628 4629 4630 4631 4632 4633 4634 4635 4636 4637 4638 4639 4640 4641 4642 4643 4644 4645 4646 4647 4648 4649 4650      4652 4653 4654 4655 4656 4657 4658 4659 4660 4661 4662 4663 4664 4665 4666 4667 4668 4669 4670 4671 4672 4673 4674 4675 4676 4677 4678 4679 4680 4681 4682 4683 4684 4685 4686 4687 4688 4689 4690 4691      4693 4694 4695 4696 4697 4698 4699)
 #pragma warning (error: 4700      4702 4703 4704 4705 4706 4707 4708 4709 4710 4711 4712 4713 4714 4715 4716 4717 4718 4719 4720 4721 4722 4723 4724 4725 4726 4727 4728 4729      4731 4732 4733 4734 4735 4736 4737      4739 4740 4741 4742 4743 4744 4745 4746 4747 4748 4749 4750 4751 4752 4753 4754 4755 4756 4757 4758 4759 4760 4761 4762 4763 4764 4765 4766      4768 4769 4770 4771 4772 4773      4775 4776 4777 4778 4779 4780 4781 4782 4783 4784 4785 4786 4787 4788 4789 4790 4791 4792 4793 4794 4795 4796 4797 4798 4799)
@@ -64,6 +64,7 @@ static_assert(_MSC_VER >= 1800, "Visual Studio 2013 or later is required to comp
 
 // Shadow variable declaration warnings. These should eventually be fixed up and reenabled. Only effect VS2015.
 // 4456 - declaration of 'LocalVariable' hides previous local declaration
+// 4457 - declaration of 'LocalVariable' hides function parameter
 // 4458 - declaration of 'parameter' hides class member
 // 4459 - declaration of 'LocalVariable' hides global declaration
 
diff --git a/Engine/Source/Runtime/Engine/Private/VorbisAudioInfo.cpp b/Engine/Source/Runtime/Engine/Private/VorbisAudioInfo.cpp
--- a/Engine/Source/Runtime/Engine/Private/VorbisAudioInfo.cpp
+++ b/Engine/Source/Runtime/Engine/Private/VorbisAudioInfo.cpp
@@ -325,6 +325,9 @@ void LoadVorbisLibraries()
 		FString VSVersion = TEXT("VS2013/");
 		FString PlatformString = TEXT("Win32");
 		FString DLLNameStub = TEXT(".dll");
+#if _MSC_VER == 1900
+		VSVersion = TEXT("VS2015/");
+#endif
 #if PLATFORM_64BITS
 		PlatformString = TEXT("Win64");
 		DLLNameStub = TEXT("_64.dll");
