- New device types (IPodTouch6, IPhone6S, IPhone6SPlus, IPadMini4) [UE-20919]

diff --git a/Engine/Config/BaseDeviceProfiles.ini b/Engine/Config/BaseDeviceProfiles.ini
--- a/Engine/Config/BaseDeviceProfiles.ini
+++ b/Engine/Config/BaseDeviceProfiles.ini
@@ -9,13 +9,17 @@
 +DeviceProfileNameAndTypes=iPadAir,IOS
 +DeviceProfileNameAndTypes=iPadMini,IOS
 +DeviceProfileNameAndTypes=iPadMini2,IOS
++DeviceProfileNameAndTypes=iPadMini4,IOS
++DeviceProfileNameAndTypes=iPodTouch5,IOS
++DeviceProfileNameAndTypes=iPodTouch6,IOS
 +DeviceProfileNameAndTypes=iPhone4,IOS
 +DeviceProfileNameAndTypes=iPhone4S,IOS
 +DeviceProfileNameAndTypes=iPhone5,IOS
 +DeviceProfileNameAndTypes=iPhone5S,IOS
-+DeviceProfileNameAndTypes=iPodTouch5,IOS
 +DeviceProfileNameAndTypes=iPhone6,IOS
 +DeviceProfileNameAndTypes=iPhone6Plus,IOS
++DeviceProfileNameAndTypes=iPhone6S,IOS
++DeviceProfileNameAndTypes=iPhone6SPlus,IOS
 +DeviceProfileNameAndTypes=Android,Android
 +DeviceProfileNameAndTypes=PS4,PS4
 +DeviceProfileNameAndTypes=XboxOne,XboxOne
@@ -126,6 +130,10 @@ BaseProfileName=IOS
 DeviceType=IOS
 BaseProfileName=iPadAir
 
+[iPadMini4 DeviceProfile]
+DeviceType=IOS
+BaseProfileName=iPadMini2
+
 [iPhone4 DeviceProfile]
 DeviceType=IOS
 BaseProfileName=IOS
@@ -155,6 +163,10 @@ DeviceType=IOS
 BaseProfileName=IOS
 +CVars=r.RenderTargetSwitchWorkaround=1
 
+[iPodTouch6 DeviceProfile]
+DeviceType=IOS
+BaseProfileName=iPhone6
+
 [iPhone6 DeviceProfile]
 DeviceType=IOS
 BaseProfileName=IOS
@@ -171,6 +183,23 @@ BaseProfileName=IOS
 +CVars=r.LightShaftQuality=1
 +CVars=r.RefractionQuality=1
 
+[iPhone6S DeviceProfile]
+DeviceType=IOS
+BaseProfileName=iPhone6
++CVars=r.BloomQuality=1
++CVars=r.DepthOfFieldQuality=1
++CVars=r.LightShaftQuality=1
++CVars=r.RefractionQuality=1
+
+[iPhone6SPlus DeviceProfile]
+DeviceType=IOS
+BaseProfileName=iPhone6Plus
++CVars=r.BloomQuality=1
++CVars=r.DepthOfFieldQuality=1
++CVars=r.LightShaftQuality=1
++CVars=r.RefractionQuality=1
+
+
 [/Script/AndroidDeviceProfileSelector.AndroidDeviceProfileMatchingRules]
 MatchProfile=(Profile="Android_PowerVR540",Match=((SourceType=SRC_GpuFamily,CompareType=CMP_Equal,MatchString="PowerVR SGX 540")))
 +MatchProfile=(Profile="Android_Adreno2xx",Match=((SourceType=SRC_GpuFamily,CompareType=CMP_Regex,MatchString="Adreno \\(TM\\) 2[0-9][0-9]")))
diff --git a/Engine/Source/Runtime/Core/Private/IOS/IOSPlatformMisc.cpp b/Engine/Source/Runtime/Core/Private/IOS/IOSPlatformMisc.cpp
--- a/Engine/Source/Runtime/Core/Private/IOS/IOSPlatformMisc.cpp
+++ b/Engine/Source/Runtime/Core/Private/IOS/IOSPlatformMisc.cpp
@@ -378,10 +378,14 @@ FIOSPlatformMisc::EIOSDevice FIOSPlatformMisc::GetIOSDeviceType()
 		// get major revision number
 		int Major = [DeviceIDString characterAtIndex:4] - '0';
 
-		if (Major >= 5)
+		if (Major == 5)
 		{
 			DeviceType = IOS_IPodTouch5;
 		}
+		else if (Major >= 7)
+		{
+			DeviceType = IOS_IPodTouch6;
+		}
 	}
 	// iPads
 	else if ([DeviceIDString hasPrefix:@"iPad"])
@@ -428,10 +432,17 @@ FIOSPlatformMisc::EIOSDevice FIOSPlatformMisc::GetIOSDeviceType()
 				DeviceType = IOS_IPadAir;
 			}
 		}
-		// iPad Air 2
+		// iPad Air 2 and iPadMini 4
 		else if (Major == 5)
 		{
-			DeviceType = IOS_IPadAir2;
+			if (Minor == 1 || Minor == 2)
+			{
+				DeviceType = IOS_IPadMini4;
+			}
+			else
+			{
+				DeviceType = IOS_IPadAir2;
+			}
 		}
 		// Default to highest settings currently available for any future device
 		else if (Major > 5)
@@ -443,6 +454,7 @@ FIOSPlatformMisc::EIOSDevice FIOSPlatformMisc::GetIOSDeviceType()
 	else if ([DeviceIDString hasPrefix:@"iPhone"])
 	{
 		int Major = [DeviceIDString characterAtIndex:6] - '0';
+		int Minor = [DeviceIDString characterAtIndex:8] - '0';
 
 		if (Major == 3)
 		{
@@ -460,17 +472,40 @@ FIOSPlatformMisc::EIOSDevice FIOSPlatformMisc::GetIOSDeviceType()
 		{
 			DeviceType = IOS_IPhone5S;
 		}
-		else if (Major >= 7)
+		else if (Major == 7)
+		{
+			if (Minor == 1)
+			{
+				DeviceType == IOS_IPhone6Plus;
+			}
+			else if (Minor == 2)
+			{
+				DeviceType == IOS_IPhone6;
+			}
+		}
+		else if (Major == 8)
+		{
+			// note that Apple switched the minor order around between 6 and 6S (gotta keep us on our toes!)
+			if (Minor == 1)
+			{
+				DeviceType == IOS_IPhone6;
+			}
+			else if (Minor == 2)
+			{
+				DeviceType == IOS_IPhone6Plus;
+			}
+		}
+		else if (Major >= 9)
 		{
-			// this could just check the minor where 1 == 6Plus, and 2 == 6, but that won't help going forward (7/7+?)
-			// so treat devices with a scale > 2.5 to be 6Plus type devices, < 2.5 to be 6 type devices
+			// for going forward into unknown devices (like 7/7+?), we can't use Minor,
+			// so treat devices with a scale > 2.5 to be 6SPlus type devices, < 2.5 to be 6S type devices
 			if ([UIScreen mainScreen].scale > 2.5f)
 			{
-				DeviceType = IOS_IPhone6Plus;
+				DeviceType = IOS_IPhone6SPlus;
 			}
 			else
 			{
-				DeviceType = IOS_IPhone6;
+				DeviceType = IOS_IPhone6S;
 			}
 		}
 	}
diff --git a/Engine/Source/Runtime/Core/Public/IOS/IOSPlatformMisc.h b/Engine/Source/Runtime/Core/Public/IOS/IOSPlatformMisc.h
--- a/Engine/Source/Runtime/Core/Public/IOS/IOSPlatformMisc.h
+++ b/Engine/Source/Runtime/Core/Public/IOS/IOSPlatformMisc.h
@@ -118,15 +118,19 @@ struct CORE_API FIOSPlatformMisc : public FGenericPlatformMisc
 		IOS_IPhone5, // also the IPhone5C
 		IOS_IPhone5S,
 		IOS_IPodTouch5,
+		IOS_IPodTouch6,
 		IOS_IPad2,
 		IOS_IPad3,
 		IOS_IPad4,
 		IOS_IPadMini,
-		IOS_IPadMini2,
+		IOS_IPadMini2, // also the iPadMini3
+		IOS_IPadMini4,
 		IOS_IPadAir,
 		IOS_IPadAir2,
 		IOS_IPhone6,
 		IOS_IPhone6Plus,
+		IOS_IPhone6S,
+		IOS_IPhone6SPlus,
 		IOS_Unknown,
 	};
 
@@ -141,15 +145,19 @@ struct CORE_API FIOSPlatformMisc : public FGenericPlatformMisc
 			L"IPhone5",
 			L"IPhone5S",
 			L"IPodTouch5",
+			L"IPodTouch6",
 			L"IPad2",
 			L"IPad3",
 			L"IPad4",
 			L"IPadMini",
 			L"IPadMini2",
+			L"IPadMini5",
 			L"IPadAir",
 			L"IPadAir2",
 			L"IPhone6",
 			L"IPhone6Plus",
+			L"IPhone6S",
+			L"IPhone6SPlus",
 			L"Unknown",
 		};
 		static_assert((sizeof(IOSDeviceNames) / sizeof(IOSDeviceNames[0])) == ((int32)IOS_Unknown + 1), "Mismatched IOSDeviceNames and EIOSDevice.");
