Linux: fix custom cursors.

- PR #1582: Locked Custom Cursor (Widget) on Linux Fix contributed by [REDACTED_PERSON].
- From PR description: "When using a custom cursor (widget) on Linux the cursor was permanently stuck to the center of the screen and was treating the cursor as if it was hidden rather than replaced."

diff --git a/Engine/Source/Runtime/Core/Private/Linux/LinuxCursor.cpp b/Engine/Source/Runtime/Core/Private/Linux/LinuxCursor.cpp
--- a/Engine/Source/Runtime/Core/Private/Linux/LinuxCursor.cpp
+++ b/Engine/Source/Runtime/Core/Private/Linux/LinuxCursor.cpp
@@ -227,9 +227,12 @@ void FLinuxCursor::SetType( const EMouseCursor::Type InNewCursor )
 {	
 	checkf( InNewCursor < EMouseCursor::TotalCursorCount, TEXT("Invalid cursor(%d) supplied"), InNewCursor );
 	CurrentType = InNewCursor;
-	if(CursorHandles[InNewCursor] == NULL)
+	if(CursorHandles[InNewCursor] == nullptr)
 	{
-		bHidden = true;
+		if (InNewCursor != EMouseCursor::Custom)
+		{
+			bHidden = true;
+		}
 		SDL_ShowCursor(SDL_DISABLE);
 		SDL_SetCursor(CursorHandles[0]);
 	}
diff --git a/Engine/Source/Runtime/Slate/Private/Framework/Application/SlateApplication.cpp b/Engine/Source/Runtime/Slate/Private/Framework/Application/SlateApplication.cpp
--- a/Engine/Source/Runtime/Slate/Private/Framework/Application/SlateApplication.cpp
+++ b/Engine/Source/Runtime/Slate/Private/Framework/Application/SlateApplication.cpp
@@ -2855,7 +2855,7 @@ void FSlateApplication::ProcessCursorReply(const FCursorReply& CursorReply)
 		{
 			CursorReply.GetCursorWidget()->SetVisibility(EVisibility::HitTestInvisible);
 			CursorWindowPtr = CursorReply.GetCursorWindow();
-			PlatformApplication->Cursor->SetType(EMouseCursor::None);
+			PlatformApplication->Cursor->SetType(EMouseCursor::Custom);
 		}
 		else
 		{
