Linux: Respect -usehyperthreading command-line switch

- PR #1594: [Linux] Respect -usehyperthreading command-line switch like other platforms do (Contributed by [REDACTED_PERSON])

diff --git a/Engine/Source/Runtime/Core/Private/Linux/LinuxPlatformMisc.cpp b/Engine/Source/Runtime/Core/Private/Linux/LinuxPlatformMisc.cpp
--- a/Engine/Source/Runtime/Core/Private/Linux/LinuxPlatformMisc.cpp
+++ b/Engine/Source/Runtime/Core/Private/Linux/LinuxPlatformMisc.cpp
@@ -603,84 +603,91 @@ int32 FLinuxPlatformMisc::NumberOfCores()
 {
 	// WARNING: this function ignores edge cases like affinity mask changes (and even more fringe cases like CPUs going offline)
 	// in the name of performance (higher level code calls NumberOfCores() way too often...)
-	static int32 NumCoreIds = 0;
-	if (NumCoreIds == 0)
+	static int32 NumberOfCores = 0;
+	if (NumberOfCores == 0)
 	{
-		cpu_set_t AvailableCpusMask;
-		CPU_ZERO(&AvailableCpusMask);
-
-		if (0 != sched_getaffinity(0, sizeof(AvailableCpusMask), &AvailableCpusMask))
+		if (FParse::Param(FCommandLine::Get(), TEXT("usehyperthreading")))
 		{
-			NumCoreIds = 1;	// we are running on something, right?
+			NumberOfCores = NumberOfCoresIncludingHyperthreads();
 		}
 		else
 		{
-			char FileNameBuffer[1024];
-			struct CpuInfo
+			cpu_set_t AvailableCpusMask;
+			CPU_ZERO(&AvailableCpusMask);
+
+			if (0 != sched_getaffinity(0, sizeof(AvailableCpusMask), &AvailableCpusMask))
 			{
-				int Core;
-				int Package;
+				NumberOfCores = 1;	// we are running on something, right?
 			}
-			CpuInfos[CPU_SETSIZE];
+			else
+			{
+				char FileNameBuffer[1024];
+				struct CpuInfo
+				{
+					int Core;
+					int Package;
+				}
+				CpuInfos[CPU_SETSIZE];
 
-			FMemory::Memzero(CpuInfos);
-			int MaxCoreId = 0;
-			int MaxPackageId = 0;
+				FMemory::Memzero(CpuInfos);
+				int MaxCoreId = 0;
+				int MaxPackageId = 0;
 
-			for(int32 CpuIdx = 0; CpuIdx < CPU_SETSIZE; ++CpuIdx)
-			{
-				if (CPU_ISSET(CpuIdx, &AvailableCpusMask))
+				for(int32 CpuIdx = 0; CpuIdx < CPU_SETSIZE; ++CpuIdx)
 				{
-					sprintf(FileNameBuffer, "/sys/devices/system/cpu/cpu%d/topology/core_id", CpuIdx);
-					
-					if (FILE* CoreIdFile = fopen(FileNameBuffer, "r"))
+					if (CPU_ISSET(CpuIdx, &AvailableCpusMask))
 					{
-						if (1 != fscanf(CoreIdFile, "%d", &CpuInfos[CpuIdx].Core))
+						sprintf(FileNameBuffer, "/sys/devices/system/cpu/cpu%d/topology/core_id", CpuIdx);
+
+						if (FILE* CoreIdFile = fopen(FileNameBuffer, "r"))
 						{
-							CpuInfos[CpuIdx].Core = 0;
+							if (1 != fscanf(CoreIdFile, "%d", &CpuInfos[CpuIdx].Core))
+							{
+								CpuInfos[CpuIdx].Core = 0;
+							}
+							fclose(CoreIdFile);
 						}
-						fclose(CoreIdFile);
-					}
 
-					sprintf(FileNameBuffer, "/sys/devices/system/cpu/cpu%d/topology/physical_package_id", CpuIdx);
+						sprintf(FileNameBuffer, "/sys/devices/system/cpu/cpu%d/topology/physical_package_id", CpuIdx);
 
-					unsigned int PackageId = 0;
-					if (FILE* PackageIdFile = fopen(FileNameBuffer, "r"))
-					{
-						if (1 != fscanf(PackageIdFile, "%d", &CpuInfos[CpuIdx].Package))
+						unsigned int PackageId = 0;
+						if (FILE* PackageIdFile = fopen(FileNameBuffer, "r"))
 						{
-							CpuInfos[CpuIdx].Package = 0;
+							if (1 != fscanf(PackageIdFile, "%d", &CpuInfos[CpuIdx].Package))
+							{
+								CpuInfos[CpuIdx].Package = 0;
+							}
+							fclose(PackageIdFile);
 						}
-						fclose(PackageIdFile);
-					}
 
-					MaxCoreId = FMath::Max(MaxCoreId, CpuInfos[CpuIdx].Core);
-					MaxPackageId = FMath::Max(MaxPackageId, CpuInfos[CpuIdx].Package);
+						MaxCoreId = FMath::Max(MaxCoreId, CpuInfos[CpuIdx].Core);
+						MaxPackageId = FMath::Max(MaxPackageId, CpuInfos[CpuIdx].Package);
+					}
 				}
-			}
 
-			int NumCores = MaxCoreId + 1;
-			int NumPackages = MaxPackageId + 1;
-			int NumPairs = NumPackages * NumCores;
-			unsigned char * Pairs = reinterpret_cast<unsigned char *>(FMemory_Alloca(NumPairs * sizeof(unsigned char)));
-			FMemory::Memzero(Pairs, NumPairs * sizeof(unsigned char));
+				int NumCores = MaxCoreId + 1;
+				int NumPackages = MaxPackageId + 1;
+				int NumPairs = NumPackages * NumCores;
+				unsigned char * Pairs = reinterpret_cast<unsigned char *>(FMemory_Alloca(NumPairs * sizeof(unsigned char)));
+				FMemory::Memzero(Pairs, NumPairs * sizeof(unsigned char));
 
-			for (int32 CpuIdx = 0; CpuIdx < CPU_SETSIZE; ++CpuIdx)
-			{
-				if (CPU_ISSET(CpuIdx, &AvailableCpusMask))
+				for (int32 CpuIdx = 0; CpuIdx < CPU_SETSIZE; ++CpuIdx)
 				{
-					Pairs[CpuInfos[CpuIdx].Package * NumCores + CpuInfos[CpuIdx].Core] = 1;
+					if (CPU_ISSET(CpuIdx, &AvailableCpusMask))
+					{
+						Pairs[CpuInfos[CpuIdx].Package * NumCores + CpuInfos[CpuIdx].Core] = 1;
+					}
 				}
-			}
 
-			for (int32 Idx = 0; Idx < NumPairs; ++Idx)
-			{
-				NumCoreIds += Pairs[Idx];
+				for (int32 Idx = 0; Idx < NumPairs; ++Idx)
+				{
+					NumberOfCores += Pairs[Idx];
+				}
 			}
 		}
 	}
 
-	return NumCoreIds;
+	return NumberOfCores;
 }
 
 int32 FLinuxPlatformMisc::NumberOfCoresIncludingHyperthreads()
