Linux: improvements reporting ensure and callstacks (UE-21322).

- The bulk of this change is PR #1583: [Linux] Show stacktrace on ensure (Contributed by [REDACTED_PERSON]).
- Other misc changes:
-   Do not limit strings to 1024 when printing a log line to stdout - allows to print the callstack in full to console.
-   Do not crash when hitting ensure() in debugger in a program that did not initialize SDL.

diff --git a/Engine/Source/Programs/TestPAL/Private/Main.cpp b/Engine/Source/Programs/TestPAL/Private/Main.cpp
--- a/Engine/Source/Programs/TestPAL/Private/Main.cpp
+++ b/Engine/Source/Programs/TestPAL/Private/Main.cpp
@@ -264,6 +264,22 @@ int32 CrashTest(const TCHAR* CommandLine)
 	GEngineLoop.PreInit(CommandLine);
 	UE_LOG(LogTestPAL, Display, TEXT("Running crash test (this should not exit)."));
 
+	// try ensures first (each ensure fires just once)
+	{
+		for (int IdxEnsure = 0; IdxEnsure < 5; ++IdxEnsure)
+		{
+			FScopeLogTime EnsureLogTime(*FString::Printf(TEXT("Handled FIRST ensure() #%d times"), IdxEnsure), nullptr, FScopeLogTime::ScopeLog_Seconds);
+			ensure(false);
+		}
+	}
+	{
+		for (int IdxEnsure = 0; IdxEnsure < 5; ++IdxEnsure)
+		{
+			FScopeLogTime EnsureLogTime(*FString::Printf(TEXT("Handled SECOND ensure() #%d times"), IdxEnsure), nullptr, FScopeLogTime::ScopeLog_Seconds);
+			ensure(false);
+		}
+	}
+
 	if (FParse::Param(CommandLine, TEXT("logfatal")))
 	{
 		UE_LOG(LogTestPAL, Fatal, TEXT("  Opa!"));
diff --git a/Engine/Source/Runtime/Core/Private/Linux/LinuxPlatformCrashContext.cpp b/Engine/Source/Runtime/Core/Private/Linux/LinuxPlatformCrashContext.cpp
--- a/Engine/Source/Runtime/Core/Private/Linux/LinuxPlatformCrashContext.cpp
+++ b/Engine/Source/Runtime/Core/Private/Linux/LinuxPlatformCrashContext.cpp
@@ -8,12 +8,6 @@
 #include <sys/utsname.h>	// for uname()
 #include <signal.h>
 
-// these are not actually system headers, but a TPS library (see ThirdParty/elftoolchain)
-#include <libelf.h>
-#include <_libelf.h>
-#include <libdwarf.h>
-#include <dwarf.h>
-
 
 FString DescribeSignal(int32 Signal, siginfo_t* Info)
 {
@@ -23,6 +17,9 @@ FString DescribeSignal(int32 Signal, siginfo_t* Info)
 
 	switch (Signal)
 	{
+	case 0:
+		// No signal - used for initialization stacktrace on non-fatal errors (ex: ensure)
+		break;
 	case SIGSEGV:
 		ErrorString += TEXT("SIGSEGV: invalid attempt to access memory at address ");
 		ErrorString += FString::Printf(TEXT("0x%08x"), (uint32*)Info->si_addr);
@@ -60,25 +57,6 @@ FLinuxCrashContext::~FLinuxCrashContext()
 		free(BacktraceSymbols);
 		BacktraceSymbols = NULL;
 	}
-
-	if (DebugInfo)
-	{
-		Dwarf_Error ErrorInfo;
-		dwarf_finish(DebugInfo, &ErrorInfo);
-		DebugInfo = NULL;
-	}
-
-	if (ElfHdr)
-	{
-		elf_end(ElfHdr);
-		ElfHdr = NULL;
-	}
-
-	if (ExeFd >= 0)
-	{
-		close(ExeFd);
-		ExeFd = -1;
-	}
 }
 
 void FLinuxCrashContext::InitFromSignal(int32 InSignal, siginfo_t* InInfo, void* InContext)
@@ -87,333 +65,9 @@ void FLinuxCrashContext::InitFromSignal(int32 InSignal, siginfo_t* InInfo, void*
 	Info = InInfo;
 	Context = reinterpret_cast< ucontext_t* >( InContext );
 
-	// open ourselves for examination
-	if (!FParse::Param( FCommandLine::Get(), TEXT(CMDARG_SUPPRESS_DWARF_PARSING)))
-	{
-		ExeFd = open("/proc/self/exe", O_RDONLY);
-		if (ExeFd >= 0)
-		{
-			Dwarf_Error ErrorInfo;
-			// allocate DWARF debug descriptor
-			if (dwarf_init(ExeFd, DW_DLC_READ, NULL, NULL, &DebugInfo, &ErrorInfo) == DW_DLV_OK)
-			{
-				// get ELF descritor
-				if (dwarf_get_elf(DebugInfo, &ElfHdr, &ErrorInfo) != DW_DLV_OK)
-				{
-					dwarf_finish(DebugInfo, &ErrorInfo);
-					DebugInfo = NULL;
-
-					close(ExeFd);
-					ExeFd = -1;
-				}
-			}
-			else
-			{
-				DebugInfo = NULL;
-				close(ExeFd);
-				ExeFd = -1;
-			}
-		}
-	}
-
 	FCString::Strcat(SignalDescription, ARRAY_COUNT( SignalDescription ) - 1, *DescribeSignal(Signal, Info));
 }
 
-/**
- * Finds a function name in DWARF DIE (Debug Information Entry).
- * For more info on DWARF format, see http://www.dwarfstd.org/Download.php , http://www.ibm.com/developerworks/library/os-debugging/
- *
- * @return true if we need to stop search (i.e. either found it or some error happened)
- */
-bool FindFunctionNameInDIE(Dwarf_Debug DebugInfo, Dwarf_Die Die, Dwarf_Addr Addr, const char **OutFuncName)
-{
-	Dwarf_Error ErrorInfo;
-	Dwarf_Half Tag;
-	Dwarf_Unsigned LowerPC, HigherPC;
-	char *TempFuncName;
-	int ReturnCode;
-
-	if (dwarf_tag(Die, &Tag, &ErrorInfo) != DW_DLV_OK || Tag != DW_TAG_subprogram ||
-		dwarf_attrval_unsigned(Die, DW_AT_low_pc, &LowerPC, &ErrorInfo) != DW_DLV_OK ||
-		dwarf_attrval_unsigned(Die, DW_AT_high_pc, &HigherPC, &ErrorInfo) != DW_DLV_OK ||
-		Addr < LowerPC || HigherPC <= Addr
-		) 
-	{
-		return false;
-	}
-	
-	// found it
-	*OutFuncName = NULL;
-	Dwarf_Attribute SubAt;
-	ReturnCode = dwarf_attr(Die, DW_AT_name, &SubAt, &ErrorInfo);
-	if (ReturnCode == DW_DLV_ERROR)
-	{
-		return true;	// error, but stop the search
-	}
-	else if (ReturnCode == DW_DLV_OK) 
-	{
-		if (dwarf_formstring(SubAt, &TempFuncName, &ErrorInfo))
-		{
-			*OutFuncName = NULL;
-		}
-		else
-		{
-			*OutFuncName = TempFuncName;
-		}
-		return true;
-	}
-
-	// DW_AT_Name is not present, look in DW_AT_specification
-	Dwarf_Attribute SpecAt;
-	if (dwarf_attr(Die, DW_AT_specification, &SpecAt, &ErrorInfo))
-	{
-		// not found, tough luck
-		return false;
-	}
-
-	Dwarf_Off Offset;
-	if (dwarf_global_formref(SpecAt, &Offset, &ErrorInfo))
-	{
-		return false;
-	}
-
-	Dwarf_Die SpecDie;
-	if (dwarf_offdie(DebugInfo, Offset, &SpecDie, &ErrorInfo))
-	{
-		return false;
-	}
-
-	if (dwarf_attrval_string(SpecDie, DW_AT_name, OutFuncName, &ErrorInfo))
-	{
-		*OutFuncName = NULL;
-	}
-
-	return true;
-}
-
-/**
- * Finds a function name in DWARF DIE (Debug Information Entry) and its children.
- * For more info on DWARF format, see http://www.dwarfstd.org/Download.php , http://www.ibm.com/developerworks/library/os-debugging/
- * Note: that function is not exactly traversing the tree, but this "seems to work"(tm). Not sure if we need to descend properly (taking child of every sibling), this
- * takes too much time (and callstacks seem to be fine without it).
- */
-void FindFunctionNameInDIEAndChildren(Dwarf_Debug DebugInfo, Dwarf_Die Die, Dwarf_Addr Addr, const char **OutFuncName)
-{
-	if (OutFuncName == NULL || *OutFuncName != NULL)
-	{
-		return;
-	}
-
-	// search for this Die
-	if (FindFunctionNameInDIE(DebugInfo, Die, Addr, OutFuncName))
-	{
-		return;
-	}
-
-	Dwarf_Die PrevChild = Die, Current = NULL;
-	Dwarf_Error ErrorInfo;
-
-	int32 MaxChildrenAllowed = 32 * 1024 * 1024;	// safeguard to make sure we never get into an infinite loop
-	for(;;)
-	{
-		if (--MaxChildrenAllowed <= 0)
-		{
-			fprintf(stderr, "Breaking out from what seems to be an infinite loop during DWARF parsing (too many children).\n");
-			return;
-		}
-
-		// Get the child
-		if (dwarf_child(PrevChild, &Current, &ErrorInfo) != DW_DLV_OK)
-		{
-			return;	// bail out
-		}
-
-		PrevChild = Current;
-
-		// look for in the child
-		if (FindFunctionNameInDIE(DebugInfo, Current, Addr, OutFuncName))
-		{
-			return;	// got the function name!
-		}
-
-		// search among child's siblings
-		int32 MaxSiblingsAllowed = 64 * 1024 * 1024;	// safeguard to make sure we never get into an infinite loop
-		for(;;)
-		{
-			if (--MaxSiblingsAllowed <= 0)
-			{
-				fprintf(stderr, "Breaking out from what seems to be an infinite loop during DWARF parsing (too many siblings).\n");
-				break;
-			}
-
-			Dwarf_Die Prev = Current;
-			if (dwarf_siblingof(DebugInfo, Prev, &Current, &ErrorInfo) != DW_DLV_OK || Current == NULL)
-			{
-				break;
-			}
-
-			if (FindFunctionNameInDIE(DebugInfo, Current, Addr, OutFuncName))
-			{
-				return;	// got the function name!
-			}
-		}
-	};
-}
-
-bool FLinuxCrashContext::GetInfoForAddress(void * Address, const char **OutFunctionNamePtr, const char **OutSourceFilePtr, int *OutLineNumberPtr)
-{
-	if (DebugInfo == NULL)
-	{
-		return false;
-	}
-
-	Dwarf_Die Die;
-	Dwarf_Unsigned Addr = reinterpret_cast< Dwarf_Unsigned >( Address ), LineNumber = 0;
-	const char * SrcFile = NULL;
-
-	static_assert(sizeof(Dwarf_Unsigned) >= sizeof(Address), "Dwarf_Unsigned type should be long enough to represent pointers. Check libdwarf bitness.");
-
-	int ReturnCode = DW_DLV_OK;
-	Dwarf_Error ErrorInfo;
-	bool bExitHeaderLoop = false;
-	int32 MaxCompileUnitsAllowed = 16 * 1024 * 1024;	// safeguard to make sure we never get into an infinite loop
-	const int32 kMaxBufferLinesAllowed = 16 * 1024 * 1024;	// safeguard to prevent too long line loop
-	for(;;)
-	{
-		if (--MaxCompileUnitsAllowed <= 0)
-		{
-			fprintf(stderr, "Breaking out from what seems to be an infinite loop during DWARF parsing (too many compile units).\n");
-			ReturnCode = DW_DLE_DIE_NO_CU_CONTEXT;	// invalidate
-			break;
-		}
-
-		if (bExitHeaderLoop)
-			break;
-
-		ReturnCode = dwarf_next_cu_header(DebugInfo, NULL, NULL, NULL, NULL, NULL, &ErrorInfo);
-		if (ReturnCode != DW_DLV_OK)
-			break;
-
-		Die = NULL;
-
-		while(dwarf_siblingof(DebugInfo, Die, &Die, &ErrorInfo) == DW_DLV_OK)
-		{
-			Dwarf_Half Tag;
-			if (dwarf_tag(Die, &Tag, &ErrorInfo) != DW_DLV_OK)
-			{
-				bExitHeaderLoop = true;
-				break;
-			}
-
-			if (Tag == DW_TAG_compile_unit)
-			{
-				break;
-			}
-		}
-
-		if (Die == NULL)
-		{
-			break;
-		}
-
-		// check if address is inside this CU
-		Dwarf_Unsigned LowerPC, HigherPC;
-		if (!dwarf_attrval_unsigned(Die, DW_AT_low_pc, &LowerPC, &ErrorInfo) && !dwarf_attrval_unsigned(Die, DW_AT_high_pc, &HigherPC, &ErrorInfo))
-		{
-			if (Addr < LowerPC || Addr >= HigherPC)
-			{
-				continue;
-			}
-		}
-
-		Dwarf_Line * LineBuf;
-		Dwarf_Signed NumLines = kMaxBufferLinesAllowed;
-		if (dwarf_srclines(Die, &LineBuf, &NumLines, &ErrorInfo) != DW_DLV_OK)
-		{
-			// could not get line info for some reason
-			break;
-		}
-
-		if (NumLines >= kMaxBufferLinesAllowed)
-		{
-			fprintf(stderr, "Number of lines associated with a DIE looks unreasonable (%d), early quitting.\n", static_cast<int32>(NumLines));
-			ReturnCode = DW_DLE_DIE_NO_CU_CONTEXT;	// invalidate
-			break;
-		}
-
-		// look which line is that
-		Dwarf_Addr LineAddress, PrevLineAddress = ~0ULL;
-		Dwarf_Unsigned PrevLineNumber = 0;
-		const char * PrevSrcFile = NULL;
-		char * SrcFileTemp = NULL;
-		for (int Idx = 0; Idx < NumLines; ++Idx)
-		{
-			if (dwarf_lineaddr(LineBuf[Idx], &LineAddress, &ErrorInfo) != 0 ||
-				dwarf_lineno(LineBuf[Idx], &LineNumber, &ErrorInfo) != 0)
-			{
-				bExitHeaderLoop = true;
-				break;
-			}
-
-			if (!dwarf_linesrc(LineBuf[Idx], &SrcFileTemp, &ErrorInfo))
-			{
-				SrcFile = SrcFileTemp;
-			}
-
-			// check if we hit the exact line
-			if (Addr == LineAddress)
-			{
-				bExitHeaderLoop = true;
-				break;
-			}
-			else if (PrevLineAddress < Addr && Addr < LineAddress)
-			{
-				LineNumber = PrevLineNumber;
-				SrcFile = PrevSrcFile;
-				bExitHeaderLoop = true;
-				break;
-			}
-
-			PrevLineAddress = LineAddress;
-			PrevLineNumber = LineNumber;
-			PrevSrcFile = SrcFile;
-		}
-
-	}
-
-	const char * FunctionName = NULL;
-	if (ReturnCode == DW_DLV_OK)
-	{
-		FindFunctionNameInDIEAndChildren(DebugInfo, Die, Addr, &FunctionName);
-	}
-
-	if (OutFunctionNamePtr != NULL && FunctionName != NULL)
-	{
-		*OutFunctionNamePtr = FunctionName;
-	}
-
-	if (OutSourceFilePtr != NULL && SrcFile != NULL)
-	{
-		*OutSourceFilePtr = SrcFile;
-		
-		if (OutLineNumberPtr != NULL)
-		{
-			*OutLineNumberPtr = LineNumber;
-		}
-	}
-
-	// Resets internal CU pointer, so next time we get here it begins from the start
-	while (ReturnCode != DW_DLV_NO_ENTRY) 
-	{
-		if (ReturnCode == DW_DLV_ERROR)
-			break;
-		ReturnCode = dwarf_next_cu_header(DebugInfo, NULL, NULL, NULL, NULL, NULL, &ErrorInfo);
-	}
-
-	// if we weren't able to find a function name, don't trust the source file either
-	return FunctionName != NULL;
-}
-
 /**
  * Handles graceful termination. Gives time to exit gracefully, but second signal will quit immediately.
  */
diff --git a/Engine/Source/Runtime/Core/Private/Linux/LinuxPlatformMisc.cpp b/Engine/Source/Runtime/Core/Private/Linux/LinuxPlatformMisc.cpp
--- a/Engine/Source/Runtime/Core/Private/Linux/LinuxPlatformMisc.cpp
+++ b/Engine/Source/Runtime/Core/Private/Linux/LinuxPlatformMisc.cpp
@@ -818,14 +818,17 @@ bool FLinuxPlatformMisc::IsDebuggerPresent()
 #if !UE_BUILD_SHIPPING
 void FLinuxPlatformMisc::UngrabAllInput()
 {
-	SDL_Window * GrabbedWindow = SDL_GetGrabbedWindow();
-	if (GrabbedWindow)
+	if (GInitializedSDL)
 	{
-		SDL_SetWindowGrab(GrabbedWindow, SDL_FALSE);
-		SDL_SetKeyboardGrab(GrabbedWindow, SDL_FALSE);
-	}
+		SDL_Window * GrabbedWindow = SDL_GetGrabbedWindow();
+		if (GrabbedWindow)
+		{
+			SDL_SetWindowGrab(GrabbedWindow, SDL_FALSE);
+			SDL_SetKeyboardGrab(GrabbedWindow, SDL_FALSE);
+		}
 
-	SDL_CaptureMouse(SDL_FALSE);
+		SDL_CaptureMouse(SDL_FALSE);
+	}
 }
 
 #endif // !UE_BUILD_SHIPPING
diff --git a/Engine/Source/Runtime/Core/Private/Linux/LinuxPlatformOutputDevices.cpp b/Engine/Source/Runtime/Core/Private/Linux/LinuxPlatformOutputDevices.cpp
--- a/Engine/Source/Runtime/Core/Private/Linux/LinuxPlatformOutputDevices.cpp
+++ b/Engine/Source/Runtime/Core/Private/Linux/LinuxPlatformOutputDevices.cpp
@@ -208,9 +208,8 @@ void FOutputDeviceConsoleLinux::Serialize(const TCHAR* Data, ELogVerbosity::Type
 					bNeedToResetColor = true;
 				}
 			}
-			TCHAR OutputString[MAX_SPRINTF]=TEXT(""); //@warning: this is safe as FCString::Sprintf only use 1024 characters max
-			FCString::Sprintf(OutputString,TEXT("%s%s"),*FOutputDevice::FormatLogLine(Verbosity, Category, Data, GPrintLogTimes),LINE_TERMINATOR);
-			printf("%s", TCHAR_TO_UTF8(OutputString));
+
+			printf("%ls\n", *FOutputDevice::FormatLogLine(Verbosity, Category, Data, GPrintLogTimes));
 
 			if (bNeedToResetColor)
 			{
diff --git a/Engine/Source/Runtime/Core/Private/Linux/LinuxPlatformStackWalk.cpp b/Engine/Source/Runtime/Core/Private/Linux/LinuxPlatformStackWalk.cpp
--- a/Engine/Source/Runtime/Core/Private/Linux/LinuxPlatformStackWalk.cpp
+++ b/Engine/Source/Runtime/Core/Private/Linux/LinuxPlatformStackWalk.cpp
@@ -12,14 +12,512 @@
 #include <cxxabi.h>
 #include <stdio.h>
 
+// these are not actually system headers, but a TPS library (see ThirdParty/elftoolchain)
+#include <libelf.h>
+#include <_libelf.h>
+#include <libdwarf.h>
+#include <dwarf.h>
+
 namespace LinuxStackWalkHelpers
 {
+	struct LinuxBacktraceSymbols
+	{
+		/** Lock for thread-safe initialization. */
+		FCriticalSection CriticalSection;
+
+		/** Initialized flag. If initialization failed, it don't run again. */
+		bool Inited;
+
+		/** File descriptor needed for libelf to open (our own) binary */
+		int ExeFd;
+
+		/** Elf header as used by libelf (forward-declared in the same way as libelf does it, it's normally of Elf type) */
+		struct _Elf* ElfHdr;
+
+		/** DWARF handle used by libdwarf (forward-declared in the same way as libdwarf does it, it's normally of Dwarf_Debug type) */
+		struct _Dwarf_Debug	* DebugInfo;
+
+		LinuxBacktraceSymbols()
+		:	Inited(false)
+		,	ExeFd(-1)
+		,	ElfHdr(NULL)
+		,	DebugInfo(NULL)
+		{
+		}
+
+		~LinuxBacktraceSymbols();
+
+		void Init();
+
+		/**
+		 * Gets information for the crash.
+		 *
+		 * @param Address the address to look up info for
+		 * @param OutFunctionNamePtr pointer to function name (may be NULL). Caller doesn't have to free it, but need to consider it temporary (i.e. next GetInfoForAddress() call on any thread may change it).
+		 * @param OutSourceFilePtr pointer to source filename (may be NULL). Caller doesn't have to free it, but need to consider it temporary (i.e. next GetInfoForAddress() call on any thread may change it).
+		 * @param OutLineNumberPtr pointer to line in a source file (may be NULL). Caller doesn't have to free it, but need to consider it temporary (i.e. next GetInfoForAddress() call on any thread may change it).
+		 *
+		 * @return true if succeeded in getting the info. If false is returned, none of above parameters should be trusted to contain valid data!
+		 */
+		bool GetInfoForAddress(void * Address, const char **OutFunctionNamePtr, const char **OutSourceFilePtr, int *OutLineNumberPtr);
+
+		/**
+		 * Check check if address is inside this entry.
+		 */
+		static bool CheckAddressInRange(Dwarf_Debug DebugInfo, Dwarf_Die Die, Dwarf_Unsigned Addr);
+
+		/**
+		 * Finds a function name in DWARF DIE (Debug Information Entry) and its children.
+		 * For more info on DWARF format, see http://www.dwarfstd.org/Download.php , http://www.ibm.com/developerworks/library/os-debugging/
+		 * Note: that function is not exactly traversing the tree, but this "seems to work"(tm). Not sure if we need to descend properly (taking child of every sibling), this
+		 * takes too much time (and callstacks seem to be fine without it).
+		 */
+		static void FindFunctionNameInDIEAndChildren(Dwarf_Debug DebugInfo, Dwarf_Die Die, Dwarf_Addr Addr, const char **OutFuncName);
+
+		/**
+		 * Finds a function name in DWARF DIE (Debug Information Entry).
+		 * For more info on DWARF format, see http://www.dwarfstd.org/Download.php , http://www.ibm.com/developerworks/library/os-debugging/
+		 *
+		 * @return true if we need to stop search (i.e. either found it or some error happened)
+		 */
+		static bool FindFunctionNameInDIE(Dwarf_Debug DebugInfo, Dwarf_Die Die, Dwarf_Addr Addr, const char **OutFuncName);
+
+	};
+
 	enum
 	{
 		MaxMangledNameLength = 1024,
 		MaxDemangledNameLength = 1024
 	};
 
+	void LinuxBacktraceSymbols::Init()
+	{
+		FScopeLock ScopeLock( &CriticalSection );
+
+		if (!Inited)
+		{
+			// open ourselves for examination
+			if (!FParse::Param( FCommandLine::Get(), TEXT(CMDARG_SUPPRESS_DWARF_PARSING)))
+			{
+				ExeFd = open("/proc/self/exe", O_RDONLY);
+				if (ExeFd >= 0)
+				{
+					Dwarf_Error ErrorInfo;
+					// allocate DWARF debug descriptor
+					if (dwarf_init(ExeFd, DW_DLC_READ, NULL, NULL, &DebugInfo, &ErrorInfo) == DW_DLV_OK)
+					{
+						// get ELF descritor
+						if (dwarf_get_elf(DebugInfo, &ElfHdr, &ErrorInfo) != DW_DLV_OK)
+						{
+							dwarf_finish(DebugInfo, &ErrorInfo);
+							DebugInfo = NULL;
+
+							close(ExeFd);
+							ExeFd = -1;
+						}
+					}
+					else
+					{
+						DebugInfo = NULL;
+						close(ExeFd);
+						ExeFd = -1;
+					}
+				}			
+			}
+			Inited = true;
+		}
+	}
+
+	LinuxBacktraceSymbols::~LinuxBacktraceSymbols()
+	{
+		if (DebugInfo)
+		{
+			Dwarf_Error ErrorInfo;
+			dwarf_finish(DebugInfo, &ErrorInfo);
+			DebugInfo = NULL;
+		}
+
+		if (ElfHdr)
+		{
+			elf_end(ElfHdr);
+			ElfHdr = NULL;
+		}
+
+		if (ExeFd >= 0)
+		{
+			close(ExeFd);
+			ExeFd = -1;
+		}
+	}
+
+	bool LinuxBacktraceSymbols::GetInfoForAddress(void * Address, const char **OutFunctionNamePtr, const char **OutSourceFilePtr, int *OutLineNumberPtr)
+	{
+		if (DebugInfo == NULL)
+		{
+			return false;
+		}
+
+		Dwarf_Die Die;
+		Dwarf_Unsigned Addr = reinterpret_cast< Dwarf_Unsigned >( Address ), LineNumber = 0;
+		const char * SrcFile = NULL;
+
+		static_assert(sizeof(Dwarf_Unsigned) >= sizeof(Address), "Dwarf_Unsigned type should be long enough to represent pointers. Check libdwarf bitness.");
+
+		int ReturnCode = DW_DLV_OK;
+		Dwarf_Error ErrorInfo;
+		bool bExitHeaderLoop = false;
+		int32 MaxCompileUnitsAllowed = 16 * 1024 * 1024;	// safeguard to make sure we never get into an infinite loop
+		const int32 kMaxBufferLinesAllowed = 16 * 1024 * 1024;	// safeguard to prevent too long line loop
+		for(;;)
+		{
+			if (--MaxCompileUnitsAllowed <= 0)
+			{
+				fprintf(stderr, "Breaking out from what seems to be an infinite loop during DWARF parsing (too many compile units).\n");
+				ReturnCode = DW_DLE_DIE_NO_CU_CONTEXT;	// invalidate
+				break;
+			}
+
+			if (bExitHeaderLoop)
+				break;
+
+			ReturnCode = dwarf_next_cu_header(DebugInfo, NULL, NULL, NULL, NULL, NULL, &ErrorInfo);
+			if (ReturnCode != DW_DLV_OK)
+				break;
+
+			Die = NULL;
+
+			while(dwarf_siblingof(DebugInfo, Die, &Die, &ErrorInfo) == DW_DLV_OK)
+			{
+				Dwarf_Half Tag;
+				if (dwarf_tag(Die, &Tag, &ErrorInfo) != DW_DLV_OK)
+				{
+					bExitHeaderLoop = true;
+					break;
+				}
+
+				if (Tag == DW_TAG_compile_unit)
+				{
+					break;
+				}
+			}
+
+			if (Die == NULL)
+			{
+				break;
+			}
+
+			// check if address is inside this CU
+			if (!CheckAddressInRange(DebugInfo, Die, Addr))
+			{
+				continue;
+			}
+
+			Dwarf_Line * LineBuf;
+			Dwarf_Signed NumLines = kMaxBufferLinesAllowed;
+			if (dwarf_srclines(Die, &LineBuf, &NumLines, &ErrorInfo) != DW_DLV_OK)
+			{
+				// could not get line info for some reason
+				break;
+			}
+
+			if (NumLines >= kMaxBufferLinesAllowed)
+			{
+				fprintf(stderr, "Number of lines associated with a DIE looks unreasonable (%d), early quitting.\n", static_cast<int32>(NumLines));
+				ReturnCode = DW_DLE_DIE_NO_CU_CONTEXT;	// invalidate
+				break;
+			}
+
+			// look which line is that
+			Dwarf_Addr LineAddress, PrevLineAddress = ~0ULL;
+			Dwarf_Unsigned LineIdx = NumLines;
+			for (int Idx = 0; Idx < NumLines; ++Idx)
+			{
+				if (dwarf_lineaddr(LineBuf[Idx], &LineAddress, &ErrorInfo) != 0)
+				{
+					bExitHeaderLoop = true;
+					break;
+				}
+				// check if we hit the exact line
+				if (Addr == LineAddress)
+				{
+					LineIdx = Idx;
+					bExitHeaderLoop = true;
+					break;
+				}
+				else if (PrevLineAddress < Addr && Addr < LineAddress)
+				{
+					LineIdx = Idx - 1;
+					break;
+				}
+				PrevLineAddress = LineAddress;
+			}
+			if (LineIdx < NumLines)
+			{
+				if (dwarf_lineno(LineBuf[LineIdx], &LineNumber, &ErrorInfo) != 0)
+				{
+					fprintf(stderr, "Can't get line number by dwarf_lineno.\n");
+					break;
+				}
+				for (int Idx = LineIdx; Idx >= 0; --Idx)
+				{
+					char * SrcFileTemp = NULL;
+					if (!dwarf_linesrc(LineBuf[Idx], &SrcFileTemp, &ErrorInfo))
+					{
+						SrcFile = SrcFileTemp;
+						break;
+					}
+				}
+				bExitHeaderLoop = true;
+			}
+		}
+
+		const char * FunctionName = NULL;
+		if (ReturnCode == DW_DLV_OK)
+		{
+			FindFunctionNameInDIEAndChildren(DebugInfo, Die, Addr, &FunctionName);
+		}
+
+		if (OutFunctionNamePtr != NULL && FunctionName != NULL)
+		{
+			*OutFunctionNamePtr = FunctionName;
+		}
+
+		if (OutSourceFilePtr != NULL && SrcFile != NULL)
+		{
+			*OutSourceFilePtr = SrcFile;
+			
+			if (OutLineNumberPtr != NULL)
+			{
+				*OutLineNumberPtr = LineNumber;
+			}
+		}
+
+		// Resets internal CU pointer, so next time we get here it begins from the start
+		while (ReturnCode != DW_DLV_NO_ENTRY) 
+		{
+			if (ReturnCode == DW_DLV_ERROR)
+				break;
+			ReturnCode = dwarf_next_cu_header(DebugInfo, NULL, NULL, NULL, NULL, NULL, &ErrorInfo);
+		}
+
+		// if we weren't able to find a function name, don't trust the source file either
+		return FunctionName != NULL;
+	}
+
+	/**
+	 * Check check if address is inside this entry.
+	 */
+	bool LinuxBacktraceSymbols::CheckAddressInRange(Dwarf_Debug DebugInfo, Dwarf_Die Die, Dwarf_Unsigned Addr)
+	{
+		Dwarf_Attribute *AttrList;
+		Dwarf_Signed AttrCount;
+
+		if (dwarf_attrlist(Die, &AttrList, &AttrCount, NULL) != DW_DLV_OK)
+			return true;
+
+		for (int i = 0; i < AttrCount; i++)
+		{
+			Dwarf_Half Attr;
+			if (dwarf_whatattr(AttrList[i], &Attr, NULL) != DW_DLV_OK)
+				continue;
+			switch (Attr)
+			{
+				case DW_AT_low_pc:
+					{
+						Dwarf_Unsigned LowOffset;
+						if (dwarf_formudata(AttrList[i], &LowOffset, NULL) != DW_DLV_OK)
+							continue;
+						if (LowOffset > Addr)
+							return false;
+					}
+					break;
+				case DW_AT_high_pc:
+					{
+						Dwarf_Unsigned HighOffset;
+						if (dwarf_formudata(AttrList[i], &HighOffset, NULL) != DW_DLV_OK)
+							continue;
+						if (HighOffset <= Addr)
+							return false;
+					}
+					break;
+				case DW_AT_ranges:
+					{
+						Dwarf_Unsigned Offset;
+						if (dwarf_formudata(AttrList[i], &Offset, NULL) != DW_DLV_OK)
+							continue;
+
+						Dwarf_Ranges *Ranges;
+						Dwarf_Signed Count;
+						if (dwarf_get_ranges(DebugInfo, (Dwarf_Off) Offset, &Ranges, &Count, NULL, NULL) != DW_DLV_OK)
+							continue;
+						for (int j = 0; j < Count; j++)
+						{
+							if (Ranges[j].dwr_type == DW_RANGES_END)
+							{
+								break;
+							}
+							if (Ranges[j].dwr_type == DW_RANGES_ENTRY)
+							{
+								if ((Ranges[j].dwr_addr1 <= Addr) && (Addr < Ranges[j].dwr_addr2))
+								{
+									return true;
+								}
+								continue;
+							}
+						}
+						return false;
+					}
+					break;
+				default:
+					break;
+			}
+		}
+		return true;
+	}
+
+	/**
+	 * Finds a function name in DWARF DIE (Debug Information Entry).
+	 * For more info on DWARF format, see http://www.dwarfstd.org/Download.php , http://www.ibm.com/developerworks/library/os-debugging/
+	 *
+	 * @return true if we need to stop search (i.e. either found it or some error happened)
+	 */
+	bool LinuxBacktraceSymbols::FindFunctionNameInDIE(Dwarf_Debug DebugInfo, Dwarf_Die Die, Dwarf_Addr Addr, const char **OutFuncName)
+	{
+		Dwarf_Error ErrorInfo;
+		Dwarf_Half Tag;
+		//Dwarf_Unsigned LowerPC, HigherPC;
+		char *TempFuncName;
+		int ReturnCode;
+
+		if (dwarf_tag(Die, &Tag, &ErrorInfo) != DW_DLV_OK || Tag != DW_TAG_subprogram)
+		{
+			return false;
+		}
+
+		// check if address is inside this entry
+		if (!CheckAddressInRange(DebugInfo, Die, Addr))
+		{
+			return false;
+		}
+		
+		// found it
+		*OutFuncName = NULL;
+		Dwarf_Attribute SubAt;
+		ReturnCode = dwarf_attr(Die, DW_AT_name, &SubAt, &ErrorInfo);
+		if (ReturnCode == DW_DLV_ERROR)
+		{
+			return true;	// error, but stop the search
+		}
+		else if (ReturnCode == DW_DLV_OK) 
+		{
+			if (dwarf_formstring(SubAt, &TempFuncName, &ErrorInfo))
+			{
+				*OutFuncName = NULL;
+			}
+			else
+			{
+				*OutFuncName = TempFuncName;
+			}
+			return true;
+		}
+
+		// DW_AT_Name is not present, look in DW_AT_specification
+		Dwarf_Attribute SpecAt;
+		if (dwarf_attr(Die, DW_AT_specification, &SpecAt, &ErrorInfo))
+		{
+			// not found, tough luck
+			return false;
+		}
+
+		Dwarf_Off Offset;
+		if (dwarf_global_formref(SpecAt, &Offset, &ErrorInfo))
+		{
+			return false;
+		}
+
+		Dwarf_Die SpecDie;
+		if (dwarf_offdie(DebugInfo, Offset, &SpecDie, &ErrorInfo))
+		{
+			return false;
+		}
+
+		if (dwarf_attrval_string(SpecDie, DW_AT_name, OutFuncName, &ErrorInfo))
+		{
+			*OutFuncName = NULL;
+		}
+
+		return true;
+	}
+
+	/**
+	 * Finds a function name in DWARF DIE (Debug Information Entry) and its children.
+	 * For more info on DWARF format, see http://www.dwarfstd.org/Download.php , http://www.ibm.com/developerworks/library/os-debugging/
+	 * Note: that function is not exactly traversing the tree, but this "seems to work"(tm). Not sure if we need to descend properly (taking child of every sibling), this
+	 * takes too much time (and callstacks seem to be fine without it).
+	 */
+	void LinuxBacktraceSymbols::FindFunctionNameInDIEAndChildren(Dwarf_Debug DebugInfo, Dwarf_Die Die, Dwarf_Addr Addr, const char **OutFuncName)
+	{
+		if (OutFuncName == NULL || *OutFuncName != NULL)
+		{
+			return;
+		}
+
+		// search for this Die
+		if (FindFunctionNameInDIE(DebugInfo, Die, Addr, OutFuncName))
+		{
+			return;
+		}
+
+		Dwarf_Die PrevChild = Die, Current = NULL;
+		Dwarf_Error ErrorInfo;
+
+		int32 MaxChildrenAllowed = 32 * 1024 * 1024;	// safeguard to make sure we never get into an infinite loop
+		for(;;)
+		{
+			if (--MaxChildrenAllowed <= 0)
+			{
+				fprintf(stderr, "Breaking out from what seems to be an infinite loop during DWARF parsing (too many children).\n");
+				return;
+			}
+
+			// Get the child
+			if (dwarf_child(PrevChild, &Current, &ErrorInfo) != DW_DLV_OK)
+			{
+				return;	// bail out
+			}
+
+			PrevChild = Current;
+
+			// look for in the child
+			if (FindFunctionNameInDIE(DebugInfo, Current, Addr, OutFuncName))
+			{
+				return;	// got the function name!
+			}
+
+			// search among child's siblings
+			int32 MaxSiblingsAllowed = 64 * 1024 * 1024;	// safeguard to make sure we never get into an infinite loop
+			for(;;)
+			{
+				if (--MaxSiblingsAllowed <= 0)
+				{
+					fprintf(stderr, "Breaking out from what seems to be an infinite loop during DWARF parsing (too many siblings).\n");
+					break;
+				}
+
+				Dwarf_Die Prev = Current;
+				if (dwarf_siblingof(DebugInfo, Prev, &Current, &ErrorInfo) != DW_DLV_OK || Current == NULL)
+				{
+					break;
+				}
+
+				if (FindFunctionNameInDIE(DebugInfo, Current, Addr, OutFuncName))
+				{
+					return;	// got the function name!
+				}
+			}
+		};
+	}
+
 	/**
 	 * Finds mangled name and returns it in internal buffer.
 	 * Caller doesn't have to deallocate that.
@@ -133,6 +631,13 @@ namespace LinuxStackWalkHelpers
 			FCStringAnsi::Strcat(LinuxContext.MinidumpCallstackInfo, ARRAY_COUNT( LinuxContext.MinidumpCallstackInfo ) - 1, TempArray);
 		}
 	}
+
+	LinuxBacktraceSymbols *GetBacktraceSymbols()
+	{
+		static LinuxBacktraceSymbols Symbols;
+		Symbols.Init();
+		return &Symbols;
+	}
 }
 
 void FLinuxPlatformStackWalk::ProgramCounterToSymbolInfo( uint64 ProgramCounter, FProgramCounterSymbolInfo& out_SymbolInfo )
@@ -189,7 +694,7 @@ bool FLinuxPlatformStackWalk::ProgramCounterToHumanReadableString( int32 Current
 				{
 					const char * SourceFilename = NULL;
 					int LineNumber;
-					if (LinuxContext->GetInfoForAddress(reinterpret_cast< void* >( ProgramCounter ), NULL, &SourceFilename, &LineNumber) && FunctionName != NULL)
+					if (LinuxStackWalkHelpers::GetBacktraceSymbols()->GetInfoForAddress(reinterpret_cast< void* >( ProgramCounter ), NULL, &SourceFilename, &LineNumber) && FunctionName != NULL)
 					{
 						FCStringAnsi::Sprintf(TempArray, " [%s, line %d]", SourceFilename, LineNumber);
 						LinuxStackWalkHelpers::AppendToString(HumanReadableString, HumanReadableStringSize, Context, TempArray);
@@ -217,6 +722,20 @@ bool FLinuxPlatformStackWalk::ProgramCounterToHumanReadableString( int32 Current
 	return true;
 }
 
+void FLinuxPlatformStackWalk::StackWalkAndDump( ANSICHAR* HumanReadableString, SIZE_T HumanReadableStringSize, int32 IgnoreCount, void* Context )
+{
+	if ( Context == nullptr )
+	{
+		FLinuxCrashContext CrashContext;
+		CrashContext.InitFromSignal(0, nullptr, nullptr);
+		FGenericPlatformStackWalk::StackWalkAndDump(HumanReadableString, HumanReadableStringSize, IgnoreCount + 1, &CrashContext);
+	}
+	else
+	{
+		FGenericPlatformStackWalk::StackWalkAndDump(HumanReadableString, HumanReadableStringSize, IgnoreCount + 1, Context);
+	}
+}
+
 void FLinuxPlatformStackWalk::CaptureStackBackTrace( uint64* BackTrace, uint32 MaxDepth, void* Context )
 {
 	// Make sure we have place to store the information before we go through the process of raising
diff --git a/Engine/Source/Runtime/Core/Public/Linux/LinuxPlatformCrashContext.h b/Engine/Source/Runtime/Core/Public/Linux/LinuxPlatformCrashContext.h
--- a/Engine/Source/Runtime/Core/Public/Linux/LinuxPlatformCrashContext.h
+++ b/Engine/Source/Runtime/Core/Public/Linux/LinuxPlatformCrashContext.h
@@ -21,15 +21,6 @@ struct CORE_API FLinuxCrashContext : public FGenericCrashContext
 	/** Symbols received via backtrace_symbols(), if any (note that we will need to clean it up) */
 	char ** BacktraceSymbols;
 
-	/** File descriptor needed for libelf to open (our own) binary */
-	int ExeFd;
-
-	/** Elf header as used by libelf (forward-declared in the same way as libelf does it, it's normally of Elf type) */
-	struct _Elf* ElfHdr;
-
-	/** DWARF handle used by libdwarf (forward-declared in the same way as libdwarf does it, it's normally of Dwarf_Debug type) */
-	struct _Dwarf_Debug	* DebugInfo;
-
 	/** Memory reserved for "exception" (signal) info */
 	TCHAR SignalDescription[128];
 
@@ -41,9 +32,6 @@ struct CORE_API FLinuxCrashContext : public FGenericCrashContext
 		,	Info(NULL)
 		,	Context(NULL)
 		,	BacktraceSymbols(NULL)
-		,	ExeFd(-1)
-		,	ElfHdr(NULL)
-		,	DebugInfo(NULL)
 	{
 		SignalDescription[ 0 ] = 0;
 		MinidumpCallstackInfo[ 0 ] = 0;
@@ -60,18 +48,6 @@ struct CORE_API FLinuxCrashContext : public FGenericCrashContext
 	 */
 	void InitFromSignal(int32 InSignal, siginfo_t* InInfo, void* InContext);
 
-	/**
-	 * Gets information for the crash.
-	 *
-	 * @param Address the address to look up info for
-	 * @param OutFunctionNamePtr pointer to function name (may be NULL). Caller doesn't have to free it, but need to consider it temporary (i.e. next GetInfoForAddress() call on any thread may change it).
-	 * @param OutSourceFilePtr pointer to source filename (may be NULL). Caller doesn't have to free it, but need to consider it temporary (i.e. next GetInfoForAddress() call on any thread may change it).
-	 * @param OutLineNumberPtr pointer to line in a source file (may be NULL). Caller doesn't have to free it, but need to consider it temporary (i.e. next GetInfoForAddress() call on any thread may change it).
-	 *
-	 * @return true if succeeded in getting the info. If false is returned, none of above parameters should be trusted to contain valid data!
-	 */
-	bool GetInfoForAddress(void * Address, const char **OutFunctionNamePtr, const char **OutSourceFilePtr, int *OutLineNumberPtr);
-
 	/**
 	 * Dumps all the data from crash context to the "minidump" report.
 	 *
diff --git a/Engine/Source/Runtime/Core/Public/Linux/LinuxPlatformStackWalk.h b/Engine/Source/Runtime/Core/Public/Linux/LinuxPlatformStackWalk.h
--- a/Engine/Source/Runtime/Core/Public/Linux/LinuxPlatformStackWalk.h
+++ b/Engine/Source/Runtime/Core/Public/Linux/LinuxPlatformStackWalk.h
@@ -13,6 +13,7 @@ struct CORE_API FLinuxPlatformStackWalk : public FGenericPlatformStackWalk
 	static void ProgramCounterToSymbolInfo( uint64 ProgramCounter, FProgramCounterSymbolInfo& out_SymbolInfo );
 	static bool ProgramCounterToHumanReadableString( int32 CurrentCallDepth, uint64 ProgramCounter, ANSICHAR* HumanReadableString, SIZE_T HumanReadableStringSize, FGenericCrashContext* Context = nullptr );
 	static void CaptureStackBackTrace( uint64* BackTrace, uint32 MaxDepth, void* Context = nullptr );
+	static void StackWalkAndDump( ANSICHAR* HumanReadableString, SIZE_T HumanReadableStringSize, int32 IgnoreCount, void* Context = nullptr );
 };
 
 typedef FLinuxPlatformStackWalk FPlatformStackWalk;
