CrashReport - Print all errors to the console, previously only the callstack was visible

diff --git a/Engine/Source/Runtime/Core/Private/GenericPlatform/GenericPlatformStackWalk.cpp b/Engine/Source/Runtime/Core/Private/GenericPlatform/GenericPlatformStackWalk.cpp
--- a/Engine/Source/Runtime/Core/Private/GenericPlatform/GenericPlatformStackWalk.cpp
+++ b/Engine/Source/Runtime/Core/Private/GenericPlatform/GenericPlatformStackWalk.cpp
@@ -133,13 +133,5 @@ void FGenericPlatformStackWalk::StackWalkAndDump( ANSICHAR* HumanReadableString,
 		FCStringAnsi::Strcat( HumanReadableString, HumanReadableStringSize, LINE_TERMINATOR_ANSI );
 		CurrentDepth++;
 	}
-	if (FParse::Param(FCommandLine::Get(), TEXT("CrashForUAT")) && FParse::Param(FCommandLine::Get(), TEXT("stdout")))
-	{
-		FPlatformMisc::LowLevelOutputDebugString(ANSI_TO_TCHAR(HumanReadableString));
-		wprintf(TEXT("\nbegin: stack for UAT"));
-		wprintf(TEXT("\n%s"), ANSI_TO_TCHAR(HumanReadableString));
-		wprintf(TEXT("\nend: stack for UAT"));
-		fflush(stdout);
-	}
 }
 
diff --git a/Engine/Source/Runtime/Core/Private/Misc/OutputDevice.cpp b/Engine/Source/Runtime/Core/Private/Misc/OutputDevice.cpp
--- a/Engine/Source/Runtime/Core/Private/Misc/OutputDevice.cpp
+++ b/Engine/Source/Runtime/Core/Private/Misc/OutputDevice.cpp
@@ -179,6 +179,24 @@ void StaticFailDebug( const TCHAR* Error, const ANSICHAR* File, int32 Line, cons
 	FCString::Strncat( GErrorHist, TEXT( "\r\n\r\n" ), ARRAY_COUNT( GErrorHist ) );
 }
 
+void FDebug::ConditionallyEmitBeginCrashUATMarker()
+{
+	const bool bWriteUATMarkers = FParse::Param( FCommandLine::Get(), TEXT( "CrashForUAT" ) ) && FParse::Param( FCommandLine::Get(), TEXT( "stdout" ) );
+	if (bWriteUATMarkers)
+	{
+		UE_LOG( LogOutputDevice, Error, TEXT( "\nbegin: stack for UAT" ) );
+	}
+}
+
+void FDebug::ConditionallyEmitEndCrashUATMarker()
+{
+	const bool bWriteUATMarkers = FParse::Param( FCommandLine::Get(), TEXT( "CrashForUAT" ) ) && FParse::Param( FCommandLine::Get(), TEXT( "stdout" ) );
+	if (bWriteUATMarkers)
+	{
+		UE_LOG( LogOutputDevice, Error, TEXT( "\nend: stack for UAT" ) );
+	}
+}
+
 #if DO_CHECK || DO_GUARD_SLOW
 //
 // Failed assertion handler.
@@ -278,7 +296,9 @@ void FDebug::EnsureFailed(const ANSICHAR* Expr, const ANSICHAR* File, int32 Line
 		FMemory::SystemFree( StackTrace );
 
 		// Dump the error and flush the log.
-		UE_LOG(LogOutputDevice, Log, TEXT("=== Handled error: ===") LINE_TERMINATOR LINE_TERMINATOR TEXT("%s"), ErrorMsg);
+		FDebug::ConditionallyEmitBeginCrashUATMarker();
+		UE_LOG(LogOutputDevice, Warning, TEXT("=== Handled error: ===") LINE_TERMINATOR LINE_TERMINATOR TEXT("%s"), ErrorMsg);
+		FDebug::ConditionallyEmitEndCrashUATMarker();
 		GLog->Flush();
 
 		// Submit the error report to the server! (and display a balloon in the system tray)
diff --git a/Engine/Source/Runtime/Core/Private/Windows/WindowsPlatformOutputDevices.cpp b/Engine/Source/Runtime/Core/Private/Windows/WindowsPlatformOutputDevices.cpp
--- a/Engine/Source/Runtime/Core/Private/Windows/WindowsPlatformOutputDevices.cpp
+++ b/Engine/Source/Runtime/Core/Private/Windows/WindowsPlatformOutputDevices.cpp
@@ -122,7 +122,9 @@ void FOutputDeviceWindowsError::HandleError()
 	FCoreDelegates::OnHandleSystemError.Broadcast();
 
 	// Dump the error and flush the log.
-	UE_LOG( LogWindows, Log, TEXT( "=== Critical error: ===" ) LINE_TERMINATOR TEXT( "%s" ) LINE_TERMINATOR, GErrorHist );
+	FDebug::ConditionallyEmitBeginCrashUATMarker();
+	UE_LOG( LogWindows, Error, TEXT( "=== Critical error: ===" ) LINE_TERMINATOR TEXT( "%s" ) LINE_TERMINATOR, GErrorHist );
+	FDebug::ConditionallyEmitEndCrashUATMarker();
 
 	GLog->PanicFlushThreadedLogs();
 
diff --git a/Engine/Source/Runtime/Core/Public/Misc/OutputDevice.h b/Engine/Source/Runtime/Core/Public/Misc/OutputDevice.h
--- a/Engine/Source/Runtime/Core/Public/Misc/OutputDevice.h
+++ b/Engine/Source/Runtime/Core/Public/Misc/OutputDevice.h
@@ -389,6 +389,12 @@ struct CORE_API FMsg
  **/
 struct CORE_API FDebug
 {
+	/** Conditionally emits a special UAT marker. */
+	static void ConditionallyEmitBeginCrashUATMarker();
+
+	/** Conditionally emits a special UAT marker. */
+	static void ConditionallyEmitEndCrashUATMarker();
+
 	/** Logs final assert message and exits the program. */
 	static void VARARGS AssertFailed(const ANSICHAR* Expr, const ANSICHAR* File, int32 Line, const TCHAR* Format = TEXT(""), ...);
 
