Missing file from last checkin.

diff --git a/Engine/Source/Runtime/Core/Public/UObject/DebugSerializationFlags.h b/Engine/Source/Runtime/Core/Public/UObject/DebugSerializationFlags.h
--- a/Engine/Source/Runtime/Core/Public/UObject/DebugSerializationFlags.h
+++ b/Engine/Source/Runtime/Core/Public/UObject/DebugSerializationFlags.h
@@ -14,5 +14,5 @@ enum EDebugSerializationFlags
 	DSF_None =					0x00000000,
 
 	DSF_IgnoreDiff =			0x00000001,
-	DSF_IsTagExportsArchive =	0x00000002,
+	DSF_EnableCookerWarnings =	0x00000002,
 };
\ No newline at end of file
