Invalid and unnecessary comparison in TUniquePtr's assignment operator removed.

diff --git a/Engine/Source/Runtime/Core/Public/Templates/UniquePtr.h b/Engine/Source/Runtime/Core/Public/Templates/UniquePtr.h
--- a/Engine/Source/Runtime/Core/Public/Templates/UniquePtr.h
+++ b/Engine/Source/Runtime/Core/Public/Templates/UniquePtr.h
@@ -88,14 +88,11 @@ public:
 	template <typename OtherT>
 	FORCEINLINE TUniquePtr& operator=(TUniquePtr<OtherT>&& Other)
 	{
-		if (this != &Other)
-		{
-			// We delete last, because we don't want odd side effects if the destructor of T relies on the state of this or Other
-			T* OldPtr = Ptr;
-			Ptr = Other.Ptr;
-			Other.Ptr = nullptr;
-			delete OldPtr;
-		}
+		// We delete last, because we don't want odd side effects if the destructor of T relies on the state of this or Other
+		T* OldPtr = Ptr;
+		Ptr = Other.Ptr;
+		Other.Ptr = nullptr;
+		delete OldPtr;
 
 		return *this;
 	}
