UE4 - Removed GC frames from abtest testing.

diff --git a/Engine/Source/Runtime/Core/Private/Misc/Core.cpp b/Engine/Source/Runtime/Core/Private/Misc/Core.cpp
--- a/Engine/Source/Runtime/Core/Private/Misc/Core.cpp
+++ b/Engine/Source/Runtime/Core/Private/Misc/Core.cpp
@@ -204,6 +204,7 @@ FString					GSystemStartTime;
 bool					GIsInitialLoad					= true;
 /** Steadily increasing frame counter.																		*/
 uint64					GFrameCounter					= 0;
+uint64					GLastGCFrame					= 0;
 /** Incremented once per frame before the scene is being rendered. In split screen mode this is incremented once for all views (not for each view). */
 uint32					GFrameNumber					= 1;
 /** Render Thread copy of the frame number. */
diff --git a/Engine/Source/Runtime/Core/Private/ProfilingDebugging/ABTesting.cpp b/Engine/Source/Runtime/Core/Private/ProfilingDebugging/ABTesting.cpp
--- a/Engine/Source/Runtime/Core/Private/ProfilingDebugging/ABTesting.cpp
+++ b/Engine/Source/Runtime/Core/Private/ProfilingDebugging/ABTesting.cpp
@@ -56,6 +56,7 @@ FABTest::FABTest()
 		, MinFramesPerTrial(0)
 		, NumResamples(0)
 		, TotalScopeTimeInFrame(0)
+		, LastGCFrame(0)
 	{
 	}
 
@@ -94,6 +95,10 @@ const TCHAR* FABTest::TickAndGetCommand()
 	{
 		RemainingCoolDown--;
 	}
+	else if (LastGCFrame != GLastGCFrame && !bABScopeTestActive) // reject GC frames for whole game tests
+	{
+		FPlatformMisc::LowLevelOutputDebugStringf(TEXT("Rejecting abtest frame because of GC."));
+	}
 	else if (bABTestActive)
 	{
 		float Delta;
@@ -203,11 +208,11 @@ const TCHAR* FABTest::TickAndGetCommand()
 					FPlatformMisc::LowLevelOutputDebugStringf(TEXT("      B is %7.4fms faster than A;  %3.0f%% chance this is noise.\n"), Diff, fConf * 100.0f);
 				}
 
-				UE_LOG(LogConsoleResponse, Display, TEXT("----------------"));
+				FPlatformMisc::LowLevelOutputDebugStringf(TEXT("----------------"));
 			}
 			else
 			{
-				UE_LOG(LogConsoleResponse, Display, TEXT("No Samples?"));
+				FPlatformMisc::LowLevelOutputDebugStringf(TEXT("No Samples?"));
 			}
 			RemainingPrint = ReportNum;
 		}
@@ -224,6 +229,7 @@ const TCHAR* FABTest::TickAndGetCommand()
 	LastTime = FPlatformTime::Seconds();
 
 	TotalScopeTimeInFrame = 0.0;
+	LastGCFrame = GLastGCFrame;
 	return OutCommand;
 }
 
diff --git a/Engine/Source/Runtime/Core/Public/CoreGlobals.h b/Engine/Source/Runtime/Core/Public/CoreGlobals.h
--- a/Engine/Source/Runtime/Core/Public/CoreGlobals.h
+++ b/Engine/Source/Runtime/Core/Public/CoreGlobals.h
@@ -232,6 +232,9 @@ extern CORE_API bool GIsRetrievingVTablePtr;
 /** Steadily increasing frame counter. */
 extern CORE_API uint64 GFrameCounter;
 
+/** GFrameCounter the last time GC was run. */
+extern CORE_API uint64 GLastGCFrame;
+
 /** Incremented once per frame before the scene is being rendered. In split screen mode this is incremented once for all views (not for each view). */
 extern CORE_API uint32 GFrameNumber;
 
diff --git a/Engine/Source/Runtime/Core/Public/ProfilingDebugging/ABTesting.h b/Engine/Source/Runtime/Core/Public/ProfilingDebugging/ABTesting.h
--- a/Engine/Source/Runtime/Core/Public/ProfilingDebugging/ABTesting.h
+++ b/Engine/Source/Runtime/Core/Public/ProfilingDebugging/ABTesting.h
@@ -96,6 +96,7 @@ private:
 	uint32 Spikes;
 
 	double TotalScopeTimeInFrame;
+	uint64 LastGCFrame;
 
 };
 
diff --git a/Engine/Source/Runtime/CoreUObject/Private/UObject/GarbageCollection.cpp b/Engine/Source/Runtime/CoreUObject/Private/UObject/GarbageCollection.cpp
--- a/Engine/Source/Runtime/CoreUObject/Private/UObject/GarbageCollection.cpp
+++ b/Engine/Source/Runtime/CoreUObject/Private/UObject/GarbageCollection.cpp
@@ -1370,6 +1370,7 @@ void CollectGarbageInternal(EObjectFlags KeepFlags, bool bPerformFullPurge)
 	// Route callbacks so we can ensure that we are e.g. not in the middle of loading something by flushing
 	// the async loading, etc...
 	FCoreUObjectDelegates::PreGarbageCollect.Broadcast();
+	GLastGCFrame = GFrameCounter;
 
 	// Set 'I'm garbage collecting' flag - might be checked inside various functions.
 	FGCScopeLock GCLock;
