Added an option to convert only a sub-section of a given FString to an icu::UnicodeString

This can help avoid an intermediate FString copy.

diff --git a/Engine/Source/Runtime/Core/Private/Internationalization/ICUUtilities.cpp b/Engine/Source/Runtime/Core/Private/Internationalization/ICUUtilities.cpp
--- a/Engine/Source/Runtime/Core/Private/Internationalization/ICUUtilities.cpp
+++ b/Engine/Source/Runtime/Core/Private/Internationalization/ICUUtilities.cpp
@@ -23,19 +23,24 @@ namespace ICUUtilities
 
 	void FStringConverter::ConvertString(const FString& Source, icu::UnicodeString& Destination, const bool ShouldNullTerminate)
 	{
-		if (Source.Len() > 0)
+		ConvertString(*Source, 0, Source.Len(), Destination, ShouldNullTerminate);
+	}
+
+	void FStringConverter::ConvertString(const TCHAR* Source, const int32 SourceStartIndex, const int32 SourceLen, icu::UnicodeString& Destination, const bool ShouldNullTerminate)
+	{
+		if (SourceLen > 0)
 		{
 			UErrorCode ICUStatus = U_ZERO_ERROR;
 
 			ucnv_reset(ICUConverter);
 
 			// Get the internal buffer of the string, we're going to use it as scratch space
-			const int32_t DestinationCapacityUChars = Source.Len() * 2;
+			const int32_t DestinationCapacityUChars = SourceLen * 2;
 			UChar* InternalStringBuffer = Destination.getBuffer(DestinationCapacityUChars);
 
 			// Perform the conversion into the string buffer
-			const int32_t SourceSizeBytes = Source.Len() * sizeof(TCHAR);
-			const int32_t DestinationLength = ucnv_toUChars(ICUConverter, InternalStringBuffer, DestinationCapacityUChars, reinterpret_cast<const char*>(*Source), SourceSizeBytes, &ICUStatus);
+			const int32_t SourceSizeBytes = SourceLen * sizeof(TCHAR);
+			const int32_t DestinationLength = ucnv_toUChars(ICUConverter, InternalStringBuffer, DestinationCapacityUChars, reinterpret_cast<const char*>(Source + SourceStartIndex), SourceSizeBytes, &ICUStatus);
 
 			// Optionally null terminate the string
 			if (ShouldNullTerminate)
@@ -61,6 +66,13 @@ namespace ICUUtilities
 		return Destination;
 	}
 
+	icu::UnicodeString FStringConverter::ConvertString(const TCHAR* Source, const int32 SourceStartIndex, const int32 SourceLen, const bool ShouldNullTerminate)
+	{
+		icu::UnicodeString Destination;
+		ConvertString(Source, SourceStartIndex, SourceLen, Destination, ShouldNullTerminate);
+		return Destination;
+	}
+
 	void FStringConverter::ConvertString(const icu::UnicodeString& Source, FString& Destination)
 	{
 		if (Source.length() > 0)
@@ -111,6 +123,19 @@ namespace ICUUtilities
 		}
 	}
 
+	void ConvertString(const TCHAR* Source, const int32 SourceStartIndex, const int32 SourceLen, icu::UnicodeString& Destination, const bool ShouldNullTerminate)
+	{
+		if (SourceLen > 0)
+		{
+			FStringConverter StringConverter;
+			StringConverter.ConvertString(Source, SourceStartIndex, SourceLen, Destination, ShouldNullTerminate);
+		}
+		else
+		{
+			Destination.remove();
+		}
+	}
+
 	icu::UnicodeString ConvertString(const FString& Source, const bool ShouldNullTerminate)
 	{
 		icu::UnicodeString Destination;
@@ -118,6 +143,13 @@ namespace ICUUtilities
 		return Destination;
 	}
 
+	icu::UnicodeString ConvertString(const TCHAR* Source, const int32 SourceStartIndex, const int32 SourceLen, const bool ShouldNullTerminate)
+	{
+		icu::UnicodeString Destination;
+		ConvertString(Source, SourceStartIndex, SourceLen, Destination, ShouldNullTerminate);
+		return Destination;
+	}
+
 	void ConvertString(const icu::UnicodeString& Source, FString& Destination)
 	{
 		if (Source.length() > 0)
diff --git a/Engine/Source/Runtime/Core/Private/Internationalization/ICUUtilities.h b/Engine/Source/Runtime/Core/Private/Internationalization/ICUUtilities.h
--- a/Engine/Source/Runtime/Core/Private/Internationalization/ICUUtilities.h
+++ b/Engine/Source/Runtime/Core/Private/Internationalization/ICUUtilities.h
@@ -19,7 +19,9 @@ namespace ICUUtilities
 
 		/** Convert FString -> icu::UnicodeString */
 		void ConvertString(const FString& Source, icu::UnicodeString& Destination, const bool ShouldNullTerminate = true);
+		void ConvertString(const TCHAR* Source, const int32 SourceStartIndex, const int32 SourceLen, icu::UnicodeString& Destination, const bool ShouldNullTerminate = true);
 		icu::UnicodeString ConvertString(const FString& Source, const bool ShouldNullTerminate = true);
+		icu::UnicodeString ConvertString(const TCHAR* Source, const int32 SourceStartIndex, const int32 SourceLen, const bool ShouldNullTerminate = true);
 
 		/** Convert icu::UnicodeString -> FString */
 		void ConvertString(const icu::UnicodeString& Source, FString& Destination);
@@ -35,7 +37,9 @@ namespace ICUUtilities
 
 	/** Convert FString -> icu::UnicodeString */
 	void ConvertString(const FString& Source, icu::UnicodeString& Destination, const bool ShouldNullTerminate = true);
+	void ConvertString(const TCHAR* Source, const int32 SourceStartIndex, const int32 SourceLen, icu::UnicodeString& Destination, const bool ShouldNullTerminate = true);
 	icu::UnicodeString ConvertString(const FString& Source, const bool ShouldNullTerminate = true);
+	icu::UnicodeString ConvertString(const TCHAR* Source, const int32 SourceStartIndex, const int32 SourceLen, const bool ShouldNullTerminate = true);
 
 	/** Convert icu::UnicodeString -> FString */
 	void ConvertString(const icu::UnicodeString& Source, FString& Destination);
