Added an option to detect the direction of a sub-section of a given FString

This can help avoid an intermediate FString copy.

diff --git a/Engine/Source/Runtime/Core/Private/Internationalization/ICUText.cpp b/Engine/Source/Runtime/Core/Private/Internationalization/ICUText.cpp
--- a/Engine/Source/Runtime/Core/Private/Internationalization/ICUText.cpp
+++ b/Engine/Source/Runtime/Core/Private/Internationalization/ICUText.cpp
@@ -265,7 +265,7 @@ ETextDirection ComputeTextDirection(UBiDi* InICUBiDi, const icu::UnicodeString&
 	return ETextDirection::LeftToRight;
 }
 
-ETextDirection ComputeTextDirection(UBiDi* InICUBiDi, const icu::UnicodeString& InICUString, TArray<FTextDirectionInfo>& OutTextDirectionInfo)
+ETextDirection ComputeTextDirection(UBiDi* InICUBiDi, const icu::UnicodeString& InICUString, const int32 InStringOffset, TArray<FTextDirectionInfo>& OutTextDirectionInfo)
 {
 	UErrorCode ICUStatus = U_ZERO_ERROR;
 
@@ -281,6 +281,7 @@ ETextDirection ComputeTextDirection(UBiDi* InICUBiDi, const icu::UnicodeString&
 		{
 			FTextDirectionInfo& CurTextDirectionInfo = OutTextDirectionInfo[RunIndex];
 			CurTextDirectionInfo.TextDirection = Internal::ICUToUE(ubidi_getVisualRun(InICUBiDi, RunIndex, &CurTextDirectionInfo.StartIndex, &CurTextDirectionInfo.Length));
+			CurTextDirectionInfo.StartIndex += InStringOffset;
 		}
 
 		return ReturnDirection;
@@ -314,12 +315,17 @@ public:
 
 	virtual ETextDirection ComputeTextDirection(const FString& InString) override
 	{
-		if (InString.IsEmpty())
+		return FICUTextBiDi::ComputeTextDirection(*InString, 0, InString.Len());
+	}
+
+	virtual ETextDirection ComputeTextDirection(const TCHAR* InString, const int32 InStringStartIndex, const int32 InStringLen) override
+	{
+		if (InStringLen == 0)
 		{
 			return ETextDirection::LeftToRight;
 		}
 
-		StringConverter.ConvertString(InString, ICUString);
+		StringConverter.ConvertString(InString, InStringStartIndex, InStringLen, ICUString);
 
 		return Internal::ComputeTextDirection(ICUBiDi, ICUString);
 	}
@@ -330,17 +336,22 @@ public:
 	}
 
 	virtual ETextDirection ComputeTextDirection(const FString& InString, TArray<FTextDirectionInfo>& OutTextDirectionInfo) override
+	{
+		return FICUTextBiDi::ComputeTextDirection(*InString, 0, InString.Len(), OutTextDirectionInfo);
+	}
+
+	virtual ETextDirection ComputeTextDirection(const TCHAR* InString, const int32 InStringStartIndex, const int32 InStringLen, TArray<FTextDirectionInfo>& OutTextDirectionInfo) override
 	{
 		OutTextDirectionInfo.Reset();
 
-		if (InString.IsEmpty())
+		if (InStringLen == 0)
 		{
 			return ETextDirection::LeftToRight;
 		}
 
-		StringConverter.ConvertString(InString, ICUString);
+		StringConverter.ConvertString(InString, InStringStartIndex, InStringLen, ICUString);
 
-		return Internal::ComputeTextDirection(ICUBiDi, ICUString, OutTextDirectionInfo);
+		return Internal::ComputeTextDirection(ICUBiDi, ICUString, InStringStartIndex, OutTextDirectionInfo);
 	}
 
 private:
@@ -367,12 +378,17 @@ ETextDirection ComputeTextDirection(const FText& InText)
 
 ETextDirection ComputeTextDirection(const FString& InString)
 {
-	if (InString.IsEmpty())
+	return ComputeTextDirection(*InString, 0, InString.Len());
+}
+
+ETextDirection ComputeTextDirection(const TCHAR* InString, const int32 InStringStartIndex, const int32 InStringLen)
+{
+	if (InStringLen == 0)
 	{
 		return ETextDirection::LeftToRight;
 	}
 
-	icu::UnicodeString ICUString = ICUUtilities::ConvertString(InString);
+	icu::UnicodeString ICUString = ICUUtilities::ConvertString(InString, InStringStartIndex, InStringLen);
 
 	UErrorCode ICUStatus = U_ZERO_ERROR;
 	UBiDi* ICUBiDi = ubidi_openSized(ICUString.length(), 0, &ICUStatus);
@@ -399,21 +415,26 @@ ETextDirection ComputeTextDirection(const FText& InText, TArray<FTextDirectionIn
 }
 
 ETextDirection ComputeTextDirection(const FString& InString, TArray<FTextDirectionInfo>& OutTextDirectionInfo)
+{
+	return ComputeTextDirection(*InString, 0, InString.Len(), OutTextDirectionInfo);
+}
+
+ETextDirection ComputeTextDirection(const TCHAR* InString, const int32 InStringStartIndex, const int32 InStringLen, TArray<FTextDirectionInfo>& OutTextDirectionInfo)
 {
 	OutTextDirectionInfo.Reset();
 
-	if (InString.IsEmpty())
+	if (InStringLen == 0)
 	{
 		return ETextDirection::LeftToRight;
 	}
 
-	icu::UnicodeString ICUString = ICUUtilities::ConvertString(InString);
+	icu::UnicodeString ICUString = ICUUtilities::ConvertString(InString, InStringStartIndex, InStringLen);
 
 	UErrorCode ICUStatus = U_ZERO_ERROR;
 	UBiDi* ICUBiDi = ubidi_openSized(ICUString.length(), 0, &ICUStatus);
 	if (ICUBiDi && U_SUCCESS(ICUStatus))
 	{
-		const ETextDirection ReturnDirection = Internal::ComputeTextDirection(ICUBiDi, ICUString, OutTextDirectionInfo);
+		const ETextDirection ReturnDirection = Internal::ComputeTextDirection(ICUBiDi, ICUString, InStringStartIndex, OutTextDirectionInfo);
 
 		ubidi_close(ICUBiDi);
 		ICUBiDi = nullptr;
diff --git a/Engine/Source/Runtime/Core/Private/Internationalization/LegacyText.cpp b/Engine/Source/Runtime/Core/Private/Internationalization/LegacyText.cpp
--- a/Engine/Source/Runtime/Core/Private/Internationalization/LegacyText.cpp
+++ b/Engine/Source/Runtime/Core/Private/Internationalization/LegacyText.cpp
@@ -112,6 +112,11 @@ public:
 	}
 
 	virtual ETextDirection ComputeTextDirection(const FString& InString) override
+	{
+		return FLegacyTextBiDi::ComputeTextDirection(*InString, 0, InString.Len());
+	}
+
+	virtual ETextDirection ComputeTextDirection(const TCHAR*, const int32 InStringStartIndex, const int32 InStringLen) override
 	{
 		return ETextDirection::LeftToRight;
 	}
@@ -122,14 +127,19 @@ public:
 	}
 
 	virtual ETextDirection ComputeTextDirection(const FString& InString, TArray<FTextDirectionInfo>& OutTextDirectionInfo) override
+	{
+		return FLegacyTextBiDi::ComputeTextDirection(*InString, 0, InString.Len(), OutTextDirectionInfo);
+	}
+
+	virtual ETextDirection ComputeTextDirection(const TCHAR*, const int32 InStringStartIndex, const int32 InStringLen, TArray<FTextDirectionInfo>& OutTextDirectionInfo) override
 	{
 		OutTextDirectionInfo.Reset();
 
-		if (!InString.IsEmpty())
+		if (InStringLen > 0)
 		{
 			FTextDirectionInfo TextDirectionInfo;
-			TextDirectionInfo.StartIndex = 0;
-			TextDirectionInfo.Length = InString.Len();
+			TextDirectionInfo.StartIndex = InStringStartIndex;
+			TextDirectionInfo.Length = InStringLen;
 			TextDirectionInfo.TextDirection = ETextDirection::LeftToRight;
 			OutTextDirectionInfo.Add(MoveTemp(TextDirectionInfo));
 		}
@@ -151,6 +161,11 @@ ETextDirection ComputeTextDirection(const FText& InText)
 }
 
 ETextDirection ComputeTextDirection(const FString& InString)
+{
+	return ComputeTextDirection(*InString, 0, InString.Len());
+}
+
+ETextDirection ComputeTextDirection(const TCHAR* InString, const int32 InStringStartIndex, const int32 InStringLen)
 {
 	return ETextDirection::LeftToRight;
 }
@@ -161,14 +176,19 @@ ETextDirection ComputeTextDirection(const FText& InText, TArray<FTextDirectionIn
 }
 
 ETextDirection ComputeTextDirection(const FString& InString, TArray<FTextDirectionInfo>& OutTextDirectionInfo)
+{
+	return ComputeTextDirection(*InString, 0, InString.Len(), OutTextDirectionInfo);
+}
+
+ETextDirection ComputeTextDirection(const TCHAR* InString, const int32 InStringStartIndex, const int32 InStringLen, TArray<FTextDirectionInfo>& OutTextDirectionInfo)
 {
 	OutTextDirectionInfo.Reset();
 
-	if (!InString.IsEmpty())
+	if (InStringLen > 0)
 	{
 		FTextDirectionInfo TextDirectionInfo;
-		TextDirectionInfo.StartIndex = 0;
-		TextDirectionInfo.Length = InString.Len();
+		TextDirectionInfo.StartIndex = InStringStartIndex;
+		TextDirectionInfo.Length = InStringLen;
 		TextDirectionInfo.TextDirection = ETextDirection::LeftToRight;
 		OutTextDirectionInfo.Add(MoveTemp(TextDirectionInfo));
 	}
diff --git a/Engine/Source/Runtime/Core/Public/Internationalization/Text.h b/Engine/Source/Runtime/Core/Public/Internationalization/Text.h
--- a/Engine/Source/Runtime/Core/Public/Internationalization/Text.h
+++ b/Engine/Source/Runtime/Core/Public/Internationalization/Text.h
@@ -786,10 +786,12 @@ namespace TextBiDi
 		/** See TextBiDi::ComputeTextDirection */
 		virtual ETextDirection ComputeTextDirection(const FText& InText) = 0;
 		virtual ETextDirection ComputeTextDirection(const FString& InString) = 0;
+		virtual ETextDirection ComputeTextDirection(const TCHAR* InString, const int32 InStringStartIndex, const int32 InStringLen) = 0;
 
 		/** See TextBiDi::ComputeTextDirection */
 		virtual ETextDirection ComputeTextDirection(const FText& InText, TArray<FTextDirectionInfo>& OutTextDirectionInfo) = 0;
 		virtual ETextDirection ComputeTextDirection(const FString& InString, TArray<FTextDirectionInfo>& OutTextDirectionInfo) = 0;
+		virtual ETextDirection ComputeTextDirection(const TCHAR* InString, const int32 InStringStartIndex, const int32 InStringLen, TArray<FTextDirectionInfo>& OutTextDirectionInfo) = 0;
 	};
 
 	/**
@@ -805,6 +807,7 @@ namespace TextBiDi
 	 */
 	CORE_API ETextDirection ComputeTextDirection(const FText& InText);
 	CORE_API ETextDirection ComputeTextDirection(const FString& InString);
+	CORE_API ETextDirection ComputeTextDirection(const TCHAR* InString, const int32 InStringStartIndex, const int32 InStringLen);
 
 	/**
 	 * Utility function which will compute the reading direction of the given text, as well as populate any advanced layout data for the text.
@@ -812,6 +815,7 @@ namespace TextBiDi
 	 */
 	CORE_API ETextDirection ComputeTextDirection(const FText& InText, TArray<FTextDirectionInfo>& OutTextDirectionInfo);
 	CORE_API ETextDirection ComputeTextDirection(const FString& InString, TArray<FTextDirectionInfo>& OutTextDirectionInfo);
+	CORE_API ETextDirection ComputeTextDirection(const TCHAR* InString, const int32 InStringStartIndex, const int32 InStringLen, TArray<FTextDirectionInfo>& OutTextDirectionInfo);
 } // namespace TextBiDi
 
 Expose_TNameOf(FText)
