Linux: fix IsRunningOnBattery() broken on some setups (UE-18358).

diff --git a/Engine/Source/Runtime/Core/Private/Linux/LinuxPlatformMisc.cpp b/Engine/Source/Runtime/Core/Private/Linux/LinuxPlatformMisc.cpp
--- a/Engine/Source/Runtime/Core/Private/Linux/LinuxPlatformMisc.cpp
+++ b/Engine/Source/Runtime/Core/Private/Linux/LinuxPlatformMisc.cpp
@@ -965,18 +965,30 @@ bool FLinuxPlatformMisc::IsRunningOnBattery()
 	LastBatteryCheck = Seconds;
 	bIsOnBattery = false;
 
-	int State = open("/sys/class/power_supply/ADP0/online", O_RDONLY);
-	if (State != -1)
+	// [RCL] 2015-09-30 FIXME: find a more robust way?
+	const int kHardCodedNumBatteries = 10;
+	for (int IdxBattery = 0; IdxBattery < kHardCodedNumBatteries; ++IdxBattery)
 	{
-		// found ACAD device. check its state.
-		ssize_t ReadBytes = read(State, Scratch, 1);
-		close(State);
+		char Filename[128];
+		sprintf(Filename, "/sys/class/power_supply/ADP%d/online", IdxBattery);
 
-		if (ReadBytes > 0)
+		int State = open(Filename, O_RDONLY);
+		if (State != -1)
 		{
-			bIsOnBattery = (Scratch[0] == '0');
+			// found ACAD device. check its state.
+			ssize_t ReadBytes = read(State, Scratch, 1);
+			close(State);
+
+			if (ReadBytes > 0)
+			{
+				bIsOnBattery = (Scratch[0] == '0');
+			}
+
+			break;	// quit checking after we found at least one
 		}
 	}
 
+	// lack of ADP most likely means that we're not on laptop at all
+
 	return bIsOnBattery;
 }
