Zero out profiler callstacks to reduce false entries in the map on some platforms.

diff --git a/Engine/Source/Runtime/Core/Private/ProfilingDebugging/MallocProfiler.cpp b/Engine/Source/Runtime/Core/Private/ProfilingDebugging/MallocProfiler.cpp
--- a/Engine/Source/Runtime/Core/Private/ProfilingDebugging/MallocProfiler.cpp
+++ b/Engine/Source/Runtime/Core/Private/ProfilingDebugging/MallocProfiler.cpp
@@ -631,7 +631,7 @@ int32 FMallocProfiler::GetCallStackIndex()
 	int32 Index = INDEX_NONE;
 
 	// Capture callstack and create CRC.
-	uint64 FullBackTrace[MEMORY_PROFILER_MAX_BACKTRACE_DEPTH + MEMORY_PROFILER_SKIP_NUM_BACKTRACE_ENTRIES];
+	uint64 FullBackTrace[MEMORY_PROFILER_MAX_BACKTRACE_DEPTH + MEMORY_PROFILER_SKIP_NUM_BACKTRACE_ENTRIES] = { 0 };
 	FPlatformStackWalk::CaptureStackBackTrace( FullBackTrace, MEMORY_PROFILER_MAX_BACKTRACE_DEPTH + MEMORY_PROFILER_SKIP_NUM_BACKTRACE_ENTRIES );
 	// Skip first 5 entries as they are inside the allocator.
 	uint64* BackTrace = &FullBackTrace[MEMORY_PROFILER_SKIP_NUM_BACKTRACE_ENTRIES];
