Added 'DoesSectionExist' function to config-cache/GConfig.

diff --git a/Engine/Source/Runtime/Core/Private/Misc/ConfigCacheIni.cpp b/Engine/Source/Runtime/Core/Private/Misc/ConfigCacheIni.cpp
--- a/Engine/Source/Runtime/Core/Private/Misc/ConfigCacheIni.cpp
+++ b/Engine/Source/Runtime/Core/Private/Misc/ConfigCacheIni.cpp
@@ -1739,6 +1739,18 @@ FConfigSection* FConfigCacheIni::GetSectionPrivate( const TCHAR* Section, bool F
 	return Sec;
 }
 
+bool FConfigCacheIni::DoesSectionExist(const TCHAR* Section, const FString& Filename)
+{
+	bool bReturnVal = false;
+
+	FRemoteConfig::Get()->FinishRead(*Filename); // Ensure the remote file has been loaded and processed
+	FConfigFile* File = Find(Filename, false);
+
+	bReturnVal = File != NULL && File->Find(Section) != NULL;
+
+	return bReturnVal;
+}
+
 void FConfigCacheIni::SetString( const TCHAR* Section, const TCHAR* Key, const TCHAR* Value, const FString& Filename )
 {
 	FConfigFile* File = Find( Filename, 1 );
diff --git a/Engine/Source/Runtime/Core/Public/Misc/ConfigCacheIni.h b/Engine/Source/Runtime/Core/Public/Misc/ConfigCacheIni.h
--- a/Engine/Source/Runtime/Core/Public/Misc/ConfigCacheIni.h
+++ b/Engine/Source/Runtime/Core/Public/Misc/ConfigCacheIni.h
@@ -307,6 +307,7 @@ public:
 	bool GetString( const TCHAR* Section, const TCHAR* Key, FString& Value, const FString& Filename );
 	bool GetText( const TCHAR* Section, const TCHAR* Key, FText& Value, const FString& Filename );
 	bool GetSection( const TCHAR* Section, TArray<FString>& Result, const FString& Filename );
+	bool DoesSectionExist(const TCHAR* Section, const FString& Filename);
 	FConfigSection* GetSectionPrivate( const TCHAR* Section, bool Force, bool Const, const FString& Filename );
 	void SetString( const TCHAR* Section, const TCHAR* Key, const TCHAR* Value, const FString& Filename );
 	void SetText( const TCHAR* Section, const TCHAR* Key, const FText& Value, const FString& Filename );
