Engine: Changed a failure to set a cvar due to a previous set via system settings ini to be displayed as Display instead of Warning
- Rationale: There is no way in a game to prevent the project settings or scalability from trying to set the value and failing with the message, and a user had to go out of their way to add one and intended them to exist (there's no UI way to set them in this manner in the editor)

[INTEGRATE] Change 2712339 by [REDACTED_PERSON][REDACTED_HANDLE] on 2015/09/30 22:00:34

diff --git a/Engine/Source/Runtime/Core/Private/HAL/ConsoleManager.cpp b/Engine/Source/Runtime/Core/Private/HAL/ConsoleManager.cpp
--- a/Engine/Source/Runtime/Core/Private/HAL/ConsoleManager.cpp
+++ b/Engine/Source/Runtime/Core/Private/HAL/ConsoleManager.cpp
@@ -121,12 +121,24 @@ public:
 		{
 			FConsoleManager& ConsoleManager = (FConsoleManager&)IConsoleManager::Get();
 			FString CVarName = ConsoleManager.FindConsoleObjectName(this);
-			UE_LOG(LogConsoleManager, Warning,
-				TEXT("A console variable Set() for '%s' failed. The access with 'SetBy%s' failed because 'SetBy%s' did set it before and and has a higher priority (this might be an intentional override)"),
+
+			const FString Message = FString::Printf(TEXT("Setting the console variable '%s' with 'SetBy%s' was ignored as it is lower priority than the previous 'SetBy%s'"),
 				CVarName.IsEmpty() ? TEXT("unknown?") : *CVarName,
 				GetSetByTCHAR((EConsoleVariableFlags)NewPri),
 				GetSetByTCHAR((EConsoleVariableFlags)OldPri)
 				);
+
+			// If it was set by an ini that has to be hand edited, it is not an issue if a lower priority system tried and failed to set it afterwards
+			const bool bIntentionallyIgnored = (OldPri == EConsoleVariableFlags::ECVF_SetBySystemSettingsIni);
+
+			if (bIntentionallyIgnored)
+			{
+				UE_LOG(LogConsoleManager, Display, TEXT("%s"), *Message);
+			}
+			else
+			{
+				UE_LOG(LogConsoleManager, Warning, TEXT("%s"), *Message);
+			}
 		}
 
 		return bRet;
