Merging using Orion->UE4

CL# 2712683
Fat-free TArray and FString ranged-for iteration in Shipping/Test builds.
Validation checks for modifications to a TArray mid-iteration in non-Shipping/Test builds, or whenever you set TARRAY_RANGED_FOR_CHECKS to 1.
Fixes to revealed ranged-for bugs in ULevel::RouteActorInitialize and UK2Node::ReallocatePinsDuringReconstruction.

diff --git a/Engine/Source/Editor/BlueprintGraph/Private/K2Node.cpp b/Engine/Source/Editor/BlueprintGraph/Private/K2Node.cpp
--- a/Engine/Source/Editor/BlueprintGraph/Private/K2Node.cpp
+++ b/Engine/Source/Editor/BlueprintGraph/Private/K2Node.cpp
@@ -299,6 +299,7 @@ void UK2Node::ReallocatePinsDuringReconstruction(TArray<UEdGraphPin*>& OldPins)
 					{
 						NewPin->PinType = OldPin->ParentPin->PinType;
 						GetSchema()->SplitPin(NewPin);
+						break;
 					}
 				}
 			}
diff --git a/Engine/Source/Runtime/Core/Public/Containers/Array.h b/Engine/Source/Runtime/Core/Public/Containers/Array.h
--- a/Engine/Source/Runtime/Core/Public/Containers/Array.h
+++ b/Engine/Source/Runtime/Core/Public/Containers/Array.h
@@ -12,6 +12,12 @@
 
 #define DEBUG_HEAP 0
 
+#if UE_BUILD_SHIPPING || UE_BUILD_TEST
+	#define TARRAY_RANGED_FOR_CHECKS 0
+#else
+	#define TARRAY_RANGED_FOR_CHECKS 1
+#endif
+
 
 /**
  * Generic iterator which can operate on types that expose the following:
@@ -131,6 +137,91 @@ FORCEINLINE TIndexedContainerIterator<ContainerType, ElementType, IndexType> ope
 	return RHS + Offset;
 }
 
+#if TARRAY_RANGED_FOR_CHECKS
+
+	/**
+	 * Pointer-like iterator type for ranged-for loops which checks that the
+	 * container hasn't been resized during iteration.
+	 */
+	template <typename ElementType>
+	struct TCheckedPointerIterator
+	{
+		// This iterator type only supports the minimal functionality needed to support
+		// C++ ranged-for syntax.  For example, it does not provide post-increment ++ nor ==.
+		//
+		// We do add an operator-- to help FString implementation
+
+		explicit TCheckedPointerIterator(const int32& InNum, ElementType* InPtr)
+			: Ptr       (InPtr)
+			, CurrentNum(InNum)
+			, InitialNum(InNum)
+		{
+		}
+
+		FORCEINLINE ElementType& operator*() const
+		{
+			return *Ptr;
+		}
+
+		FORCEINLINE TCheckedPointerIterator& operator++()
+		{
+			++Ptr;
+			return *this;
+		}
+
+		FORCEINLINE TCheckedPointerIterator& operator--()
+		{
+			--Ptr;
+			return *this;
+		}
+
+	private:
+		ElementType* Ptr;
+		const int32& CurrentNum;
+		int32        InitialNum;
+
+		friend bool operator!=(const TCheckedPointerIterator& Lhs, const TCheckedPointerIterator& Rhs)
+		{
+			// We only need to do the check in this operator, because no other operator will be
+			// called until after this one returns.
+			//
+			// Also, we should only need to check one side of this comparison - if the other iterator isn't
+			// even from the same array then the compiler has generated bad code.
+			checkf(Lhs.CurrentNum == Lhs.InitialNum, TEXT("Array has changed during ranged-for iteration!"));
+			return Lhs.Ptr != Rhs.Ptr;
+		}
+	};
+
+#endif
+
+template <typename ElementType, typename IteratorType>
+struct TDereferencingIterator
+{
+	explicit TDereferencingIterator(IteratorType InIter)
+		: Iter(InIter)
+	{
+	}
+
+	FORCEINLINE ElementType& operator*() const
+	{
+		return *(ElementType*)*Iter;
+	}
+
+	FORCEINLINE TDereferencingIterator& operator++()
+	{
+		++Iter;
+		return *this;
+	}
+
+private:
+	IteratorType Iter;
+
+	FORCEINLINE friend bool operator!=(const TDereferencingIterator& Lhs, const TDereferencingIterator& Rhs)
+	{
+		return Lhs.Iter != Rhs.Iter;
+	}
+};
+
 /**
  * Base dynamic array.
  * An untyped data array; mirrors a TArray's members, but doesn't need an exact C++ type for its elements.
@@ -2217,15 +2308,30 @@ public:
 		return TConstIterator(*this);
 	}
 
+	#if TARRAY_RANGED_FOR_CHECKS
+		typedef TCheckedPointerIterator<      ElementType> RangedForIteratorType;
+		typedef TCheckedPointerIterator<const ElementType> RangedForConstIteratorType;
+	#else
+		typedef       ElementType* RangedForIteratorType;
+		typedef const ElementType* RangedForConstIteratorType;
+	#endif
+
 private:
 	/**
 	 * DO NOT USE DIRECTLY
 	 * STL-like iterators to enable range-based for loop support.
 	 */
-	FORCEINLINE friend TIterator      begin(      TArray& Array) { return TIterator     (Array); }
-	FORCEINLINE friend TConstIterator begin(const TArray& Array) { return TConstIterator(Array); }
-	FORCEINLINE friend TIterator      end  (      TArray& Array) { return TIterator     (Array, Array.Num()); }
-	FORCEINLINE friend TConstIterator end  (const TArray& Array) { return TConstIterator(Array, Array.Num()); }
+	#if TARRAY_RANGED_FOR_CHECKS
+		FORCEINLINE friend RangedForIteratorType      begin(      TArray& Array) { return RangedForIteratorType     (Array.ArrayNum, Array.GetData()); }
+		FORCEINLINE friend RangedForConstIteratorType begin(const TArray& Array) { return RangedForConstIteratorType(Array.ArrayNum, Array.GetData()); }
+		FORCEINLINE friend RangedForIteratorType      end  (      TArray& Array) { return RangedForIteratorType     (Array.ArrayNum, Array.GetData() + Array.Num()); }
+		FORCEINLINE friend RangedForConstIteratorType end  (const TArray& Array) { return RangedForConstIteratorType(Array.ArrayNum, Array.GetData() + Array.Num()); }
+	#else
+		FORCEINLINE friend RangedForIteratorType      begin(      TArray& Array) { return Array.GetData(); }
+		FORCEINLINE friend RangedForConstIteratorType begin(const TArray& Array) { return Array.GetData(); }
+		FORCEINLINE friend RangedForIteratorType      end  (      TArray& Array) { return Array.GetData() + Array.Num(); }
+		FORCEINLINE friend RangedForConstIteratorType end  (const TArray& Array) { return Array.GetData() + Array.Num(); }
+	#endif
 
 public:
 	/**
@@ -3373,10 +3479,10 @@ private:
 	 * DO NOT USE DIRECTLY
 	 * STL-like iterators to enable range-based for loop support.
 	 */
-	FORCEINLINE friend TIterator      begin(      TIndirectArray& IndirectArray) { return TIterator     (IndirectArray); }
-	FORCEINLINE friend TConstIterator begin(const TIndirectArray& IndirectArray) { return TConstIterator(IndirectArray); }
-	FORCEINLINE friend TIterator      end  (      TIndirectArray& IndirectArray) { return TIterator     (IndirectArray, IndirectArray.Array.Num()); }
-	FORCEINLINE friend TConstIterator end  (const TIndirectArray& IndirectArray) { return TConstIterator(IndirectArray, IndirectArray.Array.Num()); }
+	FORCEINLINE friend TDereferencingIterator<      ElementType, typename InternalArrayType::RangedForIteratorType     > begin(      TIndirectArray& IndirectArray) { return TDereferencingIterator<      ElementType, typename InternalArrayType::RangedForIteratorType     >(begin(IndirectArray.Array)); }
+	FORCEINLINE friend TDereferencingIterator<const ElementType, typename InternalArrayType::RangedForConstIteratorType> begin(const TIndirectArray& IndirectArray) { return TDereferencingIterator<const ElementType, typename InternalArrayType::RangedForConstIteratorType>(begin(IndirectArray.Array)); }
+	FORCEINLINE friend TDereferencingIterator<      ElementType, typename InternalArrayType::RangedForIteratorType     > end  (      TIndirectArray& IndirectArray) { return TDereferencingIterator<      ElementType, typename InternalArrayType::RangedForIteratorType     >(end  (IndirectArray.Array)); }
+	FORCEINLINE friend TDereferencingIterator<const ElementType, typename InternalArrayType::RangedForConstIteratorType> end  (const TIndirectArray& IndirectArray) { return TDereferencingIterator<const ElementType, typename InternalArrayType::RangedForConstIteratorType>(end  (IndirectArray.Array)); }
 
 	InternalArrayType Array;
 };
diff --git a/Engine/Source/Runtime/Core/Public/Containers/UnrealString.h b/Engine/Source/Runtime/Core/Public/Containers/UnrealString.h
--- a/Engine/Source/Runtime/Core/Public/Containers/UnrealString.h
+++ b/Engine/Source/Runtime/Core/Public/Containers/UnrealString.h
@@ -240,10 +240,10 @@ private:
 	 * DO NOT USE DIRECTLY
 	 * STL-like iterators to enable range-based for loop support.
 	 */
-	FORCEINLINE friend TIterator      begin(      FString& Str) { return begin(Str.Data); }
-	FORCEINLINE friend TConstIterator begin(const FString& Str) { return begin(Str.Data); }
-	FORCEINLINE friend TIterator      end  (      FString& Str) { TIterator      Result = end(Str.Data); if (Str.Data.Num()) { --Result; } return Result; }
-	FORCEINLINE friend TConstIterator end  (const FString& Str) { TConstIterator Result = end(Str.Data); if (Str.Data.Num()) { --Result; } return Result; }
+	FORCEINLINE friend DataType::RangedForIteratorType      begin(      FString& Str) { auto Result = begin(Str.Data);                                   return Result; }
+	FORCEINLINE friend DataType::RangedForConstIteratorType begin(const FString& Str) { auto Result = begin(Str.Data);                                   return Result; }
+	FORCEINLINE friend DataType::RangedForIteratorType      end  (      FString& Str) { auto Result = end  (Str.Data); if (Str.Data.Num()) { --Result; } return Result; }
+	FORCEINLINE friend DataType::RangedForConstIteratorType end  (const FString& Str) { auto Result = end  (Str.Data); if (Str.Data.Num()) { --Result; } return Result; }
 
 public:
 	FORCEINLINE uint32 GetAllocatedSize() const
diff --git a/Engine/Source/Runtime/Engine/Private/Level.cpp b/Engine/Source/Runtime/Engine/Private/Level.cpp
--- a/Engine/Source/Runtime/Engine/Private/Level.cpp
+++ b/Engine/Source/Runtime/Engine/Private/Level.cpp
@@ -1675,8 +1675,9 @@ void ULevel::ReleaseRenderingResources()
 void ULevel::RouteActorInitialize()
 {
 	// Send PreInitializeComponents and collect volumes.
-	for( AActor* const Actor : Actors )
+	for( int32 Index = 0; Index < Actors.Num(); ++Index )
 	{
+		AActor* const Actor = Actors[Index];
 		if( Actor && !Actor->IsActorInitialized() )
 		{
 			Actor->PreInitializeComponents();
