Fix for buffer overrun in FArchive::SerializeBits when loading an exact multiple of 8 bits.

diff --git a/Engine/Source/Runtime/Core/Public/Serialization/ArchiveBase.h b/Engine/Source/Runtime/Core/Public/Serialization/ArchiveBase.h
--- a/Engine/Source/Runtime/Core/Public/Serialization/ArchiveBase.h
+++ b/Engine/Source/Runtime/Core/Public/Serialization/ArchiveBase.h
@@ -321,7 +321,7 @@ public:
 	{
 		Serialize(V, (LengthBits + 7) / 8);
 
-		if (IsLoading())
+		if (IsLoading() && (LengthBits % 8) != 0)
 		{
 			((uint8*)V)[LengthBits / 8] &= ((1 << (LengthBits & 7)) - 1);
 		}
