Added TextBiDi functions to compute the base direction of a piece of text.

This will be needed to flow text layout runs/blocks correctly when dealing with text that is primarily RTL.

diff --git a/Engine/Source/Runtime/Core/Private/Internationalization/ICUText.cpp b/Engine/Source/Runtime/Core/Private/Internationalization/ICUText.cpp
--- a/Engine/Source/Runtime/Core/Private/Internationalization/ICUText.cpp
+++ b/Engine/Source/Runtime/Core/Private/Internationalization/ICUText.cpp
@@ -294,6 +294,14 @@ ETextDirection ComputeTextDirection(UBiDi* InICUBiDi, const icu::UnicodeString&
 	return ETextDirection::LeftToRight;
 }
 
+ETextDirection ComputeBaseDirection(const icu::UnicodeString& InICUString)
+{
+	const UBiDiDirection ICUBaseDirection = ubidi_getBaseDirection(InICUString.getBuffer(), InICUString.length());
+
+	// ICUToUE will treat UBIDI_NEUTRAL as LTR
+	return Internal::ICUToUE(ICUBaseDirection);
+}
+
 class FICUTextBiDi : public ITextBiDi
 {
 public:
@@ -354,6 +362,28 @@ public:
 		return Internal::ComputeTextDirection(ICUBiDi, ICUString, InStringStartIndex, OutTextDirectionInfo);
 	}
 
+	virtual ETextDirection ComputeBaseDirection(const FText& InText) override
+	{
+		return FICUTextBiDi::ComputeBaseDirection(InText.ToString());
+	}
+
+	virtual ETextDirection ComputeBaseDirection(const FString& InString) override
+	{
+		return FICUTextBiDi::ComputeBaseDirection(*InString, 0, InString.Len());
+	}
+
+	virtual ETextDirection ComputeBaseDirection(const TCHAR* InString, const int32 InStringStartIndex, const int32 InStringLen) override
+	{
+		if (InStringLen == 0)
+		{
+			return ETextDirection::LeftToRight;
+		}
+
+		StringConverter.ConvertString(InString, InStringStartIndex, InStringLen, ICUString);
+
+		return Internal::ComputeBaseDirection(ICUString);
+	}
+
 private:
 	/** Non-copyable */
 	FICUTextBiDi(const FICUTextBiDi&);
@@ -449,6 +479,28 @@ ETextDirection ComputeTextDirection(const TCHAR* InString, const int32 InStringS
 	return ETextDirection::LeftToRight;
 }
 
+ETextDirection ComputeBaseDirection(const FText& InText)
+{
+	return ComputeBaseDirection(InText.ToString());
+}
+
+ETextDirection ComputeBaseDirection(const FString& InString)
+{
+	return ComputeBaseDirection(*InString, 0, InString.Len());
+}
+
+ETextDirection ComputeBaseDirection(const TCHAR* InString, const int32 InStringStartIndex, const int32 InStringLen)
+{
+	if (InStringLen == 0)
+	{
+		return ETextDirection::LeftToRight;
+	}
+
+	icu::UnicodeString ICUString = ICUUtilities::ConvertString(InString, InStringStartIndex, InStringLen);
+
+	return Internal::ComputeBaseDirection(ICUString);
+}
+
 } // namespace TextBiDi
 
 #endif
diff --git a/Engine/Source/Runtime/Core/Private/Internationalization/LegacyText.cpp b/Engine/Source/Runtime/Core/Private/Internationalization/LegacyText.cpp
--- a/Engine/Source/Runtime/Core/Private/Internationalization/LegacyText.cpp
+++ b/Engine/Source/Runtime/Core/Private/Internationalization/LegacyText.cpp
@@ -146,6 +146,21 @@ public:
 
 		return ETextDirection::LeftToRight;
 	}
+
+	virtual ETextDirection ComputeBaseDirection(const FText& InText) override
+	{
+		return FLegacyTextBiDi::ComputeBaseDirection(InText.ToString());
+	}
+
+	virtual ETextDirection ComputeBaseDirection(const FString& InString) override
+	{
+		return FLegacyTextBiDi::ComputeBaseDirection(*InString, 0, InString.Len());
+	}
+
+	virtual ETextDirection ComputeBaseDirection(const TCHAR*, const int32 InStringStartIndex, const int32 InStringLen) override
+	{
+		return ETextDirection::LeftToRight;
+	}
 };
 
 } // namespace Internal
@@ -196,6 +211,21 @@ ETextDirection ComputeTextDirection(const TCHAR* InString, const int32 InStringS
 	return ETextDirection::LeftToRight;
 }
 
+ETextDirection ComputeBaseDirection(const FText& InText)
+{
+	return ComputeBaseDirection(InText.ToString());
+}
+
+ETextDirection ComputeBaseDirection(const FString& InString)
+{
+	return ComputeBaseDirection(*InString, 0, InString.Len());
+}
+
+ETextDirection ComputeBaseDirection(const TCHAR* InString, const int32 InStringStartIndex, const int32 InStringLen)
+{
+	return ETextDirection::LeftToRight;
+}
+
 } // namespace TextBiDi
 
 #endif
diff --git a/Engine/Source/Runtime/Core/Public/Internationalization/Text.h b/Engine/Source/Runtime/Core/Public/Internationalization/Text.h
--- a/Engine/Source/Runtime/Core/Public/Internationalization/Text.h
+++ b/Engine/Source/Runtime/Core/Public/Internationalization/Text.h
@@ -792,6 +792,11 @@ namespace TextBiDi
 		virtual ETextDirection ComputeTextDirection(const FText& InText, TArray<FTextDirectionInfo>& OutTextDirectionInfo) = 0;
 		virtual ETextDirection ComputeTextDirection(const FString& InString, TArray<FTextDirectionInfo>& OutTextDirectionInfo) = 0;
 		virtual ETextDirection ComputeTextDirection(const TCHAR* InString, const int32 InStringStartIndex, const int32 InStringLen, TArray<FTextDirectionInfo>& OutTextDirectionInfo) = 0;
+
+		/** See TextBiDi::ComputeBaseDirection */
+		virtual ETextDirection ComputeBaseDirection(const FText& InText) = 0;
+		virtual ETextDirection ComputeBaseDirection(const FString& InString) = 0;
+		virtual ETextDirection ComputeBaseDirection(const TCHAR* InString, const int32 InStringStartIndex, const int32 InStringLen) = 0;
 	};
 
 	/**
@@ -816,6 +821,15 @@ namespace TextBiDi
 	CORE_API ETextDirection ComputeTextDirection(const FText& InText, TArray<FTextDirectionInfo>& OutTextDirectionInfo);
 	CORE_API ETextDirection ComputeTextDirection(const FString& InString, TArray<FTextDirectionInfo>& OutTextDirectionInfo);
 	CORE_API ETextDirection ComputeTextDirection(const TCHAR* InString, const int32 InStringStartIndex, const int32 InStringLen, TArray<FTextDirectionInfo>& OutTextDirectionInfo);
+
+	/**
+	 * Utility function which will compute the base direction of the given text.
+	 * This provides the text flow direction that should be used when combining bidirectional text runs together.
+	 * @return RightToLeft if the first character in the string has a bidirectional character type of R or AL, otherwise LeftToRight.
+	 */
+	CORE_API ETextDirection ComputeBaseDirection(const FText& InText);
+	CORE_API ETextDirection ComputeBaseDirection(const FString& InString);
+	CORE_API ETextDirection ComputeBaseDirection(const TCHAR* InString, const int32 InStringStartIndex, const int32 InStringLen);
 } // namespace TextBiDi
 
 Expose_TNameOf(FText)
