Core: Removed unused code from AsyncTask

diff --git a/Engine/Source/Runtime/Core/Public/Async/Async.h b/Engine/Source/Runtime/Core/Public/Async/Async.h
--- a/Engine/Source/Runtime/Core/Public/Async/Async.h
+++ b/Engine/Source/Runtime/Core/Public/Async/Async.h
@@ -11,7 +11,7 @@
  */
 enum class EAsyncExecution
 {
-	/** Execute in Task Graph (for short running tasks). */
+	/** Execute in Task Graph (for short running tyayerasks). */
 	TaskGraph,
 
 	/** Execute in separate thread (for long running tasks). */
@@ -50,16 +50,6 @@ class FAsyncGraphTaskBase
 {
 public:
 
-	/**
-	 * Returns the name of the thread that this task should run on.
-	 *
-	 * @return Always run on any thread.
-	 */
-	ENamedThreads::Type GetDesiredThread()
-	{
-		return ENamedThreads::AnyThread;
-	}
-
 	/**
 	 * Gets the task's stats tracking identifier.
 	 *
@@ -115,6 +105,16 @@ public:
 		SetPromise(Promise, Function);
 	}
 
+	/**
+	 * Returns the name of the thread that this task should run on.
+	 *
+	 * @return Always run on any thread.
+	 */
+	ENamedThreads::Type GetDesiredThread()
+	{
+		return ENamedThreads::AnyThread;
+	}
+
 	/**
 	 * Gets the future that will hold the asynchronous result.
 	 *
@@ -268,6 +268,7 @@ struct FAsyncThreadIndex
  *			return 123;
  *		}
  *
+ * @param ResultType The type of the function's return value.
  * @param Execution The execution method to use, i.e. on Task Graph or in a separate thread.
  * @param Function The function to execute.
  * @result A TFuture object that will receive the return value from the function.
@@ -314,7 +315,20 @@ TFuture<ResultType> Async(EAsyncExecution Execution, TFunction<ResultType()> Fun
 }
 
 
-template<typename ResultType> uint32 TAsyncRunnable<ResultType>::Run()
+/**
+ * Convenience function for executing code asynchronously on the Task Graph.
+ *
+ * @param Thread The name of the thread to run on.
+ * @param Function The function to execute.
+ */
+CORE_API void AsyncTask(ENamedThreads::Type Thread, TFunction<void()> Function);
+
+
+/* Inline functions
+ *****************************************************************************/
+
+template<typename ResultType>
+uint32 TAsyncRunnable<ResultType>::Run()
 {
 	SetPromise(Promise, Function);
 
