Fix issue when loading cvar from ini and game module not loaded yet

Merging 2713350

//depot/UE4-UT/...

to //depot/UE4/...

diff --git a/Engine/Source/Runtime/Core/Private/Misc/ConfigCacheIni.cpp b/Engine/Source/Runtime/Core/Private/Misc/ConfigCacheIni.cpp
--- a/Engine/Source/Runtime/Core/Private/Misc/ConfigCacheIni.cpp
+++ b/Engine/Source/Runtime/Core/Private/Misc/ConfigCacheIni.cpp
@@ -3423,8 +3423,9 @@ void ApplyCVarSettingsFromIni(const TCHAR* InSectionName, const TCHAR* InIniFile
 			}
 			else
 			{
-				UE_LOG(LogConsoleResponse, Verbose, TEXT("Skipping Unknown console variable: '%s = %s'"), *CVarName, *CVarValue);
-				UE_LOG(LogConsoleResponse, Verbose, TEXT("  Found in ini file '%s', in section '[%s]'"), InIniFilename, InSectionName);
+				// Create a dummy that is used when someone registers the variable later on.
+				// this is important for variables created in external modules, such as the game module
+				IConsoleManager::Get().RegisterConsoleVariable(*CVarName, *CVarValue, TEXT("IAmNoRealVariable"), (uint32)ECVF_Unregistered | (uint32)ECVF_CreatedFromIni | SetBy);
 			}
 		}
 	}
