^^

diff --git a/Engine/Source/Runtime/Core/Public/Async/Async.h b/Engine/Source/Runtime/Core/Public/Async/Async.h
--- a/Engine/Source/Runtime/Core/Public/Async/Async.h
+++ b/Engine/Source/Runtime/Core/Public/Async/Async.h
@@ -11,7 +11,7 @@
  */
 enum class EAsyncExecution
 {
-	/** Execute in Task Graph (for short running tyayerasks). */
+	/** Execute in Task Graph (for short running tasks). */
 	TaskGraph,
 
 	/** Execute in separate thread (for long running tasks). */
