Core: Added missing AsyncTask implementation

diff --git a/Engine/Source/Runtime/Core/Private/Async/Async.cpp b/Engine/Source/Runtime/Core/Private/Async/Async.cpp
new file mode 100644
--- /dev/null
+++ b/Engine/Source/Runtime/Core/Private/Async/Async.cpp
@@ -0,0 +1,50 @@
+// Copyright 1998-2015 Epic Games, Inc. All Rights Reserved.
+
+#include "CorePrivatePCH.h"
+#include "Async.h"
+
+
+/* Private types
+ *****************************************************************************/
+
+/** Graph task for simple fire-and-forget asynchronous functions. */
+class FAsyncGraphTask
+	: public FAsyncGraphTaskBase
+{
+public:
+
+	/** Creates and initializes a new instance. */
+	FAsyncGraphTask(ENamedThreads::Type InDesiredThread, TFunction<void()>&& InFunction)
+		: DesiredThread(InDesiredThread)
+		, Function(MoveTemp(InFunction))
+	{ }
+
+	/** Performs the actual task. */
+	void DoTask(ENamedThreads::Type CurrentThread, const FGraphEventRef& MyCompletionGraphEvent)
+	{
+		Function();
+	}
+
+	/** Returns the name of the thread that this task should run on. */
+	ENamedThreads::Type GetDesiredThread()
+	{
+		return DesiredThread;
+	}
+
+private:
+
+	/** The thread to execute the function on. */
+	ENamedThreads::Type DesiredThread;
+
+	/** The function to execute on the Task Graph. */
+	TFunction<void()> Function;
+};
+
+
+/* Global functions
+ *****************************************************************************/
+
+void AsyncTask(ENamedThreads::Type Thread, TFunction<void()> Function)
+{
+	TGraphTask<FAsyncGraphTask>::CreateTask().ConstructAndDispatchWhenReady(Thread, MoveTemp(Function));
+}
