You now have to provide the base direction when getting back text run information from TextBiDi

This allows the BiDi algorithm to correctly place some characters (such as brackets or quotes) based on whether the text will be in an LTR or RTL layout flow.

diff --git a/Engine/Source/Runtime/Core/Private/Internationalization/ICUText.cpp b/Engine/Source/Runtime/Core/Private/Internationalization/ICUText.cpp
--- a/Engine/Source/Runtime/Core/Private/Internationalization/ICUText.cpp
+++ b/Engine/Source/Runtime/Core/Private/Internationalization/ICUText.cpp
@@ -247,11 +247,17 @@ FORCEINLINE ETextDirection ICUToUE(const UBiDiDirection InDirection)
 	return ETextDirection::LeftToRight;
 }
 
+UBiDiLevel GetParagraphDirection(const ETextDirection InBaseDirection)
+{
+	check(InBaseDirection != ETextDirection::Mixed);
+	return (InBaseDirection == ETextDirection::LeftToRight) ? 0 : 1; // 0 = LTR, 1 = RTL
+}
+
 ETextDirection ComputeTextDirection(UBiDi* InICUBiDi, const icu::UnicodeString& InICUString)
 {
 	UErrorCode ICUStatus = U_ZERO_ERROR;
 
-	ubidi_setPara(InICUBiDi, InICUString.getBuffer(), InICUString.length(), 0, nullptr, &ICUStatus);
+	ubidi_setPara(InICUBiDi, InICUString.getBuffer(), InICUString.length(), GetParagraphDirection(ETextDirection::LeftToRight), nullptr, &ICUStatus);
 
 	if (U_SUCCESS(ICUStatus))
 	{
@@ -265,11 +271,11 @@ ETextDirection ComputeTextDirection(UBiDi* InICUBiDi, const icu::UnicodeString&
 	return ETextDirection::LeftToRight;
 }
 
-ETextDirection ComputeTextDirection(UBiDi* InICUBiDi, const icu::UnicodeString& InICUString, const int32 InStringOffset, TArray<FTextDirectionInfo>& OutTextDirectionInfo)
+ETextDirection ComputeTextDirection(UBiDi* InICUBiDi, const icu::UnicodeString& InICUString, const int32 InStringOffset, const ETextDirection InBaseDirection, TArray<FTextDirectionInfo>& OutTextDirectionInfo)
 {
 	UErrorCode ICUStatus = U_ZERO_ERROR;
 
-	ubidi_setPara(InICUBiDi, InICUString.getBuffer(), InICUString.length(), 0, nullptr, &ICUStatus);
+	ubidi_setPara(InICUBiDi, InICUString.getBuffer(), InICUString.length(), GetParagraphDirection(InBaseDirection), nullptr, &ICUStatus);
 
 	if (U_SUCCESS(ICUStatus))
 	{
@@ -338,17 +344,17 @@ public:
 		return Internal::ComputeTextDirection(ICUBiDi, ICUString);
 	}
 
-	virtual ETextDirection ComputeTextDirection(const FText& InText, TArray<FTextDirectionInfo>& OutTextDirectionInfo) override
+	virtual ETextDirection ComputeTextDirection(const FText& InText, const ETextDirection InBaseDirection, TArray<FTextDirectionInfo>& OutTextDirectionInfo) override
 	{
-		return FICUTextBiDi::ComputeTextDirection(InText.ToString(), OutTextDirectionInfo);
+		return FICUTextBiDi::ComputeTextDirection(InText.ToString(), InBaseDirection, OutTextDirectionInfo);
 	}
 
-	virtual ETextDirection ComputeTextDirection(const FString& InString, TArray<FTextDirectionInfo>& OutTextDirectionInfo) override
+	virtual ETextDirection ComputeTextDirection(const FString& InString, const ETextDirection InBaseDirection, TArray<FTextDirectionInfo>& OutTextDirectionInfo) override
 	{
-		return FICUTextBiDi::ComputeTextDirection(*InString, 0, InString.Len(), OutTextDirectionInfo);
+		return FICUTextBiDi::ComputeTextDirection(*InString, 0, InString.Len(), InBaseDirection, OutTextDirectionInfo);
 	}
 
-	virtual ETextDirection ComputeTextDirection(const TCHAR* InString, const int32 InStringStartIndex, const int32 InStringLen, TArray<FTextDirectionInfo>& OutTextDirectionInfo) override
+	virtual ETextDirection ComputeTextDirection(const TCHAR* InString, const int32 InStringStartIndex, const int32 InStringLen, const ETextDirection InBaseDirection, TArray<FTextDirectionInfo>& OutTextDirectionInfo) override
 	{
 		OutTextDirectionInfo.Reset();
 
@@ -359,7 +365,7 @@ public:
 
 		StringConverter.ConvertString(InString, InStringStartIndex, InStringLen, ICUString);
 
-		return Internal::ComputeTextDirection(ICUBiDi, ICUString, InStringStartIndex, OutTextDirectionInfo);
+		return Internal::ComputeTextDirection(ICUBiDi, ICUString, InStringStartIndex, InBaseDirection, OutTextDirectionInfo);
 	}
 
 	virtual ETextDirection ComputeBaseDirection(const FText& InText) override
@@ -439,17 +445,17 @@ ETextDirection ComputeTextDirection(const TCHAR* InString, const int32 InStringS
 	return ETextDirection::LeftToRight;
 }
 
-ETextDirection ComputeTextDirection(const FText& InText, TArray<FTextDirectionInfo>& OutTextDirectionInfo)
+ETextDirection ComputeTextDirection(const FText& InText, const ETextDirection InBaseDirection, TArray<FTextDirectionInfo>& OutTextDirectionInfo)
 {
-	return ComputeTextDirection(InText.ToString(), OutTextDirectionInfo);
+	return ComputeTextDirection(InText.ToString(), InBaseDirection, OutTextDirectionInfo);
 }
 
-ETextDirection ComputeTextDirection(const FString& InString, TArray<FTextDirectionInfo>& OutTextDirectionInfo)
+ETextDirection ComputeTextDirection(const FString& InString, const ETextDirection InBaseDirection, TArray<FTextDirectionInfo>& OutTextDirectionInfo)
 {
-	return ComputeTextDirection(*InString, 0, InString.Len(), OutTextDirectionInfo);
+	return ComputeTextDirection(*InString, 0, InString.Len(), InBaseDirection, OutTextDirectionInfo);
 }
 
-ETextDirection ComputeTextDirection(const TCHAR* InString, const int32 InStringStartIndex, const int32 InStringLen, TArray<FTextDirectionInfo>& OutTextDirectionInfo)
+ETextDirection ComputeTextDirection(const TCHAR* InString, const int32 InStringStartIndex, const int32 InStringLen, const ETextDirection InBaseDirection, TArray<FTextDirectionInfo>& OutTextDirectionInfo)
 {
 	OutTextDirectionInfo.Reset();
 
@@ -464,7 +470,7 @@ ETextDirection ComputeTextDirection(const TCHAR* InString, const int32 InStringS
 	UBiDi* ICUBiDi = ubidi_openSized(ICUString.length(), 0, &ICUStatus);
 	if (ICUBiDi && U_SUCCESS(ICUStatus))
 	{
-		const ETextDirection ReturnDirection = Internal::ComputeTextDirection(ICUBiDi, ICUString, InStringStartIndex, OutTextDirectionInfo);
+		const ETextDirection ReturnDirection = Internal::ComputeTextDirection(ICUBiDi, ICUString, InStringStartIndex, InBaseDirection, OutTextDirectionInfo);
 
 		ubidi_close(ICUBiDi);
 		ICUBiDi = nullptr;
diff --git a/Engine/Source/Runtime/Core/Private/Internationalization/LegacyText.cpp b/Engine/Source/Runtime/Core/Private/Internationalization/LegacyText.cpp
--- a/Engine/Source/Runtime/Core/Private/Internationalization/LegacyText.cpp
+++ b/Engine/Source/Runtime/Core/Private/Internationalization/LegacyText.cpp
@@ -121,17 +121,17 @@ public:
 		return ETextDirection::LeftToRight;
 	}
 
-	virtual ETextDirection ComputeTextDirection(const FText& InText, TArray<FTextDirectionInfo>& OutTextDirectionInfo) override
+	virtual ETextDirection ComputeTextDirection(const FText& InText, const ETextDirection InBaseDirection, TArray<FTextDirectionInfo>& OutTextDirectionInfo) override
 	{
-		return FLegacyTextBiDi::ComputeTextDirection(InText.ToString(), OutTextDirectionInfo);
+		return FLegacyTextBiDi::ComputeTextDirection(InText.ToString(), InBaseDirection, OutTextDirectionInfo);
 	}
 
-	virtual ETextDirection ComputeTextDirection(const FString& InString, TArray<FTextDirectionInfo>& OutTextDirectionInfo) override
+	virtual ETextDirection ComputeTextDirection(const FString& InString, const ETextDirection InBaseDirection, TArray<FTextDirectionInfo>& OutTextDirectionInfo) override
 	{
-		return FLegacyTextBiDi::ComputeTextDirection(*InString, 0, InString.Len(), OutTextDirectionInfo);
+		return FLegacyTextBiDi::ComputeTextDirection(*InString, 0, InString.Len(), InBaseDirection, OutTextDirectionInfo);
 	}
 
-	virtual ETextDirection ComputeTextDirection(const TCHAR*, const int32 InStringStartIndex, const int32 InStringLen, TArray<FTextDirectionInfo>& OutTextDirectionInfo) override
+	virtual ETextDirection ComputeTextDirection(const TCHAR*, const int32 InStringStartIndex, const int32 InStringLen, const ETextDirection InBaseDirection, TArray<FTextDirectionInfo>& OutTextDirectionInfo) override
 	{
 		OutTextDirectionInfo.Reset();
 
@@ -185,17 +185,17 @@ ETextDirection ComputeTextDirection(const TCHAR* InString, const int32 InStringS
 	return ETextDirection::LeftToRight;
 }
 
-ETextDirection ComputeTextDirection(const FText& InText, TArray<FTextDirectionInfo>& OutTextDirectionInfo)
+ETextDirection ComputeTextDirection(const FText& InText, const ETextDirection InBaseDirection, TArray<FTextDirectionInfo>& OutTextDirectionInfo)
 {
-	return ComputeTextDirection(InText.ToString(), OutTextDirectionInfo);
+	return ComputeTextDirection(InText.ToString(), InBaseDirection, OutTextDirectionInfo);
 }
 
-ETextDirection ComputeTextDirection(const FString& InString, TArray<FTextDirectionInfo>& OutTextDirectionInfo)
+ETextDirection ComputeTextDirection(const FString& InString, const ETextDirection InBaseDirection, TArray<FTextDirectionInfo>& OutTextDirectionInfo)
 {
-	return ComputeTextDirection(*InString, 0, InString.Len(), OutTextDirectionInfo);
+	return ComputeTextDirection(*InString, 0, InString.Len(), InBaseDirection, OutTextDirectionInfo);
 }
 
-ETextDirection ComputeTextDirection(const TCHAR* InString, const int32 InStringStartIndex, const int32 InStringLen, TArray<FTextDirectionInfo>& OutTextDirectionInfo)
+ETextDirection ComputeTextDirection(const TCHAR* InString, const int32 InStringStartIndex, const int32 InStringLen, const ETextDirection InBaseDirection, TArray<FTextDirectionInfo>& OutTextDirectionInfo)
 {
 	OutTextDirectionInfo.Reset();
 
diff --git a/Engine/Source/Runtime/Core/Public/Internationalization/Text.h b/Engine/Source/Runtime/Core/Public/Internationalization/Text.h
--- a/Engine/Source/Runtime/Core/Public/Internationalization/Text.h
+++ b/Engine/Source/Runtime/Core/Public/Internationalization/Text.h
@@ -789,9 +789,9 @@ namespace TextBiDi
 		virtual ETextDirection ComputeTextDirection(const TCHAR* InString, const int32 InStringStartIndex, const int32 InStringLen) = 0;
 
 		/** See TextBiDi::ComputeTextDirection */
-		virtual ETextDirection ComputeTextDirection(const FText& InText, TArray<FTextDirectionInfo>& OutTextDirectionInfo) = 0;
-		virtual ETextDirection ComputeTextDirection(const FString& InString, TArray<FTextDirectionInfo>& OutTextDirectionInfo) = 0;
-		virtual ETextDirection ComputeTextDirection(const TCHAR* InString, const int32 InStringStartIndex, const int32 InStringLen, TArray<FTextDirectionInfo>& OutTextDirectionInfo) = 0;
+		virtual ETextDirection ComputeTextDirection(const FText& InText, const ETextDirection InBaseDirection, TArray<FTextDirectionInfo>& OutTextDirectionInfo) = 0;
+		virtual ETextDirection ComputeTextDirection(const FString& InString, const ETextDirection InBaseDirection, TArray<FTextDirectionInfo>& OutTextDirectionInfo) = 0;
+		virtual ETextDirection ComputeTextDirection(const TCHAR* InString, const int32 InStringStartIndex, const int32 InStringLen, const ETextDirection InBaseDirection, TArray<FTextDirectionInfo>& OutTextDirectionInfo) = 0;
 
 		/** See TextBiDi::ComputeBaseDirection */
 		virtual ETextDirection ComputeBaseDirection(const FText& InText) = 0;
@@ -816,11 +816,12 @@ namespace TextBiDi
 
 	/**
 	 * Utility function which will compute the reading direction of the given text, as well as populate any advanced layout data for the text.
+	 * The base direction is the overall reading direction of the text (see ComputeBaseDirection). This will affect where some characters (such as brackets and quotes) are placed within the resultant FTextDirectionInfo data.
 	 * @return LeftToRight if all of the text is LTR, RightToLeft if all of the text is RTL, or Mixed if the text contains both LTR and RTL text.
 	 */
-	CORE_API ETextDirection ComputeTextDirection(const FText& InText, TArray<FTextDirectionInfo>& OutTextDirectionInfo);
-	CORE_API ETextDirection ComputeTextDirection(const FString& InString, TArray<FTextDirectionInfo>& OutTextDirectionInfo);
-	CORE_API ETextDirection ComputeTextDirection(const TCHAR* InString, const int32 InStringStartIndex, const int32 InStringLen, TArray<FTextDirectionInfo>& OutTextDirectionInfo);
+	CORE_API ETextDirection ComputeTextDirection(const FText& InText, const ETextDirection InBaseDirection, TArray<FTextDirectionInfo>& OutTextDirectionInfo);
+	CORE_API ETextDirection ComputeTextDirection(const FString& InString, const ETextDirection InBaseDirection, TArray<FTextDirectionInfo>& OutTextDirectionInfo);
+	CORE_API ETextDirection ComputeTextDirection(const TCHAR* InString, const int32 InStringStartIndex, const int32 InStringLen, const ETextDirection InBaseDirection, TArray<FTextDirectionInfo>& OutTextDirectionInfo);
 
 	/**
 	 * Utility function which will compute the base direction of the given text.
