Updated FICUInternationalization::LoadDLLs to be able to load VS2015 DLLs

diff --git a/Engine/Source/Runtime/Core/Private/Internationalization/ICUInternationalization.cpp b/Engine/Source/Runtime/Core/Private/Internationalization/ICUInternationalization.cpp
--- a/Engine/Source/Runtime/Core/Private/Internationalization/ICUInternationalization.cpp
+++ b/Engine/Source/Runtime/Core/Private/Internationalization/ICUInternationalization.cpp
@@ -155,13 +155,15 @@ void FICUInternationalization::LoadDLLs()
 	const FString PlatformFolderName = TEXT("Win32");
 #endif //PLATFORM_*BITS
 
-#if _MSC_VER >= 1800
+#if _MSC_VER == 1900
+	const FString VSVersionFolderName = TEXT("VS2015");
+#elif _MSC_VER == 1800
 	const FString VSVersionFolderName = TEXT("VS2013");
 #else
-	const FString VSVersionFolderName = TEXT("VS2012");
+	#error "FICUInternationalization::LoadDLLs - Unknown _MSC_VER! Please update this code for this version of MSVC."
 #endif //_MSVC_VER
 
-	// Windows requires support for 32/64 bit and also VC11 and VC12 runtimes.
+	// Windows requires support for 32/64 bit and different MSVC runtimes.
 	const FString TargetSpecificPath = ICUBinariesRoot / PlatformFolderName / VSVersionFolderName;
 
 	// Windows libraries use a specific naming convention.
