Added nan for noninitialized transform
Merging

//depot/UE4-Orion/Engine/Source/Runtime/Core/Public/Math/...

to //depot/UE4/Engine/Source/Runtime/Core/Public/Math/...

diff --git a/Engine/Source/Runtime/Core/Public/Math/Transform.h b/Engine/Source/Runtime/Core/Public/Math/Transform.h
--- a/Engine/Source/Runtime/Core/Public/Math/Transform.h
+++ b/Engine/Source/Runtime/Core/Public/Math/Transform.h
@@ -78,15 +78,6 @@ public:
 		, Translation(0.f)
 		, Scale3D(1.f)
 	{
-		// Note: This can be used to track down initialization issues with bone atom arrays; but it will
-		// cause issues with transient fields such as RootMotionDelta that get initialized to 0 by default
-#if 0
-		float qnan = FMath::Log2(-5.3f);
-		check(FMath::IsNaN(qnan));
-		Translation = FVector(qnan, qnan, qnan);
-		Rotation = FQuat(qnan, qnan, qnan, qnan);
-		Scale3D = FVector(qnan, qnan, qnan);
-#endif
 	}
 
 	/**
@@ -107,6 +98,15 @@ public:
 	 */
 	FORCEINLINE explicit FTransform(ENoInit) 
 	{
+		// Note: This can be used to track down initialization issues with bone transform arrays; but it will
+		// cause issues with transient fields such as RootMotionDelta that get initialized to 0 by default
+#if ENABLE_NAN_DIAGNOSTIC
+		float qnan = FMath::Log2(-5.3f);
+		check(FMath::IsNaN(qnan));
+		Translation = FVector(qnan, qnan, qnan);
+		Rotation = FQuat(qnan, qnan, qnan, qnan);
+		Scale3D = FVector(qnan, qnan, qnan);
+#endif
 	}
 
 	/**
diff --git a/Engine/Source/Runtime/Core/Public/Math/TransformVectorized.h b/Engine/Source/Runtime/Core/Public/Math/TransformVectorized.h
--- a/Engine/Source/Runtime/Core/Public/Math/TransformVectorized.h
+++ b/Engine/Source/Runtime/Core/Public/Math/TransformVectorized.h
@@ -80,16 +80,6 @@ public:
 		Translation = VectorZero();
 		// Scale3D = {1,1,1,0);
 		Scale3D = VectorSet_W0( VectorOne() );
-
-		// Note: This can be used to track down initialization issues with bone transform arrays; but it will
-		// cause issues with transient fields such as RootMotionDelta that get initialized to 0 by default
-#if 0
-		float qnan = FMath::Log2(-5.3f);
-		check(FMath::IsNaN(qnan));
-		Translation = MakeVectorRegister(qnan, qnan, qnan, qnan);
-		Rotation = MakeVectorRegister(qnan, qnan, qnan, qnan);
-		Scale3D = MakeVectorRegister(qnan, qnan, qnan, qnan);
-#endif
 	}
 
 	/**
@@ -203,6 +193,15 @@ public:
 	 */
 	FORCEINLINE explicit FTransform(ENoInit) 
 	{
+		// Note: This can be used to track down initialization issues with bone transform arrays; but it will
+		// cause issues with transient fields such as RootMotionDelta that get initialized to 0 by default
+#if ENABLE_NAN_DIAGNOSTIC
+		float qnan = FMath::Log2(-5.3f);
+		check(FMath::IsNaN(qnan));
+		Translation = MakeVectorRegister(qnan, qnan, qnan, qnan);
+		Rotation = MakeVectorRegister(qnan, qnan, qnan, qnan);
+		Scale3D = MakeVectorRegister(qnan, qnan, qnan, qnan);
+#endif
 	}
 
 	/**
