Linux: cache thread id in monolithic builds only.

- Alleviates static TLS usage. Proper way to solve is to cache in TLS dynamically.

diff --git a/Engine/Source/Runtime/Core/Public/Linux/LinuxPlatformTLS.h b/Engine/Source/Runtime/Core/Public/Linux/LinuxPlatformTLS.h
--- a/Engine/Source/Runtime/Core/Public/Linux/LinuxPlatformTLS.h
+++ b/Engine/Source/Runtime/Core/Public/Linux/LinuxPlatformTLS.h
@@ -25,10 +25,16 @@ struct CORE_API FLinuxTLS : public FGenericPlatformTLS
 	{
 		// note: cannot use pthread_self() without updating the rest of API to opaque (or at least 64-bit) thread handles
 #if defined(_GNU_SOURCE)
+
+	#if IS_MONOLITHIC
 		// syscall() is relatively heavy and shows up in the profiler, given that IsInGameThread() is used quite often. Cache thread id in TLS.
 		static __thread uint32 ThreadIdTLS = 0;
 		if (ThreadIdTLS == 0)
 		{
+	#else
+		uint32 ThreadIdTLS;
+		{
+	#endif // IS_MONOLITHIC
 			pid_t ThreadId = static_cast<pid_t>(syscall(SYS_gettid));
 			static_assert(sizeof(pid_t) <= sizeof(uint32), "pid_t is larger than uint32, reconsider implementation of GetCurrentThreadId()");
 			ThreadIdTLS = static_cast<uint32>(ThreadId);
@@ -38,6 +44,7 @@ struct CORE_API FLinuxTLS : public FGenericPlatformTLS
 
 #else
 		// better than nothing...
+		static_assert(sizeof(uint32) == sizeof(pthread_t), "pthread_t cannot be converted to uint32 one to one - different number of bits. Review FLinuxTLS::GetCurrentThreadId() implementation.");
 		return static_cast< uint32 >(pthread_self());
 #endif
 	}
