TArray ranged-for iteration check temporarily changed to an ensure until some fixes are made.

diff --git a/Engine/Source/Runtime/Core/Public/Containers/Array.h b/Engine/Source/Runtime/Core/Public/Containers/Array.h
--- a/Engine/Source/Runtime/Core/Public/Containers/Array.h
+++ b/Engine/Source/Runtime/Core/Public/Containers/Array.h
@@ -187,7 +187,7 @@ FORCEINLINE TIndexedContainerIterator<ContainerType, ElementType, IndexType> ope
 			//
 			// Also, we should only need to check one side of this comparison - if the other iterator isn't
 			// even from the same array then the compiler has generated bad code.
-			checkf(Lhs.CurrentNum == Lhs.InitialNum, TEXT("Array has changed during ranged-for iteration!"));
+			ensureMsgf(Lhs.CurrentNum == Lhs.InitialNum, TEXT("Array has changed during ranged-for iteration!"));
 			return Lhs.Ptr != Rhs.Ptr;
 		}
 	};
