Added extra debugging for shutdown crash issues.

diff --git a/Engine/Source/Runtime/Core/Private/HAL/OutputDevices.cpp b/Engine/Source/Runtime/Core/Private/HAL/OutputDevices.cpp
--- a/Engine/Source/Runtime/Core/Private/HAL/OutputDevices.cpp
+++ b/Engine/Source/Runtime/Core/Private/HAL/OutputDevices.cpp
@@ -883,6 +883,15 @@ void FOutputDeviceRedirector::Serialize( const TCHAR* Data, ELogVerbosity::Type
 
 	FScopeLock ScopeLock( &SynchronizationObject );
 
+#if PLATFORM_MAC || PLATFORM_WINDOWS || PLATFORM_LINUX
+	// this is for errors which occur after shutdown we might be able to salvage information from stdout 
+	if (OutputDevices.Num() == 0)
+	{
+		printf(TCHAR_TO_ANSI(Data));
+	}
+#endif
+
+
 	if ( bEnableBacklog )
 	{
 		new(BacklogLines)FBufferedLine( Data, Category, Verbosity, RealTime );
