Fixed up name of buffer archive to better reflect the file which is being saved.

diff --git a/Engine/Source/Runtime/Core/Public/Serialization/Archive.h b/Engine/Source/Runtime/Core/Public/Serialization/Archive.h
--- a/Engine/Source/Runtime/Core/Public/Serialization/Archive.h
+++ b/Engine/Source/Runtime/Core/Public/Serialization/Archive.h
@@ -144,7 +144,7 @@ public:
 	 *
 	 * This is overridden for the specific Archive Types
 	 **/
-	virtual FString GetArchiveName() const { return TEXT("FBufferArchive"); }
+	virtual FString GetArchiveName() const { return *FString::Printf( TEXT("FBufferArchive %s"), *ArchiveName.ToString()); }
 };
 
 /**
