Merging CL#2719191

UE4 - Fixed abtest spewing GC skips when it wasn't active anyway.

diff --git a/Engine/Source/Runtime/Core/Private/ProfilingDebugging/ABTesting.cpp b/Engine/Source/Runtime/Core/Private/ProfilingDebugging/ABTesting.cpp
--- a/Engine/Source/Runtime/Core/Private/ProfilingDebugging/ABTesting.cpp
+++ b/Engine/Source/Runtime/Core/Private/ProfilingDebugging/ABTesting.cpp
@@ -95,7 +95,7 @@ const TCHAR* FABTest::TickAndGetCommand()
 	{
 		RemainingCoolDown--;
 	}
-	else if (LastGCFrame != GLastGCFrame && !bABScopeTestActive) // reject GC frames for whole game tests
+	else if (bABTestActive && LastGCFrame != GLastGCFrame && !bABScopeTestActive) // reject GC frames for whole game tests
 	{
 		FPlatformMisc::LowLevelOutputDebugStringf(TEXT("Rejecting abtest frame because of GC."));
 	}
