Probable fix for older Android Clang versions.

diff --git a/Engine/Source/Runtime/Core/Public/Delegates/Delegate.h b/Engine/Source/Runtime/Core/Public/Delegates/Delegate.h
--- a/Engine/Source/Runtime/Core/Public/Delegates/Delegate.h
+++ b/Engine/Source/Runtime/Core/Public/Delegates/Delegate.h
@@ -279,7 +279,7 @@ class DynamicMulticastDelegateName : public FUNC_DECLARE_DYNAMIC_DELEGATE_BASE(T
 
 
 
-#define ENABLE_STATIC_FUNCTION_FNAMES defined(__clang__)
+#define ENABLE_STATIC_FUNCTION_FNAMES (defined(__clang__) && (__clang_major__ > 3 || (__clang_major__ == 3 && __clang_minor__ >= 5)))
 
 #if ENABLE_STATIC_FUNCTION_FNAMES
 
