UE4 - Moved shader working dir to the OS temp directory on windows to avoid problems with this being in dropbox or other shared folders.

diff --git a/Engine/Source/Runtime/Core/Private/GenericPlatform/GenericPlatformProcess.cpp b/Engine/Source/Runtime/Core/Private/GenericPlatform/GenericPlatformProcess.cpp
--- a/Engine/Source/Runtime/Core/Private/GenericPlatform/GenericPlatformProcess.cpp
+++ b/Engine/Source/Runtime/Core/Private/GenericPlatform/GenericPlatformProcess.cpp
@@ -132,7 +132,7 @@ const FString FGenericPlatformProcess::ShaderWorkingDir()
 void FGenericPlatformProcess::CleanShaderWorkingDir()
 {
 	// Path to the working directory where files are written for multi-threaded compilation
-	FString ShaderWorkingDirectory = ShaderWorkingDir();
+	FString ShaderWorkingDirectory =  FPlatformProcess::ShaderWorkingDir();
 	IFileManager::Get().DeleteDirectory(*ShaderWorkingDirectory, false, true);
 
 	FString LegacyShaderWorkingDirectory = FPaths::GameIntermediateDir() / TEXT("Shaders/WorkingDirectory/");
diff --git a/Engine/Source/Runtime/Core/Private/Windows/WindowsPlatformProcess.cpp b/Engine/Source/Runtime/Core/Private/Windows/WindowsPlatformProcess.cpp
--- a/Engine/Source/Runtime/Core/Private/Windows/WindowsPlatformProcess.cpp
+++ b/Engine/Source/Runtime/Core/Private/Windows/WindowsPlatformProcess.cpp
@@ -855,6 +855,12 @@ FString FWindowsPlatformProcess::GetCurrentWorkingDirectory()
 	return CurrentDirectory;
 }
 
+const FString FWindowsPlatformProcess::ShaderWorkingDir()
+{
+	return (FString(FPlatformProcess::UserTempDir()) / TEXT("UnrealShaderWorkingDir/"));
+}
+
+
 const TCHAR* FWindowsPlatformProcess::ExecutableName(bool bRemoveExtension)
 {
 	static TCHAR Result[512]=TEXT("");
diff --git a/Engine/Source/Runtime/Core/Public/Windows/WindowsPlatformProcess.h b/Engine/Source/Runtime/Core/Public/Windows/WindowsPlatformProcess.h
--- a/Engine/Source/Runtime/Core/Public/Windows/WindowsPlatformProcess.h
+++ b/Engine/Source/Runtime/Core/Public/Windows/WindowsPlatformProcess.h
@@ -140,6 +140,7 @@ public:
 	static const TCHAR* UserName(bool bOnlyAlphaNumeric = true);
 	static void SetCurrentWorkingDirectoryToBaseDir();
 	static FString GetCurrentWorkingDirectory();
+	static const FString ShaderWorkingDir();
 	static const TCHAR* ExecutableName(bool bRemoveExtension = true);
 	static FString GenerateApplicationPath( const FString& AppName, EBuildConfigurations::Type BuildConfiguration);
 	static const TCHAR* GetModuleExtension();
