Fix server crashing on load

diff --git a/Engine/Source/Runtime/Core/Private/ProfilingDebugging/DDCStatsHelper.cpp b/Engine/Source/Runtime/Core/Private/ProfilingDebugging/DDCStatsHelper.cpp
--- a/Engine/Source/Runtime/Core/Private/ProfilingDebugging/DDCStatsHelper.cpp
+++ b/Engine/Source/Runtime/Core/Private/ProfilingDebugging/DDCStatsHelper.cpp
@@ -13,7 +13,8 @@ namespace DDCStats
 		static bool bInitialized = false;
 		if (bInitialized == false)
 		{
-			FCookingStatsModule* CookingStatsModule = FModuleManager::LoadModulePtr<FCookingStatsModule>(TEXT("CookingStats"));
+			// FIXME: DanielL, this module needs to be loaded somewhere on the main thread first.
+			FCookingStatsModule* CookingStatsModule = nullptr; //FModuleManager::LoadModulePtr<FCookingStatsModule>(TEXT("CookingStats"));
 			if (CookingStatsModule)
 			{
 				CookingStats = &CookingStatsModule->Get();
