Fix crash with logging if UTF8 enabled and no OutDevices on Windows
#ue4
#windows

diff --git a/Engine/Source/Runtime/Core/Private/HAL/OutputDevices.cpp b/Engine/Source/Runtime/Core/Private/HAL/OutputDevices.cpp
--- a/Engine/Source/Runtime/Core/Private/HAL/OutputDevices.cpp
+++ b/Engine/Source/Runtime/Core/Private/HAL/OutputDevices.cpp
@@ -887,7 +887,11 @@ void FOutputDeviceRedirector::Serialize( const TCHAR* Data, ELogVerbosity::Type
 	// this is for errors which occur after shutdown we might be able to salvage information from stdout 
 	if (OutputDevices.Num() == 0)
 	{
+#if PLATFORM_WINDOWS
+		_tprintf(_T("%s\n"), Data);
+#else
 		printf("%s\n", TCHAR_TO_ANSI(Data));
+#endif
 	}
 #endif
 
