Fixed performance issue with printf hack on startup.

diff --git a/Engine/Source/Runtime/Core/Private/HAL/OutputDevices.cpp b/Engine/Source/Runtime/Core/Private/HAL/OutputDevices.cpp
--- a/Engine/Source/Runtime/Core/Private/HAL/OutputDevices.cpp
+++ b/Engine/Source/Runtime/Core/Private/HAL/OutputDevices.cpp
@@ -885,13 +885,15 @@ void FOutputDeviceRedirector::Serialize( const TCHAR* Data, ELogVerbosity::Type
 
 #if PLATFORM_DESKTOP
 	// this is for errors which occur after shutdown we might be able to salvage information from stdout 
-	if (OutputDevices.Num() == 0)
+	if ((OutputDevices.Num() == 0)&& GIsRequestingExit)
 	{
 #if PLATFORM_WINDOWS
 		_tprintf(_T("%s\n"), Data);
 #else
-		printf("%s\n", TCHAR_TO_ANSI(Data));
+		FGenericPlatformMisc::LocalPrint(Data);
+		// printf("%s\n", TCHAR_TO_ANSI(Data));
 #endif
+		return;
 	}
 #endif
 
