Copying //UE4/Dev-Editor to Dev-Main (//UE4/Dev-Main)

Fixes the compiler issue with FText, and the crash when changing languages.

diff --git a/Engine/Source/Editor/MaterialEditor/Private/MaterialEditor.cpp b/Engine/Source/Editor/MaterialEditor/Private/MaterialEditor.cpp
--- a/Engine/Source/Editor/MaterialEditor/Private/MaterialEditor.cpp
+++ b/Engine/Source/Editor/MaterialEditor/Private/MaterialEditor.cpp
@@ -470,8 +470,11 @@ void FMaterialEditor::InitMaterialEditor( const EToolkitMode::Type Mode, const T
 			// If this is an empty functions, create an output by default and start previewing it
 			if (GraphEditor.IsValid())
 			{
+				check(!bMaterialDirty);
 				UMaterialExpression* Expression = CreateNewMaterialExpression(UMaterialExpressionFunctionOutput::StaticClass(), FVector2D(200, 300), false, true);
 				SetPreviewExpression(Expression);
+				// This shouldn't count as having dirtied the material, so reset the flag
+				bMaterialDirty = false;
 			}
 		}
 		else
diff --git a/Engine/Source/Runtime/Core/Private/Internationalization/Text.cpp b/Engine/Source/Runtime/Core/Private/Internationalization/Text.cpp
--- a/Engine/Source/Runtime/Core/Private/Internationalization/Text.cpp
+++ b/Engine/Source/Runtime/Core/Private/Internationalization/Text.cpp
@@ -877,14 +877,15 @@ FText FText::FormatInternal(FText Pattern, TArray< struct FFormatArgumentData >
 	}
 
 	FFormatArgumentValue TmpArgumentValue;
-	auto GetArgumentValue = [&Arguments, &TmpArgumentValue](const FTextFormatHelper::FArgumentName& ArgumentName, int32 ArgumentNumber) -> const FFormatArgumentValue*
+	FFormatArgumentValue* TmpArgumentValuePtr = &TmpArgumentValue; // Need to do this to avoid the error "address of stack memory associated with local variable 'TmpArgumentValue' returned"
+	auto GetArgumentValue = [&Arguments, &TmpArgumentValuePtr](const FTextFormatHelper::FArgumentName& ArgumentName, int32 ArgumentNumber) -> const FFormatArgumentValue*
 	{
 		for (const auto& Arg : Arguments)
 		{
 			if (ArgumentName.NameLen == Arg.ArgumentName.Len() && FCString::Strnicmp(ArgumentName.NamePtr, *Arg.ArgumentName, ArgumentName.NameLen) == 0)
 			{
-				TmpArgumentValue = FFormatArgumentValue(Arg.ArgumentValue);
-				return &TmpArgumentValue;
+				(*TmpArgumentValuePtr) = FFormatArgumentValue(Arg.ArgumentValue);
+				return TmpArgumentValuePtr;
 			}
 		}
 		return nullptr;
diff --git a/Engine/Source/Runtime/SlateCore/Private/Fonts/FontCache.cpp b/Engine/Source/Runtime/SlateCore/Private/Fonts/FontCache.cpp
--- a/Engine/Source/Runtime/SlateCore/Private/Fonts/FontCache.cpp
+++ b/Engine/Source/Runtime/SlateCore/Private/Fonts/FontCache.cpp
@@ -536,10 +536,14 @@ FSlateFontCache::FSlateFontCache( TSharedRef<ISlateFontAtlasFactory> InFontAtlas
 	, bFlushRequested( false )
 {
 	UE_LOG(LogSlate, Log, TEXT("SlateFontCache - WITH_FREETYPE: %d, WITH_HARFBUZZ: %d"), WITH_FREETYPE, WITH_HARFBUZZ);
+
+	FInternationalization::Get().OnCultureChanged().AddRaw(this, &FSlateFontCache::FlushCache);
 }
 
 FSlateFontCache::~FSlateFontCache()
-{	
+{
+	FInternationalization::Get().OnCultureChanged().RemoveAll(this);
+
 	// Make sure things get destroyed in the correct order
 	TextShaper.Reset();
 	FontRenderer.Reset();
diff --git a/Engine/Source/Runtime/SlateCore/Private/Fonts/SlateFontRenderer.cpp b/Engine/Source/Runtime/SlateCore/Private/Fonts/SlateFontRenderer.cpp
--- a/Engine/Source/Runtime/SlateCore/Private/Fonts/SlateFontRenderer.cpp
+++ b/Engine/Source/Runtime/SlateCore/Private/Fonts/SlateFontRenderer.cpp
@@ -310,12 +310,6 @@ void FSlateFontRenderer::GetRenderData(const FFreeTypeFaceGlyphData& InFaceGlyph
 		}
 	}
 
-	// todo: what is this doing?
-	FT_BBox GlyphBox;
-	FT_Glyph Glyph;
-	FT_Get_Glyph( Slot, &Glyph );
-	FT_Glyph_Get_CBox( Glyph, FT_GLYPH_BBOX_PIXELS, &GlyphBox );
-
 	const int32 Height = static_cast<int32>(FreeTypeUtils::Convert26Dot6ToRoundedPixel<int32>(FT_MulFix(Face->height, Face->size->metrics.y_scale)) * InScale);
 
 	// Set measurement info for this character
@@ -338,8 +332,6 @@ void FSlateFontRenderer::GetRenderData(const FFreeTypeFaceGlyphData& InFaceGlyph
 	{
 		FT_Bitmap_Done(FTLibrary->GetLibrary(), Bitmap);
 	}
-
-	FT_Done_Glyph(Glyph);
 }
 
 FT_Face FSlateFontRenderer::GetFontFace(const FFontData& InFontData) const
