Copying //UE4/Dev-Framework to Dev-Main (//UE4/Dev-Main)

diff --git a/Engine/Source/Runtime/AIModule/Private/Navigation/CrowdManager.cpp b/Engine/Source/Runtime/AIModule/Private/Navigation/CrowdManager.cpp
--- a/Engine/Source/Runtime/AIModule/Private/Navigation/CrowdManager.cpp
+++ b/Engine/Source/Runtime/AIModule/Private/Navigation/CrowdManager.cpp
@@ -57,6 +57,7 @@ namespace CrowdDebugDrawing
 	const FColor CollisionRange(192, 0, 128);
 	const FColor CollisionSeg0(192, 0, 128);
 	const FColor CollisionSeg1(96, 0, 64);
+	const FColor CollisionSegIgnored(128, 128, 128);
 	const FColor Path(255, 255, 255);
 	const FColor PathSpecial(255, 192, 203);
 	const FColor PathOpt(0, 128, 0);
@@ -426,7 +427,8 @@ bool UCrowdManager::SetAgentMoveDirection(const UCrowdFollowingComponent* AgentC
 	return bSuccess;
 }
 
-bool UCrowdManager::SetAgentMovePath(const UCrowdFollowingComponent* AgentComponent, const FNavMeshPath* Path, int32 PathSectionStart, int32 PathSectionEnd) const
+bool UCrowdManager::SetAgentMovePath(const UCrowdFollowingComponent* AgentComponent, const FNavMeshPath* Path,
+	int32 PathSectionStart, int32 PathSectionEnd, const FVector& PathSectionEndLocation) const
 {
 	SCOPE_CYCLE_COUNTER(STAT_AI_Crowd_AgentUpdateTime);
 
@@ -440,7 +442,7 @@ bool UCrowdManager::SetAgentMovePath(const UCrowdFollowingComponent* AgentCompon
 		Path && (Path->GetPathPoints().Num() > 1) &&
 		Path->PathCorridor.IsValidIndex(PathSectionStart) && Path->PathCorridor.IsValidIndex(PathSectionEnd))
 	{
-		FVector TargetPos = Path->GetPathPoints().Last().Location;
+		FVector TargetPos = PathSectionEndLocation;
 		if (PathSectionEnd < (Path->PathCorridor.Num() - 1))
 		{
 			RecastNavData->GetPolyCenter(Path->PathCorridor[PathSectionEnd], TargetPos);
@@ -457,7 +459,7 @@ bool UCrowdManager::SetAgentMovePath(const UCrowdFollowingComponent* AgentCompon
 		DetourCrowd->updateAgentFilter(AgentData->AgentIndex, DetourFilter);
 		DetourCrowd->updateAgentState(AgentData->AgentIndex, false);
 
-		const FVector RcTargetPos = Unreal2RecastPoint(TargetPos);		
+		const FVector RcTargetPos = Unreal2RecastPoint(TargetPos);
 		bSuccess = DetourCrowd->requestMoveTarget(AgentData->AgentIndex, PathRefs.Last(), &RcTargetPos.X);
 		if (bSuccess)
 		{
@@ -926,7 +928,11 @@ void UCrowdManager::DrawDebugCollisionSegments(const dtCrowdAgent* CrowdAgent) c
 	for (int32 Idx = 0; Idx < CrowdAgent->boundary.getSegmentCount(); Idx++)
 	{
 		const float* s = CrowdAgent->boundary.getSegment(Idx);
-		FColor Color = (dtTriArea2D(CrowdAgent->npos, s, s + 3) < 0.0f) ? CrowdDebugDrawing::CollisionSeg1 : CrowdDebugDrawing::CollisionSeg0;
+		const int32 SegFlags = CrowdAgent->boundary.getSegmentFlags(Idx);
+		const FColor Color = (SegFlags & DT_CROWD_BOUNDARY_IGNORE) ? CrowdDebugDrawing::CollisionSegIgnored :
+			(dtTriArea2D(CrowdAgent->npos, s, s + 3) < 0.0f) ? CrowdDebugDrawing::CollisionSeg1 :
+			CrowdDebugDrawing::CollisionSeg0;
+
 		FVector Pt0 = Recast2UnrealPoint(s);
 		FVector Pt1 = Recast2UnrealPoint(s + 3);
 
@@ -1142,7 +1148,7 @@ void UCrowdManager::DebugTick() const
 					{
 						FVector P1 = Recast2UnrealPoint(&CrowdAgent->cornerVerts[Idx * 3]);
 						UE_VLOG_SEGMENT(LogOwner, LogCrowdFollowing, Log, P0 + CrowdDebugDrawing::Offset, P1 + CrowdDebugDrawing::Offset, CrowdDebugDrawing::Corner, TEXT(""));
-						UE_VLOG_BOX(LogOwner, LogCrowdFollowing, Log, FBox::BuildAABB(P1 + CrowdDebugDrawing::Offset, FVector(2, 2, 2)), CrowdDebugDrawing::Corner, TEXT(""));
+						UE_VLOG_BOX(LogOwner, LogCrowdFollowing, Log, FBox::BuildAABB(P1 + CrowdDebugDrawing::Offset, FVector(2, 2, 2)), CrowdDebugDrawing::Corner, TEXT("%d"), CrowdAgent->cornerFlags[Idx]);
 						P0 = P1;
 					}
 				}
@@ -1181,11 +1187,15 @@ void UCrowdManager::DebugTick() const
 				for (int32 Idx = 0; Idx < CrowdAgent->boundary.getSegmentCount(); Idx++)
 				{
 					const float* s = CrowdAgent->boundary.getSegment(Idx);
-					FColor Color = (dtTriArea2D(CrowdAgent->npos, s, s + 3) < 0.0f) ? CrowdDebugDrawing::CollisionSeg1 : CrowdDebugDrawing::CollisionSeg0;
+					const int32 SegFlags = CrowdAgent->boundary.getSegmentFlags(Idx);
+					const FColor Color = (SegFlags & DT_CROWD_BOUNDARY_IGNORE) ? CrowdDebugDrawing::CollisionSegIgnored :
+						(dtTriArea2D(CrowdAgent->npos, s, s + 3) < 0.0f) ? CrowdDebugDrawing::CollisionSeg1 :
+						CrowdDebugDrawing::CollisionSeg0;
+
 					FVector Pt0 = Recast2UnrealPoint(s);
 					FVector Pt1 = Recast2UnrealPoint(s + 3);
 
-					UE_VLOG_SEGMENT(LogOwner, LogCrowdFollowing, Log, Pt0 + CrowdDebugDrawing::Offset, Pt1 + CrowdDebugDrawing::Offset, Color, TEXT(""));
+					UE_VLOG_SEGMENT_THICK(LogOwner, LogCrowdFollowing, Log, Pt0 + CrowdDebugDrawing::Offset, Pt1 + CrowdDebugDrawing::Offset, Color, 3.0f, TEXT(""));
 				}
 			}
 		}
diff --git a/Engine/Source/Runtime/AIModule/Private/Perception/AIPerceptionComponent.cpp b/Engine/Source/Runtime/AIModule/Private/Perception/AIPerceptionComponent.cpp
--- a/Engine/Source/Runtime/AIModule/Private/Perception/AIPerceptionComponent.cpp
+++ b/Engine/Source/Runtime/AIModule/Private/Perception/AIPerceptionComponent.cpp
@@ -184,7 +184,8 @@ void UAIPerceptionComponent::OnRegister()
 	// this should not be needed but aparently AAIController::PostRegisterAllComponents
 	// gets called component's OnRegister
 	AIOwner = Cast<AAIController>(GetOwner());
-	ensure(AIOwner == nullptr || AIOwner->GetAIPerceptionComponent() == nullptr || AIOwner->GetAIPerceptionComponent() == this);
+	ensure(AIOwner == nullptr || AIOwner->GetAIPerceptionComponent() == nullptr || AIOwner->GetAIPerceptionComponent() == this
+		|| (AIOwner->GetWorld() && AIOwner->GetWorld()->WorldType != EWorldType::Editor));
 	if (AIOwner && AIOwner->GetAIPerceptionComponent() == nullptr)
 	{
 		AIOwner->SetPerceptionComponent(*this);
@@ -520,12 +521,30 @@ bool UAIPerceptionComponent::AgeStimuli(const float ConstPerceptionAgingRate)
 
 void UAIPerceptionComponent::ForgetActor(AActor* ActorToForget)
 {
-	PerceptualData.Remove(ActorToForget);
+	if (PerceptualData.Num() > 0)
+	{
+		UAIPerceptionSystem* AIPerceptionSys = UAIPerceptionSystem::GetCurrent(GetWorld());
+		if (AIPerceptionSys != nullptr && ActorToForget != nullptr)
+		{
+			AIPerceptionSys->OnListenerForgetsActor(*this, *ActorToForget);
+		}
+
+		PerceptualData.Remove(ActorToForget);
+	}
 }
 
 void UAIPerceptionComponent::ForgetAll()
 {
-	PerceptualData.Reset();
+	if (PerceptualData.Num() > 0)
+	{
+		UAIPerceptionSystem* AIPerceptionSys = UAIPerceptionSystem::GetCurrent(GetWorld());
+		if (AIPerceptionSys != nullptr)
+		{
+			AIPerceptionSys->OnListenerForgetsAll(*this);
+		}
+
+		PerceptualData.Reset();
+	}
 }
 
 float UAIPerceptionComponent::GetYoungestStimulusAge(const AActor& Source) const
@@ -639,48 +658,45 @@ bool UAIPerceptionComponent::GetActorsPerception(AActor* Actor, FActorPerception
 // debug
 //----------------------------------------------------------------------//
 #if !UE_BUILD_SHIPPING
-void UAIPerceptionComponent::GetDebugData(TArray<FString>& OnScreenStrings, TArray<FDrawDebugShapeElement>& DebugShapes) const
+void UAIPerceptionComponent::GrabGameplayDebuggerData(TArray<FString>& OnScreenStrings, TArray<FGameplayDebuggerShapeElement>& DebugShapes) const
 {
-	{
-		UAIPerceptionSystem* PerceptionSys = UAIPerceptionSystem::GetCurrent(GetWorld());
-		check(PerceptionSys);
+#if ENABLED_GAMEPLAY_DEBUGGER
+	UAIPerceptionSystem* PerceptionSys = UAIPerceptionSystem::GetCurrent(GetWorld());
+	check(PerceptionSys);
 
-		for (UAIPerceptionComponent::TActorPerceptionContainer::TConstIterator It(GetPerceptualDataConstIterator()); It; ++It)
+	for (UAIPerceptionComponent::TActorPerceptionContainer::TConstIterator It(GetPerceptualDataConstIterator()); It; ++It)
+	{
+		if (It->Key == NULL)
 		{
-			if (It->Key == NULL)
-			{
-				continue;
-			}
+			continue;
+		}
 
-			const FActorPerceptionInfo& ActorPerceptionInfo = It->Value;
+		const FActorPerceptionInfo& ActorPerceptionInfo = It->Value;
 
-			if (ActorPerceptionInfo.Target.IsValid())
+		if (ActorPerceptionInfo.Target.IsValid())
+		{
+			const FVector TargetLocation = ActorPerceptionInfo.Target->GetActorLocation();
+			for (const auto& Stimulus : ActorPerceptionInfo.LastSensedStimuli)
 			{
-				const FVector TargetLocation = ActorPerceptionInfo.Target->GetActorLocation();
-				for (const auto& Stimulus : ActorPerceptionInfo.LastSensedStimuli)
+				if (Stimulus.Strength >= 0)
 				{
-					if (Stimulus.Strength >= 0)
-					{
-						const FString Description = FString::Printf(TEXT("%s: %.2f a:%.2f"), *PerceptionSys->GetSenseName(Stimulus.Type), Stimulus.Strength, Stimulus.GetAge());
-						DebugShapes.Add(FDrawDebugShapeElement::MakeString(Description, Stimulus.StimulusLocation + FVector(0, 0, 30)));
+					const FString Description = FString::Printf(TEXT("%s: %.2f a:%.2f"), *PerceptionSys->GetSenseName(Stimulus.Type), Stimulus.Strength, Stimulus.GetAge());
+					DebugShapes.Add(UGameplayDebuggerHelper::MakeString(Description, Stimulus.StimulusLocation + FVector(0, 0, 30)));
 
-						const FColor DebugColor = PerceptionSys->GetSenseDebugColor(Stimulus.Type);
-						DebugShapes.Add(FDrawDebugShapeElement::MakePoint(Stimulus.StimulusLocation, DebugColor, 30));
-						DebugShapes.Add(FDrawDebugShapeElement::MakeLine(Stimulus.ReceiverLocation, Stimulus.StimulusLocation, DebugColor));
-						DebugShapes.Add(FDrawDebugShapeElement::MakeLine(TargetLocation, Stimulus.StimulusLocation, FColor::Black));
-					}
+					const FColor DebugColor = PerceptionSys->GetSenseDebugColor(Stimulus.Type);
+					DebugShapes.Add(UGameplayDebuggerHelper::MakePoint(Stimulus.StimulusLocation, DebugColor, 30));
+					DebugShapes.Add(UGameplayDebuggerHelper::MakeLine(Stimulus.ReceiverLocation, Stimulus.StimulusLocation, DebugColor));
+					DebugShapes.Add(UGameplayDebuggerHelper::MakeLine(TargetLocation, Stimulus.StimulusLocation, FColor::Black));
 				}
 			}
 		}
+	}
 
-		for (auto Sense : SensesConfig)
-		{
-			if (Sense)
-			{
-				Sense->GetDebugData(OnScreenStrings, DebugShapes, *this);
-			}
-		}
+	for (auto Sense : SensesConfig)
+	{
+		Sense->GetDebugData(OnScreenStrings, DebugShapes, *this);
 	}
+#endif
 }
 
 #endif // !UE_BUILD_SHIPPING
diff --git a/Engine/Source/Runtime/AIModule/Private/Perception/AIPerceptionSystem.cpp b/Engine/Source/Runtime/AIModule/Private/Perception/AIPerceptionSystem.cpp
--- a/Engine/Source/Runtime/AIModule/Private/Perception/AIPerceptionSystem.cpp
+++ b/Engine/Source/Runtime/AIModule/Private/Perception/AIPerceptionSystem.cpp
@@ -400,6 +400,42 @@ void UAIPerceptionSystem::OnListenerRemoved(const FPerceptionListener& NewListen
 	}
 }
 
+void UAIPerceptionSystem::OnListenerForgetsActor(const UAIPerceptionComponent& Listener, AActor& ActorToForget)
+{
+	const FPerceptionListenerID ListenerId = Listener.GetListenerId();
+
+	if (ListenerId != FPerceptionListenerID::InvalidID())
+	{
+		FPerceptionListener& ListenerEntry = ListenerContainer[ListenerId];
+		
+		for (UAISense* Sense : Senses)
+		{
+			if (Sense != nullptr && Sense->NeedsNotificationOnForgetting() && ListenerEntry.HasSense(Sense->GetSenseID()))
+			{
+				Sense->OnListenerForgetsActor(ListenerEntry, ActorToForget);
+			}
+		}
+	}
+}
+
+void UAIPerceptionSystem::OnListenerForgetsAll(const UAIPerceptionComponent& Listener)
+{
+	const FPerceptionListenerID ListenerId = Listener.GetListenerId();
+
+	if (ListenerId != FPerceptionListenerID::InvalidID())
+	{
+		FPerceptionListener& ListenerEntry = ListenerContainer[ListenerId];
+
+		for (UAISense* Sense : Senses)
+		{
+			if (Sense != nullptr && Sense->NeedsNotificationOnForgetting() && ListenerEntry.HasSense(Sense->GetSenseID()))
+			{
+				Sense->OnListenerForgetsAll(ListenerEntry);
+			}
+		}
+	}
+}
+
 void UAIPerceptionSystem::RegisterDelayedStimulus(FPerceptionListenerID ListenerId, float Delay, AActor* Instigator, const FAIStimulus& Stimulus)
 {
 	FDelayedStimulus DelayedStimulus;
diff --git a/Engine/Source/Runtime/AIModule/Private/Perception/AISense.cpp b/Engine/Source/Runtime/AIModule/Private/Perception/AISense.cpp
--- a/Engine/Source/Runtime/AIModule/Private/Perception/AISense.cpp
+++ b/Engine/Source/Runtime/AIModule/Private/Perception/AISense.cpp
@@ -21,6 +21,8 @@ UAISense::UAISense(const FObjectInitializer& ObjectInitializer)
 	DebugName = GetName();
 	DefaultExpirationAge = FAIStimulus::NeverHappenedAge;
 
+	bNeedsForgettingNotification = false;
+
 	if (HasAnyFlags(RF_ClassDefaultObject) == false)
 	{
 		SenseID = ((const UAISense*)GetClass()->GetDefaultObject())->GetSenseID();
@@ -127,7 +129,7 @@ TSubclassOf<UAISense> UAISenseConfig_Sight::GetSenseImplementation() const
 }
 
 #if !UE_BUILD_SHIPPING
-void UAISenseConfig_Sight::GetDebugData(TArray<FString>& OnScreenStrings, TArray<FDrawDebugShapeElement>& DebugShapes, const UAIPerceptionComponent& PerceptionComponent) const
+void UAISenseConfig_Sight::GetDebugData(TArray<FString>& OnScreenStrings, TArray<FGameplayDebuggerShapeElement>& DebugShapes, const UAIPerceptionComponent& PerceptionComponent) const
 {
 	/*PeripheralVisionAngleDegrees*/
 	const AActor* BodyActor = PerceptionComponent.GetBodyActor();
@@ -137,13 +139,13 @@ void UAISenseConfig_Sight::GetDebugData(TArray<FString>& OnScreenStrings, TArray
 		FVector BodyLocation, BodyFacing;
 		PerceptionComponent.GetLocationAndDirection(BodyLocation, BodyFacing);
 
-		DebugShapes.Add(FDrawDebugShapeElement::MakeCylinder(BodyLocation, BodyLocation + FVector(0, 0, -50), LoseSightRadius, 32, UAISense_Sight::GetDebugLoseSightColor()));
-		DebugShapes.Add(FDrawDebugShapeElement::MakeCylinder(BodyLocation, BodyLocation + FVector(0, 0, -50), SightRadius, 32, UAISense_Sight::GetDebugSightRangeColor()));
+		DebugShapes.Add(UGameplayDebuggerHelper::MakeCylinder(BodyLocation, BodyLocation + FVector(0, 0, -50), LoseSightRadius, 32, UAISense_Sight::GetDebugLoseSightColor()));
+		DebugShapes.Add(UGameplayDebuggerHelper::MakeCylinder(BodyLocation, BodyLocation + FVector(0, 0, -50), SightRadius, 32, UAISense_Sight::GetDebugSightRangeColor()));
 
 		const float SightPieLength = FMath::Max(LoseSightRadius, SightRadius);
-		DebugShapes.Add(FDrawDebugShapeElement::MakeLine(BodyLocation, BodyLocation + (BodyFacing * SightPieLength), UAISense_Sight::GetDebugSightRangeColor()));
-		DebugShapes.Add(FDrawDebugShapeElement::MakeLine(BodyLocation, BodyLocation + (BodyFacing.RotateAngleAxis(PeripheralVisionAngleDegrees, FVector::UpVector) * SightPieLength), UAISense_Sight::GetDebugSightRangeColor()));
-		DebugShapes.Add(FDrawDebugShapeElement::MakeLine(BodyLocation, BodyLocation + (BodyFacing.RotateAngleAxis(-PeripheralVisionAngleDegrees, FVector::UpVector) * SightPieLength), UAISense_Sight::GetDebugSightRangeColor()));
+		DebugShapes.Add(UGameplayDebuggerHelper::MakeLine(BodyLocation, BodyLocation + (BodyFacing * SightPieLength), UAISense_Sight::GetDebugSightRangeColor()));
+		DebugShapes.Add(UGameplayDebuggerHelper::MakeLine(BodyLocation, BodyLocation + (BodyFacing.RotateAngleAxis(PeripheralVisionAngleDegrees, FVector::UpVector) * SightPieLength), UAISense_Sight::GetDebugSightRangeColor()));
+		DebugShapes.Add(UGameplayDebuggerHelper::MakeLine(BodyLocation, BodyLocation + (BodyFacing.RotateAngleAxis(-PeripheralVisionAngleDegrees, FVector::UpVector) * SightPieLength), UAISense_Sight::GetDebugSightRangeColor()));
 	}
 }
 #endif // !UE_BUILD_SHIPPING
@@ -162,22 +164,20 @@ TSubclassOf<UAISense> UAISenseConfig_Hearing::GetSenseImplementation() const
 	return *Implementation; 
 }
 #if !UE_BUILD_SHIPPING
-void UAISenseConfig_Hearing::GetDebugData(TArray<FString>& OnScreenStrings, TArray<FDrawDebugShapeElement>& DebugShapes, const UAIPerceptionComponent& PerceptionComponent) const
+void UAISenseConfig_Hearing::GetDebugData(TArray<FString>& OnScreenStrings, TArray<FGameplayDebuggerShapeElement>& DebugShapes, const UAIPerceptionComponent& PerceptionComponent) const
 {
 	const AActor* BodyActor = PerceptionComponent.GetBodyActor();
 	if (BodyActor != nullptr)
 	{
 		UWorld* World = BodyActor->GetWorld();
 		FVector OwnerLocation = BodyActor->GetActorLocation();
-		DebugShapes.Add(FDrawDebugShapeElement::MakeCylinder(OwnerLocation, OwnerLocation + FVector(0, 0, -50), HearingRange, 32, UAISense_Hearing::GetDebugHearingRangeColor()));
+		DebugShapes.Add(UGameplayDebuggerHelper::MakeCylinder(OwnerLocation, OwnerLocation + FVector(0, 0, -50), HearingRange, 32, UAISense_Hearing::GetDebugHearingRangeColor()));
 		if (bUseLoSHearing)
 		{
-			DebugShapes.Add(FDrawDebugShapeElement::MakeCylinder(OwnerLocation, OwnerLocation + FVector(0, 0, -50), LoSHearingRange, 32, UAISense_Hearing::GetDebugLoSHearingRangeeColor()));
+			DebugShapes.Add(UGameplayDebuggerHelper::MakeCylinder(OwnerLocation, OwnerLocation + FVector(0, 0, -50), LoSHearingRange, 32, UAISense_Hearing::GetDebugLoSHearingRangeeColor()));
 		}
 	}
 }
-
-
 #endif // !UE_BUILD_SHIPPING
 
 //----------------------------------------------------------------------//
diff --git a/Engine/Source/Runtime/AIModule/Private/Perception/AISense_Blueprint.cpp b/Engine/Source/Runtime/AIModule/Private/Perception/AISense_Blueprint.cpp
--- a/Engine/Source/Runtime/AIModule/Private/Perception/AISense_Blueprint.cpp
+++ b/Engine/Source/Runtime/AIModule/Private/Perception/AISense_Blueprint.cpp
@@ -22,7 +22,7 @@ UAISense_Blueprint::UAISense_Blueprint(const FObjectInitializer& ObjectInitializ
 		OnNewListenerDelegate.BindUObject(this, &UAISense_Blueprint::OnNewListenerImpl);
 		OnListenerRemovedDelegate.BindUObject(this, &UAISense_Blueprint::OnListenerRemovedImpl);
 		// update is optional
-		if (BlueprintNodeHelpers::HasBlueprintFunction(TEXT("OnListenerUpdated"), this, StopAtClass))
+		if (BlueprintNodeHelpers::HasBlueprintFunction(TEXT("OnListenerUpdated"), *this, *StopAtClass))
 		{
 			OnListenerUpdateDelegate.BindUObject(this, &UAISense_Blueprint::OnListenerUpdateImpl);
 		}
diff --git a/Engine/Source/Runtime/AIModule/Private/Perception/AISense_Sight.cpp b/Engine/Source/Runtime/AIModule/Private/Perception/AISense_Sight.cpp
--- a/Engine/Source/Runtime/AIModule/Private/Perception/AISense_Sight.cpp
+++ b/Engine/Source/Runtime/AIModule/Private/Perception/AISense_Sight.cpp
@@ -102,6 +102,7 @@ UAISense_Sight::UAISense_Sight(const FObjectInitializer& ObjectInitializer)
 	NotifyType = EAISenseNotifyType::OnPerceptionChange;
 	
 	bAutoRegisterAllPawnsAsSources = true;
+	bNeedsForgettingNotification = true;
 }
 
 FORCEINLINE_DEBUGGABLE float UAISense_Sight::CalcQueryImportance(const FPerceptionListener& Listener, const FVector& TargetLocation, const float SightRadiusSq) const
@@ -193,7 +194,8 @@ float UAISense_Sight::Update()
 				const float SightRadiusSq = SightQuery->bLastResult ? PropDigest.LoseSightRadiusSq : PropDigest.SightRadiusSq;
 				
 				float StimulusStrength = 1.f;
-
+				
+				// @Note that automagical "seeing" does not care about sight range nor vision cone
 				if (ShouldAutomaticallySeeTarget(PropDigest, SightQuery, Listener, TargetActor, StimulusStrength))
 				{
 					// Pretend like we've seen this target where we last saw them
@@ -576,6 +578,34 @@ void UAISense_Sight::RemoveAllQueriesToTarget(const FAISightTarget::FTargetId& T
 	}
 }
 
+void UAISense_Sight::OnListenerForgetsActor(const FPerceptionListener& Listener, AActor& ActorToForget)
+{
+	const uint32 ListenerId = Listener.GetListenerID();
+	
+	for (FAISightQuery& SightQuery : SightQueryQueue)
+	{
+		if (SightQuery.ObserverId == ListenerId)
+		{
+			// assuming one query per observer-target pair
+			SightQuery.ForgetPreviousResult();
+			break;
+		}
+	}
+}
+
+void UAISense_Sight::OnListenerForgetsAll(const FPerceptionListener& Listener)
+{
+	const uint32 ListenerId = Listener.GetListenerID();
+
+	for (FAISightQuery& SightQuery : SightQueryQueue)
+	{
+		if (SightQuery.ObserverId == ListenerId)
+		{
+			SightQuery.ForgetPreviousResult();
+		}
+	}
+}
+
 #if !UE_BUILD_SHIPPING
 //----------------------------------------------------------------------//
 // DEBUG
diff --git a/Engine/Source/Runtime/AIModule/Public/BlueprintNodeHelpers.h b/Engine/Source/Runtime/AIModule/Public/BlueprintNodeHelpers.h
--- a/Engine/Source/Runtime/AIModule/Public/BlueprintNodeHelpers.h
+++ b/Engine/Source/Runtime/AIModule/Public/BlueprintNodeHelpers.h
@@ -3,6 +3,7 @@
 #pragma once
 
 class UBTNode;
+class UBehaviorTreeComponent;
 
 namespace BlueprintNodeHelpers
 {
@@ -23,14 +24,16 @@ namespace BlueprintNodeHelpers
 
 	void AbortLatentActions(UActorComponent& OwnerOb, const UObject& Ob);
 
-	FORCEINLINE bool HasBlueprintFunction(FName FuncName, const UObject* Ob, const UClass* StopAtClass)
+	FORCEINLINE bool HasBlueprintFunction(FName FuncName, const UObject& Object, const UClass& StopAtClass)
 	{
-		return (Ob->GetClass()->FindFunctionByName(FuncName)->GetOuter() != StopAtClass);
+		const UFunction* Function = Object.GetClass()->FindFunctionByName(FuncName);
+		ensure(Function);
+		return (Function != nullptr) && (Function->GetOuter() != &StopAtClass);
 	}
 
-	FORCEINLINE FString GetNodeName(const UObject* Ob)
+	FORCEINLINE FString GetNodeName(const UObject& NodeObject)
 	{
-		return Ob->GetClass()->GetName().LeftChop(2);
+		return NodeObject.GetClass()->GetName().LeftChop(2);
 	}
 
 	//----------------------------------------------------------------------//
@@ -38,4 +41,7 @@ namespace BlueprintNodeHelpers
 	//----------------------------------------------------------------------//
 	DEPRECATED(4.7, "This version is deprecated. Please use the one taking reference to UActorComponent rather than a pointer.")
 	void AbortLatentActions(UActorComponent* OwnerOb, const UObject* Ob);
+
+	DEPRECATED(4.11, "This version of HasBlueprintFunction is deprecated. Please use the one taking reference to UObject and StopAtClass rather than a pointers.")
+	AIMODULE_API bool HasBlueprintFunction(FName FuncName, const UObject* Object, const UClass* StopAtClass);
 }
diff --git a/Engine/Source/Runtime/AnimGraphRuntime/Private/AnimNodes/AnimNode_ApplyAdditive.cpp b/Engine/Source/Runtime/AnimGraphRuntime/Private/AnimNodes/AnimNode_ApplyAdditive.cpp
--- a/Engine/Source/Runtime/AnimGraphRuntime/Private/AnimNodes/AnimNode_ApplyAdditive.cpp
+++ b/Engine/Source/Runtime/AnimGraphRuntime/Private/AnimNodes/AnimNode_ApplyAdditive.cpp
@@ -24,7 +24,7 @@ void FAnimNode_ApplyAdditive::Update(const FAnimationUpdateContext& Context)
 {
 	Base.Update(Context);
 
-	if (IsLODEnabled(Context.AnimInstance, LODThreshold))
+	if (IsLODEnabled(Context.AnimInstanceProxy, LODThreshold))
 	{
 		// @note: If you derive from this class, and if you have input that you rely on for base
 		// this is not going to work	
@@ -39,7 +39,7 @@ void FAnimNode_ApplyAdditive::Update(const FAnimationUpdateContext& Context)
 
 void FAnimNode_ApplyAdditive::Evaluate(FPoseContext& Output)
 {
-	if (IsLODEnabled(Output.AnimInstance, LODThreshold))
+	if (IsLODEnabled(Output.AnimInstanceProxy, LODThreshold))
 	{
 		//@TODO: Could evaluate Base into Output and save a copy
 		const float ActualAlpha = AlphaScaleBias.ApplyTo(Alpha);
diff --git a/Engine/Source/Runtime/AnimGraphRuntime/Private/AnimNodes/AnimNode_BlendListBase.cpp b/Engine/Source/Runtime/AnimGraphRuntime/Private/AnimNodes/AnimNode_BlendListBase.cpp
--- a/Engine/Source/Runtime/AnimGraphRuntime/Private/AnimNodes/AnimNode_BlendListBase.cpp
+++ b/Engine/Source/Runtime/AnimGraphRuntime/Private/AnimNodes/AnimNode_BlendListBase.cpp
@@ -125,7 +125,7 @@ void FAnimNode_BlendListBase::Update(const FAnimationUpdateContext& Context)
 			// when this flag is true, we'll reinitialize the children
 			if (bResetChildOnActivation)
 			{
-				FAnimationInitializeContext ReinitializeContext(Context.AnimInstance);
+				FAnimationInitializeContext ReinitializeContext(Context.AnimInstanceProxy);
 
 				// reinitialize
 				BlendPose[ChildIndex].Initialize(ReinitializeContext);
@@ -200,7 +200,7 @@ void FAnimNode_BlendListBase::Update(const FAnimationUpdateContext& Context)
 
 void FAnimNode_BlendListBase::Evaluate(FPoseContext& Output)
 {
-	ANIM_MT_SCOPE_CYCLE_COUNTER(BlendPosesInGraph, Output.AnimInstance->IsRunningParallelEvaluation());
+	ANIM_MT_SCOPE_CYCLE_COUNTER(BlendPosesInGraph, !IsInGameThread());
 
 	const int32 NumPoses = PosesToEvaluate.Num();
 
diff --git a/Engine/Source/Runtime/AnimGraphRuntime/Private/AnimNodes/AnimNode_BlendSpacePlayer.cpp b/Engine/Source/Runtime/AnimGraphRuntime/Private/AnimNodes/AnimNode_BlendSpacePlayer.cpp
--- a/Engine/Source/Runtime/AnimGraphRuntime/Private/AnimNodes/AnimNode_BlendSpacePlayer.cpp
+++ b/Engine/Source/Runtime/AnimGraphRuntime/Private/AnimNodes/AnimNode_BlendSpacePlayer.cpp
@@ -3,6 +3,7 @@
 #include "AnimGraphRuntimePrivatePCH.h"
 #include "AnimNodes/AnimNode_BlendSpacePlayer.h"
 #include "Animation/BlendSpaceBase.h"
+#include "Animation/AnimInstanceProxy.h"
 
 /////////////////////////////////////////////////////
 // FAnimNode_BlendSpacePlayer
@@ -52,15 +53,15 @@ void FAnimNode_BlendSpacePlayer::UpdateAssetPlayer(const FAnimationUpdateContext
 
 void FAnimNode_BlendSpacePlayer::UpdateInternal(const FAnimationUpdateContext& Context)
 {
-	if ((BlendSpace != NULL) && (Context.AnimInstance->CurrentSkeleton->IsCompatible(BlendSpace->GetSkeleton())))
+	if ((BlendSpace != NULL) && (Context.AnimInstanceProxy->IsSkeletonCompatible(BlendSpace->GetSkeleton())))
 	{
 		// Create a tick record and fill it out
 		FAnimGroupInstance* SyncGroup;
-		FAnimTickRecord& TickRecord = Context.AnimInstance->CreateUninitializedTickRecord(GroupIndex, /*out*/ SyncGroup);
+		FAnimTickRecord& TickRecord = Context.AnimInstanceProxy->CreateUninitializedTickRecord(GroupIndex, /*out*/ SyncGroup);
 
 		const FVector BlendInput(X, Y, Z);
 	
-		Context.AnimInstance->MakeBlendSpaceTickRecord(TickRecord, BlendSpace, BlendInput, BlendSampleDataCache, BlendFilter, bLoop, PlayRate, Context.GetFinalBlendWeight(), /*inout*/ InternalTimeAccumulator, MarkerTickRecord);
+		Context.AnimInstanceProxy->MakeBlendSpaceTickRecord(TickRecord, BlendSpace, BlendInput, BlendSampleDataCache, BlendFilter, bLoop, PlayRate, Context.GetFinalBlendWeight(), /*inout*/ InternalTimeAccumulator, MarkerTickRecord);
 
 		// Update the sync group if it exists
 		if (SyncGroup != NULL)
@@ -72,7 +73,7 @@ void FAnimNode_BlendSpacePlayer::UpdateInternal(const FAnimationUpdateContext& C
 
 void FAnimNode_BlendSpacePlayer::Evaluate(FPoseContext& Output)
 {
-	if ((BlendSpace != NULL) && (Output.AnimInstance->CurrentSkeleton->IsCompatible(BlendSpace->GetSkeleton())))
+	if ((BlendSpace != NULL) && (Output.AnimInstanceProxy->IsSkeletonCompatible(BlendSpace->GetSkeleton())))
 	{
 		BlendSpace->GetAnimationPose(BlendSampleDataCache, Output.Pose, Output.Curve);
 	}
diff --git a/Engine/Source/Runtime/AnimGraphRuntime/Private/AnimNodes/AnimNode_CopyPoseFromMesh.cpp b/Engine/Source/Runtime/AnimGraphRuntime/Private/AnimNodes/AnimNode_CopyPoseFromMesh.cpp
--- a/Engine/Source/Runtime/AnimGraphRuntime/Private/AnimNodes/AnimNode_CopyPoseFromMesh.cpp
+++ b/Engine/Source/Runtime/AnimGraphRuntime/Private/AnimNodes/AnimNode_CopyPoseFromMesh.cpp
@@ -2,6 +2,7 @@
 
 #include "AnimGraphRuntimePrivatePCH.h"
 #include "AnimNodes/AnimNode_CopyPoseFromMesh.h"
+#include "Animation/AnimInstanceProxy.h"
 
 /////////////////////////////////////////////////////
 // FAnimNode_CopyPoseFromMesh
@@ -12,7 +13,8 @@ FAnimNode_CopyPoseFromMesh::FAnimNode_CopyPoseFromMesh()
 
 void FAnimNode_CopyPoseFromMesh::Initialize(const FAnimationInitializeContext& Context)
 {
-	ReinitializeMeshComponent(Context.AnimInstance);
+	FAnimNode_Base::Initialize(Context);
+	ReinitializeMeshComponent(Context.AnimInstanceProxy);
 }
 
 void FAnimNode_CopyPoseFromMesh::CacheBones(const FAnimationCacheBonesContext& Context)
@@ -25,11 +27,11 @@ void FAnimNode_CopyPoseFromMesh::Update(const FAnimationUpdateContext& Context)
 
 	if (CurrentlyUsedSourceMeshComponent.IsValid() && CurrentlyUsedSourceMeshComponent.Get() != SourceMeshComponent)
 	{
-		ReinitializeMeshComponent(Context.AnimInstance);
+		ReinitializeMeshComponent(Context.AnimInstanceProxy);
 	}
 	else if (!CurrentlyUsedSourceMeshComponent.IsValid() && SourceMeshComponent)
 	{
-		ReinitializeMeshComponent(Context.AnimInstance);
+		ReinitializeMeshComponent(Context.AnimInstanceProxy);
 	}
 }
 
@@ -44,7 +46,8 @@ void FAnimNode_CopyPoseFromMesh::Evaluate(FPoseContext& Output)
 		for(FCompactPoseBoneIndex PoseBoneIndex : OutPose.ForEachBoneIndex())
 		{
 			const int32& SkeletonBoneIndex = RequiredBones.GetSkeletonIndex(PoseBoneIndex);
-			const int32* Value = BoneMapToSource.Find(SkeletonBoneIndex);
+			const int32& MeshBoneIndex = RequiredBones.GetSkeletonToPoseBoneIndexArray()[SkeletonBoneIndex];
+			const int32* Value = BoneMapToSource.Find(MeshBoneIndex);
 			if(Value && *Value!=INDEX_NONE)
 			{
 				const int32 SourceBoneIndex = *Value;
@@ -70,13 +73,13 @@ void FAnimNode_CopyPoseFromMesh::GatherDebugData(FNodeDebugData& DebugData)
 {
 }
 
-void FAnimNode_CopyPoseFromMesh::ReinitializeMeshComponent(UAnimInstance* AnimInstance)
+void FAnimNode_CopyPoseFromMesh::ReinitializeMeshComponent(FAnimInstanceProxy* AnimInstanceProxy)
 {
 	CurrentlyUsedSourceMeshComponent = SourceMeshComponent;
 	BoneMapToSource.Reset();
 	if (SourceMeshComponent && SourceMeshComponent->SkeletalMesh && !SourceMeshComponent->IsPendingKill())
 	{
-		USkeletalMeshComponent* TargetMeshComponent = AnimInstance->GetSkelMeshComponent();
+		USkeletalMeshComponent* TargetMeshComponent = AnimInstanceProxy->GetSkelMeshComponent();
 		if (TargetMeshComponent)
 		{
 			USkeletalMesh* SourceSkelMesh = SourceMeshComponent->SkeletalMesh;
diff --git a/Engine/Source/Runtime/AnimGraphRuntime/Private/AnimNodes/AnimNode_LayeredBoneBlend.cpp b/Engine/Source/Runtime/AnimGraphRuntime/Private/AnimNodes/AnimNode_LayeredBoneBlend.cpp
--- a/Engine/Source/Runtime/AnimGraphRuntime/Private/AnimNodes/AnimNode_LayeredBoneBlend.cpp
+++ b/Engine/Source/Runtime/AnimGraphRuntime/Private/AnimNodes/AnimNode_LayeredBoneBlend.cpp
@@ -4,6 +4,7 @@
 #include "AnimNodes/AnimNode_LayeredBoneBlend.h"
 #include "AnimationRuntime.h"
 #include "Animation/AnimStats.h"
+#include "Animation/AnimInstanceProxy.h"
 
 /////////////////////////////////////////////////////
 // FAnimNode_LayeredBoneBlend
@@ -26,8 +27,8 @@ void FAnimNode_LayeredBoneBlend::Initialize(const FAnimationInitializeContext& C
 		}
 
 		// initialize mask weight now
-		check (Context.AnimInstance->CurrentSkeleton);
-		ReinitializeBoneBlendWeights(Context.AnimInstance->RequiredBones, Context.AnimInstance->CurrentSkeleton);
+		check (Context.AnimInstanceProxy->GetSkeleton());
+		ReinitializeBoneBlendWeights(Context.AnimInstanceProxy->GetRequiredBones(), Context.AnimInstanceProxy->GetSkeleton());
 	}
 }
 
@@ -47,9 +48,9 @@ void FAnimNode_LayeredBoneBlend::CacheBones(const FAnimationCacheBonesContext& C
 		BlendPoses[ChildIndex].CacheBones(Context);
 	}
 
-	if (Context.AnimInstance->RequiredBones.GetBoneIndicesArray().Num() != DesiredBoneBlendWeights.Num())
+	if (Context.AnimInstanceProxy->GetRequiredBones().GetBoneIndicesArray().Num() != DesiredBoneBlendWeights.Num())
 	{
-		ReinitializeBoneBlendWeights(Context.AnimInstance->RequiredBones, Context.AnimInstance->CurrentSkeleton);
+		ReinitializeBoneBlendWeights(Context.AnimInstanceProxy->GetRequiredBones(), Context.AnimInstanceProxy->GetSkeleton());
 	}
 }
 
@@ -77,7 +78,7 @@ void FAnimNode_LayeredBoneBlend::Update(const FAnimationUpdateContext& Context)
 
 void FAnimNode_LayeredBoneBlend::Evaluate(FPoseContext& Output)
 {
-	ANIM_MT_SCOPE_CYCLE_COUNTER(BlendPosesInGraph, Output.AnimInstance->IsRunningParallelEvaluation());
+	ANIM_MT_SCOPE_CYCLE_COUNTER(BlendPosesInGraph, !IsInGameThread());
 
 	const int NumPoses = BlendPoses.Num();
 	if ((NumPoses == 0) || !bHasRelevantPoses)
diff --git a/Engine/Source/Runtime/AnimGraphRuntime/Private/AnimNodes/AnimNode_RandomPlayer.cpp b/Engine/Source/Runtime/AnimGraphRuntime/Private/AnimNodes/AnimNode_RandomPlayer.cpp
new file mode 100644
--- /dev/null
+++ b/Engine/Source/Runtime/AnimGraphRuntime/Private/AnimNodes/AnimNode_RandomPlayer.cpp
@@ -0,0 +1,317 @@
+// Copyright 1998-2015 Epic Games, Inc. All Rights Reserved.
+
+#include "AnimGraphRuntimePrivatePCH.h"
+#include "AnimNode_RandomPlayer.h"
+#include "AnimInstanceProxy.h"
+
+FAnimNode_RandomPlayer::FAnimNode_RandomPlayer()
+: CurrentEntry(INDEX_NONE)
+, NextEntry(INDEX_NONE)
+, CurrentDataIndex(0)
+{
+
+}
+
+void FAnimNode_RandomPlayer::Initialize(const FAnimationInitializeContext& Context)
+{
+	FAnimNode_Base::Initialize(Context);
+	EvaluateGraphExposedInputs.Execute(Context);
+
+	const int32 NumEntries = Entries.Num();
+
+	if(NumEntries == 0)
+	{
+		// early out here, no need to do anything at all if we're not playing anything
+		return;
+	}
+
+	NormalizedPlayChances.Empty(NormalizedPlayChances.Num());
+	NormalizedPlayChances.AddUninitialized(NumEntries);
+
+	// Initialize normalized play chance for each entry and validate entry data
+	float SumChances = 0.0f;
+	for(FRandomPlayerSequenceEntry& Entry : Entries)
+	{
+		SumChances += Entry.ChanceToPlay;
+
+		if(Entry.MaxLoopCount < Entry.MinLoopCount)
+		{
+			Swap(Entry.MaxLoopCount, Entry.MinLoopCount);
+		}
+
+		if(Entry.MaxPlayRate < Entry.MinPlayRate)
+		{
+			Swap(Entry.MaxLoopCount, Entry.MinLoopCount);
+		}
+	}
+
+	for(int32 Idx = 0 ; Idx < NumEntries ; ++Idx)
+	{
+		NormalizedPlayChances[Idx] = Entries[Idx].ChanceToPlay / SumChances;
+	}
+
+	// Initialize random stream and pick first entry
+	RandomStream.Initialize(FPlatformTime::Cycles());
+	
+	CurrentEntry = GetNextEntryIndex();
+	NextEntry = GetNextEntryIndex();
+
+	PlayData.Empty(2);
+	PlayData.AddDefaulted(2);
+
+	FRandomAnimPlayData& CurrentData = PlayData[GetDataIndex(ERandomDataIndexType::Current)];
+	FRandomAnimPlayData& NextData = PlayData[GetDataIndex(ERandomDataIndexType::Next)];
+
+	// Init play data
+	CurrentData.BlendWeight = 1.0f;
+	CurrentData.PlayRate = RandomStream.FRandRange(Entries[CurrentEntry].MinPlayRate, Entries[CurrentEntry].MaxPlayRate);
+	CurrentData.RemainingLoops = FMath::Clamp(RandomStream.RandRange(Entries[CurrentEntry].MinLoopCount, Entries[CurrentEntry].MaxLoopCount), 0, MAX_int32);
+
+	NextData.BlendWeight = 0.0f;
+	NextData.PlayRate = RandomStream.FRandRange(Entries[NextEntry].MinPlayRate, Entries[NextEntry].MaxPlayRate);
+	NextData.RemainingLoops = FMath::Clamp(RandomStream.RandRange(Entries[NextEntry].MinLoopCount, Entries[NextEntry].MaxLoopCount), 0, MAX_int32);
+}
+
+void FAnimNode_RandomPlayer::Update(const FAnimationUpdateContext& Context)
+{
+	EvaluateGraphExposedInputs.Execute(Context);
+
+	FRandomAnimPlayData* CurrentData = &PlayData[GetDataIndex(ERandomDataIndexType::Current)];
+	FRandomAnimPlayData* NextData = &PlayData[GetDataIndex(ERandomDataIndexType::Next)];
+
+	if(UAnimSequence* CurrentSequence = Entries[CurrentEntry].Sequence)
+	{
+		float TimeRemaining = CurrentSequence->SequenceLength - CurrentData->InternalTimeAccumulator;
+
+		if(CurrentData->InternalTimeAccumulator < CurrentData->PreviousTimeAccumulator)
+		{
+			// We've looped, update remaining
+			--CurrentData->RemainingLoops;
+
+			// Special case for same entry - we don't do any blending so need to switch the data over here
+			if(CurrentData->RemainingLoops < 0 && CurrentEntry == NextEntry)
+			{
+				// Need to switch to the next anim, but first put our accumulator in the next data we're about to switch
+				// to so we don't see a pop
+				NextData->InternalTimeAccumulator = CurrentData->InternalTimeAccumulator;
+
+				SwitchNextToCurrent();
+
+				// Re-get data as we've switched over
+				CurrentData = &PlayData[GetDataIndex(ERandomDataIndexType::Current)];
+				NextData = &PlayData[GetDataIndex(ERandomDataIndexType::Next)];
+			}
+		}
+
+		// Cache time to detect loops
+		CurrentData->PreviousTimeAccumulator = CurrentData->InternalTimeAccumulator;
+		NextData->PreviousTimeAccumulator = NextData->InternalTimeAccumulator;
+
+		// If we're in the blend window start blending, but only if we're moving to a new animation,
+		// otherwise just keep looping.
+		const bool bInCrossfadeTime = TimeRemaining <= Entries[NextEntry].BlendIn.GetBlendTime();
+		const bool bNextAnimIsDifferent = NextEntry != CurrentEntry;
+		const bool bNeedMoreLoops = CurrentData->RemainingLoops > 0;
+
+		if(bInCrossfadeTime && !bNeedMoreLoops)
+		{
+			if(bNextAnimIsDifferent)
+			{
+				// Blending to next
+				Entries[NextEntry].BlendIn.Update(Context.GetDeltaTime());
+
+				float BlendedAlpha = Entries[NextEntry].BlendIn.GetBlendedValue();
+
+				if(BlendedAlpha < 1.0f)
+				{
+					NextData->BlendWeight = BlendedAlpha;
+					CurrentData->BlendWeight = 1.0f - BlendedAlpha;
+				}
+			}
+		}
+
+		// If we were blending but now we're done, switch play data
+		if(Entries[NextEntry].BlendIn.IsComplete())
+		{
+			SwitchNextToCurrent();
+
+			// Re-get data as we've switched over
+			CurrentData = &PlayData[GetDataIndex(ERandomDataIndexType::Current)];
+			NextData = &PlayData[GetDataIndex(ERandomDataIndexType::Next)];
+		}
+
+		if(FAnimInstanceProxy* AnimProxy = Context.AnimInstanceProxy)
+		{
+			FAnimGroupInstance* SyncGroup;
+			FAnimTickRecord& TickRecord = AnimProxy->CreateUninitializedTickRecord(INDEX_NONE, SyncGroup);
+			AnimProxy->MakeSequenceTickRecord(TickRecord, Entries[CurrentEntry].Sequence, true, CurrentData->PlayRate, CurrentData->BlendWeight, CurrentData->InternalTimeAccumulator, CurrentData->MarkerTickRecord);
+
+			if(NextData->BlendWeight > 0.0f)
+			{
+				FAnimTickRecord& NextTickRecord = AnimProxy->CreateUninitializedTickRecord(INDEX_NONE, SyncGroup);
+				AnimProxy->MakeSequenceTickRecord(NextTickRecord, Entries[NextEntry].Sequence, true, NextData->PlayRate, NextData->BlendWeight, NextData->InternalTimeAccumulator, NextData->MarkerTickRecord);
+			}
+		}
+	}
+}
+
+void FAnimNode_RandomPlayer::Evaluate(FPoseContext& Output)
+{
+	if(Entries.Num() > 0)
+	{
+		UAnimSequence* CurrentSequence = Entries[CurrentEntry].Sequence;
+
+		if(CurrentSequence)
+		{
+			FRandomAnimPlayData& CurrentData = PlayData[GetDataIndex(ERandomDataIndexType::Current)];
+			FRandomAnimPlayData& NextData = PlayData[GetDataIndex(ERandomDataIndexType::Next)];
+
+			if(CurrentData.BlendWeight != 1.0f)
+			{
+				if(FAnimInstanceProxy* AnimProxy = Output.AnimInstanceProxy)
+				{
+					// Start Blending
+					FCompactPose Poses[2];
+					FBlendedCurve Curves[2];
+					float Weights[2];
+
+					Poses[0].SetBoneContainer(&AnimProxy->GetRequiredBones());
+					Poses[1].SetBoneContainer(&AnimProxy->GetRequiredBones());
+
+					Curves[0].InitFrom(AnimProxy->GetSkeleton());
+					Curves[1].InitFrom(AnimProxy->GetSkeleton());
+
+					Weights[0] = CurrentData.BlendWeight;
+					Weights[1] = NextData.BlendWeight;
+
+					UAnimSequence* NextSequence = Entries[NextEntry].Sequence;
+
+					CurrentSequence->GetAnimationPose(Poses[0], Curves[0], FAnimExtractContext(CurrentData.InternalTimeAccumulator, AnimProxy->ShouldExtractRootMotion()));
+					NextSequence->GetAnimationPose(Poses[1], Curves[1], FAnimExtractContext(NextData.InternalTimeAccumulator, AnimProxy->ShouldExtractRootMotion()));
+
+					FAnimationRuntime::BlendPosesTogether(TFixedSizeArrayView<FCompactPose>(Poses, 2), TFixedSizeArrayView<FBlendedCurve>(Curves, 2), TFixedSizeArrayView<float>(Weights, 2), Output.Pose, Output.Curve);
+				}
+				else
+				{
+					Output.ResetToRefPose();
+				}
+			}
+			else
+			{
+				// Single anim
+				CurrentSequence->GetAnimationPose(Output.Pose, Output.Curve, FAnimExtractContext(CurrentData.InternalTimeAccumulator, Output.AnimInstanceProxy->ShouldExtractRootMotion()));
+			}
+		}
+		else
+		{
+			Output.ResetToRefPose();
+		}
+	}
+	else
+	{
+		Output.ResetToRefPose();
+	}
+}
+
+void FAnimNode_RandomPlayer::GatherDebugData(FNodeDebugData& DebugData)
+{
+	FString DebugLine = DebugData.GetNodeName(this);
+
+	DebugData.AddDebugItem(DebugLine, true);
+}
+
+int32 FAnimNode_RandomPlayer::GetNextEntryIndex()
+{
+	if(Entries.Num() > 0)
+	{
+		if(bShuffleMode)
+		{
+			if(ShuffleList.Num() == 0)
+			{
+				// Need a new list
+				BuildShuffleList();
+			}
+
+			// Get the top value, don't allow realloc
+			return ShuffleList.Pop(false);
+		}
+		else
+		{
+			float RandomVal = RandomStream.GetFraction();
+			const int32 NumEntries = Entries.Num();
+
+			// Grab the entry index corresponding to the value
+			for(int32 Idx = 0 ; Idx < NumEntries ; ++Idx)
+			{
+				RandomVal -= NormalizedPlayChances[Idx];
+				if(RandomVal <= 0.0f)
+				{
+					return Idx;
+				}
+			}
+		}
+	}
+
+	return INDEX_NONE;
+}
+
+int32 FAnimNode_RandomPlayer::GetDataIndex(const ERandomDataIndexType& Type)
+{
+	if(Type == ERandomDataIndexType::Current)
+	{
+		return CurrentDataIndex;
+	}
+	else
+	{
+		// Next Accumulator
+		return (CurrentDataIndex + 1) % 2;
+	}
+}
+
+void FAnimNode_RandomPlayer::SwitchNextToCurrent()
+{
+	// reset alpha blend we've possibly just taken
+	Entries[NextEntry].BlendIn.Reset();
+
+	// Switch which entry to get sequences and parameters from, and pre-generate the next entry index
+	CurrentEntry = NextEntry;
+	NextEntry = GetNextEntryIndex();
+
+	// Switch play data
+	CurrentDataIndex = (CurrentDataIndex + 1) %2;
+
+	// Get our play data
+	FRandomAnimPlayData& CurrentData = PlayData[GetDataIndex(ERandomDataIndexType::Current)];
+	FRandomAnimPlayData& NextData = PlayData[GetDataIndex(ERandomDataIndexType::Next)];
+
+	// Reset blendweights
+	CurrentData.BlendWeight = 1.0f;
+	NextData.BlendWeight = 0.0f;
+
+	// Set up data for next switch
+	NextData.InternalTimeAccumulator = 0.0f;
+	NextData.PreviousTimeAccumulator = 0.0f;
+	NextData.PlayRate = RandomStream.FRandRange(Entries[NextEntry].MinPlayRate, Entries[NextEntry].MaxPlayRate);
+	NextData.RemainingLoops = FMath::Clamp(RandomStream.RandRange(Entries[NextEntry].MinLoopCount, Entries[NextEntry].MaxLoopCount), 0, MAX_int32);
+	NextData.MarkerTickRecord.Reset();
+}
+
+void FAnimNode_RandomPlayer::BuildShuffleList()
+{
+	ShuffleList.Reset(Entries.Num());
+
+	// Build entry index list
+	const int32 NumEntries = Entries.Num();
+	for(int32 i = 0 ; i < NumEntries ; ++i)
+	{
+		ShuffleList.Add(i);
+	}
+
+	// Shuffle the list
+	const int32 NumShuffles = ShuffleList.Num() - 1;
+	for(int32 i = 0 ; i < NumShuffles ; ++i)
+	{
+		int32 SwapIdx = RandomStream.RandRange(i, NumShuffles);
+		ShuffleList.Swap(i, SwapIdx);
+	}
+}
diff --git a/Engine/Source/Runtime/AnimGraphRuntime/Private/AnimNodes/AnimNode_RotationOffsetBlendSpace.cpp b/Engine/Source/Runtime/AnimGraphRuntime/Private/AnimNodes/AnimNode_RotationOffsetBlendSpace.cpp
--- a/Engine/Source/Runtime/AnimGraphRuntime/Private/AnimNodes/AnimNode_RotationOffsetBlendSpace.cpp
+++ b/Engine/Source/Runtime/AnimGraphRuntime/Private/AnimNodes/AnimNode_RotationOffsetBlendSpace.cpp
@@ -20,7 +20,7 @@ void FAnimNode_RotationOffsetBlendSpace::CacheBones(const FAnimationCacheBonesCo
 
 void FAnimNode_RotationOffsetBlendSpace::UpdateAssetPlayer(const FAnimationUpdateContext& Context)
 {
-	if (IsLODEnabled(Context.AnimInstance, LODThreshold))
+	if (IsLODEnabled(Context.AnimInstanceProxy, LODThreshold))
 	{
 		FAnimNode_BlendSpacePlayer::UpdateAssetPlayer(Context);
 	}
@@ -33,7 +33,7 @@ void FAnimNode_RotationOffsetBlendSpace::Evaluate(FPoseContext& Context)
 	// Evaluate base pose
 	BasePose.Evaluate(Context);
 
-	if (IsLODEnabled(Context.AnimInstance, LODThreshold))
+	if (IsLODEnabled(Context.AnimInstanceProxy, LODThreshold))
 	{
 		// Evaluate MeshSpaceRotation additive blendspace
 		FPoseContext MeshSpaceRotationAdditivePoseContext(Context);
diff --git a/Engine/Source/Runtime/AnimGraphRuntime/Private/AnimNodes/AnimNode_SaveCachedPose.cpp b/Engine/Source/Runtime/AnimGraphRuntime/Private/AnimNodes/AnimNode_SaveCachedPose.cpp
--- a/Engine/Source/Runtime/AnimGraphRuntime/Private/AnimNodes/AnimNode_SaveCachedPose.cpp
+++ b/Engine/Source/Runtime/AnimGraphRuntime/Private/AnimNodes/AnimNode_SaveCachedPose.cpp
@@ -2,6 +2,7 @@
 
 #include "AnimGraphRuntimePrivatePCH.h"
 #include "AnimNodes/AnimNode_SaveCachedPose.h"
+#include "Animation/AnimInstanceProxy.h"
 
 /////////////////////////////////////////////////////
 // FAnimNode_SaveCachedPose
@@ -14,10 +15,10 @@ void FAnimNode_SaveCachedPose::Initialize(const FAnimationInitializeContext& Con
 {
 	// StateMachines cause reinitialization on state changes.
 	// we only want to let them through if we're not relevant as to not create a pop.
-	if (!InitializationCounter.IsSynchronizedWith(Context.AnimInstance->InitializationCounter)
-		|| ((EvaluationCounter.Get() != INDEX_NONE) && !EvaluationCounter.WasSynchronizedInTheLastFrame(Context.AnimInstance->EvaluationCounter)))
+	if (!InitializationCounter.IsSynchronizedWith(Context.AnimInstanceProxy->GetInitializationCounter())
+		|| ((EvaluationCounter.Get() != INDEX_NONE) && !EvaluationCounter.WasSynchronizedInTheLastFrame(Context.AnimInstanceProxy->GetEvaluationCounter())))
 	{
-		InitializationCounter.SynchronizeWith(Context.AnimInstance->InitializationCounter);
+		InitializationCounter.SynchronizeWith(Context.AnimInstanceProxy->GetInitializationCounter());
 
 		// Reset update counter here as Initialize can happen as part of a state machine update
 		// when the current state changes. This ensures that any subsequent updates that may 
@@ -33,9 +34,9 @@ void FAnimNode_SaveCachedPose::Initialize(const FAnimationInitializeContext& Con
 
 void FAnimNode_SaveCachedPose::CacheBones(const FAnimationCacheBonesContext& Context)
 {
-	if (!CachedBonesCounter.IsSynchronizedWith(Context.AnimInstance->CachedBonesCounter))
+	if (!CachedBonesCounter.IsSynchronizedWith(Context.AnimInstanceProxy->GetCachedBonesCounter()))
 	{
-		CachedBonesCounter.SynchronizeWith(Context.AnimInstance->CachedBonesCounter);
+		CachedBonesCounter.SynchronizeWith(Context.AnimInstanceProxy->GetCachedBonesCounter());
 
 		// Cache bones in the subgraph
 		Pose.CacheBones(Context);
@@ -44,9 +45,9 @@ void FAnimNode_SaveCachedPose::CacheBones(const FAnimationCacheBonesContext& Con
 
 void FAnimNode_SaveCachedPose::Update(const FAnimationUpdateContext& Context)
 {
-	if (!UpdateCounter.IsSynchronizedWith(Context.AnimInstance->UpdateCounter))
+	if (!UpdateCounter.IsSynchronizedWith(Context.AnimInstanceProxy->GetUpdateCounter()))
 	{
-		UpdateCounter.SynchronizeWith(Context.AnimInstance->UpdateCounter);
+		UpdateCounter.SynchronizeWith(Context.AnimInstanceProxy->GetUpdateCounter());
 
 		// Update the subgraph
 		Pose.Update(Context);
@@ -55,9 +56,9 @@ void FAnimNode_SaveCachedPose::Update(const FAnimationUpdateContext& Context)
 
 void FAnimNode_SaveCachedPose::Evaluate(FPoseContext& Output)
 {
-	if (!EvaluationCounter.IsSynchronizedWith(Output.AnimInstance->EvaluationCounter))
+	if (!EvaluationCounter.IsSynchronizedWith(Output.AnimInstanceProxy->GetEvaluationCounter()))
 	{
-		EvaluationCounter.SynchronizeWith(Output.AnimInstance->EvaluationCounter);
+		EvaluationCounter.SynchronizeWith(Output.AnimInstanceProxy->GetEvaluationCounter());
 
 		FPoseContext CachingContext(Output);
 		Pose.Evaluate(CachingContext);
diff --git a/Engine/Source/Runtime/AnimGraphRuntime/Private/AnimNodes/AnimNode_SequenceEvaluator.cpp b/Engine/Source/Runtime/AnimGraphRuntime/Private/AnimNodes/AnimNode_SequenceEvaluator.cpp
--- a/Engine/Source/Runtime/AnimGraphRuntime/Private/AnimNodes/AnimNode_SequenceEvaluator.cpp
+++ b/Engine/Source/Runtime/AnimGraphRuntime/Private/AnimNodes/AnimNode_SequenceEvaluator.cpp
@@ -2,6 +2,7 @@
 
 #include "AnimGraphRuntimePrivatePCH.h"
 #include "AnimNodes/AnimNode_SequenceEvaluator.h"
+#include "Animation/AnimInstanceProxy.h"
 
 /////////////////////////////////////////////////////
 // FAnimSequenceEvaluatorNode
@@ -18,7 +19,7 @@ void FAnimNode_SequenceEvaluator::CacheBones(const FAnimationCacheBonesContext&
 void FAnimNode_SequenceEvaluator::UpdateAssetPlayer(const FAnimationUpdateContext& Context)
 {
 	EvaluateGraphExposedInputs.Execute(Context);
-	if ((GroupIndex != INDEX_NONE) && (Sequence != NULL) && (Context.AnimInstance->CurrentSkeleton->IsCompatible(Sequence->GetSkeleton())))
+	if ((GroupIndex != INDEX_NONE) && (Sequence != NULL) && (Context.AnimInstanceProxy->IsSkeletonCompatible(Sequence->GetSkeleton())))
 	{
 		float TimeJump = ExplicitTime - InternalTimeAccumulator;
 		if (bShouldLoopWhenInSyncGroup)
@@ -48,9 +49,9 @@ void FAnimNode_SequenceEvaluator::UpdateAssetPlayer(const FAnimationUpdateContex
 
 void FAnimNode_SequenceEvaluator::Evaluate(FPoseContext& Output)
 {
-	if ((Sequence != NULL) && (Output.AnimInstance->CurrentSkeleton->IsCompatible(Sequence->GetSkeleton())))
+	if ((Sequence != NULL) && (Output.AnimInstanceProxy->IsSkeletonCompatible(Sequence->GetSkeleton())))
 	{
-		Sequence->GetAnimationPose(Output.Pose, Output.Curve, FAnimExtractContext(InternalTimeAccumulator, Output.AnimInstance->ShouldExtractRootMotion()));
+		Sequence->GetAnimationPose(Output.Pose, Output.Curve, FAnimExtractContext(InternalTimeAccumulator, Output.AnimInstanceProxy->ShouldExtractRootMotion()));
 	}
 	else
 	{
diff --git a/Engine/Source/Runtime/AnimGraphRuntime/Private/AnimNodes/AnimNode_Slot.cpp b/Engine/Source/Runtime/AnimGraphRuntime/Private/AnimNodes/AnimNode_Slot.cpp
--- a/Engine/Source/Runtime/AnimGraphRuntime/Private/AnimNodes/AnimNode_Slot.cpp
+++ b/Engine/Source/Runtime/AnimGraphRuntime/Private/AnimNodes/AnimNode_Slot.cpp
@@ -3,6 +3,7 @@
 #include "AnimGraphRuntimePrivatePCH.h"
 #include "AnimNodes/AnimNode_Slot.h"
 #include "Animation/AnimMontage.h"
+#include "Animation/AnimInstanceProxy.h"
 
 /////////////////////////////////////////////////////
 // FAnimNode_Slot
@@ -15,10 +16,10 @@ void FAnimNode_Slot::Initialize(const FAnimationInitializeContext& Context)
 	WeightData.Reset();
 
 	// If this node has not already been registered with the AnimInstance, do it.
-	if (!SlotNodeInitializationCounter.IsSynchronizedWith(Context.AnimInstance->SlotNodeInitializationCounter))
+	if (!SlotNodeInitializationCounter.IsSynchronizedWith(Context.AnimInstanceProxy->GetSlotNodeInitializationCounter()))
 	{
-		SlotNodeInitializationCounter.SynchronizeWith(Context.AnimInstance->SlotNodeInitializationCounter);
-		Context.AnimInstance->RegisterSlotNodeWithAnimInstance(SlotName);
+		SlotNodeInitializationCounter.SynchronizeWith(Context.AnimInstanceProxy->GetSlotNodeInitializationCounter());
+		Context.AnimInstanceProxy->RegisterSlotNodeWithAnimInstance(SlotName);
 	}
 }
 
@@ -30,11 +31,11 @@ void FAnimNode_Slot::CacheBones(const FAnimationCacheBonesContext& Context)
 void FAnimNode_Slot::Update(const FAnimationUpdateContext& Context)
 {
 	// Update weights.
-	Context.AnimInstance->GetSlotWeight(SlotName, WeightData.SlotNodeWeight, WeightData.SourceWeight, WeightData.TotalNodeWeight);
+	Context.AnimInstanceProxy->GetSlotWeight(SlotName, WeightData.SlotNodeWeight, WeightData.SourceWeight, WeightData.TotalNodeWeight);
 
 	// Update cache in AnimInstance.
-	Context.AnimInstance->UpdateSlotNodeWeight(SlotName, WeightData.SlotNodeWeight);
-	Context.AnimInstance->UpdateSlotRootMotionWeight(SlotName, Context.GetFinalBlendWeight());
+	Context.AnimInstanceProxy->UpdateSlotNodeWeight(SlotName, WeightData.SlotNodeWeight);
+	Context.AnimInstanceProxy->UpdateSlotRootMotionWeight(SlotName, Context.GetFinalBlendWeight());
 
 	if (WeightData.SourceWeight > ZERO_ANIMWEIGHT_THRESH)
 	{
@@ -57,7 +58,7 @@ void FAnimNode_Slot::Evaluate(FPoseContext & Output)
 			Source.Evaluate(SourceContext);
 		}
 
-		Output.AnimInstance->SlotEvaluatePose(SlotName, SourceContext.Pose, SourceContext.Curve, WeightData.SourceWeight, Output.Pose, Output.Curve, WeightData.SlotNodeWeight, WeightData.TotalNodeWeight);
+		Output.AnimInstanceProxy->SlotEvaluatePose(SlotName, SourceContext.Pose, SourceContext.Curve, WeightData.SourceWeight, Output.Pose, Output.Curve, WeightData.SlotNodeWeight, WeightData.TotalNodeWeight);
 
 		checkSlow(!Output.ContainsNaN());
 		checkSlow(Output.IsNormalized());
diff --git a/Engine/Source/Runtime/AnimGraphRuntime/Private/BoneControllers/AnimNode_AnimDynamics.cpp b/Engine/Source/Runtime/AnimGraphRuntime/Private/BoneControllers/AnimNode_AnimDynamics.cpp
--- a/Engine/Source/Runtime/AnimGraphRuntime/Private/BoneControllers/AnimNode_AnimDynamics.cpp
+++ b/Engine/Source/Runtime/AnimGraphRuntime/Private/BoneControllers/AnimNode_AnimDynamics.cpp
@@ -2,7 +2,9 @@
 
 #include "AnimGraphRuntimePrivatePCH.h"
 #include "AnimNode_AnimDynamics.h"
+#include "Animation/AnimInstanceProxy.h"
 
+DEFINE_STAT(STAT_AnimDynamicsOverall);
 DEFINE_STAT(STAT_AnimDynamicsWindData);
 DEFINE_STAT(STAT_AnimDynamicsBoneEval);
 DEFINE_STAT(STAT_AnimDynamicsSubSteps);
@@ -42,8 +44,9 @@ void FAnimNode_AnimDynamics::Initialize(const FAnimationInitializeContext& Conte
 {
 	FAnimNode_SkeletalControlBase::Initialize(Context);
 
-	USkeletalMeshComponent* SkelMeshComponent = Context.AnimInstance->GetSkelMeshComponent();
-	FBoneContainer& RequiredBones = Context.AnimInstance->RequiredBones;
+	Context.AnimInstanceProxy->AddGameThreadPreUpdateEvent(FGameThreadPreUpdateEvent::CreateRaw(this, &FAnimNode_AnimDynamics::HandleGameThreadPreUpdateEvent));
+
+	FBoneContainer& RequiredBones = Context.AnimInstanceProxy->GetRequiredBones();
 
 	BoundBone.Initialize(RequiredBones);
 
@@ -75,6 +78,8 @@ void FAnimNode_AnimDynamics::UpdateInternal(const FAnimationUpdateContext& Conte
 
 void FAnimNode_AnimDynamics::EvaluateBoneTransforms(USkeletalMeshComponent* SkelComp, FCSPose<FCompactPose>& MeshBases, TArray<FBoneTransform>& OutBoneTransforms)
 {
+	SCOPE_CYCLE_COUNTER(STAT_AnimDynamicsOverall);
+
 	int32 RestrictToLOD = CVarRestrictLod.GetValueOnAnyThread();
 	bool bEnabledForLod = RestrictToLOD >= 0 ? SkelComp->PredictedLODLevel == RestrictToLOD : true;
 
@@ -130,12 +135,6 @@ void FAnimNode_AnimDynamics::EvaluateBoneTransforms(USkeletalMeshComponent* Skel
 					Body->bWindEnabled = false;
 				}
 			}
-			
-			float CurrentTimeDilation = 1.0f;
-			if(UAnimInstance* AnimationInstance = SkelComp->GetAnimInstance())
-			{
-				CurrentTimeDilation = AnimationInstance->CurrentTimeDilation;
-			}
 
 			if (CVarEnableAdaptiveSubstep.GetValueOnAnyThread() == 1)
 			{
@@ -551,3 +550,11 @@ void FAnimNode_AnimDynamics::UpdateLimits(USkeletalMeshComponent* SkelComp, FCSP
 		}
 	}
 }
+
+void FAnimNode_AnimDynamics::HandleGameThreadPreUpdateEvent(const UAnimInstance* InAnimInstance)
+{
+	const USkeletalMeshComponent* SkelComp = InAnimInstance->GetSkelMeshComponent();
+	const UWorld* World = SkelComp->GetWorld();
+	check(World->GetWorldSettings());
+	CurrentTimeDilation = World->GetWorldSettings()->GetEffectiveTimeDilation();
+}
\ No newline at end of file
diff --git a/Engine/Source/Runtime/AnimGraphRuntime/Private/BoneControllers/AnimNode_SkeletalControlBase.cpp b/Engine/Source/Runtime/AnimGraphRuntime/Private/BoneControllers/AnimNode_SkeletalControlBase.cpp
--- a/Engine/Source/Runtime/AnimGraphRuntime/Private/BoneControllers/AnimNode_SkeletalControlBase.cpp
+++ b/Engine/Source/Runtime/AnimGraphRuntime/Private/BoneControllers/AnimNode_SkeletalControlBase.cpp
@@ -4,6 +4,7 @@
 #include "BoneControllers/AnimNode_SkeletalControlBase.h"
 #include "AnimationUtils.h"
 #include "AnimationRuntime.h"
+#include "Animation/AnimInstanceProxy.h"
 
 /////////////////////////////////////////////////////
 // FAnimNode_SkeletalControlBase
@@ -17,7 +18,7 @@ void FAnimNode_SkeletalControlBase::Initialize(const FAnimationInitializeContext
 
 void FAnimNode_SkeletalControlBase::CacheBones(const FAnimationCacheBonesContext& Context) 
 {
-	InitializeBoneReferences(Context.AnimInstance->RequiredBones);
+	InitializeBoneReferences(Context.AnimInstanceProxy->GetRequiredBones());
 	ComponentPose.CacheBones(Context);
 }
 
@@ -29,13 +30,13 @@ void FAnimNode_SkeletalControlBase::Update(const FAnimationUpdateContext& Contex
 {
 	ComponentPose.Update(Context);
 
-	if (IsLODEnabled(Context.AnimInstance, LODThreshold))
+	if (IsLODEnabled(Context.AnimInstanceProxy, LODThreshold))
 	{
 		EvaluateGraphExposedInputs.Execute(Context);
 
 		// Apply the skeletal control if it's valid
 		const float ActualAlpha = AlphaScaleBias.ApplyTo(Alpha);
-		if ((ActualAlpha >= ZERO_ANIMWEIGHT_THRESH) && IsValidToEvaluate(Context.AnimInstance->CurrentSkeleton, Context.AnimInstance->RequiredBones))
+		if ((ActualAlpha >= ZERO_ANIMWEIGHT_THRESH) && IsValidToEvaluate(Context.AnimInstanceProxy->GetSkeleton(), Context.AnimInstanceProxy->GetRequiredBones()))
 		{
 			UpdateInternal(Context);
 		}
@@ -60,13 +61,13 @@ void FAnimNode_SkeletalControlBase::EvaluateComponentSpace(FComponentSpacePoseCo
 	// Evaluate the input
 	ComponentPose.EvaluateComponentSpace(Output);
 
-	if (IsLODEnabled(Output.AnimInstance, LODThreshold))
+	if (IsLODEnabled(Output.AnimInstanceProxy, LODThreshold))
 	{
 		// Apply the skeletal control if it's valid
 		const float ActualAlpha = AlphaScaleBias.ApplyTo(Alpha);
-		if ((ActualAlpha >= ZERO_ANIMWEIGHT_THRESH) && IsValidToEvaluate(Output.AnimInstance->CurrentSkeleton, Output.AnimInstance->RequiredBones))
+		if ((ActualAlpha >= ZERO_ANIMWEIGHT_THRESH) && IsValidToEvaluate(Output.AnimInstanceProxy->GetSkeleton(), Output.AnimInstanceProxy->GetRequiredBones()))
 		{
-			USkeletalMeshComponent* Component = Output.AnimInstance->GetSkelMeshComponent();
+			USkeletalMeshComponent* Component = Output.AnimInstanceProxy->GetSkelMeshComponent();
 
 #if WITH_EDITORONLY_DATA
 			// save current pose before applying skeletal control to compute the exact gizmo location in AnimGraphNode
diff --git a/Engine/Source/Runtime/AnimGraphRuntime/Private/BoneControllers/AnimNode_SpringBone.cpp b/Engine/Source/Runtime/AnimGraphRuntime/Private/BoneControllers/AnimNode_SpringBone.cpp
--- a/Engine/Source/Runtime/AnimGraphRuntime/Private/BoneControllers/AnimNode_SpringBone.cpp
+++ b/Engine/Source/Runtime/AnimGraphRuntime/Private/BoneControllers/AnimNode_SpringBone.cpp
@@ -2,6 +2,7 @@
 
 #include "AnimGraphRuntimePrivatePCH.h"
 #include "BoneControllers/AnimNode_SpringBone.h"
+#include "Animation/AnimInstanceProxy.h"
 
 /////////////////////////////////////////////////////
 // FAnimNode_SpringBone
@@ -30,6 +31,8 @@ void FAnimNode_SpringBone::Initialize(const FAnimationInitializeContext& Context
 {
 	FAnimNode_SkeletalControlBase::Initialize(Context);
 
+	Context.AnimInstanceProxy->AddGameThreadPreUpdateEvent(FGameThreadPreUpdateEvent::CreateRaw(this, &FAnimNode_SpringBone::HandleGameThreadPreUpdateEvent));
+
 	RemainingTime = 0.0f;
 }
 
@@ -45,7 +48,7 @@ void FAnimNode_SpringBone::UpdateInternal(const FAnimationUpdateContext& Context
 	RemainingTime += Context.GetDeltaTime();
 
 	// Fixed step simulation at 120hz
-	FixedTimeStep = (1.f / 120.f) * Context.AnimInstance->CurrentTimeDilation;
+	FixedTimeStep = (1.f / 120.f) * TimeDilation;
 }
 
 void FAnimNode_SpringBone::GatherDebugData(FNodeDebugData& DebugData)
@@ -215,3 +218,11 @@ void FAnimNode_SpringBone::InitializeBoneReferences(const FBoneContainer& Requir
 {
 	SpringBone.Initialize(RequiredBones);
 }
+
+void FAnimNode_SpringBone::HandleGameThreadPreUpdateEvent(const UAnimInstance* InAnimInstance)
+{
+	const USkeletalMeshComponent* SkelComp = InAnimInstance->GetSkelMeshComponent();
+	const UWorld* World = SkelComp->GetWorld();
+	check(World->GetWorldSettings());
+	TimeDilation = World->GetWorldSettings()->GetEffectiveTimeDilation();
+}
\ No newline at end of file
diff --git a/Engine/Source/Runtime/AnimGraphRuntime/Private/BoneControllers/AnimNode_Trail.cpp b/Engine/Source/Runtime/AnimGraphRuntime/Private/BoneControllers/AnimNode_Trail.cpp
--- a/Engine/Source/Runtime/AnimGraphRuntime/Private/BoneControllers/AnimNode_Trail.cpp
+++ b/Engine/Source/Runtime/AnimGraphRuntime/Private/BoneControllers/AnimNode_Trail.cpp
@@ -2,6 +2,7 @@
 
 #include "AnimGraphRuntimePrivatePCH.h"
 #include "BoneControllers/AnimNode_Trail.h"
+#include "Animation/AnimInstanceProxy.h"
 
 /////////////////////////////////////////////////////
 // FAnimNode_Trail
@@ -296,7 +297,7 @@ void FAnimNode_Trail::Initialize(const FAnimationInitializeContext& Context)
 {
 	FAnimNode_SkeletalControlBase::Initialize(Context);
 
-	const USkeleton* Skeleton = Context.AnimInstance->CurrentSkeleton;
+	const USkeleton* Skeleton = Context.AnimInstanceProxy->GetSkeleton();
 	check (Skeleton);
 
 	BaseJoint.Initialize(Skeleton);
diff --git a/Engine/Source/Runtime/AnimGraphRuntime/Private/BoneControllers/AnimNode_WheelHandler.cpp b/Engine/Source/Runtime/AnimGraphRuntime/Private/BoneControllers/AnimNode_WheelHandler.cpp
--- a/Engine/Source/Runtime/AnimGraphRuntime/Private/BoneControllers/AnimNode_WheelHandler.cpp
+++ b/Engine/Source/Runtime/AnimGraphRuntime/Private/BoneControllers/AnimNode_WheelHandler.cpp
@@ -6,6 +6,8 @@
 #include "GameFramework/WheeledVehicle.h"
 #include "Vehicles/WheeledVehicleMovementComponent.h"
 #include "Vehicles/VehicleWheel.h"
+#include "Animation/AnimInstanceProxy.h"
+
 
 /////////////////////////////////////////////////////
 // FAnimNode_WheelHandler
@@ -132,7 +134,7 @@ void FAnimNode_WheelHandler::Initialize(const FAnimationInitializeContext& Conte
 {
 	// TODO: only check vehicle anim instance
 	// UVehicleAnimInstance
-	AWheeledVehicle * Vehicle = Cast<AWheeledVehicle> ((CastChecked<USkeletalMeshComponent> (Context.AnimInstance->GetOuter()))->GetOwner());
+	AWheeledVehicle * Vehicle = Cast<AWheeledVehicle> (Context.AnimInstanceProxy->GetSkelMeshComponent()->GetOwner());
 	;
 	// we only support vehicle for this class
 	if(Vehicle != nullptr)
diff --git a/Engine/Source/Runtime/AnimGraphRuntime/Public/AnimNodes/AnimNode_CopyPoseFromMesh.h b/Engine/Source/Runtime/AnimGraphRuntime/Public/AnimNodes/AnimNode_CopyPoseFromMesh.h
--- a/Engine/Source/Runtime/AnimGraphRuntime/Public/AnimNodes/AnimNode_CopyPoseFromMesh.h
+++ b/Engine/Source/Runtime/AnimGraphRuntime/Public/AnimNodes/AnimNode_CopyPoseFromMesh.h
@@ -33,5 +33,5 @@ private:
 	TMap<int32, int32> BoneMapToSource;
 
 	// reinitialize mesh component 
-	void ReinitializeMeshComponent(UAnimInstance* AnimInstance);
+	void ReinitializeMeshComponent(FAnimInstanceProxy* AnimInstanceProxy);
 };
diff --git a/Engine/Source/Runtime/AnimGraphRuntime/Public/AnimNodes/AnimNode_RandomPlayer.h b/Engine/Source/Runtime/AnimGraphRuntime/Public/AnimNodes/AnimNode_RandomPlayer.h
new file mode 100644
--- /dev/null
+++ b/Engine/Source/Runtime/AnimGraphRuntime/Public/AnimNodes/AnimNode_RandomPlayer.h
@@ -0,0 +1,150 @@
+// Copyright 1998-2015 Epic Games, Inc. All Rights Reserved.
+
+#pragma once
+
+#include "Animation/AnimNodeBase.h"
+#include "AnimNode_RandomPlayer.generated.h"
+
+enum class ERandomDataIndexType
+{
+	Current,
+	Next,
+};
+
+struct FRandomAnimPlayData
+{
+	FRandomAnimPlayData()
+		: PreviousTimeAccumulator(0.0f)
+		, InternalTimeAccumulator(0.0f)
+		, PlayRate(1.0f)
+		, BlendWeight(0.0f)
+		, RemainingLoops(0)
+	{
+
+	}
+
+	// Prev time, used to track loops (prev > current)
+	float PreviousTimeAccumulator;
+
+	// Current time through the sequence
+	float InternalTimeAccumulator;
+
+	// Calculated play rate 
+	float PlayRate;
+
+	// Current blend weight
+	float BlendWeight;
+
+	// Calculated loops remaining
+	int32 RemainingLoops;
+
+	// Marker tick record for this play through
+	FMarkerTickRecord MarkerTickRecord;
+};
+
+/** The random player node holds a list of sequences and parameter ranges which will be played continuously
+  * In a random order. If shuffle mode is enabled then each entry will be played once before repeating any
+  */
+USTRUCT()
+struct FRandomPlayerSequenceEntry
+{
+	GENERATED_BODY()
+
+	FRandomPlayerSequenceEntry()
+	: Sequence(nullptr)
+	, ChanceToPlay(1.0f)
+	, MinLoopCount(0)
+	, MaxLoopCount(0)
+	, MinPlayRate(1.0f)
+	, MaxPlayRate(1.0f)
+	{
+
+	}
+
+	/** Sequence to play when this entry is picked */
+	UPROPERTY(EditAnywhere, BlueprintReadWrite, Category = "Settings")
+	UAnimSequence* Sequence;
+
+	/** When not in shuffle mode, this is the chance this entry will play (normalized against all other sample chances) */
+	UPROPERTY(EditAnywhere, BlueprintReadWrite, Category = "Settings")
+	float ChanceToPlay;
+
+	/** Minimum number of times this entry will loop before ending */
+	UPROPERTY(EditAnywhere, BlueprintReadWrite, Category = "Settings", meta=(UIMin="0", ClampMin="0"))
+	int32 MinLoopCount;
+
+	/** Maximum number of times this entry will loop before ending */
+	UPROPERTY(EditAnywhere, BlueprintReadWrite, Category = "Settings", meta=(UIMin="0", ClampMin="0"))
+	int32 MaxLoopCount;
+
+	/** Minimum playrate for this entry */
+	UPROPERTY(EditAnywhere, BlueprintReadWrite, Category = "Settings", meta=(UIMin="0", ClampMin="0"))
+	float MinPlayRate;
+
+	/** Maximum playrate for this entry */
+	UPROPERTY(EditAnywhere, BlueprintReadWrite, Category = "Settings", meta=(UIMin="0", ClampMin="0"))
+	float MaxPlayRate;
+
+	/** Blending properties used when this entry is blending in ontop of another entry */
+	UPROPERTY(EditAnywhere, BlueprintReadWrite, Category = "Settings")
+	FAlphaBlend BlendIn;
+};
+
+USTRUCT()
+struct ANIMGRAPHRUNTIME_API FAnimNode_RandomPlayer : public FAnimNode_Base
+{
+	GENERATED_BODY()
+
+	FAnimNode_RandomPlayer();
+
+public:
+
+	/** When shuffle mode is active we will never loop a sequence beyond MaxLoopCount
+	  * without visiting each sequence in turn (no repeats). Enabling this will ignore
+	  * ChanceToPlay for each entry
+	  */
+	UPROPERTY(EditAnywhere, BlueprintReadWrite, Category = "Settings")
+	bool bShuffleMode;
+
+	/** List of sequences to randomly step through */
+	UPROPERTY(EditAnywhere, BlueprintReadWrite, Category = "Settings")
+	TArray<FRandomPlayerSequenceEntry> Entries;
+
+	// FAnimNode_Base interface
+	virtual void Initialize(const FAnimationInitializeContext& Context) override;
+	virtual void Update(const FAnimationUpdateContext& Context) override;
+	virtual void Evaluate(FPoseContext& Output) override;
+	virtual void GatherDebugData(FNodeDebugData& DebugData) override;
+	// End of FAnimNode_Base interface
+
+private:
+
+	int32 GetNextEntryIndex();
+	int32 GetDataIndex(const ERandomDataIndexType& Type);
+
+	void SwitchNextToCurrent();
+
+	void BuildShuffleList();
+
+	// Normalized list of play chances when we aren't using shuffle mode
+	TArray<float> NormalizedPlayChances;
+
+	// The currently playing entry in the entries list
+	int32 CurrentEntry;
+
+	// The next entry to play, we need this early so we can handle the crossfade correctly as entries
+	// can all have different blend in times.
+	int32 NextEntry;
+
+	// List to store transient shuffle stack in shuffle mode.
+	TArray<int32> ShuffleList;
+
+	// Index of the 'current' data set in the PlayData array
+	int32 CurrentDataIndex;
+
+	// Play data for the current and next sequence
+	TArray<FRandomAnimPlayData> PlayData;
+
+	// Random number source
+	FRandomStream RandomStream;
+};
\ No newline at end of file
diff --git a/Engine/Source/Runtime/AnimGraphRuntime/Public/BoneControllers/AnimNode_AnimDynamics.h b/Engine/Source/Runtime/AnimGraphRuntime/Public/BoneControllers/AnimNode_AnimDynamics.h
--- a/Engine/Source/Runtime/AnimGraphRuntime/Public/BoneControllers/AnimNode_AnimDynamics.h
+++ b/Engine/Source/Runtime/AnimGraphRuntime/Public/BoneControllers/AnimNode_AnimDynamics.h
@@ -6,8 +6,9 @@
 #include "AnimNode_SkeletalControlBase.h"
 #include "AnimNode_AnimDynamics.generated.h"
 
+DECLARE_CYCLE_STAT_EXTERN(TEXT("Anim Dynamics Overall"), STAT_AnimDynamicsOverall, STATGROUP_Physics, ANIMGRAPHRUNTIME_API);
 DECLARE_CYCLE_STAT_EXTERN(TEXT("Anim Dynamics Wind Data Update"), STAT_AnimDynamicsWindData, STATGROUP_Physics, ANIMGRAPHRUNTIME_API);
-DECLARE_CYCLE_STAT_EXTERN(TEXT("Anim Dynamics Wind Bone Evaluation"), STAT_AnimDynamicsBoneEval, STATGROUP_Physics, ANIMGRAPHRUNTIME_API);
+DECLARE_CYCLE_STAT_EXTERN(TEXT("Anim Dynamics Bone Evaluation"), STAT_AnimDynamicsBoneEval, STATGROUP_Physics, ANIMGRAPHRUNTIME_API);
 DECLARE_DWORD_COUNTER_STAT_EXTERN(TEXT("Anim Dynamics Sub-Steps"), STAT_AnimDynamicsSubSteps, STATGROUP_Physics, ANIMGRAPHRUNTIME_API);
 
 /** Supported angular constraint types */
@@ -277,6 +278,9 @@ protected:
 	virtual bool IsValidToEvaluate(const USkeleton* Skeleton, const FBoneContainer& RequiredBones);
 	// End of FAnimNode_SkeletalControlBase protected interface
 
+private:
+	void HandleGameThreadPreUpdateEvent(const UAnimInstance* InAnimInstance);
+
 private:
 
 	// We can't get clean bone positions unless we are in the evaluate step.
@@ -292,6 +296,9 @@ private:
 	// Current amount of time debt
 	float TimeDebt;
 
+	// Current time dilation
+	float CurrentTimeDilation;
+
 	// Cached physics settings. We cache these on initialise to avoid the cost of accessing UPhysicsSettings a lot each frame
 	float MaxPhysicsDeltaTime;
 	float MaxSubstepDeltaTime;
diff --git a/Engine/Source/Runtime/AnimGraphRuntime/Public/BoneControllers/AnimNode_SpringBone.h b/Engine/Source/Runtime/AnimGraphRuntime/Public/BoneControllers/AnimNode_SpringBone.h
--- a/Engine/Source/Runtime/AnimGraphRuntime/Public/BoneControllers/AnimNode_SpringBone.h
+++ b/Engine/Source/Runtime/AnimGraphRuntime/Public/BoneControllers/AnimNode_SpringBone.h
@@ -72,6 +72,9 @@ struct ANIMGRAPHRUNTIME_API FAnimNode_SpringBone : public FAnimNode_SkeletalCont
 	/** Internal use - Current timestep */
 	float FixedTimeStep;
 
+	/** Internal use - Current time dilation */
+	float TimeDilation;
+
 	/** Did we have a non-zero ControlStrength last frame. */
 	bool bHadValidStrength;
 
@@ -100,4 +103,6 @@ private:
 	// FAnimNode_SkeletalControlBase interface
 	virtual void InitializeBoneReferences(const FBoneContainer& RequiredBones) override;
 	// End of FAnimNode_SkeletalControlBase interface
+
+	void HandleGameThreadPreUpdateEvent(const UAnimInstance* InAnimInstance);
 };
diff --git a/Engine/Source/Runtime/AssetRegistry/Private/AssetRegistry.cpp b/Engine/Source/Runtime/AssetRegistry/Private/AssetRegistry.cpp
--- a/Engine/Source/Runtime/AssetRegistry/Private/AssetRegistry.cpp
+++ b/Engine/Source/Runtime/AssetRegistry/Private/AssetRegistry.cpp
@@ -471,7 +471,7 @@ bool FAssetRegistry::GetAssets(const FARFilter& Filter, TArray<FAssetData>& OutA
 				}
 
 				// Skip assets that were loaded for diffing
-				if ( InMemoryPackage->PackageFlags & PKG_ForDiffing )
+				if (InMemoryPackage->HasAnyPackageFlags(PKG_ForDiffing))
 				{
 					return;
 				}
@@ -548,7 +548,7 @@ bool FAssetRegistry::GetAssets(const FARFilter& Filter, TArray<FAssetData>& OutA
 				}
 
 				// This asset is in memory and passes all filters
-				FAssetData* AssetData = new (OutAssetData)FAssetData(PackageName, PackagePath, FName(*GroupNamesStr), Obj->GetFName(), Obj->GetClass()->GetFName(), TagMap, InMemoryPackage->GetChunkIDs(), InMemoryPackage->PackageFlags);
+				FAssetData* AssetData = new (OutAssetData)FAssetData(PackageName, PackagePath, FName(*GroupNamesStr), Obj->GetFName(), Obj->GetClass()->GetFName(), TagMap, InMemoryPackage->GetChunkIDs(), InMemoryPackage->GetPackageFlags());
 			}
 		};
 
@@ -1270,7 +1270,7 @@ void FAssetRegistry::AssetCreated(UObject* NewAsset)
 		UPackage* NewPackage = NewAsset->GetOutermost();
 
 		// Mark this package as newly created.
-		NewPackage->PackageFlags |= PKG_NewlyCreated;
+		NewPackage->SetPackageFlags(PKG_NewlyCreated);
 
 		const FString NewPackageName = NewPackage->GetName();
 		const FString Filename = FPackageName::LongPackageNameToFilename(NewPackageName, FPackageName::GetAssetPackageExtension());
diff --git a/Engine/Source/Runtime/AssetRegistry/Public/AssetData.h b/Engine/Source/Runtime/AssetRegistry/Public/AssetData.h
--- a/Engine/Source/Runtime/AssetRegistry/Public/AssetData.h
+++ b/Engine/Source/Runtime/AssetRegistry/Public/AssetData.h
@@ -95,7 +95,7 @@ public:
 			TagsAndValues = MakeSharedMapView(MoveTemp(TagsAndValuesMap));
 
 			ChunkIDs = Outermost->GetChunkIDs();
-			PackageFlags = Outermost->PackageFlags;
+			PackageFlags = Outermost->GetPackageFlags();
 		}
 	}
 
@@ -271,7 +271,7 @@ public:
 			UE_LOG(LogAssetData, Log, TEXT("                 %d"), Chunk);
 		}
 
-		UE_LOG(LogAssetData, Log, TEXT("        PackageFlahs: %d"), PackageFlags);
+		UE_LOG(LogAssetData, Log, TEXT("        PackageFlags: %d"), PackageFlags);
 	}
 
 	/** Get the first FAssetData of a particular class from an Array of FAssetData */
diff --git a/Engine/Source/Runtime/Core/Private/Stats/StatsCommand.cpp b/Engine/Source/Runtime/Core/Private/Stats/StatsCommand.cpp
--- a/Engine/Source/Runtime/Core/Private/Stats/StatsCommand.cpp
+++ b/Engine/Source/Runtime/Core/Private/Stats/StatsCommand.cpp
@@ -76,36 +76,19 @@ struct FGroupFilter : public IItemFiler
 
 struct FBudgetData
 {
-	float TotalGameTime;
-	float TotalWorkerTime;
-	TArray<FString> GameStats;
-	TArray<FString> WorkerStats;
-	TSet<FName>	GameStatsChildren;
-	TSet<FName>	WorkerStatsChildren;
-
-	FBudgetData()
-		: TotalGameTime(-1)
-		, TotalWorkerTime(-1)
-	{
-	}
-
+	TArray<FString> Stats;
+	TSet<FName>	NonAccumulatingStats;
+	TMap<FName, float> ThreadBudgetMap;
+	
 	/** Builds any extra meta data from the stats provided **/
 	void Process()
-	{
-		ProcessStats(GameStats, GameStatsChildren);
-		ProcessStats(WorkerStats, WorkerStatsChildren);
-	}
-
-private:
-
-	void ProcessStats(TArray<FString>& Stats, TSet<FName>& Children)
 	{
 		FString ChildPrefix(TEXT("-"));
-		for(FString& Stat : Stats)
+		for (FString& Stat : Stats)
 		{
-			if(Stat.RemoveFromStart(ChildPrefix))
+			if (Stat.RemoveFromStart(ChildPrefix))
 			{
-				Children.Add(FName(*Stat));
+				NonAccumulatingStats.Add(FName(*Stat));
 			}
 		}
 	}
@@ -247,6 +230,22 @@ void DumpHistoryFrame(FStatsThreadState const& StatsData, int64 TargetFrame, flo
 			}
 			UE_LOG(LogStats, Log, TEXT("  %s"), *FStatsUtils::DebugPrint(Meta));
 		}
+
+		UE_LOG(LogStats, Log, TEXT("Inclusive aggregate stack data with thread breakdown ---------------"));
+		Stats.Empty();
+		TMap<FName, TArray<FStatMessage>> ByThread;
+		StatsData.GetInclusiveAggregateStackStats(TargetFrame, Stats, nullptr, false, &ByThread);
+		for(TMap<FName, TArray<FStatMessage>>::TConstIterator It(ByThread); It; ++It)
+		{
+			const FName ShortThreadName = FStatNameAndInfo::GetShortNameFrom(It.Key());
+			UE_LOG(LogStats, Log, TEXT("  %s"), *ShortThreadName.ToString());
+
+			const TArray<FStatMessage>& StatMessages = It.Value();
+			for(const FStatMessage& Meta : Stats)
+			{
+				UE_LOG(LogStats, Log, TEXT("    %s"), *FStatsUtils::DebugPrint(Meta))
+			}
+		}
 	}
 }
 
@@ -703,16 +702,24 @@ FStatGroupGameThreadNotifier& FStatGroupGameThreadNotifier::Get()
 struct FInternalGroup
 {
 	/** Initialization constructor. */
-	FInternalGroup(const FName InGroupName, const FName InGroupCategory, const EStatDisplayMode::Type InDisplayMode, TSet<FName>& InEnabledItems, const FString& InGroupDescription, float InTotalBudget = -1.f, TSet<FName> InBudgetIgnore = TSet<FName>())
+	FInternalGroup(const FName InGroupName, const FName InGroupCategory, const EStatDisplayMode::Type InDisplayMode, TSet<FName>& InEnabledItems, const FString& InGroupDescription, TMap<FName, float>* InThreadBudgetMap = nullptr, TSet<FName>* InBudgetIgnore = nullptr)
 		: GroupName( InGroupName )
 		, GroupCategory(InGroupCategory)
 		, GroupDescription( InGroupDescription )
-		, TotalBudget( InTotalBudget )
-		, BudgetIgnoreStats( InBudgetIgnore )
 		, DisplayMode( InDisplayMode )	
 	{
 		// To avoid copy.
 		Exchange( EnabledItems, InEnabledItems );
+		
+		if(InThreadBudgetMap)
+		{
+			Exchange(ThreadBudgetMap, *InThreadBudgetMap);	//avoid copy
+		}
+
+		if (InBudgetIgnore)
+		{
+			Exchange(BudgetIgnoreStats, *InBudgetIgnore);	//avoid copy
+		}
 	}
 
 	/** Set of elements which should be included in this group stats. */
@@ -728,7 +735,7 @@ struct FInternalGroup
 	FString GroupDescription;
 
 	/** If budget mode is used, this is the expected cost of the stats in the group added up. */
-	float TotalBudget;
+	TMap<FName, float> ThreadBudgetMap;
 
 	/** If budget mode is used, these are the stats that we display, but ignore during summation */
 	TSet<FName> BudgetIgnoreStats;
@@ -744,6 +751,7 @@ struct FHudFrame
 	TArray<FStatMessage> ExclusiveAggregate;
 	TArray<FStatMessage> NonStackStats;
 	FRawStatStackNode HierarchyInclusive;
+	TMap<FName, TArray<FStatMessage>> InclusiveAggregateThreadBreakdown;
 };
 
 struct FHUDGroupManager 
@@ -763,10 +771,12 @@ struct FHUDGroupManager
 	/** Flat array of messages, it's accumulating all the time, but can be reset with a command 'stat hier -reset'. */
 	TArray<FStatMessage> TotalAggregateInclusive;
 	TArray<FStatMessage> TotalNonStackStats;
+	TMap<FName, TArray<FStatMessage>> TotalAggregateInclusiveThreadBreakdown;
 
 	/** Root stat stack for history frames, by default it's for the last 20 frames. */
 	FComplexRawStatStackNode AggregatedHierarchyHistory;
 	TArray<FComplexStatMessage> AggregatedFlatHistory;
+	TMap<FName, TArray<FComplexStatMessage>> AggregatedFlatHistoryThreadBreakdown;
 	TArray<FComplexStatMessage> AggregatedNonStackStatsHistory;
 
 	/** Copy of the stat group command parameters. */
@@ -908,47 +918,29 @@ struct FHUDGroupManager
 				}
 				else
 				{
-					float GameThreadBudget = -1.f;
-					float WorkerThreadBudget = -1.f;
-					TArray<FName> GameStatShortNames;
-					TSet<FName> GameBudgetIgnoreStats;
-					TArray<FName> WorkerStatShortNames;
-					TSet<FName> WorkerBudgetIgnoreStats;
+					TMap<FName, float> ThreadBudgetMap;
+					TArray<FName> StatShortNames;
+					TSet<FName> NonAccumulatingStats;
 					{
 						FScopeLock BudgetLock(&BudgetStatMapCS);
 						if(FBudgetData* BudgetData = BudgetStatMapping.Find(Params.BudgetSection))
 						{
-							for(const FString& StatEntry : BudgetData->GameStats)
-							{
-								GameStatShortNames.Add(FName(*StatEntry));
-							}				
-
-							for (const FString& StatEntry : BudgetData->WorkerStats)
+							for(const FString& StatEntry : BudgetData->Stats)
 							{
-								WorkerStatShortNames.Add(FName(*StatEntry));
+								StatShortNames.Add(FName(*StatEntry));
 							}
 
-							GameThreadBudget = BudgetData->TotalGameTime;
-							GameBudgetIgnoreStats = BudgetData->GameStatsChildren;
-
-							WorkerThreadBudget = BudgetData->TotalWorkerTime;
-							WorkerBudgetIgnoreStats = BudgetData->WorkerStatsChildren;
+							NonAccumulatingStats = BudgetData->NonAccumulatingStats;
+							ThreadBudgetMap = BudgetData->ThreadBudgetMap;
 						}
 					}
 					
 					{
-						TSet<FName> GameStatSet;
-						GetStatsForNames(GameStatSet, GameStatShortNames);
-						FName GameBudgetGroupName(TEXT("Budget_GameThread"));
-						EnabledGroups.Add(GameBudgetGroupName, FInternalGroup(GameBudgetGroupName, NAME_None, EStatDisplayMode::Flat, GameStatSet, *Params.BudgetSection, GameThreadBudget, GameBudgetIgnoreStats));
-						HandleToggleCommandBroadcast( GameBudgetGroupName, bCurrentEnabled, bOthersEnabled );
-					}
-					{
-						TSet<FName> WorkerStatSet;
-						GetStatsForNames(WorkerStatSet, WorkerStatShortNames);
-						FName WorkerBudgetGroupName(TEXT("Budget_WorkerThread"));
-						EnabledGroups.Add(WorkerBudgetGroupName, FInternalGroup(WorkerBudgetGroupName, NAME_None, EStatDisplayMode::Flat, WorkerStatSet, *Params.BudgetSection, WorkerThreadBudget, WorkerBudgetIgnoreStats));
-						HandleToggleCommandBroadcast( WorkerBudgetGroupName, bCurrentEnabled, bOthersEnabled );
+						TSet<FName> StatSet;
+						GetStatsForNames(StatSet, StatShortNames);
+						FName BudgetGroupName(*Params.BudgetSection);
+						EnabledGroups.Add(BudgetGroupName, FInternalGroup(*Params.BudgetSection, NAME_None, EStatDisplayMode::Flat, StatSet, TEXT("Budget"), &ThreadBudgetMap, &NonAccumulatingStats));
+						HandleToggleCommandBroadcast( BudgetGroupName, bCurrentEnabled, bOthersEnabled );
 					}
 				}
 			}
@@ -1068,8 +1060,25 @@ struct FHUDGroupManager
 
 			{
 				SCOPE_CYCLE_COUNTER(STAT_GetFlatAggregates);
-				Stats.GetInclusiveAggregateStackStats( TargetFrame, NewFrame.InclusiveAggregate, &Filter, false );
+				Stats.GetInclusiveAggregateStackStats( TargetFrame, NewFrame.InclusiveAggregate, &Filter, false, &NewFrame.InclusiveAggregateThreadBreakdown );
 				Stats.GetExclusiveAggregateStackStats( TargetFrame, NewFrame.ExclusiveAggregate, &Filter, false );
+
+				//Merge all task graph stats into 1
+				TArray<FStatMessage> MergedTaskGraphThreads;
+				for(TMap<FName, TArray<FStatMessage>>::TIterator It(NewFrame.InclusiveAggregateThreadBreakdown); It; ++It)
+				{
+					const FName ThreadName = FStatNameAndInfo::GetShortNameFrom(It.Key());
+					if (ThreadName.ToString().Contains(TEXT("TaskGraphThread")))
+					{
+						FStatsUtils::AddMergeStatArray(MergedTaskGraphThreads, It.Value());
+						It.RemoveCurrent();
+					}
+				}
+				
+				if(MergedTaskGraphThreads.Num())
+				{
+					NewFrame.InclusiveAggregateThreadBreakdown.Add(FName(TEXT("MergedTaskGraphThreads")), MergedTaskGraphThreads);
+				}
 			}
 		}
 
@@ -1092,10 +1101,15 @@ struct FHUDGroupManager
 		if( NumTotalStackFrames == 0 )
 		{
 			TotalAggregateInclusive = NewFrame.InclusiveAggregate;
+			TotalAggregateInclusiveThreadBreakdown = NewFrame.InclusiveAggregateThreadBreakdown;
 		}
 		else
 		{
 			FStatsUtils::AddMergeStatArray( TotalAggregateInclusive, NewFrame.InclusiveAggregate );
+			for(TMap<FName, TArray<FStatMessage>>::TConstIterator It(NewFrame.InclusiveAggregateThreadBreakdown); It; ++It)
+			{
+				FStatsUtils::AddMergeStatArray(TotalAggregateInclusiveThreadBreakdown.FindOrAdd(It.Key()), It.Value());
+			}
 		}
 
 		// Aggregate non-stack stats.
@@ -1137,18 +1151,31 @@ struct FHUDGroupManager
 			{
 				TotalHierarchyInclusive.Sort(FStatDurationComparer<FRawStatStackNode>());
 				TotalAggregateInclusive.Sort(FStatDurationComparer<FStatMessage>());
+				for(TMap<FName, TArray<FStatMessage>>::TIterator It(TotalAggregateInclusiveThreadBreakdown); It; ++It)
+				{
+					It.Value().Sort(FStatDurationComparer<FStatMessage>());
+				}
+				
 				TotalNonStackStats.Sort(FStatValueComparer());
 			}
 			else if (StatCompare == EStatCompareBy::CallCount)
 			{
 				TotalHierarchyInclusive.Sort(FStatCallCountComparer<FRawStatStackNode>());
 				TotalAggregateInclusive.Sort(FStatCallCountComparer<FStatMessage>());
+				for (TMap<FName, TArray<FStatMessage>>::TIterator It(TotalAggregateInclusiveThreadBreakdown); It; ++It)
+				{
+					It.Value().Sort(FStatCallCountComparer<FStatMessage>());
+				}
 				TotalNonStackStats.Sort(FStatValueComparer());
 			}
 			else if (StatCompare == EStatCompareBy::Name)
 			{
 				TotalHierarchyInclusive.Sort(FStatNameComparer<FRawStatStackNode>());
 				TotalAggregateInclusive.Sort(FStatNameComparer<FStatMessage>());
+				for (TMap<FName, TArray<FStatMessage>>::TIterator It(TotalAggregateInclusiveThreadBreakdown); It; ++It)
+				{
+					It.Value().Sort(FStatNameComparer<FStatMessage>());
+				}
 				TotalNonStackStats.Sort(FStatNameComparer<FStatMessage>());
 			}
 		}
@@ -1181,6 +1208,17 @@ struct FHUDGroupManager
 				new(AggregatedFlatHistory) FComplexStatMessage(StatMessage);
 			}
 
+			// Copy flat-stack stats by thread
+			AggregatedFlatHistoryThreadBreakdown.Reset();
+			for(TMap<FName, TArray<FStatMessage>>::TConstIterator It(TotalAggregateInclusiveThreadBreakdown); It; ++It)
+			{
+				TArray<FComplexStatMessage>& AggregatedFlatHistoryThreadBreakdownArray = AggregatedFlatHistoryThreadBreakdown.Add(It.Key());
+				for (const FStatMessage& StatMessage : It.Value())
+				{
+					new (AggregatedFlatHistoryThreadBreakdownArray)FComplexStatMessage(StatMessage);
+				}
+			}
+
 			// Copy non-stack stats
 			AggregatedNonStackStatsHistory.Reset( TotalNonStackStats.Num() );
 			for( int32 Index = 0; Index < TotalNonStackStats.Num(); ++Index )
@@ -1190,16 +1228,21 @@ struct FHUDGroupManager
 			}
 			
 			// Accumulate hierarchy, flat and non-stack stats.
-			for( auto It = History.CreateConstIterator(); It; ++It )
+			for( auto FrameIt = History.CreateConstIterator(); FrameIt; ++FrameIt )
 			{
 				SCOPE_CYCLE_COUNTER(STAT_Accumulate);
-				const FHudFrame& Frame = It.Value();
+				const FHudFrame& Frame = FrameIt.Value();
 
 				AggregatedHierarchyHistory.MergeAddAndMax( Frame.HierarchyInclusive );
 
 				FComplexStatUtils::MergeAddAndMaxArray( AggregatedFlatHistory, Frame.InclusiveAggregate, EComplexStatField::IncSum, EComplexStatField::IncMax );
-				FComplexStatUtils::MergeAddAndMaxArray( AggregatedFlatHistory, Frame.ExclusiveAggregate, EComplexStatField::ExcSum, EComplexStatField::ExcMax );
 
+				for (TMap<FName, TArray<FStatMessage>>::TConstIterator It(Frame.InclusiveAggregateThreadBreakdown); It; ++It)
+				{
+					FComplexStatUtils::MergeAddAndMaxArray( AggregatedFlatHistoryThreadBreakdown.FindChecked(It.Key()), It.Value(), EComplexStatField::IncSum, EComplexStatField::IncMax );
+				}
+
+				FComplexStatUtils::MergeAddAndMaxArray( AggregatedFlatHistory, Frame.ExclusiveAggregate, EComplexStatField::ExcSum, EComplexStatField::ExcMax );
 				FComplexStatUtils::MergeAddAndMaxArray( AggregatedNonStackStatsHistory, Frame.NonStackStats, EComplexStatField::IncSum, EComplexStatField::IncMax );
 			}
 
@@ -1218,13 +1261,18 @@ struct FHUDGroupManager
 			FComplexStatUtils::DiviveStatArray( AggregatedFlatHistory, NumFrames, EComplexStatField::IncSum, EComplexStatField::IncAve );
 			FComplexStatUtils::DiviveStatArray( AggregatedFlatHistory, NumFrames, EComplexStatField::ExcSum, EComplexStatField::ExcAve );
 
+			for(TMap<FName, TArray<FComplexStatMessage>>::TIterator It(AggregatedFlatHistoryThreadBreakdown); It; ++It)
+			{
+				FComplexStatUtils::DiviveStatArray(It.Value(), NumFrames, EComplexStatField::IncSum, EComplexStatField::IncAve);
+			}
+
 			FComplexStatUtils::DiviveStatArray( AggregatedNonStackStatsHistory, NumFrames, EComplexStatField::IncSum, EComplexStatField::IncAve );
 	
 			// Iterate through all enabled groups.
-			for( auto It = EnabledGroups.CreateIterator(); It; ++It )
+			for( auto GroupIt = EnabledGroups.CreateIterator(); GroupIt; ++GroupIt )
 			{
-				const FName& GroupName = It.Key();
-				FInternalGroup& InternalGroup = It.Value();
+				const FName& GroupName = GroupIt.Key();
+				FInternalGroup& InternalGroup = GroupIt.Value();
 
 				// Create a new hud group.
 				new(ToGame->HudGroups) FHudGroup();
@@ -1232,7 +1280,7 @@ struct FHUDGroupManager
 
 				ToGame->GroupNames.Add( GroupName );
 				ToGame->GroupDescriptions.Add( InternalGroup.GroupDescription );
-				HudGroup.TotalGroupBudget = InternalGroup.TotalBudget;
+				HudGroup.ThreadBudgetMap = InternalGroup.ThreadBudgetMap;
 				HudGroup.BudgetIgnoreStats = InternalGroup.BudgetIgnoreStats;
 
 				if (Params.bSlowMode)
@@ -1261,6 +1309,21 @@ struct FHUDGroupManager
 							new(HudGroup.FlatAggregate) FComplexStatMessage( AggregatedStatMessage );
 						}
 					}
+
+					for(TMap<FName, TArray<FComplexStatMessage>>::TConstIterator It(AggregatedFlatHistoryThreadBreakdown); It; ++It)
+					{
+						const TArray<FComplexStatMessage>& SrcArray = It.Value();
+						
+						for(const FComplexStatMessage& AggregatedStatMessage : SrcArray)
+						{
+							const bool bToBeAdded = InternalGroup.EnabledItems.Contains(AggregatedStatMessage.NameAndInfo.GetRawName());
+							if(bToBeAdded)
+							{
+								TArray<FComplexStatMessage>& DestArray = HudGroup.FlatAggregateThreadBreakdown.FindOrAdd(It.Key());	
+								new(DestArray) FComplexStatMessage( AggregatedStatMessage );
+							}
+						}
+					}
 				}
 
 				// Copy non-stack stats assigned to memory and counter groups.
@@ -2048,11 +2111,23 @@ bool DirectStatsCommand(const TCHAR* Cmd, bool bBlockForCompletion /*= false*/,
 				{
 					FBudgetData& BudgetData = BudgetStatMapping.FindOrAdd(BudgetSection);
 					BudgetData = FBudgetData();
-					GConfig->GetArray(*BudgetSection, TEXT("Game"), BudgetData.GameStats, GEngineIni);
-					GConfig->GetFloat(*BudgetSection, TEXT("GameTime"), BudgetData.TotalGameTime, GEngineIni);
-
-					GConfig->GetArray(*BudgetSection, TEXT("Worker"), BudgetData.WorkerStats, GEngineIni);
-					GConfig->GetFloat(*BudgetSection, TEXT("WorkerTime"), BudgetData.TotalWorkerTime, GEngineIni);
+					GConfig->GetArray(*BudgetSection, TEXT("Stats"), BudgetData.Stats, GEngineIni);
+					
+					TArray<FString> Lines;
+					GConfig->GetSection(*BudgetSection, Lines, GEngineIni);
+					for(const FString& Line : Lines)
+					{
+						if(!Line.Contains(TEXT("+Stats=")))	//ignore stats array
+						{
+							FString ThreadName;
+							Line.Split(FString(TEXT("=")), &ThreadName, nullptr);
+							float Budget = -1.f;
+							if(GConfig->GetFloat(*BudgetSection, *ThreadName, Budget, GEngineIni))
+							{
+								BudgetData.ThreadBudgetMap.FindOrAdd(FName(*ThreadName)) = Budget;
+							}
+						}
+					}
 
 					BudgetData.Process();
 				}
diff --git a/Engine/Source/Runtime/Core/Private/Stats/StatsData.cpp b/Engine/Source/Runtime/Core/Private/Stats/StatsData.cpp
--- a/Engine/Source/Runtime/Core/Private/Stats/StatsData.cpp
+++ b/Engine/Source/Runtime/Core/Private/Stats/StatsData.cpp
@@ -1064,7 +1064,7 @@ void FStatsThreadState::UpdateStatMessagesMemoryUsage()
 	}
 }
 
-void FStatsThreadState::GetInclusiveAggregateStackStats(int64 TargetFrame, TArray<FStatMessage>& OutStats, IItemFiler* Filter, bool bAddNonStackStats /*= true*/) const
+void FStatsThreadState::GetInclusiveAggregateStackStats(int64 TargetFrame, TArray<FStatMessage>& OutStats, IItemFiler* Filter, bool bAddNonStackStats, TMap<FName, TArray<FStatMessage>>* OptionalOutThreadBreakdownMap) const
 {
 	struct FTimeInfo
 	{
@@ -1081,10 +1081,42 @@ void FStatsThreadState::GetInclusiveAggregateStackStats(int64 TargetFrame, TArra
 	};
 	TMap<FName, FTimeInfo> Timing;
 	TMap<FName, FStatMessage> ThisFrameMetaData;
+	TMap<FName, TMap<FName, FStatMessage>> ThisFrameMetaDataPerThread;
+	TMap<FName, FStatMessage> ThreadStarts;
+	TMap<FName, FStatMessage> ThreadEnds;
 	TArray<FStatMessage> const& Data = GetCondensedHistory(TargetFrame);
+	TMap<FName, FStatMessage>* ThisFrameMetaDataPerThreadPtr = nullptr;
+	int32 Depth = 0;
 	for (int32 Index = 0; Index < Data.Num(); Index++)
 	{
 		FStatMessage const& Item = Data[Index];
+
+		//Need to get thread root first regardless of filter
+		if(OptionalOutThreadBreakdownMap)
+		{
+			EStatOperation::Type Op = Item.NameAndInfo.GetField<EStatOperation>();
+			if (Op == EStatOperation::ChildrenStart)
+			{
+				if(Depth++ == 1)
+				{
+					FName LongName = Item.NameAndInfo.GetRawName();
+					check(ThisFrameMetaDataPerThreadPtr == nullptr);
+					ThreadStarts.Add(LongName, Item);
+					ThisFrameMetaDataPerThreadPtr = &ThisFrameMetaDataPerThread.FindOrAdd(LongName);
+				}
+			}
+			else if (Op == EStatOperation::ChildrenEnd)
+			{
+				if(--Depth == 1)
+				{
+					FName LongName = Item.NameAndInfo.GetRawName();
+					ThreadEnds.Add(LongName, Item);
+					check(ThisFrameMetaDataPerThreadPtr);
+					ThisFrameMetaDataPerThreadPtr = nullptr;
+				}
+			}
+		}
+
 		if (!Filter || Filter->Keep(Item))
 		{
 			FName LongName = Item.NameAndInfo.GetRawName();
@@ -1101,7 +1133,21 @@ void FStatsThreadState::GetInclusiveAggregateStackStats(int64 TargetFrame, TArra
 					Result->NameAndInfo.SetFlag(EStatMetaFlags::IsPackedCCAndDuration, true);
 					Result->Clear();
 				}
+
+				if(Depth && ThisFrameMetaDataPerThreadPtr)
+				{
+					FStatMessage* ThreadResult = ThisFrameMetaDataPerThreadPtr->Find(LongName);
+					if (!ThreadResult)
+					{
+						ThreadResult = &ThisFrameMetaDataPerThreadPtr->Add(LongName, Item);
+						ThreadResult->NameAndInfo.SetField<EStatOperation>(EStatOperation::Set);
+						ThreadResult->NameAndInfo.SetFlag(EStatMetaFlags::IsPackedCCAndDuration, true);
+						ThreadResult->Clear();
+					}
+				}
+
 				FTimeInfo& ItemTime = Timing.FindOrAdd(LongName);
+
 				if (Op == EStatOperation::ChildrenStart)
 				{
 					ItemTime.StartCalls++;
@@ -1117,6 +1163,11 @@ void FStatsThreadState::GetInclusiveAggregateStackStats(int64 TargetFrame, TArra
 					if (!ItemTime.Recursion) // doing aggregates here, so ignore misleading recursion which would be counted twice
 					{
 						FStatsUtils::AccumulateStat(*Result, Item, EStatOperation::Add);
+						if (Depth && ThisFrameMetaDataPerThreadPtr)
+						{
+							FStatMessage& ThreadResult = ThisFrameMetaDataPerThreadPtr->FindChecked(LongName);
+							FStatsUtils::AccumulateStat(ThreadResult, Item, EStatOperation::Add);
+						}
 					}
 
 				}
@@ -1132,6 +1183,24 @@ void FStatsThreadState::GetInclusiveAggregateStackStats(int64 TargetFrame, TArra
 	{
 		OutStats.Add(It.Value());
 	}
+	
+	if(OptionalOutThreadBreakdownMap)
+	{
+		for (TMap<FName, TMap<FName, FStatMessage>> ::TConstIterator ItThread(ThisFrameMetaDataPerThread); ItThread; ++ItThread)
+		{
+			const TMap<FName, FStatMessage>& ItemNameToMeta = ItThread.Value();
+			const FName ThreadName = ItThread.Key();
+
+			if(ItemNameToMeta.Num())
+			{
+				TArray<FStatMessage>& MetaForThread = OptionalOutThreadBreakdownMap->FindOrAdd(ThreadName);
+				for (TMap<FName, FStatMessage>::TConstIterator It(ItemNameToMeta); It; ++It)
+				{
+					MetaForThread.Add(It.Value());
+				}
+			}
+		}
+	}
 }
 
 void FStatsThreadState::GetExclusiveAggregateStackStats(int64 TargetFrame, TArray<FStatMessage>& OutStats, IItemFiler* Filter, bool bAddNonStackStats /*= true*/) const
diff --git a/Engine/Source/Runtime/Core/Private/Windows/XInputInterface.cpp b/Engine/Source/Runtime/Core/Private/Windows/XInputInterface.cpp
--- a/Engine/Source/Runtime/Core/Private/Windows/XInputInterface.cpp
+++ b/Engine/Source/Runtime/Core/Private/Windows/XInputInterface.cpp
@@ -237,37 +237,43 @@ void XInputInterface::SetMessageHandler( const TSharedRef< FGenericApplicationMe
 
 void XInputInterface::SetChannelValue( const int32 ControllerId, const FForceFeedbackChannelType ChannelType, const float Value )
 {
-	FControllerState& ControllerState = ControllerStates[ ControllerId ];
-
-	if( ControllerState.bIsConnected )
+	if (ControllerId >= 0 && ControllerId < MAX_NUM_XINPUT_CONTROLLERS)
 	{
-		switch( ChannelType )
+		FControllerState& ControllerState = ControllerStates[ ControllerId ];
+
+		if( ControllerState.bIsConnected )
 		{
-			case FForceFeedbackChannelType::LEFT_LARGE:
-				ControllerState.ForceFeedback.LeftLarge = Value;
-				break;
+			switch( ChannelType )
+			{
+				case FForceFeedbackChannelType::LEFT_LARGE:
+					ControllerState.ForceFeedback.LeftLarge = Value;
+					break;
 
-			case FForceFeedbackChannelType::LEFT_SMALL:
-				ControllerState.ForceFeedback.LeftSmall = Value;
-				break;
+				case FForceFeedbackChannelType::LEFT_SMALL:
+					ControllerState.ForceFeedback.LeftSmall = Value;
+					break;
 
-			case FForceFeedbackChannelType::RIGHT_LARGE:
-				ControllerState.ForceFeedback.RightLarge = Value;
-				break;
+				case FForceFeedbackChannelType::RIGHT_LARGE:
+					ControllerState.ForceFeedback.RightLarge = Value;
+					break;
 
-			case FForceFeedbackChannelType::RIGHT_SMALL:
-				ControllerState.ForceFeedback.RightSmall = Value;
-				break;
+				case FForceFeedbackChannelType::RIGHT_SMALL:
+					ControllerState.ForceFeedback.RightSmall = Value;
+					break;
+			}
 		}
 	}
 }
 
 void XInputInterface::SetChannelValues( const int32 ControllerId, const FForceFeedbackValues &Values )
 {
-	FControllerState& ControllerState = ControllerStates[ ControllerId ];
-
-	if( ControllerState.bIsConnected )
+	if (ControllerId >= 0 && ControllerId < MAX_NUM_XINPUT_CONTROLLERS)
 	{
-		ControllerState.ForceFeedback = Values;
+		FControllerState& ControllerState = ControllerStates[ ControllerId ];
+
+		if( ControllerState.bIsConnected )
+		{
+			ControllerState.ForceFeedback = Values;
+		}
 	}
 }
\ No newline at end of file
diff --git a/Engine/Source/Runtime/Core/Public/Math/Vector.h b/Engine/Source/Runtime/Core/Public/Math/Vector.h
--- a/Engine/Source/Runtime/Core/Public/Math/Vector.h
+++ b/Engine/Source/Runtime/Core/Public/Math/Vector.h
@@ -44,8 +44,18 @@ public:
 			*const_cast<FVector*>(this) = ZeroVector;
 		}
 	}
+
+	FORCEINLINE void DiagnosticCheckNaN(const TCHAR* Message) const
+	{
+		if (ContainsNaN())
+		{
+			logOrEnsureNanError(TEXT("%s: FVector contains NaN: %s"), Message, *ToString());
+			*const_cast<FVector*>(this) = ZeroVector;
+		}
+	}
 #else
 	FORCEINLINE void DiagnosticCheckNaN() const {}
+	FORCEINLINE void DiagnosticCheckNaN(const TCHAR* Message) const {}
 #endif
 
 	/** Default constructor (no initialization). */
diff --git a/Engine/Source/Runtime/Core/Public/ProfilingDebugging/ExternalProfiler.h b/Engine/Source/Runtime/Core/Public/ProfilingDebugging/ExternalProfiler.h
--- a/Engine/Source/Runtime/Core/Public/ProfilingDebugging/ExternalProfiler.h
+++ b/Engine/Source/Runtime/Core/Public/ProfilingDebugging/ExternalProfiler.h
@@ -105,11 +105,11 @@ private:
  * Use this to include a body of code in profiler's captured session using 'Resume' and 'Pause' cues.  It
  * can safely be embedded within another 'timer' or 'excluder' scope.
  */
-class ExternalProfilerIncluder : public FScopedExternalProfilerBase
+class FExternalProfilerIncluder : public FScopedExternalProfilerBase
 {
 public:
 	/** Constructor */
-	ExternalProfilerIncluder()
+	FExternalProfilerIncluder()
 	{
 		// 'Timer' scopes will always 'resume' VTune
 		const bool bWantPause = false;
@@ -117,7 +117,7 @@ public:
 	}
 
 	/** Destructor */
-	~ExternalProfilerIncluder()
+	~FExternalProfilerIncluder()
 	{
 		StopScopedTimer();
 	}
@@ -149,5 +149,5 @@ public:
 
 };
 
-#define SCOPE_PROFILER_INCLUDER(X) ExternalProfilerIncluder ExternalProfilerIncluder_##X;
+#define SCOPE_PROFILER_INCLUDER(X) FExternalProfilerIncluder ExternalProfilerIncluder_##X;
 #define SCOPE_PROFILER_EXCLUDER(X) FExternalProfilerExcluder ExternalProfilerExcluder_##X;
diff --git a/Engine/Source/Runtime/Core/Public/Stats/Stats2.h b/Engine/Source/Runtime/Core/Public/Stats/Stats2.h
--- a/Engine/Source/Runtime/Core/Public/Stats/Stats2.h
+++ b/Engine/Source/Runtime/Core/Public/Stats/Stats2.h
@@ -1274,9 +1274,8 @@ public:
 			const bool bFrameHasChanged = FStats::GameThreadStatsFrame > CurrentGameFrame;
 			if( bFrameHasChanged )
 			{
-				// Other threads are one frame behind so set the frame to the previous one.
-				Packet.AssignFrame( CurrentGameFrame );
 				CurrentGameFrame = FStats::GameThreadStatsFrame;
+				Packet.AssignFrame( CurrentGameFrame );
 				return true;
 			}
 		}
diff --git a/Engine/Source/Runtime/Core/Public/Stats/StatsData.h b/Engine/Source/Runtime/Core/Public/Stats/StatsData.h
--- a/Engine/Source/Runtime/Core/Public/Stats/StatsData.h
+++ b/Engine/Source/Runtime/Core/Public/Stats/StatsData.h
@@ -634,7 +634,7 @@ public:
 	}
 
 	/** Gets the old-skool flat grouped inclusive stats. These ignore recursion, merge threads, etc and so generally the condensed callstack is less confusing. **/
-	void GetInclusiveAggregateStackStats(int64 TargetFrame, TArray<FStatMessage>& OutStats, IItemFiler* Filter = nullptr, bool bAddNonStackStats = true) const;
+	void GetInclusiveAggregateStackStats(int64 TargetFrame, TArray<FStatMessage>& OutStats, IItemFiler* Filter = nullptr, bool bAddNonStackStats = true, TMap<FName, TArray<FStatMessage>>* OptionalOutThreadBreakdownMap = nullptr) const;
 
 	/** Gets the old-skool flat grouped exclusive stats. These merge threads, etc and so generally the condensed callstack is less confusing. **/
 	void GetExclusiveAggregateStackStats(int64 TargetFrame, TArray<FStatMessage>& OutStats, IItemFiler* Filter = nullptr, bool bAddNonStackStats = true) const;
@@ -862,6 +862,9 @@ struct FHudGroup
 {
 	/** Array of all flat aggregates for the last n frames. */
 	TArray<FComplexStatMessage> FlatAggregate;
+
+	/** Array of all flat aggregates for the last n frames broken down by thread. */
+	TMap<FName, TArray<FComplexStatMessage>> FlatAggregateThreadBreakdown;
 	
 	/** Array of all aggregates for the last n frames. */
 	TArray<FComplexStatMessage> HierAggregate;
@@ -879,7 +882,7 @@ struct FHudGroup
 	TSet<FName> BudgetIgnoreStats;
 
 	/** Expected group budget */
-	float TotalGroupBudget;
+	TMap<FName, float> ThreadBudgetMap;
 };
 
 /**
