Fix for HTML5 math compile errors.

diff --git a/Engine/Source/Runtime/Core/Public/HTML5/HTML5PlatformMath.h b/Engine/Source/Runtime/Core/Public/HTML5/HTML5PlatformMath.h
--- a/Engine/Source/Runtime/Core/Public/HTML5/HTML5PlatformMath.h
+++ b/Engine/Source/Runtime/Core/Public/HTML5/HTML5PlatformMath.h
@@ -8,7 +8,7 @@
 #pragma once
 #include "GenericPlatform/GenericPlatformMath.h"
 
-struct FGenericPlatformMath : public FPlatformMath
+struct FHTML5PlatformMath : public FGenericPlatformMath
 {
 	/**
 	 * Counts the number of leading zeros in the bit representation of the value
@@ -44,3 +44,5 @@ struct FGenericPlatformMath : public FPlatformMath
 		return __builtin_ctz(Value);
 	}
 };
+
+typedef FHTML5PlatformMath FPlatformMath;
