Fixes for non-variadic version of TFunctionRefBase.

diff --git a/Engine/Source/Runtime/Core/Public/Templates/Function.h b/Engine/Source/Runtime/Core/Public/Templates/Function.h
--- a/Engine/Source/Runtime/Core/Public/Templates/Function.h
+++ b/Engine/Source/Runtime/Core/Public/Templates/Function.h
@@ -367,11 +367,11 @@ namespace UE4Function_Private
 		template <typename Ret, typename T0, typename T1, typename T2             > struct TFunctionRefAsserter<Ret (T0, T1, T2    )> { static Ret Call(void* Obj, T0&, T1&, T2&     ) { checkf(false, TEXT("Attempting to call a null TFunction!")); return FakeCall((Ret*)Obj); } };
 		template <typename Ret, typename T0, typename T1, typename T2, typename T3> struct TFunctionRefAsserter<Ret (T0, T1, T2, T3)> { static Ret Call(void* Obj, T0&, T1&, T2&, T3&) { checkf(false, TEXT("Attempting to call a null TFunction!")); return FakeCall((Ret*)Obj); } };
 
-		template <typename Ret                                                    > struct TFunctionRefBase<Ret (              )> : TFunctionRefBaseCommon<Ret (              ), Ret (void*                    )> { typedef TFunctionRefBaseCommon<Ret (              ), Ret (void*                    )> Super; explicit TFunctionRefBase(ENoInit) : Super(NoInit) {} template <typename FunctorType> explicit TFunctionRefBaseCommon(FunctorType* Functor) : Super(Functor) {} Ret operator()(                          ) const { const DerivedType* Derived = static_cast<const DerivedType*>(this); return this->Func(Derived->GetPtr()                ); } };
-		template <typename Ret, typename T0                                       > struct TFunctionRefBase<Ret (T0            )> : TFunctionRefBaseCommon<Ret (T0            ), Ret (void*, T0&               )> { typedef TFunctionRefBaseCommon<Ret (T0            ), Ret (void*, T0&               )> Super; explicit TFunctionRefBase(ENoInit) : Super(NoInit) {} template <typename FunctorType> explicit TFunctionRefBaseCommon(FunctorType* Functor) : Super(Functor) {} Ret operator()(T0 P0                     ) const { const DerivedType* Derived = static_cast<const DerivedType*>(this); return this->Func(Derived->GetPtr(), P0            ); } };
-		template <typename Ret, typename T0, typename T1                          > struct TFunctionRefBase<Ret (T0, T1        )> : TFunctionRefBaseCommon<Ret (T0, T1        ), Ret (void*, T0&, T1&          )> { typedef TFunctionRefBaseCommon<Ret (T0, T1        ), Ret (void*, T0&, T1&          )> Super; explicit TFunctionRefBase(ENoInit) : Super(NoInit) {} template <typename FunctorType> explicit TFunctionRefBaseCommon(FunctorType* Functor) : Super(Functor) {} Ret operator()(T0 P0, T1 P1              ) const { const DerivedType* Derived = static_cast<const DerivedType*>(this); return this->Func(Derived->GetPtr(), P0, P1        ); } };
-		template <typename Ret, typename T0, typename T1, typename T2             > struct TFunctionRefBase<Ret (T0, T1, T2    )> : TFunctionRefBaseCommon<Ret (T0, T1, T2    ), Ret (void*, T0&, T1&, T2&     )> { typedef TFunctionRefBaseCommon<Ret (T0, T1, T2    ), Ret (void*, T0&, T1&, T2&     )> Super; explicit TFunctionRefBase(ENoInit) : Super(NoInit) {} template <typename FunctorType> explicit TFunctionRefBaseCommon(FunctorType* Functor) : Super(Functor) {} Ret operator()(T0 P0, T1 P1, T2 P2       ) const { const DerivedType* Derived = static_cast<const DerivedType*>(this); return this->Func(Derived->GetPtr(), P0, P1, P2    ); } };
-		template <typename Ret, typename T0, typename T1, typename T2, typename T3> struct TFunctionRefBase<Ret (T0, T1, T2, T3)> : TFunctionRefBaseCommon<Ret (T0, T1, T2, T3), Ret (void*, T0&, T1&, T2&, T3&)> { typedef TFunctionRefBaseCommon<Ret (T0, T1, T2, T3), Ret (void*, T0&, T1&, T2&, T3&)> Super; explicit TFunctionRefBase(ENoInit) : Super(NoInit) {} template <typename FunctorType> explicit TFunctionRefBaseCommon(FunctorType* Functor) : Super(Functor) {} Ret operator()(T0 P0, T1 P1, T2 P2, T3 P3) const { const DerivedType* Derived = static_cast<const DerivedType*>(this); return this->Func(Derived->GetPtr(), P0, P1, P2, P3); } };
+		template <typename DerivedType, typename Ret                                                    > struct TFunctionRefBase<DerivedType, Ret (              )> : TFunctionRefBaseCommon<Ret (              ), Ret (void*                    )> { typedef TFunctionRefBaseCommon<Ret (              ), Ret (void*                    )> Super; explicit TFunctionRefBase(ENoInit) : Super(NoInit) {} template <typename FunctorType> explicit TFunctionRefBase(FunctorType* Functor) : Super(Functor) {} Ret operator()(                          ) const { const DerivedType* Derived = static_cast<const DerivedType*>(this); return this->GetCallable()(Derived->GetPtr()                ); } };
+		template <typename DerivedType, typename Ret, typename T0                                       > struct TFunctionRefBase<DerivedType, Ret (T0            )> : TFunctionRefBaseCommon<Ret (T0            ), Ret (void*, T0&               )> { typedef TFunctionRefBaseCommon<Ret (T0            ), Ret (void*, T0&               )> Super; explicit TFunctionRefBase(ENoInit) : Super(NoInit) {} template <typename FunctorType> explicit TFunctionRefBase(FunctorType* Functor) : Super(Functor) {} Ret operator()(T0 P0                     ) const { const DerivedType* Derived = static_cast<const DerivedType*>(this); return this->GetCallable()(Derived->GetPtr(), P0            ); } };
+		template <typename DerivedType, typename Ret, typename T0, typename T1                          > struct TFunctionRefBase<DerivedType, Ret (T0, T1        )> : TFunctionRefBaseCommon<Ret (T0, T1        ), Ret (void*, T0&, T1&          )> { typedef TFunctionRefBaseCommon<Ret (T0, T1        ), Ret (void*, T0&, T1&          )> Super; explicit TFunctionRefBase(ENoInit) : Super(NoInit) {} template <typename FunctorType> explicit TFunctionRefBase(FunctorType* Functor) : Super(Functor) {} Ret operator()(T0 P0, T1 P1              ) const { const DerivedType* Derived = static_cast<const DerivedType*>(this); return this->GetCallable()(Derived->GetPtr(), P0, P1        ); } };
+		template <typename DerivedType, typename Ret, typename T0, typename T1, typename T2             > struct TFunctionRefBase<DerivedType, Ret (T0, T1, T2    )> : TFunctionRefBaseCommon<Ret (T0, T1, T2    ), Ret (void*, T0&, T1&, T2&     )> { typedef TFunctionRefBaseCommon<Ret (T0, T1, T2    ), Ret (void*, T0&, T1&, T2&     )> Super; explicit TFunctionRefBase(ENoInit) : Super(NoInit) {} template <typename FunctorType> explicit TFunctionRefBase(FunctorType* Functor) : Super(Functor) {} Ret operator()(T0 P0, T1 P1, T2 P2       ) const { const DerivedType* Derived = static_cast<const DerivedType*>(this); return this->GetCallable()(Derived->GetPtr(), P0, P1, P2    ); } };
+		template <typename DerivedType, typename Ret, typename T0, typename T1, typename T2, typename T3> struct TFunctionRefBase<DerivedType, Ret (T0, T1, T2, T3)> : TFunctionRefBaseCommon<Ret (T0, T1, T2, T3), Ret (void*, T0&, T1&, T2&, T3&)> { typedef TFunctionRefBaseCommon<Ret (T0, T1, T2, T3), Ret (void*, T0&, T1&, T2&, T3&)> Super; explicit TFunctionRefBase(ENoInit) : Super(NoInit) {} template <typename FunctorType> explicit TFunctionRefBase(FunctorType* Functor) : Super(Functor) {} Ret operator()(T0 P0, T1 P1, T2 P2, T3 P3) const { const DerivedType* Derived = static_cast<const DerivedType*>(this); return this->GetCallable()(Derived->GetPtr(), P0, P1, P2, P3); } };
 
 	#endif
 }
