Added a check for presence of MetalRHI module to FMacPlatformMisc::HasPlatformFeature() to make it possible to run apps built with an older SDKs on OS X 10.11 (UE-23621)

diff --git a/Engine/Source/Runtime/Core/Private/Mac/MacPlatformMisc.cpp b/Engine/Source/Runtime/Core/Private/Mac/MacPlatformMisc.cpp
--- a/Engine/Source/Runtime/Core/Private/Mac/MacPlatformMisc.cpp
+++ b/Engine/Source/Runtime/Core/Private/Mac/MacPlatformMisc.cpp
@@ -1875,15 +1875,15 @@ typedef NSArray* (*MTLCopyAllDevices)(void);
 
 bool FMacPlatformMisc::HasPlatformFeature(const TCHAR* FeatureName)
 {
-	if (FCString::Stricmp(FeatureName, TEXT("Metal")) == 0 && !FParse::Param(FCommandLine::Get(),TEXT("opengl")))
+	if (FCString::Stricmp(FeatureName, TEXT("Metal")) == 0 && !FParse::Param(FCommandLine::Get(),TEXT("opengl")) && FModuleManager::Get().ModuleExists(TEXT("MetalRHI")))
 	{
 		// Find out if there are any Metal devices on the system - some Mac's have none
 		void* DLLHandle = FPlatformProcess::GetDllHandle(TEXT("/System/Library/Frameworks/Metal.framework/Metal"));
-		if(DLLHandle)
+		if (DLLHandle)
 		{
 			// Use the copy all function because we don't want to invoke a GPU switch at this point on dual-GPU Macbooks
 			MTLCopyAllDevices CopyDevicesPtr = (MTLCopyAllDevices)FPlatformProcess::GetDllExport(DLLHandle, TEXT("MTLCopyAllDevices"));
-			if(CopyDevicesPtr)
+			if (CopyDevicesPtr)
 			{
 				SCOPED_AUTORELEASE_POOL;
 				NSArray* MetalDevices = CopyDevicesPtr();
@@ -1893,7 +1893,7 @@ bool FMacPlatformMisc::HasPlatformFeature(const TCHAR* FeatureName)
 			}
 		}
 	}
-	
+
 	return FGenericPlatformMisc::HasPlatformFeature(FeatureName);
 }
 
