Fix iOS build in CIS after Dev-Framework merge (missing parens).

diff --git a/Engine/Source/Runtime/Core/Public/Math/UnrealMathNeon.h b/Engine/Source/Runtime/Core/Public/Math/UnrealMathNeon.h
--- a/Engine/Source/Runtime/Core/Public/Math/UnrealMathNeon.h
+++ b/Engine/Source/Runtime/Core/Public/Math/UnrealMathNeon.h
@@ -1105,7 +1105,7 @@ FORCEINLINE VectorRegister VectorFloor(const VectorRegister& X)
 //TODO: Vectorize
 FORCEINLINE VectorRegister VectorTruncate(const VectorRegister& X)
 {
-	return MakeVectorRegister(FMath::TruncToFloat(VectorGetComponent(X, 0), FMath::TruncToFloat(VectorGetComponent(X, 1), FMath::TruncToFloat(VectorGetComponent(X, 2), FMath::TruncToFloat(VectorGetComponent(X, 3));
+	return MakeVectorRegister(FMath::TruncToFloat(VectorGetComponent(X, 0)), FMath::TruncToFloat(VectorGetComponent(X, 1)), FMath::TruncToFloat(VectorGetComponent(X, 2)), FMath::TruncToFloat(VectorGetComponent(X, 3)));
 }
 
 //TODO: Vectorize
