#ue4 - Fix Xbox One build. Replace exp2f(Value) with powf(2, Value).

diff --git a/Engine/Source/Runtime/Core/Public/GenericPlatform/GenericPlatformMath.h b/Engine/Source/Runtime/Core/Public/GenericPlatform/GenericPlatformMath.h
--- a/Engine/Source/Runtime/Core/Public/GenericPlatform/GenericPlatformMath.h
+++ b/Engine/Source/Runtime/Core/Public/GenericPlatform/GenericPlatformMath.h
@@ -224,7 +224,7 @@ struct FGenericPlatformMath
 	// Returns e^Value
 	static FORCEINLINE float Exp( float Value ) { return expf(Value); }
 	// Returns 2^Value
-	static FORCEINLINE float Exp2( float Value ) { return exp2f(Value); }
+	static FORCEINLINE float Exp2( float Value ) { return powf(2.f, Value); /*exp2f(Value);*/ }
 	static FORCEINLINE float Loge( float Value ) {	return logf(Value); }
 	static FORCEINLINE float LogX( float Base, float Value ) { return Loge(Value) / Loge(Base); }
 	// 1.0 / Loge(2) = 1.4426950f
