Copying //UE4/Dev-Core to //UE4/Main

diff --git a/Engine/Source/Programs/CrashReporter/CrashReportReceiver/WebHandler.cs b/Engine/Source/Programs/CrashReporter/CrashReportReceiver/WebHandler.cs
new file mode 100644
--- /dev/null
+++ b/Engine/Source/Programs/CrashReporter/CrashReportReceiver/WebHandler.cs
@@ -0,0 +1,532 @@
+// Copyright 1998-2015 Epic Games, Inc. All Rights Reserved.
+
+using System;
+using System.Collections.Generic;
+using System.Collections.Concurrent;
+using System.Diagnostics;
+using System.IO;
+using System.Linq;
+using System.Net;
+using System.Web;
+using System.Text;
+
+using Tools.CrashReporter.CrashReportCommon;
+using Tools.DotNETCommon.XmlHandler;
+using System.ComponentModel;
+using System.Threading;
+
+namespace Tools.CrashReporter.CrashReportReceiver
+{
+	/// <summary>
+	/// A class to handle the receiving of crash reports.
+	/// </summary>
+	public class WebHandler : IDisposable
+	{
+		/// <summary>A flag to check to ensure the service started properly.</summary>
+		public bool bStartedSuccessfully = false;
+
+		/// <summary>Base http listener to which the async listeners are associated.</summary>
+		private HttpListener ServiceHttpListener = null;
+
+		/// <summary>Intermediate path crash reports are downloaded to until complete.</summary>
+		private string FileReceiptPath;
+
+		/// <summary>Object to take care of incoming files.</summary>
+		LandingZoneMonitor LandingZone;
+
+		/// <summary>Synchronisation object to cope with erratic client behaviour.</summary>
+		FUploadsInProgress UploadsInProgress = new FUploadsInProgress();
+
+		/// <summary>Leave at least this amount of time to leave between checks for abandoned reports</summary>
+		const int MinimumMinutesBetweenAbandonedReportChecks = 30;
+
+		/// <summary>Once a report is this old without being completed, it gets deleted</summary>
+		const int AgeMinutesToConsiderReportAbandoned = 4*60;
+
+		/// <summary>Time of most recent check for abandoned incomplete reports</summary>
+		DateTime LastAbandonedReportCheckTime = DateTime.Now;
+
+		/// <summary> Last day, used to create a new log file. </summary>
+		int LastDay = DateTime.UtcNow.Day;
+
+		/// <summary>Atomic flag to indicate a thread is checking for abandoned reports</summary>
+		int CheckingForAbandonedReports = 0;
+
+		/// <summary>
+		/// Implementing Dispose.
+		/// </summary>
+		public void Dispose()
+		{
+			Dispose( true );
+			GC.SuppressFinalize( this );
+		}
+
+		/// <summary>
+		/// Disposes the resources.
+		/// </summary>
+		/// <param name="Disposing">true if the Dispose call is from user code, and not system code.</param>
+		protected virtual void Dispose( bool Disposing )
+		{
+			if (ServiceHttpListener != null)
+			{
+				ServiceHttpListener.Close();
+			}
+		}
+
+		/// <summary>
+		/// Initialise the service to listen for reports.
+		/// </summary>
+		public WebHandler()
+		{
+			try
+			{
+				FileReceiptPath = Properties.Settings.Default.CrashReportRepository + "-Temp";
+
+				LandingZone = new LandingZoneMonitor(
+					Properties.Settings.Default.CrashReportRepository,
+					CrashReporterReceiverServicer.Log
+				);
+
+				Directory.CreateDirectory( FileReceiptPath );			
+
+				// Fire up a listener.
+				ServiceHttpListener = new HttpListener();
+				ServiceHttpListener.Prefixes.Add( "http://*:57005/CrashReporter/" );
+				ServiceHttpListener.Start();
+				ServiceHttpListener.BeginGetContext( AsyncHandleHttpRequest, null );
+
+				bStartedSuccessfully = true;
+			}
+			catch( Exception Ex )
+			{
+				CrashReporterReceiverServicer.WriteEvent( "Initialisation error: " + Ex.ToString() );
+			}
+		}
+
+		/// <summary>
+		/// Cleanup any used resources.
+		/// </summary>
+		public void Release()
+		{
+			if( bStartedSuccessfully )
+			{
+				if( ServiceHttpListener != null )
+				{
+					ServiceHttpListener.Stop();
+					ServiceHttpListener.Abort();
+				}
+			}
+		}
+
+		/// <summary>
+		/// Read in the web request payload as a string.
+		/// </summary>
+		/// <param name="Request">A listener request.</param>
+		/// <returns>A string of the input stream in the requests encoding.</returns>
+		private string GetContentStreamString( HttpListenerRequest Request )
+		{
+			string Result = "";
+
+			if( Request.HasEntityBody )
+			{
+				using( StreamReader Reader = new StreamReader( Request.InputStream, Request.ContentEncoding ) )
+				{
+					Result = Reader.ReadToEnd();
+				}
+
+				Request.InputStream.Close();
+			}
+
+			return Result;
+		}
+
+		/// <summary>
+		/// Check to see if a report has already been uploaded.
+		/// </summary>
+		/// <param name="Request">A request containing either the Report Id as a string or an XML representation of a CheckReportRequest class instance.</param>
+		/// <returns>Result object, indicating whether the report has already been uploaded.</returns>
+		private CrashReporterResult CheckReport(HttpListenerRequest Request)
+		{
+			var ReportResult = new CrashReporterResult();
+
+			var RequestClass = new CheckReportRequest();
+			RequestClass.ReportId = GetReportIdFromPostData(GetContentStreamString(Request));
+
+			ReportResult.bSuccess = !LandingZone.HasReportAlreadyBeenReceived(RequestClass.ReportId);
+
+			if( !ReportResult.bSuccess )
+			{
+				CrashReporterReceiverServicer.WriteEvent( string.Format( "Report \"{0}\" has already been received", RequestClass.ReportId ) );
+			}
+
+			return ReportResult;
+		}
+
+		/// <summary>
+		/// Check to see if we wish to reject a report based on the WER meta data.
+		/// </summary>
+		/// <param name="Request">A request containing the XML representation of a WERReportMetadata class instance.</param>
+		/// <returns>true if we do not reject.</returns>
+		private CrashReporterResult CheckReportDetail( HttpListenerRequest Request )
+		{
+			CrashReporterResult ReportResult = new CrashReporterResult();
+
+			string WERReportMetadataString = GetContentStreamString( Request );
+			WERReportMetadata WERData = null;
+
+			if( WERReportMetadataString.Length > 0 )
+			{
+				try
+				{
+					WERData = XmlHandler.FromXmlString<WERReportMetadata>( WERReportMetadataString );
+				}
+				catch( System.Exception Ex )
+				{
+					CrashReporterReceiverServicer.WriteEvent( "Error during XmlHandler.FromXmlString, probably incorrect encoding, trying to fix: " + Ex.Message );
+
+					byte[] StringBytes = System.Text.Encoding.Unicode.GetBytes( WERReportMetadataString );
+					string ConvertedXML = System.Text.Encoding.UTF8.GetString( StringBytes );
+					WERData = XmlHandler.FromXmlString<WERReportMetadata>( ConvertedXML );
+				}
+			}
+
+			if( WERData != null )
+			{
+				// Ignore crashes in the minidump parser itself
+				ReportResult.bSuccess = true;
+				if( WERData.ProblemSignatures.Parameter0.ToLower() == "MinidumpDiagnostics".ToLower() )
+				{
+					ReportResult.bSuccess = false;
+					ReportResult.Message = "Rejecting MinidumpDiagnostics crash";
+				}
+
+				// Ignore Debug and DebugGame crashes
+				string CrashingModule = WERData.ProblemSignatures.Parameter3.ToLower();
+				if( CrashingModule.Contains( "-debug" ) )
+				{
+					ReportResult.bSuccess = false;
+					ReportResult.Message = "Rejecting Debug or DebugGame crash";
+				}
+			}
+
+			return ReportResult;
+		}
+
+		/// <summary>
+		/// Receive a file and write it to a temporary folder.
+		/// </summary>
+		/// <param name="Request">A request containing the file details in the headers (DirectoryName/FileName/FileLength).</param>
+		/// <returns>true if the file is received successfully.</returns>
+		/// <remarks>There is an arbitrary file size limit of CrashReporterConstants.MaxFileSizeToUpload as a simple exploit prevention method.</remarks>
+		private CrashReporterResult ReceiveFile( HttpListenerRequest Request )
+		{
+			CrashReporterResult ReportResult = new CrashReporterResult();
+
+			if( !Request.HasEntityBody )
+			{
+				return ReportResult;
+			}
+
+			// Take this opportunity to clean out folders for reports that were never completed
+			CheckForAbandonedReports();
+
+			// Make sure we have a sensible file size
+			long BytesToReceive = 0;
+			if (long.TryParse(Request.Headers["FileLength"], out BytesToReceive))
+			{
+				if (BytesToReceive >= CrashReporterConstants.MaxFileSizeToUpload)
+				{
+					return ReportResult;
+				}
+			}
+
+			string DirectoryName = Request.Headers["DirectoryName"];
+			string FileName = Request.Headers["FileName"];
+			var T = Request.ContentLength64;
+			bool bIsOverloaded = false;
+			if( !UploadsInProgress.TryReceiveFile( DirectoryName, FileName, BytesToReceive, ref ReportResult.Message, ref bIsOverloaded ) )
+			{
+				CrashReporterReceiverServicer.WriteEvent(ReportResult.Message);
+				ReportResult.bSuccess = false;
+				return ReportResult;
+			}
+
+			string PathName = Path.Combine( FileReceiptPath, DirectoryName, FileName );
+
+			// Recreate the file receipt directory, just in case.
+			Directory.CreateDirectory( FileReceiptPath );
+
+			// Create the folder to save files to
+			DirectoryInfo DirInfo = new DirectoryInfo( Path.GetDirectoryName( PathName ) );
+			DirInfo.Create();
+
+			// Make sure the file doesn't already exist. If it does, delete it.
+			if (File.Exists(PathName))
+			{
+				File.Delete(PathName);
+			}
+
+			FileInfo Info = new FileInfo( PathName );
+			FileStream FileWriter = Info.OpenWrite();
+
+			// Read in the input stream from the request, and write to a file
+			long OriginalBytesToReceive = BytesToReceive;
+			try
+			{
+				using (BinaryReader Reader = new BinaryReader(Request.InputStream))
+				{
+					byte[] Buffer = new byte[CrashReporterConstants.StreamChunkSize];
+
+					while (BytesToReceive > 0)
+					{
+						int BytesToRead = Math.Min((int)BytesToReceive, CrashReporterConstants.StreamChunkSize);
+						Interlocked.Add( ref UploadsInProgress.CurrentReceivedData, BytesToRead );
+
+						int ReceivedChunkSize = Reader.Read(Buffer, 0, BytesToRead);
+
+						if (ReceivedChunkSize == 0)
+						{
+							ReportResult.Message = string.Format("Partial file \"{0}\" received", FileName);
+							ReportResult.bSuccess = false;
+							CrashReporterReceiverServicer.WriteEvent(ReportResult.Message);
+
+							return ReportResult;
+						}
+
+						BytesToReceive -= ReceivedChunkSize;
+						FileWriter.Write(Buffer, 0, ReceivedChunkSize);
+					}
+				}
+			}
+			finally
+			{
+				FileWriter.Close();
+				Request.InputStream.Close();
+
+				UploadsInProgress.FileUploadAttemptDone( OriginalBytesToReceive, bIsOverloaded );
+
+				bool bWriteMetadata = Path.GetExtension( FileName ) == ".ue4crash";
+				if( bWriteMetadata )
+				{
+					string CompressedSize = Request.Headers["CompressedSize"];
+					string UncompressedSize = Request.Headers["UncompressedSize"];
+					string NumberOfFiles = Request.Headers["NumberOfFiles"];
+
+					string MetadataPath = Path.Combine( FileReceiptPath, DirectoryName, Path.GetFileNameWithoutExtension( FileName ) + ".xml" );
+					XmlHandler.WriteXml<FCompressedCrashInformation>( new FCompressedCrashInformation( CompressedSize, UncompressedSize, NumberOfFiles ), MetadataPath );
+				}
+			}
+
+			ReportResult.bSuccess = true;
+			return ReportResult;
+		}
+
+		/// <summary>
+		/// Rename to the temporary landing zone directory to the final location.
+		/// </summary>
+		/// <param name="Request">A request containing either the Report Id as a string or an XML representation of a CheckReportRequest class instance.</param>
+		/// <returns>true if everything is renamed correctly.</returns>
+		private CrashReporterResult UploadComplete(HttpListenerRequest Request)
+		{
+			var ReportResult = new CrashReporterResult();
+
+			var RequestClass = new CheckReportRequest();
+			RequestClass.ReportId = GetReportIdFromPostData(GetContentStreamString(Request));
+
+			string IntermediatePathName = Path.Combine(FileReceiptPath, RequestClass.ReportId);
+
+			if (!UploadsInProgress.TrySetReportComplete(RequestClass.ReportId))
+			{
+				ReportResult.Message = string.Format("Report \"{0}\" has already been completed", RequestClass.ReportId);
+				ReportResult.bSuccess = false;
+				return ReportResult;
+			}
+
+			DirectoryInfo DirInfo = new DirectoryInfo(IntermediatePathName);
+			if (!DirInfo.Exists)
+			{
+				return ReportResult;
+			}
+
+			LandingZone.ReceiveReport(DirInfo, RequestClass.ReportId);
+			ReportResult.bSuccess = true;
+
+			int CurrentDay = DateTime.UtcNow.Day;
+			if( CurrentDay > LastDay )
+			{
+				// Check the log and create a new one for a new day.
+				CrashReporterReceiverServicer.Log.CreateNewLogFile();
+				LastDay = CurrentDay;
+			}
+
+			return ReportResult;
+		}
+
+		/// <summary>
+		/// The main listener callback to handle client requests.
+		/// </summary>
+		/// <param name="ClientRequest">The request from the client.</param>
+		private void AsyncHandleHttpRequest( IAsyncResult ClientRequest )
+		{
+			try
+			{
+				HttpListenerContext Context = ServiceHttpListener.EndGetContext( ClientRequest );
+				ServiceHttpListener.BeginGetContext( AsyncHandleHttpRequest, null );
+				HttpListenerRequest Request = Context.Request;
+				bool bIgnorePerfData = false;
+
+				using( HttpListenerResponse Response = Context.Response )
+				{
+					// Extract the URL parameters
+					string[] UrlElements = Request.RawUrl.Split( "/".ToCharArray(), StringSplitOptions.RemoveEmptyEntries );
+
+					// http://*:57005/CrashReporter/CheckReport
+					// http://*:57005/CrashReporter/CheckReportDetail
+					CrashReporterResult ReportResult = new CrashReporterResult();
+					if( UrlElements[0].ToLower() == "crashreporter" )
+					{
+						switch( UrlElements[1].ToLower() )
+						{
+						case "ping":
+							ReportResult.bSuccess = true;
+							break;
+						case "checkreport":
+							ReportResult = CheckReport( Context.Request );
+							break;
+						case "checkreportdetail":
+							ReportResult = CheckReportDetail( Context.Request );
+							break;
+						case "uploadreportfile":
+							ReportResult = ReceiveFile( Context.Request );
+							bIgnorePerfData = true;
+							break;
+						case "uploadcomplete":
+							ReportResult = UploadComplete( Context.Request );
+							break;
+						default:
+							ReportResult.bSuccess = false;
+							ReportResult.Message = "Invalid command: " + UrlElements[1];
+							break;
+						}
+					}
+					else
+					{
+						ReportResult.bSuccess = false;
+						ReportResult.Message = "Invalid application: " + UrlElements[0] + " (expecting CrashReporter)";
+					}
+
+					string ResponseString = XmlHandler.ToXmlString<CrashReporterResult>( ReportResult );
+
+					Response.SendChunked = true;
+					Response.ContentType = "text/xml";
+
+					byte[] Buffer = Encoding.UTF8.GetBytes( ResponseString );
+					Response.ContentLength64 = Buffer.Length;
+					Response.OutputStream.Write( Buffer, 0, Buffer.Length );
+
+					Response.StatusCode = ( int )HttpStatusCode.OK;
+
+					if( !bIgnorePerfData )
+					{
+						// Update the overhead data.
+						Int64 ContentLenght = Response.ContentLength64 + Request.ContentLength64;
+						Interlocked.Add( ref UploadsInProgress.CurrentReceivedData, ContentLenght );
+					}
+				}
+			}
+			catch( Exception Ex )
+			{
+				CrashReporterReceiverServicer.WriteEvent( "Error during async listen: " + Ex.Message );
+			}
+		}
+
+		/// <summary>
+		/// Translate post data, which may be XML or a raw string, to a report ID.
+		/// </summary>
+		/// <param name="ReportIdPostData">A string that's either the Report Id or an XML representation of a CheckReportRequest class instance.</param>
+		/// <returns>The report ID string</returns>
+		private string GetReportIdFromPostData(string ReportIdPostData)
+		{
+			if (ReportIdPostData.Length > 0)
+			{
+				// XML snippet will always start with <
+				if (ReportIdPostData[0] != '<')
+				{
+					return ReportIdPostData;
+				}
+				else
+				{
+					// Report id is embedded in a serialised request object (sent by old C# uploader)
+					var RequestObject = XmlHandler.FromXmlString<CheckReportRequest>(ReportIdPostData);
+					if ( RequestObject.ReportId != null && RequestObject.ReportId.Length > 0 )
+					{
+						return RequestObject.ReportId;
+					}
+					// LEGACY SUPPORT. ReportId was once named DirectoryName
+					else if ( RequestObject.DirectoryName != null && RequestObject.DirectoryName.Length > 0 )
+					{
+						return RequestObject.DirectoryName;
+					}
+				}
+			}
+
+			return "";
+		}
+
+		/// <summary>
+		/// Periodically delete folders of abandoned reports
+		/// </summary>
+		void CheckForAbandonedReports()
+		{
+			if (Interlocked.Exchange(ref CheckingForAbandonedReports, 1) == 1)
+			{
+				// Being checked by another thread, so no need
+				return;
+			}
+
+			try
+			{
+				var Now = DateTime.Now;
+
+				// No need to do this very often
+				if ((Now - LastAbandonedReportCheckTime).Minutes < MinimumMinutesBetweenAbandonedReportChecks)
+				{
+					return;
+				}
+				LastAbandonedReportCheckTime = Now;
+
+				try
+				{
+					foreach (string Dir in Directory.EnumerateDirectories(FileReceiptPath))
+					{
+						var LastAccessTime = new DateTime();
+						string ReportId = new DirectoryInfo(Dir).Name;
+						bool bGotAccessTime = UploadsInProgress.TryGetLastAccessTime(ReportId, ref LastAccessTime);
+						if (!bGotAccessTime || (Now - LastAccessTime).Minutes > AgeMinutesToConsiderReportAbandoned)
+						{
+							try
+							{
+								Directory.Delete(Dir, true /* delete all contents */);
+							}
+							catch (Exception Ex)
+							{
+								CrashReporterReceiverServicer.WriteEvent(string.Format("Failed to delete directory {0}: {1}", ReportId, Ex.Message));
+							}
+						}
+					}
+				}
+				catch (Exception Ex)
+				{
+					// If we get in here, the reports won't get cleaned up. May happen if
+					// permissions are incorrect or the directory is locked
+					CrashReporterReceiverServicer.WriteEvent( "Error during folder tidy: " + Ex.Message );
+				}
+			}
+			finally
+			{
+				// Allow other threads to check again
+				Interlocked.Exchange(ref CheckingForAbandonedReports, 0);
+			}
+		}
+	}
+}
diff --git a/Engine/Source/Programs/CrashReporter/CrashReportWebSite/Controllers/ReportsController.cs b/Engine/Source/Programs/CrashReporter/CrashReportWebSite/Controllers/ReportsController.cs
--- a/Engine/Source/Programs/CrashReporter/CrashReportWebSite/Controllers/ReportsController.cs
+++ b/Engine/Source/Programs/CrashReporter/CrashReportWebSite/Controllers/ReportsController.cs
@@ -44,7 +44,6 @@ namespace Tools.CrashReporter.CrashReportWebSite.Controllers
 			BuggRepository BuggsRepo = new BuggRepository();
 			CrashRepository CrashRepo = new CrashRepository();
 
-			// @TODO yrx 2015-02-17 BuggIDToBeAddedToJira replace with List<int> based on check box and Submit?
 			// It would be great to have a CSV export of this as well with buggs ID being the key I can then use to join them :)
 			// 
 			// Enumerate JIRA projects if needed.
diff --git a/Engine/Source/Programs/CrashReporter/CrashReportWebSite/Models/BuggsViewModel.cs b/Engine/Source/Programs/CrashReporter/CrashReportWebSite/Models/BuggsViewModel.cs
--- a/Engine/Source/Programs/CrashReporter/CrashReportWebSite/Models/BuggsViewModel.cs
+++ b/Engine/Source/Programs/CrashReporter/CrashReportWebSite/Models/BuggsViewModel.cs
@@ -10,7 +10,7 @@ namespace Tools.CrashReporter.CrashReportWebSite.Models
 	/// <summary>
 	/// The view model for the Bugg index page.
 	/// </summary>
-	/// // @TODO yrx 2014-08-14 Add base class
+	/// @TODO 2014-08-14 Add base class
 	public class BuggsViewModel
 	{
 		/// <summary>Tooltip for "Crash Groups" tab</summary>
diff --git a/Engine/Source/Programs/CrashReporter/CrashReportWebSite/Models/CrashRepository.cs b/Engine/Source/Programs/CrashReporter/CrashReportWebSite/Models/CrashRepository.cs
--- a/Engine/Source/Programs/CrashReporter/CrashReportWebSite/Models/CrashRepository.cs
+++ b/Engine/Source/Programs/CrashReporter/CrashReportWebSite/Models/CrashRepository.cs
@@ -89,8 +89,8 @@ namespace Tools.CrashReporter.CrashReportWebSite.Models
 				{
 					var BranchList = Context.Crashes
 					.Where( n => n.TimeOfCrash > DateTime.Now.AddMonths( -3 ) )
-					// Depot || Stream
-					.Where( n => n.Branch.StartsWith( "UE4" ) || n.Branch.StartsWith( "//UE4" ) )
+					// Depot - //depot/UE4* || Stream //UE4, //Something etc.
+					.Where( n => n.Branch.StartsWith( "UE4" ) || n.Branch.StartsWith( "//" ) )
 					.Select( n => n.Branch )
 					.Distinct()
 					.ToList();
diff --git a/Engine/Source/Programs/CrashReporter/CrashReportWebSite/Models/Repository.cs b/Engine/Source/Programs/CrashReporter/CrashReportWebSite/Models/Repository.cs
--- a/Engine/Source/Programs/CrashReporter/CrashReportWebSite/Models/Repository.cs
+++ b/Engine/Source/Programs/CrashReporter/CrashReportWebSite/Models/Repository.cs
@@ -271,7 +271,7 @@ namespace Tools.CrashReporter.CrashReportWebSite.Models
 		/// <returns>A dictionary of user group names, and the count of crashes for each group.</returns>
 		public Dictionary<string, int> GetCountsByGroup()
 		{
-			// @TODO yrx 2014-11-06 Optimize?
+			// @TODO 2014-11-06 Optimize?
 			using( FAutoScopedLogTimer LogTimer = new FAutoScopedLogTimer( this.GetType().ToString() + " SQL OPT" ) )
 			{
 				Dictionary<string, int> Results = new Dictionary<string, int>();
diff --git a/Engine/Source/Programs/UnrealFrontend/Private/Commands/StatsConvertCommand.cpp b/Engine/Source/Programs/UnrealFrontend/Private/Commands/StatsConvertCommand.cpp
--- a/Engine/Source/Programs/UnrealFrontend/Private/Commands/StatsConvertCommand.cpp
+++ b/Engine/Source/Programs/UnrealFrontend/Private/Commands/StatsConvertCommand.cpp
@@ -3,173 +3,65 @@
 #include "UnrealFrontendPrivatePCH.h"
 #include "StatsConvertCommand.h"
 
-
-void FStatsConvertCommand::WriteString( FArchive& Writer, const ANSICHAR* Format, ... )
+/** Helper class used to extract stats data into CSV file. */
+class FCSVStatsProfiler : public FStatsReadFile
 {
-	ANSICHAR Array[1024];
-	va_list ArgPtr;
-	va_start(ArgPtr,Format);
-	// Build the string
-	int32 Result = FCStringAnsi::GetVarArgs(Array,ARRAY_COUNT(Array),ARRAY_COUNT(Array)-1,Format,ArgPtr);
-	// Now write that to the file
-	Writer.Serialize((void*)Array,Result);
-}
-
-
-void FStatsConvertCommand::InternalRun()
-{
-#if	STATS
-	// get the target file
-	FString TargetFile;
-	FParse::Value(FCommandLine::Get(), TEXT("-INFILE="), TargetFile);
-	FString OutFile;
-	FParse::Value(FCommandLine::Get(), TEXT("-OUTFILE="), OutFile);
-	FString StatListString;
-	FParse::Value(FCommandLine::Get(), TEXT("-STATLIST="), StatListString);
+	friend struct FStatsReader<FCSVStatsProfiler>;
+	typedef FStatsReadFile Super;
 
-	// get the list of stats
-	TArray<FString> StatArrayString;
-	if (StatListString.ParseIntoArray(StatArrayString, TEXT("+"), true) == 0)
-	{
-		StatArrayString.Add(TEXT("STAT_FrameTime"));
-	}
+protected:
 
-	// convert to FNames for faster compare
-	for( const FString& It : StatArrayString )
+	/** Initialization constructor. */
+	FCSVStatsProfiler( const TCHAR* InFilename )
+		: FStatsReadFile( InFilename, false )
+		, CSVWriter( nullptr )
 	{
-		StatList.Add(*It);
+		// Keep only the last frame.
+		SetHistoryFrames( 1 );
 	}
 
-	// open a csv file for write
-	TAutoPtr<FArchive> FileWriter( IFileManager::Get().CreateFileWriter( *OutFile ) );
-	if (!FileWriter)
-	{
-		UE_LOG( LogStats, Error, TEXT( "Could not open output file: %s" ), *OutFile );
-		return;
-	}
+	/** Called every each frame has been read from the file. */
+	virtual void ReadStatsFrame( const TArray<FStatMessage>& CondensedMessages, const int64 Frame ) override;
 
-	// @TODO yrx 2014-03-24 move to function
-	// attempt to read the data and convert to csv
-	const int64 Size = IFileManager::Get().FileSize( *TargetFile );
-	if( Size < 4 )
-	{
-		UE_LOG( LogStats, Error, TEXT( "Could not open input file: %s" ), *TargetFile );
-		return;
-	}
+	/** Writes a formatted string to the CSV file. */
+	void WriteString( const ANSICHAR* Format, ... );
+public:
 
-	TAutoPtr<FArchive> FileReader( IFileManager::Get().CreateFileReader( *TargetFile ) );
-	if( !FileReader )
+	/** Sets a writer used to serialize the data for the CSV file. */
+	void Initialize( FArchive* InCSVWriter, const TArray<FString>& StatArrayString )
 	{
-		UE_LOG( LogStats, Error, TEXT( "Could not open input file: %s" ), *TargetFile );
-		return;
-	}
-
-	if( !Stream.ReadHeader( *FileReader ) )
-	{
-		UE_LOG( LogStats, Error, TEXT( "Could not open input file, bad magic: %s" ), *TargetFile );
-		return;
-	}
-
-	// This is not supported yet.
-	if (Stream.Header.bRawStatsFile)
-	{
-		UE_LOG( LogStats, Error, TEXT( "Could not open input file, not supported type (raw): %s" ), *TargetFile );
-		return;
-	}
-
-	const bool bIsFinalized = Stream.Header.IsFinalized();
-	if( bIsFinalized )
-	{
-		// Read metadata.
-		TArray<FStatMessage> MetadataMessages;
-		Stream.ReadFNamesAndMetadataMessages( *FileReader, MetadataMessages );
-		ThreadState.ProcessMetaDataOnly( MetadataMessages );
-
-		// Read frames offsets.
-		Stream.ReadFramesOffsets( *FileReader );
-		FileReader->Seek( Stream.FramesInfo[0].FrameFileOffset );
-	}
-
-	if( Stream.Header.HasCompressedData() )
-	{
-		UE_CLOG( !bIsFinalized, LogStats, Fatal, TEXT( "Compressed stats file has to be finalized" ) );
-	}
-
-	ReadAndConvertStatMessages( *FileReader, *FileWriter );
-
-#endif // STATS
-}
-
-
-void FStatsConvertCommand::ReadAndConvertStatMessages( FArchive& Reader, FArchive& Writer )
-{
-	// output the csv header
-	WriteString( Writer, "Frame,Name,Value\r\n" );
-
-	uint64 ReadMessages = 0;
-	TArray<FStatMessage> Messages;
-
-	// Buffer used to store the compressed and decompressed data.
-	TArray<uint8> SrcData;
-	TArray<uint8> DestData;
-
-	const bool bHasCompressedData = Stream.Header.HasCompressedData();
-	const bool bIsFinalized = Stream.Header.IsFinalized();
-	float DataLoadingProgress = 0.0f;
-
-	// Sanity checks.
-	check( bHasCompressedData );
+		CSVWriter = InCSVWriter;
 
-	while( Reader.Tell() < Reader.TotalSize() )
-	{
-		// Read the compressed data.
-		FCompressedStatsData UncompressedData( SrcData, DestData );
-		Reader << UncompressedData;
-		if( UncompressedData.HasReachedEndOfCompressedData() )
+		// Find the raw names for faster compare.
+		for (const FString& It : StatArrayString)
 		{
-			return;
-		}
-
-		FMemoryReader MemoryReader( DestData, true );
-
-		while( MemoryReader.Tell() < MemoryReader.TotalSize() )
-		{
-			// read the message
-			FStatMessage Message( Stream.ReadMessage( MemoryReader, bIsFinalized ) );
-			ReadMessages++;
-			if( ReadMessages % 32768 == 0 )
+			// Check for the short name.
+			const FStatMessage* LongNamePtr = State.ShortNameToLongName.Find( *It );
+			if (LongNamePtr != nullptr)
 			{
-				UE_LOG( LogStats, Log, TEXT( "StatsConvertCommand progress: %.1f%%" ), DataLoadingProgress );
-			}
-
-			if( Message.NameAndInfo.GetShortName() != TEXT( "Unknown FName" ) )
-			{
-				if( Message.NameAndInfo.GetField<EStatOperation>() == EStatOperation::AdvanceFrameEventGameThread && ReadMessages > 2 )
-				{
-					new (Messages) FStatMessage( Message );
-					ThreadState.AddMessages( Messages );
-					Messages.Reset();
-
-					CollectAndWriteStatsValues( Writer );
-					DataLoadingProgress = (double)Reader.Tell() / (double)Reader.TotalSize() * 100.0f;
-				}
-
-				new (Messages) FStatMessage( Message );
-			}
-			else
-			{
-				break;
+				StatRawNames.Add( LongNamePtr->NameAndInfo.GetRawName() );
+				StatShortNames.Add( LongNamePtr->NameAndInfo.GetShortName() );
 			}
 		}
+
+		// Output the csv header.
+		WriteString( "Frame,Name,Value\r\n" );
 	}
-}
 
+protected:
+	FArchive* CSVWriter;
+	/** Stats' raw names for fast comparison. */
+	TArray<FName> StatRawNames;
 
-void FStatsConvertCommand::CollectAndWriteStatsValues( FArchive& Writer )
+	/** Stats' short names for writing into the CSV file. */
+	TArray<FName> StatShortNames;
+};
+
+void FCSVStatsProfiler::ReadStatsFrame( const TArray<FStatMessage>& CondensedMessages, const int64 Frame )
 {
 	// get the thread stats
 	TArray<FStatMessage> Stats;
-	ThreadState.GetInclusiveAggregateStackStats(ThreadState.CurrentGameFrame, Stats);
+	State.GetInclusiveAggregateStackStats( CondensedMessages, Stats );
 
 	// The tick rate for different platforms will be different. We cannot use the Windows tick rate accurately here. We will pull it from the stats,
 	// but until we encounter one our best bet is to leave values as full ticks that can be analysed by hand.
@@ -177,36 +69,92 @@ void FStatsConvertCommand::CollectAndWriteStatsValues( FArchive& Writer )
 
 	for (int32 Index = 0; Index < Stats.Num(); ++Index)
 	{
-		FStatMessage const& Meta = Stats[Index];
+		const FStatMessage& StatMessage = Stats[Index];
 		//UE_LOG(LogTemp, Display, TEXT("Stat: %s"), *Meta.NameAndInfo.GetShortName().ToString());
 
-		if (Meta.NameAndInfo.GetShortName() == FStatConstants::NAME_SecondsPerCycle)
+		if (StatMessage.NameAndInfo.GetRawName() == FStatConstants::RAW_SecondsPerCycle)
 		{
 			// SecondsPerCycle may vary over time, so we update it here
-			MillisecondsPerCycle = Meta.GetValue_double() * 1000.0f;
+			MillisecondsPerCycle = StatMessage.GetValue_double() * 1000.0f;
 		}
 
-		for (int32 Jndex = 0; Jndex < StatList.Num(); ++Jndex)
+		for (int32 Jndex = 0; Jndex < StatRawNames.Num(); ++Jndex)
 		{
-			if (Meta.NameAndInfo.GetShortName() == StatList[Jndex])
+			const FName StatRawName = StatMessage.NameAndInfo.GetRawName();
+			if (StatRawName == StatRawNames[Jndex])
 			{
 				double StatValue = 0.0f;
-				if (Meta.NameAndInfo.GetFlag(EStatMetaFlags::IsPackedCCAndDuration))
+				if (StatMessage.NameAndInfo.GetFlag( EStatMetaFlags::IsPackedCCAndDuration ))
 				{
-					StatValue = MillisecondsPerCycle * FromPackedCallCountDuration_Duration(Meta.GetValue_int64());
+					StatValue = MillisecondsPerCycle * FromPackedCallCountDuration_Duration( StatMessage.GetValue_int64() );
 				}
-				else if (Meta.NameAndInfo.GetFlag(EStatMetaFlags::IsCycle))
+				else if (StatMessage.NameAndInfo.GetFlag( EStatMetaFlags::IsCycle ))
 				{
-					StatValue = MillisecondsPerCycle * Meta.GetValue_int64();
+					StatValue = MillisecondsPerCycle * StatMessage.GetValue_int64();
 				}
 				else
 				{
-					StatValue = Meta.GetValue_int64();
+					StatValue = StatMessage.GetValue_int64();
 				}
 
 				// write out to the csv file
-				WriteString( Writer, "%d,%S,%f\r\n", ThreadState.CurrentGameFrame, *StatList[Jndex].ToString(), StatValue );
+				WriteString( "%d,%S,%f\r\n", Frame, *StatShortNames[Jndex].ToString(), StatValue );
 			}
 		}
 	}
 }
+
+void FCSVStatsProfiler::WriteString( const ANSICHAR* Format, ... )
+{
+	ANSICHAR Array[1024];
+	va_list ArgPtr;
+	va_start( ArgPtr, Format );
+	// Build the string.
+	int32 Result = FCStringAnsi::GetVarArgs( Array, ARRAY_COUNT( Array ), ARRAY_COUNT( Array ) - 1, Format, ArgPtr );
+	// Now write that to the file.
+	CSVWriter->Serialize( (void*)Array, Result );
+}
+
+void FStatsConvertCommand::InternalRun()
+{
+	// Get the target file.
+	FString TargetFile;
+	FParse::Value(FCommandLine::Get(), TEXT("-INFILE="), TargetFile);
+
+	FString OutFile;
+	FParse::Value(FCommandLine::Get(), TEXT("-OUTFILE="), OutFile);
+
+	// Stat list can contains only stat's FName.
+	FString StatListString;
+	FParse::Value(FCommandLine::Get(), TEXT("-STATLIST="), StatListString);
+
+	// Open a CSV file for write.
+	TAutoPtr<FArchive> FileWriter( IFileManager::Get().CreateFileWriter( *OutFile ) );
+	if (!FileWriter)
+	{
+		UE_LOG( LogStats, Error, TEXT( "Could not open output file: %s" ), *OutFile );
+		return;
+	}
+
+	TAutoPtr<FCSVStatsProfiler> Instance( FStatsReader<FCSVStatsProfiler>::Create( *TargetFile ) );
+	if (Instance)
+	{
+		TArray<FName> StatList;
+
+		// Get the list of stats.
+		TArray<FString> StatArrayString;
+		if (StatListString.ParseIntoArray( StatArrayString, TEXT( "+" ), true ) == 0)
+		{
+			StatArrayString.Add( TEXT( "STAT_FrameTime" ) );
+		}
+
+		Instance->Initialize( FileWriter, StatArrayString );
+		Instance->ReadAndProcessAsynchronously();
+
+		while (Instance->IsBusy())
+		{
+			FPlatformProcess::Sleep( 2.0f );
+			UE_LOG( LogStats, Log, TEXT( "FStatsConvertCommand: Stage: %s / %3i%%" ), *Instance->GetProcessingStageAsString(), Instance->GetStageProgress() );
+		}
+	}
+}
\ No newline at end of file
diff --git a/Engine/Source/Programs/UnrealFrontend/Private/Commands/StatsConvertCommand.h b/Engine/Source/Programs/UnrealFrontend/Private/Commands/StatsConvertCommand.h
--- a/Engine/Source/Programs/UnrealFrontend/Private/Commands/StatsConvertCommand.h
+++ b/Engine/Source/Programs/UnrealFrontend/Private/Commands/StatsConvertCommand.h
@@ -20,13 +20,5 @@ public:
 protected:
 
 	void InternalRun();
-	void WriteString( FArchive& Writer, const ANSICHAR* Format, ... );
-	void CollectAndWriteStatsValues( FArchive& Writer );
-	void ReadAndConvertStatMessages( FArchive& Reader, FArchive& Writer );
 
-private:
-
-	FStatsReadStream Stream;
-	FStatsThreadState ThreadState;
-	TArray<FName> StatList;
 };
diff --git a/Engine/Source/Programs/UnrealHeaderTool/Private/ParserHelper.h b/Engine/Source/Programs/UnrealHeaderTool/Private/ParserHelper.h
--- a/Engine/Source/Programs/UnrealHeaderTool/Private/ParserHelper.h
+++ b/Engine/Source/Programs/UnrealHeaderTool/Private/ParserHelper.h
@@ -1739,7 +1739,6 @@ private:
 	TMap<UStruct*,TCHAR*> StructNameMap;
 	TArray<TCHAR*> InterfaceAllocations;
 
-	friend struct FNameLookupCPPArchiveProxy;
 	friend class FUHTMakefile;
 	FUHTMakefile* UHTMakefile;
 	FUnrealSourceFile* UnrealSourceFile;
diff --git a/Engine/Source/Runtime/AssetRegistry/Private/AssetDataGatherer.cpp b/Engine/Source/Runtime/AssetRegistry/Private/AssetDataGatherer.cpp
--- a/Engine/Source/Runtime/AssetRegistry/Private/AssetDataGatherer.cpp
+++ b/Engine/Source/Runtime/AssetRegistry/Private/AssetDataGatherer.cpp
@@ -372,7 +372,6 @@ FAssetDataGatherer::FAssetDataGatherer(const TArray<FString>& InPaths, bool bInI
 	, NumPathsToSearchAtLastSyncPoint( InPaths.Num() )
 	, bLoadAndSaveCache( bInLoadAndSaveCache )
 	, bFinishedInitialDiscovery( false )
-	, DiskCachedAssetDataBuffer( nullptr )
 	, Thread(nullptr)
 {
 	bGatherDependsData = GIsEditor && !FParse::Param( FCommandLine::Get(), TEXT("NoDependsGathering") );
@@ -405,12 +404,6 @@ FAssetDataGatherer::~FAssetDataGatherer()
 	NewCachedAssetDataMap.Empty();
 	DiskCachedAssetDataMap.Empty();
 
-	if ( DiskCachedAssetDataBuffer )
-	{
-		delete[] DiskCachedAssetDataBuffer;
-		DiskCachedAssetDataBuffer = nullptr;
-	}
-
 	for ( auto CacheIt = NewCachedAssetData.CreateConstIterator(); CacheIt; ++CacheIt )
 	{
 		delete (*CacheIt);
@@ -519,14 +512,13 @@ uint32 FAssetDataGatherer::Run()
 				bool bLoadedFromCache = false;
 				if (bLoadAndSaveCache)
 				{
-					FDiskCachedAssetData** DiskCachedAssetDataPtr = DiskCachedAssetDataMap.Find(PackageName);
-					FDiskCachedAssetData* DiskCachedAssetData = nullptr;
-					if (DiskCachedAssetDataPtr && *DiskCachedAssetDataPtr)
+					FDiskCachedAssetData* DiskCachedAssetData = DiskCachedAssetDataMap.Find(PackageName);
+					if (DiskCachedAssetData)
 					{
-						const FDateTime& CachedTimestamp = (*DiskCachedAssetDataPtr)->Timestamp;
-						if (AssetFileData.PackageTimestamp == CachedTimestamp)
+						const FDateTime& CachedTimestamp = DiskCachedAssetData->Timestamp;
+						if (AssetFileData.PackageTimestamp != CachedTimestamp)
 						{
-							DiskCachedAssetData = *DiskCachedAssetDataPtr;
+							DiskCachedAssetData = nullptr;
 						}
 					}
 
@@ -581,7 +573,7 @@ uint32 FAssetDataGatherer::Run()
 							++NumFilesProcessedSinceLastCacheSave;
 
 							// Update the cache
-							FDiskCachedAssetData* NewData = new FDiskCachedAssetData(PackageName, AssetFileData.PackageTimestamp);
+							FDiskCachedAssetData* NewData = new FDiskCachedAssetData(AssetFileData.PackageTimestamp);
 							NewData->AssetDataList.Reserve(AssetDataFromFile.Num());
 							for (const FAssetData* BackgroundAssetData : AssetDataFromFile)
 							{
@@ -819,6 +811,8 @@ void FAssetDataGatherer::SerializeCache(FArchive& Ar)
 		// save out by walking the TMap
 		for (auto CacheIt = NewCachedAssetDataMap.CreateConstIterator(); CacheIt; ++CacheIt)
 		{
+			FName PackageName = CacheIt.Key();
+			Ar << PackageName;
 			Ar << *CacheIt.Value();
 		}
 	}
@@ -826,24 +820,29 @@ void FAssetDataGatherer::SerializeCache(FArchive& Ar)
 	{
 		// allocate one single block for all asset data structs (to reduce tens of thousands of heap allocations)
 		DiskCachedAssetDataMap.Empty(LocalNumAssets);
-		DiskCachedAssetDataBuffer = new FDiskCachedAssetData[LocalNumAssets];
 
 		for (int32 AssetIndex = 0; AssetIndex < LocalNumAssets; ++AssetIndex)
 		{
-			// make a new asset data object
-			FDiskCachedAssetData* NewCachedAssetDataPtr = &DiskCachedAssetDataBuffer[AssetIndex];
+			// Load the name first to add the entry to the tmap below
+			FName PackageName;
+			Ar << PackageName;
+			if (Ar.IsError())
+			{
+				// There was an error reading the cache. Bail out.
+				break;
+			}
 
-			// load it
-			Ar << *NewCachedAssetDataPtr;
+			// Add to the cached map
+			FDiskCachedAssetData& CachedAssetData = DiskCachedAssetDataMap.Add(PackageName);
+
+			// Now load the data
+			Ar << CachedAssetData;
 
 			if (Ar.IsError())
 			{
 				// There was an error reading the cache. Bail out.
 				break;
 			}
-
-			// hash it
-			DiskCachedAssetDataMap.Add(NewCachedAssetDataPtr->PackageName, NewCachedAssetDataPtr);
 		}
 
 		// If there was an error loading the cache, abandon all data loaded from it so we can build a clean one.
@@ -851,8 +850,6 @@ void FAssetDataGatherer::SerializeCache(FArchive& Ar)
 		{
 			UE_LOG(LogAssetRegistry, Error, TEXT("There was an error loading the asset registry cache. Generating a new one."));
 			DiskCachedAssetDataMap.Empty();
-			delete DiskCachedAssetDataBuffer;
-			DiskCachedAssetDataBuffer = nullptr;
 		}
 	}
 
diff --git a/Engine/Source/Runtime/AssetRegistry/Private/AssetDataGatherer.h b/Engine/Source/Runtime/AssetRegistry/Private/AssetDataGatherer.h
--- a/Engine/Source/Runtime/AssetRegistry/Private/AssetDataGatherer.h
+++ b/Engine/Source/Runtime/AssetRegistry/Private/AssetDataGatherer.h
@@ -233,14 +233,15 @@ private:
 	/** The name of the file that contains the timestamped cache of discovered assets */
 	FString CacheFilename;
 
-	/** When loading a registry from disk, we can allocate all the FAssetData objects in one chunk, to save on 10s of thousands of heap allocations */
-	FDiskCachedAssetData* DiskCachedAssetDataBuffer;
-
 	/** An array of all cached data that was newly discovered this run. This array is just used to make sure they are all deleted at shutdown */
 	TArray<FDiskCachedAssetData*> NewCachedAssetData;
 
-	/** Map of PackageName to cached discovered assets that were loaded from disk */
-	TMap<FName, FDiskCachedAssetData*> DiskCachedAssetDataMap;
+	/** 
+	 * Map of PackageName to cached discovered assets that were loaded from disk. 
+	 * This should only be modified in the loading section of SerializeCache and remain const after
+	 * since NewCachedAssetDataMap keeps pointers to the values in this map.
+	 */
+	TMap<FName, FDiskCachedAssetData> DiskCachedAssetDataMap;
 
 	/** Map of PackageName to cached discovered assets that will be written to disk at shutdown */
 	TMap<FName, FDiskCachedAssetData*> NewCachedAssetDataMap;
diff --git a/Engine/Source/Runtime/AssetRegistry/Private/DiskCachedAssetData.h b/Engine/Source/Runtime/AssetRegistry/Private/DiskCachedAssetData.h
--- a/Engine/Source/Runtime/AssetRegistry/Private/DiskCachedAssetData.h
+++ b/Engine/Source/Runtime/AssetRegistry/Private/DiskCachedAssetData.h
@@ -5,7 +5,6 @@
 class FDiskCachedAssetData
 {
 public:
-	FName PackageName;
 	FDateTime Timestamp;
 	TArray<FAssetData> AssetDataList;
 	FPackageDependencyData DependencyData;
@@ -13,19 +12,13 @@ public:
 	FDiskCachedAssetData()
 	{}
 
-	FDiskCachedAssetData(FName InPackageName, const FDateTime& InTimestamp)
-		: PackageName(InPackageName), Timestamp(InTimestamp)
+	FDiskCachedAssetData(const FDateTime& InTimestamp)
+		: Timestamp(InTimestamp)
 	{}
 
 	/** Operator for serialization */
 	friend FArchive& operator<<(FArchive& Ar, FDiskCachedAssetData& DiskCachedAssetData)
 	{
-		Ar << DiskCachedAssetData.PackageName;
-		if (Ar.IsError())
-		{
-			return Ar;
-		}
-
 		Ar << DiskCachedAssetData.Timestamp;
 		if (Ar.IsError())
 		{
diff --git a/Engine/Source/Runtime/AssetRegistry/Private/NameTableArchive.cpp b/Engine/Source/Runtime/AssetRegistry/Private/NameTableArchive.cpp
--- a/Engine/Source/Runtime/AssetRegistry/Private/NameTableArchive.cpp
+++ b/Engine/Source/Runtime/AssetRegistry/Private/NameTableArchive.cpp
@@ -15,12 +15,12 @@ bool FNameTableArchiveReader::LoadFile(const TCHAR* Filename, int32 Serializatio
 		int32 MagicNumber = 0;
 		*this << MagicNumber;
 
-		if ( MagicNumber == PACKAGE_FILE_TAG )
+		if (!IsError() && MagicNumber == PACKAGE_FILE_TAG)
 		{
 			int32 VersionNumber = 0;
 			*this << VersionNumber;
 
-			if (VersionNumber == SerializationVersion)
+			if (!IsError() && VersionNumber == SerializationVersion)
 			{
 				return SerializeNameMap();
 			}
@@ -35,7 +35,7 @@ bool FNameTableArchiveReader::SerializeNameMap()
 	int64 NameOffset = 0;
 	*this << NameOffset;
 
-	if (NameOffset > TotalSize())
+	if (IsError() || NameOffset > TotalSize())
 	{
 		// The file was corrupted. Return false to fail to load the cache an thus regenerate it.
 		return false;
@@ -49,12 +49,22 @@ bool FNameTableArchiveReader::SerializeNameMap()
 		int32 NameCount = 0;
 		*this << NameCount;
 
+		if (IsError())
+		{
+			return false;
+		}
+
 		for ( int32 NameMapIdx = 0; NameMapIdx < NameCount; ++NameMapIdx )
 		{
 			// Read the name entry from the file.
 			FNameEntry NameEntry(ENAME_LinkerConstructor);
 			*this << NameEntry;
 
+			if (IsError())
+			{
+				return false;
+			}
+
 			NameMap.Add( 
 				NameEntry.IsWide() ? 
 					FName(ENAME_LinkerConstructor, NameEntry.GetWideName()) : 
diff --git a/Engine/Source/Runtime/Core/Private/Containers/String.cpp b/Engine/Source/Runtime/Core/Private/Containers/String.cpp
--- a/Engine/Source/Runtime/Core/Private/Containers/String.cpp
+++ b/Engine/Source/Runtime/Core/Private/Containers/String.cpp
@@ -37,7 +37,7 @@ int32 FString::Find(const TCHAR* SubStr, ESearchCase::Type SearchCase, ESearchDi
 			? FCString::Stristr(Start, SubStr)
 			: FCString::Strstr(Start, SubStr);
 
-		return Tmp ? (Tmp-**this) : -1;
+		return Tmp ? (Tmp-**this) : INDEX_NONE;
 	}
 	else
 	{
@@ -72,7 +72,7 @@ int32 FString::Find(const TCHAR* SubStr, ESearchCase::Type SearchCase, ESearchDi
 					return i;
 				}
 			}
-			return -1;
+			return INDEX_NONE;
 		}
 	}
 }
diff --git a/Engine/Source/Runtime/Core/Private/GenericPlatform/GenericPlatformFile.cpp b/Engine/Source/Runtime/Core/Private/GenericPlatform/GenericPlatformFile.cpp
--- a/Engine/Source/Runtime/Core/Private/GenericPlatform/GenericPlatformFile.cpp
+++ b/Engine/Source/Runtime/Core/Private/GenericPlatform/GenericPlatformFile.cpp
@@ -18,6 +18,19 @@ const TCHAR* IPlatformFile::GetPhysicalTypeName()
 	return TEXT("PhysicalFile");
 }
 
+void IPlatformFile::GetTimeStampPair(const TCHAR* PathA, const TCHAR* PathB, FDateTime& OutTimeStampA, FDateTime& OutTimeStampB)
+{
+	if (GetLowerLevel())
+	{
+		GetLowerLevel()->GetTimeStampPair(PathA, PathB, OutTimeStampA, OutTimeStampB);
+	}
+	else
+	{
+		OutTimeStampA = GetTimeStamp(PathA);
+		OutTimeStampB = GetTimeStamp(PathB);
+	}
+}
+
 bool IPlatformFile::IterateDirectoryRecursively(const TCHAR* Directory, FDirectoryVisitor& Visitor)
 {
 	class FRecurse : public FDirectoryVisitor
diff --git a/Engine/Source/Runtime/Core/Private/HAL/Allocators/CachedOSPageAllocator.cpp b/Engine/Source/Runtime/Core/Private/HAL/Allocators/CachedOSPageAllocator.cpp
new file mode 100644
--- /dev/null
+++ b/Engine/Source/Runtime/Core/Private/HAL/Allocators/CachedOSPageAllocator.cpp
@@ -0,0 +1,95 @@
+// Copyright 1998-2015 Epic Games, Inc. All Rights Reserved.
+
+#include "CorePrivatePCH.h"
+
+#include "Allocators/CachedOSPageAllocator.h"
+
+void* FCachedOSPageAllocator::AllocateImpl(SIZE_T Size, FFreePageBlock* First, FFreePageBlock* Last, uint32& FreedPageBlocksNum, uint32& CachedTotal)
+{
+	if (First != Last)
+	{
+		FFreePageBlock* Found = nullptr;
+
+		for (FFreePageBlock* Block = First; Block != Last; ++Block)
+		{
+			// look for exact matches first, these are aligned to the page size, so it should be quite common to hit these on small pages sizes
+			if (Block->ByteSize == Size)
+			{
+				Found = Block;
+				break;
+			}
+		}
+
+		if (!Found)
+		{
+			SIZE_T SizeTimes4 = Size * 4;
+
+			for (FFreePageBlock* Block = First; Block != Last; ++Block)
+			{
+				// is it possible (and worth i.e. <25% overhead) to use this block
+				if (Block->ByteSize >= Size && Block->ByteSize * 3 <= SizeTimes4)
+				{
+					Found = Block;
+					break;
+				}
+			}
+		}
+
+		if (Found)
+		{
+			void* Result = Found->Ptr;
+			CachedTotal -= Found->ByteSize;
+			if (Found + 1 != Last)
+			{
+				FMemory::Memmove(Found, Found + 1, sizeof(FFreePageBlock) * ((Last - Found) - 1));
+			}
+			--FreedPageBlocksNum;
+			return Result;
+		}
+
+		if (void* Ptr = FPlatformMemory::BinnedAllocFromOS(Size))
+		{
+			return Ptr;
+		}
+
+		// Are we holding on to much mem? Release it all.
+		for (FFreePageBlock* Block = First; Block != Last; ++Block)
+		{
+			FPlatformMemory::BinnedFreeToOS(Block->Ptr);
+			Block->Ptr      = nullptr;
+			Block->ByteSize = 0;
+		}
+		FreedPageBlocksNum = 0;
+		CachedTotal        = 0;
+	}
+
+	return FPlatformMemory::BinnedAllocFromOS(Size);
+}
+
+void FCachedOSPageAllocator::FreeImpl(void* Ptr, SIZE_T Size, uint32 NumCacheBlocks, uint32 CachedByteLimit, FFreePageBlock* First, uint32& FreedPageBlocksNum, uint32& CachedTotal)
+{
+	if (Size > CachedByteLimit / 4)
+	{
+		FPlatformMemory::BinnedFreeToOS(Ptr);
+		return;
+	}
+
+	while (FreedPageBlocksNum && (FreedPageBlocksNum >= NumCacheBlocks || CachedTotal + Size > CachedByteLimit)) 
+	{
+		//Remove the oldest one
+		void* FreePtr = First->Ptr;
+		CachedTotal -= First->ByteSize;
+		FreedPageBlocksNum--;
+		if (FreedPageBlocksNum)
+		{
+			FMemory::Memmove(First, First + 1, sizeof(FFreePageBlock) * FreedPageBlocksNum);
+		}
+		FPlatformMemory::BinnedFreeToOS(FreePtr);
+	}
+
+	First[FreedPageBlocksNum].Ptr      = Ptr;
+	First[FreedPageBlocksNum].ByteSize = Size;
+
+	CachedTotal += Size;
+	++FreedPageBlocksNum;
+}
diff --git a/Engine/Source/Runtime/Core/Private/HAL/FileManagerGeneric.cpp b/Engine/Source/Runtime/Core/Private/HAL/FileManagerGeneric.cpp
--- a/Engine/Source/Runtime/Core/Private/HAL/FileManagerGeneric.cpp
+++ b/Engine/Source/Runtime/Core/Private/HAL/FileManagerGeneric.cpp
@@ -475,6 +475,11 @@ FDateTime FFileManagerGeneric::GetTimeStamp(const TCHAR* Filename)
 	return GetLowLevel().GetTimeStamp(Filename);
 }
 
+void FFileManagerGeneric::GetTimeStampPair(const TCHAR* PathA, const TCHAR* PathB, FDateTime& OutTimeStampA, FDateTime& OutTimeStampB)
+{
+	GetLowLevel().GetTimeStampPair(PathA, PathB, OutTimeStampA, OutTimeStampB);
+}
+
 bool FFileManagerGeneric::SetTimeStamp(const TCHAR* Filename, FDateTime DateTime)
 {
 	// if the file doesn't exist, fail
diff --git a/Engine/Source/Runtime/Core/Private/HAL/MallocBinned2.cpp b/Engine/Source/Runtime/Core/Private/HAL/MallocBinned2.cpp
--- a/Engine/Source/Runtime/Core/Private/HAL/MallocBinned2.cpp
+++ b/Engine/Source/Runtime/Core/Private/HAL/MallocBinned2.cpp
@@ -10,8 +10,6 @@
 #include "MemoryMisc.h"
 #include "HAL/PlatformAtomics.h"
 
-#define BINNED2_MEM_TIME(st)
-
 //when modifying the global allocator stats, if we are using COARSE locks, then all callsites for stat modification are covered by the allocator-wide access guard. Thus the stats can be modified directly.
 //If we are using FINE locks, then we must modify the stats through atomics as the locks are either not actually covering the stat callsites, or are locking specific table locks which is not sufficient for stats.
 #if STATS
@@ -57,6 +55,18 @@ DEFINE_STAT(STAT_Binned2_CurrentAllocs);
 DEFINE_STAT(STAT_Binned2_TotalAllocs);
 DEFINE_STAT(STAT_Binned2_SlackCurrent);
 
+// Block sizes are based around getting the maximum amount of allocations per pool, with as little alignment waste as possible.
+// Block sizes should be close to even divisors of the POOL_SIZE, and well distributed. They must be 16-byte aligned as well.
+const uint32 GMallocBinned2BlockSizes[] =
+{
+	8,		16,		32,		48,		64,		80,		96,		112,
+	128,	160,	192,	224,	256,	288,	320,	384,
+	448,	512,	576,	640,	704,	768,	896,	1024,
+	1168,	1360,	1632,	2048,	2336,	2720,	3264,	4096,
+	4672,	5456,	6544,	8192,	9360,	10912,	13104,	16384,
+	21840,	32768
+};
+
 /** Information about a piece of free memory. 8 bytes */
 struct FMallocBinned2::FFreeMem
 {
@@ -72,7 +82,7 @@ struct FMallocBinned2::FPoolInfo
 	/** Number of allocated elements in this pool, when counts down to zero can free the entire pool. */
 	uint16			Taken;		// 2
 	/** Index of pool. Index into MemSizeToPoolTable[]. Valid when < MAX_POOLED_ALLOCATION_SIZE, MAX_POOLED_ALLOCATION_SIZE is OsTable.
-		When AllocSize is 0, this is the number of pages to step back to find the base address of an allocation. See FindPoolInfoInternal()
+		When AllocSize is 0, this is the number of pages to step back to find the base address of an allocation. See FindPoolInfo()
 	*/
 	uint16			TableIndex; // 4		
 	/** Number of bytes allocated */
@@ -96,11 +106,6 @@ struct FMallocBinned2::FPoolInfo
 		}
 	}
 
-	uint32 GetBytes() const
-	{
-		return AllocSize;
-	}
-
 	UPTRINT GetOsBytes( uint32 InPageSize, uint32 SmallAllocLimt ) const
 	{
 		if (TableIndex == SmallAllocLimt)
@@ -137,46 +142,47 @@ struct FMallocBinned2::FPoolInfo
 /** Hash table struct for retrieving allocation book keeping information */
 struct FMallocBinned2::PoolHashBucket
 {
-	UPTRINT			Key;
-	FPoolInfo*		FirstPool;
+	UPTRINT         Key;
+	FPoolInfo*      FirstPool;
 	PoolHashBucket* Prev;
 	PoolHashBucket* Next;
 
 	PoolHashBucket()
 	{
-		Key=0;
-		FirstPool=nullptr;
-		Prev=this;
-		Next=this;
+		Key      = 0;
+		FirstPool= nullptr;
+		Prev     = this;
+		Next     = this;
 	}
 
-	void Link( PoolHashBucket* After )
+	void Link(PoolHashBucket* After)
 	{
-		Link(After, Prev, this);
-	}
-
-	static void Link( PoolHashBucket* Node, PoolHashBucket* Before, PoolHashBucket* After )
-	{
-		Node->Prev=Before;
-		Node->Next=After;
-		Before->Next=Node;
-		After->Prev=Node;
+		After->Prev = Prev;
+		After->Next = this;
+		Prev ->Next = After;
+		this ->Prev = After;
 	}
 
 	void Unlink()
 	{
 		Next->Prev = Prev;
 		Prev->Next = Next;
-		Prev=this;
-		Next=this;
+		Prev       = this;
+		Next       = this;
 	}
 };
 
 struct FMallocBinned2::Private
 {
-	/** Default alignment for binned allocator */
+	static_assert(ARRAY_COUNT(GMallocBinned2BlockSizes) == POOL_COUNT, "Block size array size must match POOL_COUNT");
+	static_assert(sizeof(FMallocBinned2::FPoolInfo)     == 32,         "sizeof(FPoolInfo) is expected to be 32 bytes");
+
+	// Default alignment for binned allocator
 	enum { DEFAULT_BINNED_ALLOCATOR_ALIGNMENT = sizeof(FFreeMem) };
+
+	// There is internal limit on page size of 64k
 	enum { PAGE_SIZE_LIMIT = 65536 };
+
 	// BINNED_ALLOC_POOL_SIZE can be increased beyond 64k to cause binned malloc to allocate
 	// the small size bins in bigger chunks. If OS Allocation is slow, increasing
 	// this number *may* help performance but YMMV.
@@ -217,210 +223,173 @@ struct FMallocBinned2::Private
 		}
 		FMemory::Memset(Indirect, 0, IndirectPoolBlockSizeBytes);
 
-		BINNED2_PEAK_STATCOUNTER(Allocator.OsPeak,    BINNED2_ADD_STATCOUNTER(Allocator.OsCurrent,    (int64)(Align(IndirectPoolBlockSizeBytes, Allocator.PageSize))));
-		BINNED2_PEAK_STATCOUNTER(Allocator.WastePeak, BINNED2_ADD_STATCOUNTER(Allocator.WasteCurrent, (int64)(Align(IndirectPoolBlockSizeBytes, Allocator.PageSize))));
+		BINNED2_PEAK_STATCOUNTER(Allocator.Stats.OsPeak,    BINNED2_ADD_STATCOUNTER(Allocator.Stats.OsCurrent,    (int64)(Align(IndirectPoolBlockSizeBytes, Allocator.PageSize))));
+		BINNED2_PEAK_STATCOUNTER(Allocator.Stats.WastePeak, BINNED2_ADD_STATCOUNTER(Allocator.Stats.WasteCurrent, (int64)(Align(IndirectPoolBlockSizeBytes, Allocator.PageSize))));
 
 		return Indirect;
 	}
 
 	/** 
-	* Gets the FPoolInfo for a memory address. If no valid info exists one is created. 
-	* NOTE: This function requires a mutex across threads, but its is the callers responsibility to 
-	* acquire the mutex before calling
-	*/
-	static FORCEINLINE FPoolInfo* GetPoolInfo(FMallocBinned2& Allocator, UPTRINT Ptr)
+	 * Gets the FPoolInfo for a memory address. If no valid info exists one is created. 
+	 * NOTE: This function requires a mutex across threads, but it's the caller's responsibility to 
+	 * acquire the mutex before calling
+	 */
+	static /*FORCEINLINE*/ FPoolInfo* GetPoolInfo(FMallocBinned2& Allocator, void* InPtr)
 	{
-		if (!Allocator.HashBuckets)
-		{
-			// Init tables.
-			Allocator.HashBuckets = (PoolHashBucket*)FPlatformMemory::BinnedAllocFromOS(Align(Allocator.MaxHashBuckets * sizeof(PoolHashBucket), Allocator.PageSize));
-
-			for (uint32 i = 0; i < Allocator.MaxHashBuckets; ++i) 
-			{
-				new (Allocator.HashBuckets + i) PoolHashBucket();
-			}
-		}
+		UPTRINT Ptr = (UPTRINT)InPtr;
 
 		UPTRINT Key       = Ptr >> Allocator.HashKeyShift;
 		UPTRINT Hash      = Key & (Allocator.MaxHashBuckets - 1);
-		UPTRINT PoolIndex = ((UPTRINT)Ptr >> Allocator.PoolBitShift) & Allocator.PoolMask;
+		UPTRINT PoolIndex = (Ptr >> Allocator.PoolBitShift) & Allocator.PoolMask;
 
 		PoolHashBucket* Collision = &Allocator.HashBuckets[Hash];
 		do
 		{
-			if (Collision->Key == Key || !Collision->FirstPool)
+			if (!Collision->FirstPool)
+			{
+				Collision->Key       = Key;
+				Collision->FirstPool = CreateIndirect(Allocator);
+
+				return &Collision->FirstPool[PoolIndex];
+			}
+
+			if (Collision->Key == Key)
 			{
-				if (!Collision->FirstPool)
-				{
-					Collision->Key = Key;
-					InitializeHashBucket(Allocator, Collision);
-					CA_ASSUME(Collision->FirstPool);
-				}
 				return &Collision->FirstPool[PoolIndex];
 			}
 
 			Collision = Collision->Next;
 		} while (Collision != &Allocator.HashBuckets[Hash]);
 
-		//Create a new hash bucket entry
-		PoolHashBucket* NewBucket = CreateHashBucket(Allocator);
-		NewBucket->Key = Key;
-		Allocator.HashBuckets[Hash].Link(NewBucket);
+		// Create a new hash bucket entry
+		if (!Allocator.HashBucketFreeList) 
+		{
+			const uint32 PageSize = Allocator.PageSize;
 
-		return &NewBucket->FirstPool[PoolIndex];
-	}
+			Allocator.HashBucketFreeList = (PoolHashBucket*)FPlatformMemory::BinnedAllocFromOS(PageSize);
 
-	static FORCEINLINE FPoolInfo* FindPoolInfo(FMallocBinned2& Allocator, UPTRINT Ptr1, UPTRINT& AllocationBase)
-	{
-		uint16  NextStep = 0;
-		UPTRINT Ptr      = Ptr1 &~ ((UPTRINT)Allocator.PageSize - 1);
-		for (uint32 i = 0, n = (BINNED_ALLOC_POOL_SIZE / Allocator.PageSize) + 1; i < n; ++i)
-		{
-			FPoolInfo* Pool = FindPoolInfoInternal(Allocator, Ptr, NextStep);
-			if (Pool)
+			BINNED2_PEAK_STATCOUNTER(Allocator.Stats.OsPeak,    BINNED2_ADD_STATCOUNTER(Allocator.Stats.OsCurrent,    PageSize));
+			BINNED2_PEAK_STATCOUNTER(Allocator.Stats.WastePeak, BINNED2_ADD_STATCOUNTER(Allocator.Stats.WasteCurrent, PageSize));
+			
+			for (UPTRINT i = 0, n = PageSize / sizeof(PoolHashBucket); i < n; ++i)
 			{
-				AllocationBase = Ptr;
-				//checkSlow(Ptr1 >= AllocationBase && Ptr1 < AllocationBase + Pool->GetBytes());
-				return Pool;
+				Allocator.HashBucketFreeList->Link(new (Allocator.HashBucketFreeList + i) PoolHashBucket());
 			}
-			Ptr = ((Ptr - (Allocator.PageSize * NextStep)) - 1) &~ ((UPTRINT)Allocator.PageSize - 1);
 		}
-		AllocationBase = 0;
-		return nullptr;
-	}
 
-	static FORCEINLINE FPoolInfo* FindPoolInfoInternal(FMallocBinned2& Allocator, UPTRINT Ptr, uint16& JumpOffset)
-	{
-		checkSlow(Allocator.HashBuckets);
+		PoolHashBucket* NextFree  = Allocator.HashBucketFreeList->Next;
+		PoolHashBucket* NewBucket = Allocator.HashBucketFreeList;
 
-		uint32 Key       = Ptr >> Allocator.HashKeyShift;
-		uint32 Hash      = Key & (Allocator.MaxHashBuckets - 1);
-		uint32 PoolIndex = ((UPTRINT)Ptr >> Allocator.PoolBitShift) & Allocator.PoolMask;
+		NewBucket->Unlink();
 
-		JumpOffset = 0;
+		if (NextFree == NewBucket)
+		{
+			// Are we leaking memory here?
+			NextFree = nullptr;
+		}
+		Allocator.HashBucketFreeList = NextFree;
 
-		PoolHashBucket* Collision = &Allocator.HashBuckets[Hash];
-		do
+		if (!NewBucket->FirstPool)
 		{
-			if (Collision->Key==Key)
-			{
-				if (!Collision->FirstPool[PoolIndex].AllocSize)
-				{
-					JumpOffset = Collision->FirstPool[PoolIndex].TableIndex;
-					return nullptr;
-				}
-				return &Collision->FirstPool[PoolIndex];
-			}
-			Collision = Collision->Next;
-		} while (Collision != &Allocator.HashBuckets[Hash]);
+			NewBucket->FirstPool = CreateIndirect(Allocator);
+		}
 
-		return nullptr;
-	}
+		NewBucket->Key = Key;
 
-	/**
-	 *	Returns a newly created and initialized PoolHashBucket for use.
-	 */
-	static FORCEINLINE PoolHashBucket* CreateHashBucket(FMallocBinned2& Allocator) 
-	{
-		PoolHashBucket* Bucket = AllocateHashBucket(Allocator);
-		InitializeHashBucket(Allocator, Bucket);
-		return Bucket;
-	}
+		Allocator.HashBuckets[Hash].Link(NewBucket);
 
-	/**
-	 *	Initializes bucket with valid parameters
-	 *	@param bucket pointer to be initialized
-	 */
-	static FORCEINLINE void InitializeHashBucket(FMallocBinned2& Allocator, PoolHashBucket* Bucket)
-	{
-		if (!Bucket->FirstPool)
-		{
-			Bucket->FirstPool = CreateIndirect(Allocator);
-		}
+		return &NewBucket->FirstPool[PoolIndex];
 	}
 
-	/**
-	 * Allocates a hash bucket from the free list of hash buckets
-	*/
-	static PoolHashBucket* AllocateHashBucket(FMallocBinned2& Allocator)
+	static FPoolInfo* FindPoolInfo(FMallocBinned2& Allocator, void* InPtr, void*& AllocationBase)
 	{
-		if (!Allocator.HashBucketFreeList) 
+		const UPTRINT PageSize = (UPTRINT)Allocator.PageSize;
+
+		UPTRINT Ptr = (UPTRINT)InPtr &~ (PageSize - 1);
+		for (uint32 Repeat = (BINNED_ALLOC_POOL_SIZE / PageSize) + 1; Repeat; --Repeat)
 		{
-			const uint32 PageSize = Allocator.PageSize;
+			uint32 Key       = Ptr >> Allocator.HashKeyShift;
+			uint32 Hash      = Key & (Allocator.MaxHashBuckets - 1);
+			uint32 PoolIndex = (Ptr >> Allocator.PoolBitShift) & Allocator.PoolMask;
 
-			Allocator.HashBucketFreeList = (PoolHashBucket*)FPlatformMemory::BinnedAllocFromOS(PageSize);
-			BINNED2_PEAK_STATCOUNTER(Allocator.OsPeak,    BINNED2_ADD_STATCOUNTER(Allocator.OsCurrent,    PageSize));
-			BINNED2_PEAK_STATCOUNTER(Allocator.WastePeak, BINNED2_ADD_STATCOUNTER(Allocator.WasteCurrent, PageSize));
-			
-			for (UPTRINT i = 0, n = PageSize / sizeof(PoolHashBucket); i < n; ++i)
+			PoolHashBucket* Collision = &Allocator.HashBuckets[Hash];
+			for (;;)
 			{
-				Allocator.HashBucketFreeList->Link(new (Allocator.HashBucketFreeList + i) PoolHashBucket());
-			}
-		}
+				if (Collision->Key == Key)
+				{
+					if (Collision->FirstPool[PoolIndex].AllocSize)
+					{
+						AllocationBase = (void*)Ptr;
+						return &Collision->FirstPool[PoolIndex];
+					}
 
-		PoolHashBucket* NextFree = Allocator.HashBucketFreeList->Next;
-		PoolHashBucket* Free     = Allocator.HashBucketFreeList;
+					Ptr = ((Ptr - (PageSize * Collision->FirstPool[PoolIndex].TableIndex)) - 1) &~ (PageSize - 1);
+					break;
+				}
 
-		Free->Unlink();
-		if (NextFree == Free) 
-		{
-			NextFree = nullptr;
+				Collision = Collision->Next;
+				if (Collision == &Allocator.HashBuckets[Hash])
+				{
+					Ptr = (Ptr - 1) &~ (PageSize - 1);
+					break;
+				}
+			}
 		}
-		Allocator.HashBucketFreeList = NextFree;
 
-		return Free;
+		return nullptr;
+	}
+
+	static FORCEINLINE FPoolInfo* FindPoolInfo(FMallocBinned2& Allocator, void* Ptr)
+	{
+		void* BasePtr;
+		return FindPoolInfo(Allocator, Ptr, BasePtr);
 	}
 
-	static FPoolInfo* AllocatePoolMemory(FMallocBinned2& Allocator, FPoolTable* Table, uint32 PoolSize, uint16 TableIndex)
+	static FPoolInfo* AllocatePoolMemory(FMallocBinned2& Allocator, FPoolTable& Table, uint32 PoolSize, uint16 TableIndex)
 	{
-		const uint32 PageSize = Allocator.PageSize;
+		const uint32 PageSize  = Allocator.PageSize;
+		const uint32 BlockSize = Table.BlockSize;
 
 		// Must create a new pool.
-		uint32 Blocks   = PoolSize / Table->BlockSize;
-		uint32 Bytes    = Blocks * Table->BlockSize;
+		uint32 Blocks   = PoolSize / BlockSize;
+		uint32 Bytes    = Blocks * BlockSize;
 		UPTRINT OsBytes = Align(Bytes, PageSize);
 
 		checkSlow(Blocks >= 1);
-		checkSlow(Blocks * Table->BlockSize <= Bytes && PoolSize >= Bytes);
+		checkSlow(Blocks * BlockSize <= Bytes && PoolSize >= Bytes);
 
 		// Allocate memory.
-		FFreeMem* Free = nullptr;
-		SIZE_T ActualPoolSize; //TODO: use this to reduce waste?
-		Free = (FFreeMem*)OSAlloc(Allocator, OsBytes, ActualPoolSize);
-
-		checkSlow(!((UPTRINT)Free & (PageSize - 1)));
+		FFreeMem* Free = (FFreeMem*)Allocator.CachedOSPageAllocator.Allocate(OsBytes);
 		if( !Free )
 		{
 			OutOfMemory(OsBytes);
 		}
 
+		checkSlow(IsAligned(Free, PageSize));
+
 		// Create pool in the indirect table.
-		FPoolInfo* Pool;
+		FPoolInfo* Pool = GetPoolInfo(Allocator, Free);
+		for (UPTRINT i = (UPTRINT)PageSize, Offset = 0; i < OsBytes; i += PageSize, ++Offset)
 		{
-			Pool = GetPoolInfo(Allocator, (UPTRINT)Free);
-			for (UPTRINT i = (UPTRINT)PageSize, Offset = 0; i < OsBytes; i += PageSize, ++Offset)
-			{
-				FPoolInfo* TrailingPool = GetPoolInfo(Allocator, ((UPTRINT)Free) + i);
-				check(TrailingPool);
-
-				//Set trailing pools to point back to first pool
-				TrailingPool->SetAllocationSizes(0, 0, Offset, Allocator.BinnedOSTableIndex);
-			}
+			FPoolInfo* TrailingPool = GetPoolInfo(Allocator, (void*)(((UPTRINT)Free) + i));
+			check(TrailingPool);
 
-			
-			BINNED2_PEAK_STATCOUNTER(Allocator.OsPeak,    BINNED2_ADD_STATCOUNTER(Allocator.OsCurrent,    OsBytes));
-			BINNED2_PEAK_STATCOUNTER(Allocator.WastePeak, BINNED2_ADD_STATCOUNTER(Allocator.WasteCurrent, (OsBytes - Bytes)));
+			//Set trailing pools to point back to first pool
+			TrailingPool->SetAllocationSizes(0, 0, Offset, Allocator.BinnedOSTableIndex);
 		}
 
+		BINNED2_PEAK_STATCOUNTER(Allocator.Stats.OsPeak,    BINNED2_ADD_STATCOUNTER(Allocator.Stats.OsCurrent,    OsBytes));
+		BINNED2_PEAK_STATCOUNTER(Allocator.Stats.WastePeak, BINNED2_ADD_STATCOUNTER(Allocator.Stats.WasteCurrent, (OsBytes - Bytes)));
+
 		// Init pool.
-		Pool->Link( Table->FirstPool );
+		Pool->Link( Table.FirstPool );
 		Pool->SetAllocationSizes(Bytes, OsBytes, TableIndex, Allocator.BinnedOSTableIndex);
 		Pool->Taken		 = 0;
 		Pool->FirstMem   = Free;
 
 #if STATS
-		Table->NumActivePools++;
-		Table->MaxActivePools = FMath::Max(Table->MaxActivePools, Table->NumActivePools);
+		Table.NumActivePools++;
+		Table.MaxActivePools = FMath::Max(Table.MaxActivePools, Table.NumActivePools);
 #endif
 		// Create first free item.
 		Free->NumFreeBlocks = Blocks;
@@ -429,7 +398,7 @@ struct FMallocBinned2::Private
 		return Pool;
 	}
 
-	static FORCEINLINE FFreeMem* AllocateBlockFromPool(FMallocBinned2& Allocator, FPoolTable* Table, FPoolInfo* Pool, uint32 Alignment)
+	static /*FORCEINLINE*/ FFreeMem* AllocateBlockFromPool(FMallocBinned2& Allocator, FPoolTable& Table, FPoolInfo* Pool, uint32 Alignment)
 	{
 		// Pick first available block and unlink it.
 		Pool->Taken++;
@@ -437,7 +406,7 @@ struct FMallocBinned2::Private
 		checkSlow(Pool->FirstMem);
 		checkSlow(Pool->FirstMem->NumFreeBlocks > 0);
 		checkSlow(Pool->FirstMem->NumFreeBlocks < PAGE_SIZE_LIMIT);
-		FFreeMem* Free = (FFreeMem*)((uint8*)Pool->FirstMem + --Pool->FirstMem->NumFreeBlocks * Table->BlockSize);
+		FFreeMem* Free = (FFreeMem*)((uint8*)Pool->FirstMem + --Pool->FirstMem->NumFreeBlocks * Table.BlockSize);
 		if( !Pool->FirstMem->NumFreeBlocks )
 		{
 			Pool->FirstMem = Pool->FirstMem->Next;
@@ -445,212 +414,25 @@ struct FMallocBinned2::Private
 			{
 				// Move to exhausted list.
 				Pool->Unlink();
-				Pool->Link( Table->ExhaustedPool );
+				Pool->Link( Table.ExhaustedPool );
 			}
 		}
-		BINNED2_PEAK_STATCOUNTER(Allocator.UsedPeak, BINNED2_ADD_STATCOUNTER(Allocator.UsedCurrent, Table->BlockSize));
+		BINNED2_PEAK_STATCOUNTER(Allocator.Stats.UsedPeak, BINNED2_ADD_STATCOUNTER(Allocator.Stats.UsedCurrent, Table.BlockSize));
 		return Align(Free, Alignment);
 	}
 
-	/**
-	* Releases memory back to the system. This is not protected from multi-threaded access and it's
-	* the callers responsibility to Lock AccessGuard before calling this.
-	*/
-	static void FreeInternal(FMallocBinned2& Allocator, void* Ptr)
-	{
-		BINNED2_MEM_TIME(MemTime -= FPlatformTime::Seconds());
-		BINNED2_DECREMENT_STATCOUNTER(Allocator.CurrentAllocs);
-
-		UPTRINT BasePtr;
-		FPoolInfo* Pool = FindPoolInfo(Allocator, (UPTRINT)Ptr, BasePtr);
-#if PLATFORM_IOS || PLATFORM_MAC
-        if (Pool == NULL)
-        {
-            UE_LOG(LogMemory, Warning, TEXT("Attempting to free a pointer we didn't allocate!"));
-            return;
-        }
-#endif
-		checkSlow(Pool);
-		checkSlow(Pool->GetBytes() != 0);
-		if (Pool->TableIndex < Allocator.BinnedOSTableIndex)
-		{
-			FPoolTable* Table = Allocator.MemSizeToPoolTable[Pool->TableIndex];
-#if STATS
-			Table->ActiveRequests--;
-#endif
-			// If this pool was exhausted, move to available list.
-			if( !Pool->FirstMem )
-			{
-				Pool->Unlink();
-				Pool->Link( Table->FirstPool );
-			}
-
-			void* BaseAddress = (void*)BasePtr;
-			uint32 BlockSize = Table->BlockSize;
-			PTRINT OffsetFromBase = (PTRINT)Ptr - (PTRINT)BaseAddress;
-			check(OffsetFromBase >= 0);
-			uint32 AlignOffset = OffsetFromBase % BlockSize;
-
-			// Patch pointer to include previously applied alignment.
-			Ptr = (void*)((PTRINT)Ptr - (PTRINT)AlignOffset);
-
-			// Free a pooled allocation.
-			FFreeMem* Free		= (FFreeMem*)Ptr;
-			Free->NumFreeBlocks	= 1;
-			Free->Next			= Pool->FirstMem;
-			Pool->FirstMem		= Free;
-			BINNED2_ADD_STATCOUNTER(Allocator.UsedCurrent, -(int64)(Table->BlockSize));
-
-			// Free this pool.
-			checkSlow(Pool->Taken >= 1);
-			if( --Pool->Taken == 0 )
-			{
-#if STATS
-				Table->NumActivePools--;
-#endif
-				// Free the OS memory.
-				SIZE_T OsBytes = Pool->GetOsBytes(Allocator.PageSize, Allocator.BinnedOSTableIndex);
-				BINNED2_ADD_STATCOUNTER(Allocator.OsCurrent,    -(int64)OsBytes);
-				BINNED2_ADD_STATCOUNTER(Allocator.WasteCurrent, -(int64)(OsBytes - Pool->GetBytes()));
-				Pool->Unlink();
-				Pool->SetAllocationSizes(0, 0, 0, Allocator.BinnedOSTableIndex);
-				OSFree(Allocator, (void*)BasePtr, OsBytes);
-			}
-		}
-		else
-		{
-			// Free an OS allocation.
-			checkSlow(!((UPTRINT)Ptr & (Allocator.PageSize - 1)));
-			SIZE_T OsBytes = Pool->GetOsBytes(Allocator.PageSize, Allocator.BinnedOSTableIndex);
-
-			BINNED2_ADD_STATCOUNTER(Allocator.UsedCurrent,  -(int64)Pool->GetBytes());
-			BINNED2_ADD_STATCOUNTER(Allocator.OsCurrent,    -(int64)OsBytes);
-			BINNED2_ADD_STATCOUNTER(Allocator.WasteCurrent, -(int64)(OsBytes - Pool->GetBytes()));
-			OSFree(Allocator, (void*)BasePtr, OsBytes);
-		}
-
-		BINNED2_MEM_TIME(MemTime += FPlatformTime::Seconds());
-	}
-
-	static void PushFreeLockless(FMallocBinned2& Allocator, void* Ptr)
-	{
-		FreeInternal(Allocator, Ptr);
-	}
-
-	static FORCEINLINE void OSFree(FMallocBinned2& Allocator, void* Ptr, SIZE_T Size)
-	{
-#ifdef BINNED2_CACHE_FREED_OS_ALLOCS
-		if (Size > BINNED2_MAX_CACHED_OS_FREES_BYTE_LIMIT / 4)
-		{
-			FPlatformMemory::BinnedFreeToOS(Ptr);
-			return;
-		}
-
-		while (Allocator.FreedPageBlocksNum && (Allocator.FreedPageBlocksNum >= BINNED2_MAX_CACHED_OS_FREES || Allocator.CachedTotal + Size > BINNED2_MAX_CACHED_OS_FREES_BYTE_LIMIT)) 
-		{
-			//Remove the oldest one
-			void* FreePtr = Allocator.FreedPageBlocks[0].Ptr;
-			Allocator.CachedTotal -= Allocator.FreedPageBlocks[0].ByteSize;
-			Allocator.FreedPageBlocksNum--;
-			if (Allocator.FreedPageBlocksNum)
-			{
-				FMemory::Memmove(&Allocator.FreedPageBlocks[0], &Allocator.FreedPageBlocks[1], sizeof(FFreePageBlock) * Allocator.FreedPageBlocksNum);
-			}
-			FPlatformMemory::BinnedFreeToOS(FreePtr);
-		}
-		Allocator.FreedPageBlocks[Allocator.FreedPageBlocksNum].Ptr      = Ptr;
-		Allocator.FreedPageBlocks[Allocator.FreedPageBlocksNum].ByteSize = Size;
-		Allocator.CachedTotal += Size;
-		++Allocator.FreedPageBlocksNum;
-#else
-		(void)Size;
-		FPlatformMemory::BinnedFreeToOS(Ptr);
-#endif
-	}
-
-	static FORCEINLINE void* OSAlloc(FMallocBinned2& Allocator, SIZE_T NewSize, SIZE_T& OutActualSize)
-	{
-#ifdef BINNED2_CACHE_FREED_OS_ALLOCS
-		// We want to hold the lock a little as possible so release it
-		// before the big call to the OS
-		{
-			for (uint32 i=0; i < Allocator.FreedPageBlocksNum; ++i)
-			{
-				// look for exact matches first, these are aligned to the page size, so it should be quite common to hit these on small pages sizes
-				if (Allocator.FreedPageBlocks[i].ByteSize == NewSize)
-				{
-					void* Ret = Allocator.FreedPageBlocks[i].Ptr;
-					OutActualSize = Allocator.FreedPageBlocks[i].ByteSize;
-					Allocator.CachedTotal -= Allocator.FreedPageBlocks[i].ByteSize;
-					if (i < Allocator.FreedPageBlocksNum - 1)
-					{
-						FMemory::Memmove(&Allocator.FreedPageBlocks[i], &Allocator.FreedPageBlocks[i + 1], sizeof(FFreePageBlock) * (Allocator.FreedPageBlocksNum - i - 1));
-					}
-					Allocator.FreedPageBlocksNum--;
-					return Ret;
-				}
-			};
-			for (uint32 i=0; i < Allocator.FreedPageBlocksNum; ++i)
-			{
-				// is it possible (and worth i.e. <25% overhead) to use this block
-				if (Allocator.FreedPageBlocks[i].ByteSize >= NewSize && Allocator.FreedPageBlocks[i].ByteSize * 3 <= NewSize * 4)
-				{
-					void* Ret = Allocator.FreedPageBlocks[i].Ptr;
-					OutActualSize = Allocator.FreedPageBlocks[i].ByteSize;
-					Allocator.CachedTotal -= Allocator.FreedPageBlocks[i].ByteSize;
-					if (i < Allocator.FreedPageBlocksNum - 1)
-					{
-						FMemory::Memmove(&Allocator.FreedPageBlocks[i], &Allocator.FreedPageBlocks[i + 1], sizeof(FFreePageBlock) * (Allocator.FreedPageBlocksNum - i - 1));
-					}
-					Allocator.FreedPageBlocksNum--;
-					return Ret;
-				}
-			};
-		}
-		OutActualSize = NewSize;
-		void* Ptr = FPlatformMemory::BinnedAllocFromOS(NewSize);
-		if (!Ptr)
-		{
-			//Are we holding on to much mem? Release it all.
-			FlushAllocCache(Allocator);
-			Ptr = FPlatformMemory::BinnedAllocFromOS(NewSize);
-		}
-		return Ptr;
-#else
-		(void)OutActualSize;
-		return FPlatformMemory::BinnedAllocFromOS(NewSize);
-#endif
-	}
-
-#ifdef BINNED2_CACHE_FREED_OS_ALLOCS
-	static void FlushAllocCache(FMallocBinned2& Allocator)
-	{
-		for (int i = 0, n = Allocator.FreedPageBlocksNum; i<n; ++i) 
-		{
-			//Remove allocs
-			FPlatformMemory::BinnedFreeToOS(Allocator.FreedPageBlocks[i].Ptr);
-			Allocator.FreedPageBlocks[i].Ptr = nullptr;
-			Allocator.FreedPageBlocks[i].ByteSize = 0;
-		}
-		Allocator.FreedPageBlocksNum = 0;
-		Allocator.CachedTotal        = 0;
-	}
-#endif
-
 #if	STATS
 	static void UpdateSlackStat(FMallocBinned2& Allocator)
 	{
-		SIZE_T LocalWaste = Allocator.WasteCurrent;
-		double Waste = 0.0;
-		for( int32 PoolIndex = 0; PoolIndex < POOL_COUNT; PoolIndex++ )
+		SIZE_T LocalWaste = Allocator.Stats.WasteCurrent;
+		double Waste      = 0.0;
+		for (FPoolTable& Table : Allocator.PoolTable)
 		{
-			FPoolTable& Table = Allocator.PoolTable[PoolIndex];
-
 			Waste += ( ( double )Table.TotalWaste / ( double )Table.TotalRequests ) * ( double )Table.ActiveRequests;
 			Waste += Table.NumActivePools * ( BINNED_ALLOC_POOL_SIZE - ( ( BINNED_ALLOC_POOL_SIZE / Table.BlockSize ) * Table.BlockSize ) );
 		}
 		LocalWaste += ( uint32 )Waste;
-		Allocator.SlackCurrent = Allocator.OsCurrent - LocalWaste - Allocator.UsedCurrent;
+		Allocator.Stats.SlackCurrent = Allocator.Stats.OsCurrent - LocalWaste - Allocator.Stats.UsedCurrent;
 	}
 #endif // STATS
 };
@@ -660,41 +442,18 @@ void FMallocBinned2::GetAllocatorStats( FGenericMemoryStats& out_Stats )
 	FMalloc::GetAllocatorStats( out_Stats );
 
 #if	STATS
-	SIZE_T	LocalOsCurrent = 0;
-	SIZE_T	LocalOsPeak = 0;
-	SIZE_T	LocalWasteCurrent = 0;
-	SIZE_T	LocalWastePeak = 0;
-	SIZE_T	LocalUsedCurrent = 0;
-	SIZE_T	LocalUsedPeak = 0;
-	SIZE_T	LocalCurrentAllocs = 0;
-	SIZE_T	LocalTotalAllocs = 0;
-	SIZE_T	LocalSlackCurrent = 0;
-
-	{
-		Private::UpdateSlackStat(*this);
-
-		// Copy memory stats.
-		LocalOsCurrent = OsCurrent;
-		LocalOsPeak = OsPeak;
-		LocalWasteCurrent = WasteCurrent;
-		LocalWastePeak = WastePeak;
-		LocalUsedCurrent = UsedCurrent;
-		LocalUsedPeak = UsedPeak;
-		LocalCurrentAllocs = CurrentAllocs;
-		LocalTotalAllocs = TotalAllocs;
-		LocalSlackCurrent = SlackCurrent;
-	}
+	Private::UpdateSlackStat(*this);
 
 	// Malloc binned stats.
-	out_Stats.Add( GET_STATDESCRIPTION( STAT_Binned2_OsCurrent ), LocalOsCurrent );
-	out_Stats.Add( GET_STATDESCRIPTION( STAT_Binned2_OsPeak ), LocalOsPeak );
-	out_Stats.Add( GET_STATDESCRIPTION( STAT_Binned2_WasteCurrent ), LocalWasteCurrent );
-	out_Stats.Add( GET_STATDESCRIPTION( STAT_Binned2_WastePeak ), LocalWastePeak );
-	out_Stats.Add( GET_STATDESCRIPTION( STAT_Binned2_UsedCurrent ), LocalUsedCurrent );
-	out_Stats.Add( GET_STATDESCRIPTION( STAT_Binned2_UsedPeak ), LocalUsedPeak );
-	out_Stats.Add( GET_STATDESCRIPTION( STAT_Binned2_CurrentAllocs ), LocalCurrentAllocs );
-	out_Stats.Add( GET_STATDESCRIPTION( STAT_Binned2_TotalAllocs ), LocalTotalAllocs );
-	out_Stats.Add( GET_STATDESCRIPTION( STAT_Binned2_SlackCurrent ), LocalSlackCurrent );
+	out_Stats.Add( GET_STATDESCRIPTION( STAT_Binned2_OsCurrent ),     Stats.OsCurrent );
+	out_Stats.Add( GET_STATDESCRIPTION( STAT_Binned2_OsPeak ),        Stats.OsPeak );
+	out_Stats.Add( GET_STATDESCRIPTION( STAT_Binned2_WasteCurrent ),  Stats.WasteCurrent );
+	out_Stats.Add( GET_STATDESCRIPTION( STAT_Binned2_WastePeak ),     Stats.WastePeak );
+	out_Stats.Add( GET_STATDESCRIPTION( STAT_Binned2_UsedCurrent ),   Stats.UsedCurrent );
+	out_Stats.Add( GET_STATDESCRIPTION( STAT_Binned2_UsedPeak ),      Stats.UsedPeak );
+	out_Stats.Add( GET_STATDESCRIPTION( STAT_Binned2_CurrentAllocs ), Stats.CurrentAllocs );
+	out_Stats.Add( GET_STATDESCRIPTION( STAT_Binned2_TotalAllocs ),   Stats.TotalAllocs );
+	out_Stats.Add( GET_STATDESCRIPTION( STAT_Binned2_SlackCurrent ),  Stats.SlackCurrent );
 #endif // STATS
 }
 
@@ -716,102 +475,56 @@ void FMallocBinned2::InitializeStatsMetadata()
 }
 
 FMallocBinned2::FMallocBinned2(uint32 InPageSize, uint64 AddressLimit)
-	:	TableAddressLimit(AddressLimit)
-	,	HashBuckets(nullptr)
-	,	HashBucketFreeList(nullptr)
-	,	PageSize		(InPageSize)
-#ifdef BINNED2_CACHE_FREED_OS_ALLOCS
-	,	FreedPageBlocksNum(0)
-	,	CachedTotal(0)
-#endif
-#if STATS
-	,	OsCurrent		( 0 )
-	,	OsPeak			( 0 )
-	,	WasteCurrent	( 0 )
-	,	WastePeak		( 0 )
-	,	UsedCurrent		( 0 )
-	,	UsedPeak		( 0 )
-	,	CurrentAllocs	( 0 )
-	,	TotalAllocs		( 0 )
-	,	SlackCurrent	( 0 )
-	,	MemTime			( 0.0 )
-#endif
+	: PageSize             (InPageSize)
+	, PoolBitShift         (FPlatformMath::CeilLogTwo(InPageSize))
+	, IndirectPoolBlockSize(InPageSize / sizeof(FPoolInfo))
+	, IndirectPoolBitShift (FPlatformMath::CeilLogTwo(IndirectPoolBlockSize))
+	, MaxHashBuckets       (AddressLimit >> (IndirectPoolBitShift + PoolBitShift))
+	, HashKeyShift         (PoolBitShift + IndirectPoolBitShift)
+	, PoolMask             ((1ull << (HashKeyShift - PoolBitShift)) - 1)
+	, BinnedSizeLimit      (Private::PAGE_SIZE_LIMIT / 2)
+	, BinnedOSTableIndex   (BinnedSizeLimit + EXTENDED_PAGE_POOL_ALLOCATION_COUNT)
+	, HashBucketFreeList   (nullptr)
 {
-	check(!(PageSize & (PageSize - 1)));
-	check(!(AddressLimit & (AddressLimit - 1)));
-	check(PageSize <= 65536); // There is internal limit on page size of 64k
+	check(FMath::IsPowerOfTwo(PageSize));
+	check(FMath::IsPowerOfTwo(AddressLimit));
+	check(FMath::IsPowerOfTwo(BinnedSizeLimit));
+	check(PageSize <= Private::PAGE_SIZE_LIMIT);
 	check(AddressLimit > PageSize); // Check to catch 32 bit overflow in AddressLimit
 
-	/** Shift to get the reference from the indirect tables */
-	PoolBitShift = FPlatformMath::CeilLogTwo(PageSize);
-	IndirectPoolBitShift = FPlatformMath::CeilLogTwo(PageSize/sizeof(FPoolInfo));
-	IndirectPoolBlockSize = PageSize/sizeof(FPoolInfo);
-
-	MaxHashBuckets = AddressLimit >> (IndirectPoolBitShift+PoolBitShift); 
-	MaxHashBucketBits = FPlatformMath::CeilLogTwo(MaxHashBuckets);
-	MaxHashBucketWaste = (MaxHashBuckets*sizeof(PoolHashBucket))/1024;
-	MaxBookKeepingOverhead = ((AddressLimit/PageSize)*sizeof(PoolHashBucket))/(1024*1024);
-	/** 
-	* Shift required to get required hash table key.
-	*/
-	HashKeyShift = PoolBitShift+IndirectPoolBitShift;
-	/** Used to mask off the bits that have been used to lookup the indirect table */
-	PoolMask =  ( ( 1ull << ( HashKeyShift - PoolBitShift ) ) - 1 );
-	BinnedSizeLimit = Private::PAGE_SIZE_LIMIT/2;
-	BinnedOSTableIndex = BinnedSizeLimit+EXTENDED_PAGE_POOL_ALLOCATION_COUNT;
-
-	check((BinnedSizeLimit & (BinnedSizeLimit-1)) == 0);
+	// Init tables.
 
+	/** The following options are not valid for page sizes less than 64k. They are here to reduce waste */
+	PagePoolTable[0].BlockSize = (PageSize == Private::PAGE_SIZE_LIMIT) ? BinnedSizeLimit + (BinnedSizeLimit / 2) : 0;
+	PagePoolTable[1].BlockSize = (PageSize == Private::PAGE_SIZE_LIMIT) ? PageSize        +  BinnedSizeLimit      : 0;
 
-	// Init tables.
-	OsTable.FirstPool = nullptr;
-	OsTable.ExhaustedPool = nullptr;
-	OsTable.BlockSize = 0;
-
-	/** The following options are not valid for page sizes less than 64k. They are here to reduce waste*/
-	PagePoolTable[0].FirstPool = nullptr;
-	PagePoolTable[0].ExhaustedPool = nullptr;
-	PagePoolTable[0].BlockSize = PageSize == Private::PAGE_SIZE_LIMIT ? BinnedSizeLimit+(BinnedSizeLimit/2) : 0;
-
-	PagePoolTable[1].FirstPool = nullptr;
-	PagePoolTable[1].ExhaustedPool = nullptr;
-	PagePoolTable[1].BlockSize = PageSize == Private::PAGE_SIZE_LIMIT ? PageSize+BinnedSizeLimit : 0;
-
-	// Block sizes are based around getting the maximum amount of allocations per pool, with as little alignment waste as possible.
-	// Block sizes should be close to even divisors of the POOL_SIZE, and well distributed. They must be 16-byte aligned as well.
-	static const uint32 BlockSizes[POOL_COUNT] =
+	for (uint32 Index = 0; Index != POOL_COUNT; ++Index)
 	{
-		8,		16,		32,		48,		64,		80,		96,		112,
-		128,	160,	192,	224,	256,	288,	320,	384,
-		448,	512,	576,	640,	704,	768,	896,	1024,
-		1168,	1360,	1632,	2048,	2336,	2720,	3264,	4096,
-		4672,	5456,	6544,	8192,	9360,	10912,	13104,	16384,
-		21840,	32768
-	};
-
-	for( uint32 i = 0; i < POOL_COUNT; i++ )
-	{
-		PoolTable[i].FirstPool = nullptr;
-		PoolTable[i].ExhaustedPool = nullptr;
-		PoolTable[i].BlockSize = BlockSizes[i];
+		PoolTable[Index].BlockSize  = GMallocBinned2BlockSizes[Index];
 #if STATS
-		PoolTable[i].MinRequest = PoolTable[i].BlockSize;
+		PoolTable[Index].MinRequest = GMallocBinned2BlockSizes[Index];
 #endif
 	}
 
-	for( uint32 i=0; i<MAX_POOLED_ALLOCATION_SIZE; i++ )
+	for (uint32 Index = 0; Index != MAX_POOLED_ALLOCATION_SIZE; ++Index)
 	{
-		uint32 Index = 0;
-		while( PoolTable[Index].BlockSize < i )
+		FPoolTable* Pool = PoolTable;
+		while (Pool->BlockSize < Index)
 		{
-			++Index;
+			++Pool;
+			checkSlow(Pool != PoolTable + POOL_COUNT);
 		}
-		checkSlow(Index < POOL_COUNT);
-		MemSizeToPoolTable[i] = &PoolTable[Index];
+		MemSizeToPoolTable[Index] = Pool;
 	}
 
-	MemSizeToPoolTable[BinnedSizeLimit] = &PagePoolTable[0];
-	MemSizeToPoolTable[BinnedSizeLimit+1] = &PagePoolTable[1];
+	MemSizeToPoolTable[BinnedSizeLimit    ] = &PagePoolTable[0];
+	MemSizeToPoolTable[BinnedSizeLimit + 1] = &PagePoolTable[1];
+
+	HashBuckets = (PoolHashBucket*)FPlatformMemory::BinnedAllocFromOS(Align(MaxHashBuckets * sizeof(PoolHashBucket), PageSize));
+	for (uint32 i = 0; i < MaxHashBuckets; ++i) 
+	{
+		new (HashBuckets + i) PoolHashBucket();
+	}
 
 	check(MAX_POOLED_ALLOCATION_SIZE - 1 == PoolTable[POOL_COUNT - 1].BlockSize);
 }
@@ -836,191 +549,253 @@ void* FMallocBinned2::Malloc(SIZE_T Size, uint32 Alignment)
 	Alignment = FMath::Max<uint32>(Alignment, Private::DEFAULT_BINNED_ALLOCATOR_ALIGNMENT);
 	SIZE_T SpareBytesCount = FMath::Min<SIZE_T>(Private::DEFAULT_BINNED_ALLOCATOR_ALIGNMENT, Size);
 	Size = FMath::Max<SIZE_T>(PoolTable[0].BlockSize, Size + (Alignment - SpareBytesCount));
-	BINNED2_MEM_TIME(MemTime -= FPlatformTime::Seconds());
-	
-	BINNED2_INCREMENT_STATCOUNTER(CurrentAllocs);
-	BINNED2_INCREMENT_STATCOUNTER(TotalAllocs);
-	
-	FFreeMem* Free;
-	if( Size < BinnedSizeLimit )
+
+	BINNED2_INCREMENT_STATCOUNTER(Stats.CurrentAllocs);
+	BINNED2_INCREMENT_STATCOUNTER(Stats.TotalAllocs);
+
+	if (Size < BinnedSizeLimit)
 	{
 		// Allocate from pool.
 		FPoolTable* Table = MemSizeToPoolTable[Size];
+
 		checkSlow(Size <= Table->BlockSize);
 
 		Private::TrackStats(Table, Size);
 
 		FPoolInfo* Pool = Table->FirstPool;
-		if( !Pool )
+		if (!Pool)
 		{
-			Pool = Private::AllocatePoolMemory(*this, Table, Private::BINNED_ALLOC_POOL_SIZE/*PageSize*/, Size);
+			Pool = Private::AllocatePoolMemory(*this, *Table, Private::BINNED_ALLOC_POOL_SIZE/*PageSize*/, Size);
 		}
 
-		Free = Private::AllocateBlockFromPool(*this, Table, Pool, Alignment);
+		FFreeMem* Result = Private::AllocateBlockFromPool(*this, *Table, Pool, Alignment);
+		return Result;
 	}
-	else if ( ((Size >= BinnedSizeLimit && Size <= PagePoolTable[0].BlockSize) ||
-		(Size > PageSize && Size <= PagePoolTable[1].BlockSize)))
+
+	if (Size <= PagePoolTable[0].BlockSize || (Size > PageSize && Size <= PagePoolTable[1].BlockSize))
 	{
 		// Bucket in a pool of 3*PageSize or 6*PageSize
-		uint32 BinType = Size < PageSize ? 0 : 1;
-		uint32 PageCount = 3*BinType + 3;
-		FPoolTable* Table = &PagePoolTable[BinType];
+		uint32      BinType   = Size < PageSize ? 0 : 1;
+		uint32      PageCount = 3 * BinType + 3;
+		FPoolTable* Table     = &PagePoolTable[BinType];
+
 		checkSlow(Size <= Table->BlockSize);
 
 		Private::TrackStats(Table, Size);
 
 		FPoolInfo* Pool = Table->FirstPool;
-		if( !Pool )
+		if (!Pool)
 		{
-			Pool = Private::AllocatePoolMemory(*this, Table, PageCount*PageSize, BinnedSizeLimit+BinType);
+			Pool = Private::AllocatePoolMemory(*this, *Table, PageCount*PageSize, BinnedSizeLimit+BinType);
 		}
 
-		Free = Private::AllocateBlockFromPool(*this, Table, Pool, Alignment);
+		FFreeMem* Result = Private::AllocateBlockFromPool(*this, *Table, Pool, Alignment);
+		return Result;
 	}
-	else
+
+	// Use OS for large allocations.
+	UPTRINT AlignedSize = Align(Size, PageSize);
+	FFreeMem* Result = (FFreeMem*)CachedOSPageAllocator.Allocate(AlignedSize);
+	if (!Result)
 	{
-		// Use OS for large allocations.
-		UPTRINT AlignedSize = Align(Size,PageSize);
-		SIZE_T ActualPoolSize; //TODO: use this to reduce waste?
-		Free = (FFreeMem*)Private::OSAlloc(*this, AlignedSize, ActualPoolSize);
-		if( !Free )
-		{
-			Private::OutOfMemory(AlignedSize);
-		}
+		Private::OutOfMemory(AlignedSize);
+	}
+
+	// OS allocations should always be page-aligned.
+	checkSlow(IsAligned(Result, PageSize));
 
-		void* AlignedFree = Align(Free, Alignment);
+	void* AlignedResult = Align(Result, Alignment);
 
-		// Create indirect.
-		FPoolInfo* Pool;
+	// Create indirect.
+	FPoolInfo* Pool = Private::GetPoolInfo(*this, Result);
+	if (Result != AlignedResult)
+	{
+		// Mark the FPoolInfo for AlignedResult to jump back to the FPoolInfo for ptr.
+		for (UPTRINT i = (UPTRINT)PageSize, Offset = 0; i < AlignedSize; i += PageSize, ++Offset)
 		{
-			Pool = Private::GetPoolInfo(*this, (UPTRINT)Free);
+			FPoolInfo* TrailingPool = Private::GetPoolInfo(*this, (void*)((UPTRINT)Result + i));
+			check(TrailingPool);
 
-			if ((UPTRINT)Free != ((UPTRINT)AlignedFree & ~((UPTRINT)PageSize - 1)))
-			{
-				// Mark the FPoolInfo for AlignedFree to jump back to the FPoolInfo for ptr.
-				for (UPTRINT i = (UPTRINT)PageSize, Offset = 0; i < AlignedSize; i += PageSize, ++Offset)
-				{
-					FPoolInfo* TrailingPool = Private::GetPoolInfo(*this, ((UPTRINT)Free) + i);
-					check(TrailingPool);
-					//Set trailing pools to point back to first pool
-					TrailingPool->SetAllocationSizes(0, 0, Offset, BinnedOSTableIndex);
-				}
-			}
+			//Set trailing pools to point back to first pool
+			TrailingPool->SetAllocationSizes(0, 0, Offset, BinnedOSTableIndex);
 		}
-		Free = (FFreeMem*)AlignedFree;
-		Pool->SetAllocationSizes(Size, AlignedSize, BinnedOSTableIndex, BinnedOSTableIndex);
-		BINNED2_PEAK_STATCOUNTER(OsPeak, BINNED2_ADD_STATCOUNTER(OsCurrent, AlignedSize));
-		BINNED2_PEAK_STATCOUNTER(UsedPeak, BINNED2_ADD_STATCOUNTER(UsedCurrent, Size));
-		BINNED2_PEAK_STATCOUNTER(WastePeak, BINNED2_ADD_STATCOUNTER(WasteCurrent, (int64)(AlignedSize - Size)));
 	}
 
-	BINNED2_MEM_TIME(MemTime += FPlatformTime::Seconds());
-	return Free;
+	Pool->SetAllocationSizes(Size, AlignedSize, BinnedOSTableIndex, BinnedOSTableIndex);
+
+	BINNED2_PEAK_STATCOUNTER(Stats.OsPeak,    BINNED2_ADD_STATCOUNTER(Stats.OsCurrent,    AlignedSize));
+	BINNED2_PEAK_STATCOUNTER(Stats.UsedPeak,  BINNED2_ADD_STATCOUNTER(Stats.UsedCurrent,  Size));
+	BINNED2_PEAK_STATCOUNTER(Stats.WastePeak, BINNED2_ADD_STATCOUNTER(Stats.WasteCurrent, (int64)(AlignedSize - Size)));
+
+	return AlignedResult;
 }
 
-void* FMallocBinned2::Realloc( void* Ptr, SIZE_T NewSize, uint32 Alignment )
+void* FMallocBinned2::Realloc(void* Ptr, SIZE_T NewSize, uint32 Alignment)
 {
 	// Handle DEFAULT_ALIGNMENT for binned allocator.
-	if (Alignment == DEFAULT_ALIGNMENT)
+	Alignment = FMath::Max<uint32>(Alignment, Private::DEFAULT_BINNED_ALLOCATOR_ALIGNMENT);
+
+	if (!Ptr)
 	{
-		Alignment = Private::DEFAULT_BINNED_ALLOCATOR_ALIGNMENT;
+		void* Result = FMallocBinned2::Malloc(NewSize, Alignment);
+		return Result;
+	}
+
+	if (NewSize == 0)
+	{
+		FMallocBinned2::Free(Ptr);
+		return nullptr;
 	}
 
-	Alignment = FMath::Max<uint32>(Alignment, Private::DEFAULT_BINNED_ALLOCATOR_ALIGNMENT);
 	const uint32 NewSizeUnmodified = NewSize;
 	SIZE_T SpareBytesCount = FMath::Min<SIZE_T>(Private::DEFAULT_BINNED_ALLOCATOR_ALIGNMENT, NewSize);
 	if (NewSize)
 	{
 		NewSize = FMath::Max<SIZE_T>(PoolTable[0].BlockSize, NewSize + (Alignment - SpareBytesCount));
 	}
-	BINNED2_MEM_TIME(MemTime -= FPlatformTime::Seconds());
-	UPTRINT BasePtr;
-	void* NewPtr = Ptr;
-	if( Ptr && NewSize )
-	{
-		FPoolInfo* Pool = Private::FindPoolInfo(*this, (UPTRINT)Ptr, BasePtr);
 
-		if( Pool->TableIndex < BinnedOSTableIndex )
+	FPoolInfo* Pool = Private::FindPoolInfo(*this, Ptr);
+	if (Pool->TableIndex < BinnedOSTableIndex)
+	{
+		// Allocated from pool, so grow or shrink if necessary.
+		check(Pool->TableIndex > 0); // it isn't possible to allocate a size of 0, Malloc will increase the size to DEFAULT_BINNED_ALLOCATOR_ALIGNMENT
+		if (NewSizeUnmodified > MemSizeToPoolTable[Pool->TableIndex]->BlockSize || NewSizeUnmodified <= MemSizeToPoolTable[Pool->TableIndex - 1]->BlockSize)
 		{
-			// Allocated from pool, so grow or shrink if necessary.
-			check(Pool->TableIndex > 0); // it isn't possible to allocate a size of 0, Malloc will increase the size to DEFAULT_BINNED_ALLOCATOR_ALIGNMENT
-			if (NewSizeUnmodified > MemSizeToPoolTable[Pool->TableIndex]->BlockSize || NewSizeUnmodified <= MemSizeToPoolTable[Pool->TableIndex - 1]->BlockSize)
-			{
-				NewPtr = Malloc(NewSizeUnmodified, Alignment);
-				FMemory::Memcpy(NewPtr, Ptr, FMath::Min<SIZE_T>(NewSizeUnmodified, MemSizeToPoolTable[Pool->TableIndex]->BlockSize - (Alignment - SpareBytesCount)));
-				Free( Ptr );
-			}
-			else if (((UPTRINT)Ptr & (UPTRINT)(Alignment - 1)) != 0)
-			{
-				NewPtr = Align(Ptr, Alignment);
-				FMemory::Memmove(NewPtr, Ptr, NewSize);
-			}
+			void* Result = FMallocBinned2::Malloc(NewSizeUnmodified, Alignment);
+			FMemory::Memcpy(Result, Ptr, FMath::Min<SIZE_T>(NewSizeUnmodified, MemSizeToPoolTable[Pool->TableIndex]->BlockSize - (Alignment - SpareBytesCount)));
+			FMallocBinned2::Free(Ptr);
+
+			return Result;
 		}
-		else
+
+		if (!IsAligned(Ptr, Alignment))
 		{
-			// Allocated from OS.
-			if( NewSize > Pool->GetOsBytes(PageSize, BinnedOSTableIndex) || NewSize * 3 < Pool->GetOsBytes(PageSize, BinnedOSTableIndex) * 2 )
-			{
-				// Grow or shrink.
-				NewPtr = Malloc(NewSizeUnmodified, Alignment);
-				FMemory::Memcpy(NewPtr, Ptr, FMath::Min<SIZE_T>(NewSizeUnmodified, Pool->GetBytes()));
-				Free( Ptr );
-			}
-			else
-			{
-				//need a lock to cover the SetAllocationSizes()
-				int32 UsedChange = (NewSize - Pool->GetBytes());
-				
-				// Keep as-is, reallocation isn't worth the overhead.
-				BINNED2_ADD_STATCOUNTER(UsedCurrent, UsedChange);
-				BINNED2_PEAK_STATCOUNTER(UsedPeak, UsedCurrent);
-				BINNED2_ADD_STATCOUNTER(WasteCurrent, (Pool->GetBytes() - NewSize));
-				Pool->SetAllocationSizes(NewSizeUnmodified, Pool->GetOsBytes(PageSize, BinnedOSTableIndex), BinnedOSTableIndex, BinnedOSTableIndex);
-			}
+			void* Result = Align(Ptr, Alignment);
+			FMemory::Memmove(Result, Ptr, NewSize);
+			return Result;
 		}
+
+		return Ptr;
 	}
-	else if( Ptr == nullptr )
-	{
-		NewPtr = Malloc(NewSizeUnmodified, Alignment);
-	}
-	else
+
+	// Allocated from OS.
+	UPTRINT PoolOsBytes = Pool->GetOsBytes(PageSize, BinnedOSTableIndex);
+	if( NewSize > PoolOsBytes || NewSize * 3 < PoolOsBytes * 2 )
 	{
-		Free( Ptr );
-		NewPtr = nullptr;
+		// Grow or shrink.
+		void* Result = FMallocBinned2::Malloc(NewSizeUnmodified, Alignment);
+		FMemory::Memcpy(Result, Ptr, FMath::Min<SIZE_T>(NewSizeUnmodified, Pool->AllocSize));
+		FMallocBinned2::Free(Ptr);
+
+		return Result;
 	}
 
-	BINNED2_MEM_TIME(MemTime += FPlatformTime::Seconds());
-	return NewPtr;
+	//need a lock to cover the SetAllocationSizes()
+	int32 UsedChange = (NewSize - Pool->AllocSize);
+
+	// Keep as-is, reallocation isn't worth the overhead.
+	BINNED2_ADD_STATCOUNTER (Stats.UsedCurrent,  UsedChange);
+	BINNED2_PEAK_STATCOUNTER(Stats.UsedPeak,     Stats.UsedCurrent);
+	BINNED2_ADD_STATCOUNTER (Stats.WasteCurrent, (Pool->AllocSize - NewSize));
+
+	Pool->SetAllocationSizes(NewSizeUnmodified, PoolOsBytes, BinnedOSTableIndex, BinnedOSTableIndex);
+
+	return Ptr;
 }
 
 void FMallocBinned2::Free( void* Ptr )
 {
-	if( !Ptr )
+	if (!Ptr)
 	{
 		return;
 	}
 
-	Private::PushFreeLockless(*this, Ptr);
+	BINNED2_DECREMENT_STATCOUNTER(Stats.CurrentAllocs);
+
+	void* BasePtr;
+	FPoolInfo* Pool = Private::FindPoolInfo(*this, Ptr, BasePtr);
+#if PLATFORM_IOS || PLATFORM_MAC
+	if (Pool == NULL)
+	{
+		UE_LOG(LogMemory, Warning, TEXT("Attempting to free a pointer we didn't allocate!"));
+		return;
+	}
+#endif
+	checkSlow(Pool);
+	checkSlow(Pool->AllocSize != 0);
+	if (Pool->TableIndex < BinnedOSTableIndex)
+	{
+		FPoolTable* Table = MemSizeToPoolTable[Pool->TableIndex];
+#if STATS
+		Table->ActiveRequests--;
+#endif
+		// If this pool was exhausted, move to available list.
+		if (!Pool->FirstMem)
+		{
+			Pool->Unlink();
+			Pool->Link(Table->FirstPool);
+		}
+
+		check((UPTRINT)BasePtr <= (UPTRINT)Ptr);
+
+		uint32 AlignOffset = ((UPTRINT)Ptr - (UPTRINT)BasePtr) % Table->BlockSize;
+
+		// Patch pointer to include previously applied alignment.
+		Ptr = (void*)((PTRINT)Ptr - (PTRINT)AlignOffset);
+
+		// Free a pooled allocation.
+		FFreeMem* Free		= (FFreeMem*)Ptr;
+		Free->NumFreeBlocks	= 1;
+		Free->Next			= Pool->FirstMem;
+		Pool->FirstMem		= Free;
+		BINNED2_ADD_STATCOUNTER(Stats.UsedCurrent, -(int64)Table->BlockSize);
+
+		// Free this pool.
+		checkSlow(Pool->Taken >= 1);
+		if( --Pool->Taken == 0 )
+		{
+#if STATS
+			Table->NumActivePools--;
+#endif
+			// Free the OS memory.
+			SIZE_T OsBytes = Pool->GetOsBytes(PageSize, BinnedOSTableIndex);
+			BINNED2_ADD_STATCOUNTER(Stats.OsCurrent,    -(int64)OsBytes);
+			BINNED2_ADD_STATCOUNTER(Stats.WasteCurrent, -(int64)(OsBytes - Pool->AllocSize));
+			Pool->Unlink();
+			Pool->SetAllocationSizes(0, 0, 0, BinnedOSTableIndex);
+			CachedOSPageAllocator.Free(BasePtr, OsBytes);
+		}
+	}
+	else
+	{
+		// Free an OS allocation.
+		checkSlow(IsAligned(Ptr, PageSize));
+		SIZE_T OsBytes = Pool->GetOsBytes(PageSize, BinnedOSTableIndex);
+
+		BINNED2_ADD_STATCOUNTER(Stats.UsedCurrent,  -(int64)Pool->AllocSize);
+		BINNED2_ADD_STATCOUNTER(Stats.OsCurrent,    -(int64)OsBytes);
+		BINNED2_ADD_STATCOUNTER(Stats.WasteCurrent, -(int64)(OsBytes - Pool->AllocSize));
+		CachedOSPageAllocator.Free(BasePtr, OsBytes);
+	}
 }
 
-bool FMallocBinned2::GetAllocationSize(void *Original, SIZE_T &SizeOut)
+bool FMallocBinned2::GetAllocationSize(void* Original, SIZE_T& SizeOut)
 {
 	if (!Original)
 	{
 		return false;
 	}
-	UPTRINT BasePtr;
-	FPoolInfo* Pool = Private::FindPoolInfo(*this, (UPTRINT)Original, BasePtr);
-	SizeOut = Pool->TableIndex < BinnedOSTableIndex ? MemSizeToPoolTable[Pool->TableIndex]->BlockSize : Pool->GetBytes();
+
+	FPoolInfo* Pool = Private::FindPoolInfo(*this, Original);
+	SizeOut = Pool->TableIndex < BinnedOSTableIndex ? MemSizeToPoolTable[Pool->TableIndex]->BlockSize : Pool->AllocSize;
 	return true;
 }
 
 bool FMallocBinned2::ValidateHeap()
 {
-	for( int32 i = 0; i < POOL_COUNT; i++ )
+	for (FPoolTable& Table : PoolTable)
 	{
-		FPoolTable* Table = &PoolTable[i];
-		for( FPoolInfo** PoolPtr = &Table->FirstPool; *PoolPtr; PoolPtr = &(*PoolPtr)->Next )
+		for( FPoolInfo** PoolPtr = &Table.FirstPool; *PoolPtr; PoolPtr = &(*PoolPtr)->Next )
 		{
 			FPoolInfo* Pool = *PoolPtr;
 			check(Pool->PrevLink == PoolPtr);
@@ -1030,7 +805,7 @@ bool FMallocBinned2::ValidateHeap()
 				check(Free->NumFreeBlocks > 0);
 			}
 		}
-		for( FPoolInfo** PoolPtr = &Table->ExhaustedPool; *PoolPtr; PoolPtr = &(*PoolPtr)->Next )
+		for( FPoolInfo** PoolPtr = &Table.ExhaustedPool; *PoolPtr; PoolPtr = &(*PoolPtr)->Next )
 		{
 			FPoolInfo* Pool = *PoolPtr;
 			check(Pool->PrevLink == PoolPtr);
@@ -1045,40 +820,18 @@ void FMallocBinned2::UpdateStats()
 {
 	FMalloc::UpdateStats();
 #if STATS
-	SIZE_T	LocalOsCurrent = 0;
-	SIZE_T	LocalOsPeak = 0;
-	SIZE_T	LocalWasteCurrent = 0;
-	SIZE_T	LocalWastePeak = 0;
-	SIZE_T	LocalUsedCurrent = 0;
-	SIZE_T	LocalUsedPeak = 0;
-	SIZE_T	LocalCurrentAllocs = 0;
-	SIZE_T	LocalTotalAllocs = 0;
-	SIZE_T	LocalSlackCurrent = 0;
-
-	{
-		Private::UpdateSlackStat(*this);
-
-		// Copy memory stats.
-		LocalOsCurrent = OsCurrent;
-		LocalOsPeak = OsPeak;
-		LocalWasteCurrent = WasteCurrent;
-		LocalWastePeak = WastePeak;
-		LocalUsedCurrent = UsedCurrent;
-		LocalUsedPeak = UsedPeak;
-		LocalCurrentAllocs = CurrentAllocs;
-		LocalTotalAllocs = TotalAllocs;
-		LocalSlackCurrent = SlackCurrent;
-	}
 
-	SET_MEMORY_STAT( STAT_Binned2_OsCurrent, LocalOsCurrent );
-	SET_MEMORY_STAT( STAT_Binned2_OsPeak, LocalOsPeak );
-	SET_MEMORY_STAT( STAT_Binned2_WasteCurrent, LocalWasteCurrent );
-	SET_MEMORY_STAT( STAT_Binned2_WastePeak, LocalWastePeak );
-	SET_MEMORY_STAT( STAT_Binned2_UsedCurrent, LocalUsedCurrent );
-	SET_MEMORY_STAT( STAT_Binned2_UsedPeak, LocalUsedPeak );
-	SET_DWORD_STAT( STAT_Binned2_CurrentAllocs, LocalCurrentAllocs );
-	SET_DWORD_STAT( STAT_Binned2_TotalAllocs, LocalTotalAllocs );
-	SET_MEMORY_STAT( STAT_Binned2_SlackCurrent, LocalSlackCurrent );
+	Private::UpdateSlackStat(*this);
+
+	SET_MEMORY_STAT( STAT_Binned2_OsCurrent,     Stats.OsCurrent );
+	SET_MEMORY_STAT( STAT_Binned2_OsPeak,        Stats.OsPeak );
+	SET_MEMORY_STAT( STAT_Binned2_WasteCurrent,  Stats.WasteCurrent );
+	SET_MEMORY_STAT( STAT_Binned2_WastePeak,     Stats.WastePeak );
+	SET_MEMORY_STAT( STAT_Binned2_UsedCurrent,   Stats.UsedCurrent );
+	SET_MEMORY_STAT( STAT_Binned2_UsedPeak,      Stats.UsedPeak );
+	SET_DWORD_STAT ( STAT_Binned2_CurrentAllocs, Stats.CurrentAllocs );
+	SET_DWORD_STAT ( STAT_Binned2_TotalAllocs,   Stats.TotalAllocs );
+	SET_MEMORY_STAT( STAT_Binned2_SlackCurrent,  Stats.SlackCurrent );
 #endif
 }
 
@@ -1093,17 +846,15 @@ void FMallocBinned2::DumpAllocatorStats( class FOutputDevice& Ar )
 		// This is all of the memory including stuff too big for the pools
 		BufferedOutput.CategorizedLogf( LogMemory.GetCategoryName(), ELogVerbosity::Log, TEXT( "Allocator Stats for %s:" ), GetDescriptiveName() );
 		// Waste is the total overhead of the memory system
-		BufferedOutput.CategorizedLogf( LogMemory.GetCategoryName(), ELogVerbosity::Log, TEXT( "Current Memory %.2f MB used, plus %.2f MB waste" ), UsedCurrent / (1024.0f * 1024.0f), (OsCurrent - UsedCurrent) / (1024.0f * 1024.0f) );
-		BufferedOutput.CategorizedLogf( LogMemory.GetCategoryName(), ELogVerbosity::Log, TEXT( "Peak Memory %.2f MB used, plus %.2f MB waste" ), UsedPeak / (1024.0f * 1024.0f), (OsPeak - UsedPeak) / (1024.0f * 1024.0f) );
+		BufferedOutput.CategorizedLogf( LogMemory.GetCategoryName(), ELogVerbosity::Log, TEXT( "Current Memory %.2f MB used, plus %.2f MB waste" ), Stats.UsedCurrent / (1024.0f * 1024.0f), (Stats.OsCurrent - Stats.UsedCurrent) / (1024.0f * 1024.0f) );
+		BufferedOutput.CategorizedLogf( LogMemory.GetCategoryName(), ELogVerbosity::Log, TEXT( "Peak Memory %.2f MB used, plus %.2f MB waste" ), Stats.UsedPeak / (1024.0f * 1024.0f), (Stats.OsPeak - Stats.UsedPeak) / (1024.0f * 1024.0f) );
 
-		BufferedOutput.CategorizedLogf( LogMemory.GetCategoryName(), ELogVerbosity::Log, TEXT( "Current OS Memory %.2f MB, peak %.2f MB" ), OsCurrent / (1024.0f * 1024.0f), OsPeak / (1024.0f * 1024.0f) );
-		BufferedOutput.CategorizedLogf( LogMemory.GetCategoryName(), ELogVerbosity::Log, TEXT( "Current Waste %.2f MB, peak %.2f MB" ), WasteCurrent / (1024.0f * 1024.0f), WastePeak / (1024.0f * 1024.0f) );
-		BufferedOutput.CategorizedLogf( LogMemory.GetCategoryName(), ELogVerbosity::Log, TEXT( "Current Used %.2f MB, peak %.2f MB" ), UsedCurrent / (1024.0f * 1024.0f), UsedPeak / (1024.0f * 1024.0f) );
-		BufferedOutput.CategorizedLogf( LogMemory.GetCategoryName(), ELogVerbosity::Log, TEXT( "Current Slack %.2f MB" ), SlackCurrent / (1024.0f * 1024.0f) );
+		BufferedOutput.CategorizedLogf( LogMemory.GetCategoryName(), ELogVerbosity::Log, TEXT( "Current OS Memory %.2f MB, peak %.2f MB" ), Stats.OsCurrent / (1024.0f * 1024.0f), Stats.OsPeak / (1024.0f * 1024.0f) );
+		BufferedOutput.CategorizedLogf( LogMemory.GetCategoryName(), ELogVerbosity::Log, TEXT( "Current Waste %.2f MB, peak %.2f MB" ), Stats.WasteCurrent / (1024.0f * 1024.0f), Stats.WastePeak / (1024.0f * 1024.0f) );
+		BufferedOutput.CategorizedLogf( LogMemory.GetCategoryName(), ELogVerbosity::Log, TEXT( "Current Used %.2f MB, peak %.2f MB" ), Stats.UsedCurrent / (1024.0f * 1024.0f), Stats.UsedPeak / (1024.0f * 1024.0f) );
+		BufferedOutput.CategorizedLogf( LogMemory.GetCategoryName(), ELogVerbosity::Log, TEXT( "Current Slack %.2f MB" ), Stats.SlackCurrent / (1024.0f * 1024.0f) );
 
-		BufferedOutput.CategorizedLogf( LogMemory.GetCategoryName(), ELogVerbosity::Log, TEXT( "Allocs      % 6i Current / % 6i Total" ), CurrentAllocs, TotalAllocs );
-		BINNED2_MEM_TIME( BufferedOutput.CategorizedLogf( LogMemory.GetCategoryName(), ELogVerbosity::Log, TEXT( "Seconds     % 5.3f" ), MemTime ) );
-		BINNED2_MEM_TIME( BufferedOutput.CategorizedLogf( LogMemory.GetCategoryName(), ELogVerbosity::Log, TEXT( "MSec/Allc   % 5.5f" ), 1000.0 * MemTime / MemAllocs ) );
+		BufferedOutput.CategorizedLogf( LogMemory.GetCategoryName(), ELogVerbosity::Log, TEXT( "Allocs      % 6i Current / % 6i Total" ), Stats.CurrentAllocs, Stats.TotalAllocs );
 
 		// This is the memory tracked inside individual allocation pools
 		BufferedOutput.CategorizedLogf( LogMemory.GetCategoryName(), ELogVerbosity::Log, TEXT( "" ) );
diff --git a/Engine/Source/Runtime/Core/Private/Misc/CallbackDevice.cpp b/Engine/Source/Runtime/Core/Private/Misc/CallbackDevice.cpp
--- a/Engine/Source/Runtime/Core/Private/Misc/CallbackDevice.cpp
+++ b/Engine/Source/Runtime/Core/Private/Misc/CallbackDevice.cpp
@@ -26,12 +26,6 @@ FCoreDelegates::FObjectArrayForDebugVisualizersDelegate& FCoreDelegates::GetObje
 	return StaticDelegate;
 }
 
-FCoreDelegates::FGetSerialNumberBlocksForDebugVisualizersDelegate& FCoreDelegates::GetSerialNumberBlocksForDebugVisualizersDelegate()
-{
-	static FGetSerialNumberBlocksForDebugVisualizersDelegate StaticDelegate;
-	return StaticDelegate;
-}
-
 FCoreDelegates::FOnMountPak FCoreDelegates::OnMountPak;
 FCoreDelegates::FOnUnmountPak FCoreDelegates::OnUnmountPak;
 FCoreDelegates::FOnUserLoginChangedEvent FCoreDelegates::OnUserLoginChangedEvent; 
