Copying //UE4/Dev-Core to //UE4/Main

diff --git a/Engine/Source/Runtime/Core/Private/Misc/CoreMisc.cpp b/Engine/Source/Runtime/Core/Private/Misc/CoreMisc.cpp
--- a/Engine/Source/Runtime/Core/Private/Misc/CoreMisc.cpp
+++ b/Engine/Source/Runtime/Core/Private/Misc/CoreMisc.cpp
@@ -203,9 +203,9 @@ bool FFileHelper::LoadFileToString( FString& Result, const TCHAR* Filename, uint
 /**
  * Save a binary array to a file.
  */
-bool FFileHelper::SaveArrayToFile( const TArray<uint8>& Array, const TCHAR* Filename, IFileManager* FileManager /*= &IFileManager::Get()*/ )
+bool FFileHelper::SaveArrayToFile( const TArray<uint8>& Array, const TCHAR* Filename, IFileManager* FileManager /*= &IFileManager::Get()*/, uint32 WriteFlags )
 {
-	FArchive* Ar = FileManager->CreateFileWriter( Filename, 0 );
+	FArchive* Ar = FileManager->CreateFileWriter( Filename, WriteFlags );
 	if( !Ar )
 	{
 		return 0;
@@ -219,10 +219,10 @@ bool FFileHelper::SaveArrayToFile( const TArray<uint8>& Array, const TCHAR* File
  * Write the FString to a file.
  * Supports all combination of ANSI/Unicode files and platforms.
  */
-bool FFileHelper::SaveStringToFile( const FString& String, const TCHAR* Filename,  EEncodingOptions::Type EncodingOptions, IFileManager* FileManager /*= &IFileManager::Get()*/ )
+bool FFileHelper::SaveStringToFile( const FString& String, const TCHAR* Filename,  EEncodingOptions::Type EncodingOptions, IFileManager* FileManager /*= &IFileManager::Get()*/, uint32 WriteFlags )
 {
 	// max size of the string is a UCS2CHAR for each character and some UNICODE magic 
-	auto Ar = TUniquePtr<FArchive>( FileManager->CreateFileWriter( Filename, 0 ) );
+	auto Ar = TUniquePtr<FArchive>( FileManager->CreateFileWriter( Filename, WriteFlags ) );
 	if( !Ar )
 		return false;
 
@@ -585,6 +585,9 @@ bool FCommandLine::Set(const TCHAR* NewCommandLine)
 	}
 
 	FCString::Strncpy( CmdLine, NewCommandLine, ARRAY_COUNT(CmdLine) );
+	// If configured as part of the build, strip out any unapproved args
+	WhitelistCommandLines();
+
 	bIsInitialized = true;
 
 	// Check for the '-' that normal ones get converted to in Outlook. It's important to do it AFTER the command line is initialized
@@ -605,7 +608,81 @@ bool FCommandLine::Set(const TCHAR* NewCommandLine)
 void FCommandLine::Append(const TCHAR* AppendString)
 {
 	FCString::Strncat( CmdLine, AppendString, ARRAY_COUNT(CmdLine) );
+	// If configured as part of the build, strip out any unapproved args
+	WhitelistCommandLines();
+}
+
+#if WANTS_COMMANDLINE_WHITELIST
+TArray<FString> FCommandLine::ApprovedArgs;
+
+#ifdef OVERRIDE_COMMANDLINE_WHITELIST
+/**
+ * When overriding this setting make sure that your define looks like the following in your .cs file:
+ *
+ *		OutCPPEnvironmentConfiguration.Definitions.Add("OVERRIDE_COMMANDLINE_WHITELIST=\"-arg1 -arg2 -arg3 -arg4\"");
+ *
+ * The important part is the \" as they quotes get stripped off by the compiler without them
+ */
+const TCHAR* OverrideList = TEXT(OVERRIDE_COMMANDLINE_WHITELIST);
+#else
+// Default list most conservative restrictions
+const TCHAR* OverrideList = TEXT("-fullscreen /windowed");
+#endif
+
+void FCommandLine::WhitelistCommandLines()
+{
+	if (ApprovedArgs.Num() == 0)
+	{
+		TArray<FString> Ignored;
+		FCommandLine::Parse(OverrideList, ApprovedArgs, Ignored);
+	}
+	// Process the original command line
+	TArray<FString> OriginalList = FilterCommandLine(OriginalCmdLine);
+	BuildWhitelistCommandLine(OriginalCmdLine, ARRAY_COUNT(OriginalCmdLine), OriginalList);
+	// Process the current command line
+	TArray<FString> CmdList = FilterCommandLine(CmdLine);
+	BuildWhitelistCommandLine(CmdLine, ARRAY_COUNT(CmdLine), CmdList);
+}
+
+TArray<FString> FCommandLine::FilterCommandLine(TCHAR* CommandLine)
+{
+	TArray<FString> Ignored;
+	TArray<FString> ParsedList;
+	// Parse the command line list
+	FCommandLine::Parse(CommandLine, ParsedList, Ignored);
+	// Remove any that are not in our approved list
+	for (int32 Index = 0; Index < ParsedList.Num(); Index++)
+	{
+		if (!ApprovedArgs.Contains(ParsedList[Index]))
+		{
+			ParsedList.RemoveAt(Index);
+		}
+	}
+	return ParsedList;
+}
+
+void FCommandLine::BuildWhitelistCommandLine(TCHAR* CommandLine, uint32 ArrayCount, const TArray<FString>& FilteredArgs)
+{
+	check(ArrayCount > 0);
+	// Zero the whole string
+	FMemory::Memzero(OriginalCmdLine, sizeof(TCHAR) * ArrayCount);
+
+	uint32 StartIndex = 0;
+	for (auto Arg : FilteredArgs)
+	{
+		if ((StartIndex + Arg.Len() + 2) < ArrayCount)
+		{
+			if (StartIndex != 0)
+			{
+				CommandLine[StartIndex++] = TEXT(' ');
+			}
+			CommandLine[StartIndex++] = TEXT('-');
+			FCString::Strncpy(&CommandLine[StartIndex], *Arg, ArrayCount - StartIndex);
+			StartIndex += Arg.Len();
+		}
+	}
 }
+#endif
 
 void FCommandLine::AddToSubprocessCommandline( const TCHAR* Param )
 {
@@ -666,9 +743,10 @@ void FCommandLine::Parse(const TCHAR* InCmdLine, TArray<FString>& Tokens, TArray
 	FString NextToken;
 	while (FParse::Token(InCmdLine, NextToken, false))
 	{
-		if ((**NextToken == TCHAR('-')))
+		if ((**NextToken == TCHAR('-')) || (**NextToken == TCHAR('/')))
 		{
 			new(Switches) FString(NextToken.Mid(1));
+			new(Tokens) FString(NextToken.Right(NextToken.Len() - 1));
 		}
 		else
 		{
diff --git a/Engine/Source/Runtime/Core/Private/Stats/Stats2.cpp b/Engine/Source/Runtime/Core/Private/Stats/Stats2.cpp
--- a/Engine/Source/Runtime/Core/Private/Stats/Stats2.cpp
+++ b/Engine/Source/Runtime/Core/Private/Stats/Stats2.cpp
@@ -22,7 +22,6 @@ struct FStats2Globals
 static struct FForceInitAtBootFStats2 : public TForceInitAtBoot<FStats2Globals>
 {} FForceInitAtBootFStats2;
 
-DECLARE_FLOAT_COUNTER_STAT( TEXT("Seconds Per Cycle"), STAT_SecondsPerCycle, STATGROUP_Engine );
 DECLARE_DWORD_COUNTER_STAT( TEXT("Frame Packets Received"),STAT_StatFramePacketsRecv,STATGROUP_StatSystem);
 
 DECLARE_CYCLE_STAT(TEXT("WaitForStats"),STAT_WaitForStats,STATGROUP_Engine);
@@ -36,6 +35,7 @@ DECLARE_MEMORY_STAT( TEXT("Stats Descriptions"), STAT_StatDescMemory, STATGROUP_
 
 DEFINE_STAT(STAT_FrameTime);
 DEFINE_STAT(STAT_NamedMarker);
+DEFINE_STAT(STAT_SecondsPerCycle);
 
 /*-----------------------------------------------------------------------------
 	DebugLeakTest, for the stats based memory profiler
@@ -1076,6 +1076,12 @@ void FThreadStats::Flush( bool bHasBrokenCallstacks /*= false*/, bool bForceFlus
 
 void FThreadStats::FlushRegularStats( bool bHasBrokenCallstacks, bool bForceFlush )
 {
+	if (bReentranceGuard)
+	{
+		return;
+	}
+	TGuardValue<bool> Guard( bReentranceGuard, true );
+
 	enum
 	{
 		PRESIZE_MAX_NUM_ENTRIES = 10,
@@ -1133,11 +1139,11 @@ void FThreadStats::FlushRegularStats( bool bHasBrokenCallstacks, bool bForceFlus
 
 void FThreadStats::FlushRawStats( bool bHasBrokenCallstacks /*= false*/, bool bForceFlush /*= false*/ )
 {
-	if( bReentranceGuard )
+	if (bReentranceGuard)
 	{
 		return;
 	}
-	bReentranceGuard = true;
+	TGuardValue<bool> Guard( bReentranceGuard, true );
 	
 	enum
 	{
@@ -1173,8 +1179,6 @@ void FThreadStats::FlushRawStats( bool bHasBrokenCallstacks /*= false*/, bool bF
 
 		UE_LOG( LogStats, Verbose, TEXT( "FlushRawStats NumMessages: %i (%.2f MB), Thread: %u" ), NumMessages, NumMessagesAsMB, Packet.ThreadId );
 	}
-	
-	bReentranceGuard = false;
 }
 
 void FThreadStats::CheckForCollectingStartupStats()
diff --git a/Engine/Source/Runtime/Core/Private/Stats/StatsData.cpp b/Engine/Source/Runtime/Core/Private/Stats/StatsData.cpp
--- a/Engine/Source/Runtime/Core/Private/Stats/StatsData.cpp
+++ b/Engine/Source/Runtime/Core/Private/Stats/StatsData.cpp
@@ -22,7 +22,7 @@ DECLARE_MEMORY_STAT( TEXT("Stat Messages"), STAT_StatMessagesMemory, STATGROUP_S
 const FName FStatConstants::NAME_ThreadRoot = "ThreadRoot";
 const char* FStatConstants::ThreadGroupName = STAT_GROUP_TO_FStatGroup( STATGROUP_Threads )::GetGroupName();
 const FName FStatConstants::NAME_ThreadGroup = FStatConstants::ThreadGroupName;
-const FName FStatConstants::NAME_SecondsPerCycle = TEXT( "STAT_SecondsPerCycle" );
+const FName FStatConstants::RAW_SecondsPerCycle = FStatNameAndInfo( GET_STATFNAME( STAT_SecondsPerCycle ), true ).GetRawName();
 const FName FStatConstants::NAME_NoCategory = FName(TEXT("STATCAT_None"));
 
 const FString FStatConstants::StatsFileExtension = TEXT( ".ue4stats" );
@@ -901,7 +901,7 @@ void FStatsThreadState::AddToHistoryAndEmpty(FStatPacketArray& NewData)
 		}
 	}
 
-	int64 NewLatestFrame = GetLatestValidFrame();
+	const int64 NewLatestFrame = GetLatestValidFrame();
 
 	if (NewLatestFrame > 0)
 	{
@@ -925,11 +925,11 @@ void FStatsThreadState::AddToHistoryAndEmpty(FStatPacketArray& NewData)
 		}
 	}
 
-	int64 MinFrameToKeep = LatestFinishedFrame - HistoryFrames;
+	const int64 MinFrameToKeep = LatestFinishedFrame - HistoryFrames;
 
 	for (auto It = BadFrames.CreateIterator(); It; ++It)
 	{
-		int64 ThisFrame = *It ;
+		const int64 ThisFrame = *It;
 		if (ThisFrame <= LastFullFrameMetaAndNonFrame && ThisFrame < MinFrameToKeep)
 		{
 			check(ThisFrame <= LastFullFrameMetaAndNonFrame);
@@ -938,7 +938,7 @@ void FStatsThreadState::AddToHistoryAndEmpty(FStatPacketArray& NewData)
 	}
 	for (auto It = History.CreateIterator(); It; ++It)
 	{
-		int64 ThisFrame = It.Key();
+		const int64 ThisFrame = It.Key();
 		if (ThisFrame <= LastFullFrameMetaAndNonFrame && ThisFrame < MinFrameToKeep)
 		{
 			check(ThisFrame <= LastFullFrameMetaAndNonFrame);
@@ -947,7 +947,7 @@ void FStatsThreadState::AddToHistoryAndEmpty(FStatPacketArray& NewData)
 	}
 	for (auto It = EventsHistory.CreateIterator(); It; ++It)
 	{
-		int64 ThisFrame = It.Value().Frame;
+		const int64 ThisFrame = It.Value().Frame;
 		if (ThisFrame <= LastFullFrameProcessed && ThisFrame < MinFrameToKeep)
 		{
 			It.RemoveCurrent();
@@ -955,7 +955,7 @@ void FStatsThreadState::AddToHistoryAndEmpty(FStatPacketArray& NewData)
 	}
 	for (auto It = CondensedStackHistory.CreateIterator(); It; ++It)
 	{
-		int64 ThisFrame = It.Key();
+		const int64 ThisFrame = It.Key();
 		if (ThisFrame <= LastFullFrameProcessed && ThisFrame < MinFrameToKeep)
 		{
 			delete It.Value();
@@ -964,7 +964,7 @@ void FStatsThreadState::AddToHistoryAndEmpty(FStatPacketArray& NewData)
 	}
 	for (auto It = GoodFrames.CreateIterator(); It; ++It)
 	{
-		int64 ThisFrame = *It ;
+		const int64 ThisFrame = *It;
 		if (!History.Contains(ThisFrame) && !CondensedStackHistory.Contains(ThisFrame)) // if it isn't in the history anymore, it isn't good anymore
 		{
 			It.RemoveCurrent();
@@ -1067,7 +1067,23 @@ void FStatsThreadState::UpdateStatMessagesMemoryUsage()
 	}
 }
 
-void FStatsThreadState::GetInclusiveAggregateStackStats(int64 TargetFrame, TArray<FStatMessage>& OutStats, IItemFiler* Filter, bool bAddNonStackStats, TMap<FName, TArray<FStatMessage>>* OptionalOutThreadBreakdownMap) const
+void FStatsThreadState::GetInclusiveAggregateStackStats(
+	int64 TargetFrame, 
+	TArray<FStatMessage>& OutStats, 
+	IItemFiler* Filter /*= nullptr*/, 
+	bool bAddNonStackStats /*= true*/, 
+	TMap<FName, TArray<FStatMessage>>* OptionalOutThreadBreakdownMap /*= nullptr */) const
+{
+	const TArray<FStatMessage>& CondensedMessages = GetCondensedHistory( TargetFrame );
+	GetInclusiveAggregateStackStats( CondensedMessages, OutStats, Filter, bAddNonStackStats, OptionalOutThreadBreakdownMap );
+}
+
+void FStatsThreadState::GetInclusiveAggregateStackStats( 
+	const TArray<FStatMessage>& CondensedMessages, 
+	TArray<FStatMessage>& OutStats, 
+	IItemFiler* Filter /*= nullptr*/, 
+	bool bAddNonStackStats /*= true*/, 
+	TMap<FName, TArray<FStatMessage>>* OptionalOutThreadBreakdownMap /*= nullptr */ ) const
 {
 	struct FTimeInfo
 	{
@@ -1087,12 +1103,11 @@ void FStatsThreadState::GetInclusiveAggregateStackStats(int64 TargetFrame, TArra
 	TMap<FName, TMap<FName, FStatMessage>> ThisFrameMetaDataPerThread;
 	TMap<FName, FStatMessage> ThreadStarts;
 	TMap<FName, FStatMessage> ThreadEnds;
-	TArray<FStatMessage> const& Data = GetCondensedHistory(TargetFrame);
 	TMap<FName, FStatMessage>* ThisFrameMetaDataPerThreadPtr = nullptr;
 	int32 Depth = 0;
-	for (int32 Index = 0; Index < Data.Num(); Index++)
+	for (int32 Index = 0; Index < CondensedMessages.Num(); Index++)
 	{
-		FStatMessage const& Item = Data[Index];
+		FStatMessage const& Item = CondensedMessages[Index];
 
 		//Need to get thread root first regardless of filter
 		if(OptionalOutThreadBreakdownMap)
@@ -1206,16 +1221,28 @@ void FStatsThreadState::GetInclusiveAggregateStackStats(int64 TargetFrame, TArra
 	}
 }
 
-void FStatsThreadState::GetExclusiveAggregateStackStats(int64 TargetFrame, TArray<FStatMessage>& OutStats, IItemFiler* Filter, bool bAddNonStackStats /*= true*/) const
+void FStatsThreadState::GetExclusiveAggregateStackStats( 
+	int64 TargetFrame, 
+	TArray<FStatMessage>& OutStats, 
+	IItemFiler* Filter /*= nullptr*/, 
+	bool bAddNonStackStats /*= true*/ ) const
 {
-	TMap<FName, FStatMessage> ThisFrameMetaData;
-	TArray<FStatMessage> const& Data = GetCondensedHistory(TargetFrame);
+	const TArray<FStatMessage>& CondensedMessages = GetCondensedHistory( TargetFrame );
+	GetExclusiveAggregateStackStats( CondensedMessages, OutStats, Filter, bAddNonStackStats );
+}
 
+void FStatsThreadState::GetExclusiveAggregateStackStats( 
+	const TArray<FStatMessage>& CondensedMessages, 
+	TArray<FStatMessage>& OutStats, 
+	IItemFiler* Filter /*= nullptr*/, 
+	bool bAddNonStackStats /*= true */ ) const
+{
+	TMap<FName, FStatMessage> ThisFrameMetaData;
 	TArray<FStatMessage> ChildDurationStack;
 
-	for (int32 Index = 0; Index < Data.Num(); Index++)
+	for (int32 Index = 0; Index < CondensedMessages.Num(); Index++)
 	{
-		FStatMessage const& Item = Data[Index];
+		FStatMessage const& Item = CondensedMessages[Index];
 		FName LongName = Item.NameAndInfo.GetRawName();
 
 		EStatOperation::Type Op = Item.NameAndInfo.GetField<EStatOperation>();
@@ -1264,7 +1291,7 @@ void FStatsThreadState::GetExclusiveAggregateStackStats(int64 TargetFrame, TArra
 	}
 }
 
-TArray<FStatMessage> const& FStatsThreadState::GetCondensedHistory(int64 TargetFrame) const
+TArray<FStatMessage> const& FStatsThreadState::GetCondensedHistory( int64 TargetFrame ) const
 {
 	check(IsFrameValid(TargetFrame));
 
@@ -1437,9 +1464,22 @@ void FStatsThreadState::GetRawStackStats(int64 TargetFrame, FRawStatStackNode& R
 	}
 }
 
-void FStatsThreadState::UncondenseStackStats(int64 TargetFrame, FRawStatStackNode& Root, IItemFiler* Filter, TArray<FStatMessage>* OutNonStackStats) const
+void FStatsThreadState::UncondenseStackStats(
+	int64 TargetFrame, 
+	FRawStatStackNode& Root, 
+	IItemFiler* Filter /*= nullptr*/, 
+	TArray<FStatMessage>* OutNonStackStats /*= nullptr */) const
+{
+	const TArray<FStatMessage>& CondensedMessages = GetCondensedHistory( TargetFrame );
+	UncondenseStackStats( CondensedMessages, Root, Filter, OutNonStackStats );
+}
+
+void FStatsThreadState::UncondenseStackStats( 
+	const TArray<FStatMessage>& CondensedMessages, 
+	FRawStatStackNode& Root, 
+	IItemFiler* Filter /*= nullptr*/, 
+	TArray<FStatMessage>* OutNonStackStats /*= nullptr */ ) const
 {
-	TArray<FStatMessage> const& Data = GetCondensedHistory(TargetFrame);
 	TMap<FName, FStatMessage> ThisFrameNonStackStats;
 
 	{
@@ -1447,9 +1487,9 @@ void FStatsThreadState::UncondenseStackStats(int64 TargetFrame, FRawStatStackNod
 		Stack.Add(&Root);
 		FRawStatStackNode* Current = Stack.Last();
 
-		for (int32 Index = 0; Index < Data.Num(); Index++)
+		for (int32 Index = 0; Index < CondensedMessages.Num(); Index++)
 		{
-			FStatMessage const& Item = Data[Index];
+			FStatMessage const& Item = CondensedMessages[Index];
 			if (!Filter || Filter->Keep(Item))
 			{
 				EStatOperation::Type Op = Item.NameAndInfo.GetField<EStatOperation>();
@@ -1585,9 +1625,8 @@ FName FStatsThreadState::GetStatThreadName( const FStatPacket& Packet ) const
 
 void FStatsThreadState::Condense(int64 TargetFrame, TArray<FStatMessage>& OutStats) const
 {
-	static FStatNameAndInfo Adv(NAME_AdvanceFrame, "", "", TEXT(""), EStatDataType::ST_int64, true, false);
-	new (OutStats) FStatMessage(Adv.GetEncodedName(), EStatOperation::AdvanceFrameEventGameThread, TargetFrame, false);
-	new (OutStats) FStatMessage(Adv.GetEncodedName(), EStatOperation::AdvanceFrameEventRenderThread, TargetFrame, false);
+	new (OutStats) FStatMessage(FStatConstants::AdvanceFrame.GetEncodedName(), EStatOperation::AdvanceFrameEventGameThread, TargetFrame, false);
+	new (OutStats) FStatMessage(FStatConstants::AdvanceFrame.GetEncodedName(), EStatOperation::AdvanceFrameEventRenderThread, TargetFrame, false);
 	FRawStatStackNode Root;
 	GetRawStackStats(TargetFrame, Root, &OutStats);
 	TArray<FStatMessage> StackStats;
diff --git a/Engine/Source/Runtime/Core/Private/Stats/StatsFile.cpp b/Engine/Source/Runtime/Core/Private/Stats/StatsFile.cpp
--- a/Engine/Source/Runtime/Core/Private/Stats/StatsFile.cpp
+++ b/Engine/Source/Runtime/Core/Private/Stats/StatsFile.cpp
@@ -15,6 +15,8 @@
 DECLARE_CYCLE_STAT( TEXT( "Stream File" ), STAT_StreamFile, STATGROUP_StatSystem );
 DECLARE_CYCLE_STAT( TEXT( "Wait For Write" ), STAT_StreamFileWaitForWrite, STATGROUP_StatSystem );
 
+#define LOCTEXT_NAMESPACE "StatsFile"
+
 /*-----------------------------------------------------------------------------
 	FAsyncStatsWrite
 -----------------------------------------------------------------------------*/
@@ -51,13 +53,16 @@ public:
 		FArchive& Ar = *Outer->File;
 
 		// Seek to the end of the file.
-		Ar.Seek( Ar.TotalSize() );
+		const int64 TotalSize = Ar.TotalSize();
+		// Update the file size.
+		Outer->FileSize = TotalSize;
+		Ar.Seek( TotalSize );
 		const int64 FrameFileOffset = Ar.Tell();
 
 		FCompressedStatsData CompressedData( Data, Outer->CompressedData );
 		Ar << CompressedData;
 
-		Outer->FinalizeSavingData(FrameFileOffset);
+		Outer->FinalizeSavingData( FrameFileOffset );
 	}
 
 	FORCEINLINE TStatId GetStatId() const
@@ -78,125 +83,31 @@ FStatsThreadState::FStatsThreadState( FString const& Filename )
 	, MaxNumStatMessages( 0 )
 	, bFindMemoryExtensiveStats( false )
 	, CurrentGameFrame( -1 )
-	, CurrentRenderFrame( -1 )	
+	, CurrentRenderFrame( -1 )
 	, MaxFrameSeen( -1 )
 	, MinFrameSeen( -1 )
 	, bWasLoaded( true )
-{
-	const int64 Size = IFileManager::Get().FileSize( *Filename );
-	if( Size < 4 )
-	{
-		UE_LOG( LogStats, Error, TEXT( "Could not open: %s" ), *Filename );
-		return;
-	}
-	TAutoPtr<FArchive> FileReader( IFileManager::Get().CreateFileReader( *Filename ) );
-	if( !FileReader )
-	{
-		UE_LOG( LogStats, Error, TEXT( "Could not open: %s" ), *Filename );
-		return;
-	}
-
-	FStatsReadStream Stream;
-	if( !Stream.ReadHeader( *FileReader ) )
-	{
-		UE_LOG( LogStats, Error, TEXT( "Could not open, bad magic: %s" ), *Filename );
-		return;
-	}
-
-	// Test version only works for the finalized stats files.
-	const bool bIsFinalized = Stream.Header.IsFinalized();
-	check( bIsFinalized );
-	
-	TArray<FStatMessage> Messages;
-	if( Stream.Header.bRawStatsFile )
-	{
-		int32 NumReadStatPackets = 0;
-
-		// Read metadata.
-		TArray<FStatMessage> MetadataMessages;
-		Stream.ReadFNamesAndMetadataMessages( *FileReader, MetadataMessages );
-		ProcessMetaDataForLoad( MetadataMessages );
-
-		// Read the raw stats messages.
-		const bool bHasCompressedData = Stream.Header.HasCompressedData();
-		check(bHasCompressedData);
-
-		// Buffer used to store the compressed and decompressed data.
-		TArray<uint8> SrcData;
-		TArray<uint8> DestData;
-		FCompressedStatsData UncompressedData( SrcData, DestData );
-		while( true )
-		{
-			// Read the compressed data.		
-			*FileReader << UncompressedData;
-
-			if( UncompressedData.HasReachedEndOfCompressedData() )
-			{
-				// EOF
-				break;
-		}
-
-			// Select the proper archive.
-			FMemoryReader MemoryReader( DestData, true );
-			FArchive& Archive = bHasCompressedData ? MemoryReader : *FileReader;
-
-			FStatPacket StatPacket;
-			Stream.ReadStatPacket( MemoryReader, StatPacket );
-			// Process the raw stat messages.
-			NumReadStatPackets++;
-			}
-	
-		UE_LOG( LogStats, Log, TEXT( "File: %s has %i stat packets" ), *Filename, NumReadStatPackets );
-	}
-	else
-	{
-		// Read the condensed stats messages.
-		while( FileReader->Tell() < Size )
-		{
-			FStatMessage Read( Stream.ReadMessage( *FileReader ) );
-			if( Read.NameAndInfo.GetField<EStatOperation>() == EStatOperation::SpecialMessageMarker )
-			{
-				// Simply break the loop.
-				// The profiler supports more advanced handling of this message.
-				break;
-			}
-			else if( Read.NameAndInfo.GetField<EStatOperation>() == EStatOperation::AdvanceFrameEventGameThread )
-			{
-				ProcessMetaDataForLoad( Messages );
-				if( CurrentGameFrame > 0 && Messages.Num() )
-				{
-					check( !CondensedStackHistory.Contains( CurrentGameFrame ) );
-					TArray<FStatMessage>* Save = new TArray<FStatMessage>();
-					Exchange( *Save, Messages );
-					CondensedStackHistory.Add( CurrentGameFrame, Save );
-					GoodFrames.Add( CurrentGameFrame );
-				}
-			}
-
-			new (Messages)FStatMessage( Read );
-		}
-		// meh, we will discard the last frame, but we will look for meta data
-	}
-}
+{}
 
-void FStatsThreadState::AddMessages( TArray<FStatMessage>& InMessages )
+// #YRX_remove 2015-11-17 
+void FStatsThreadState::AddCondensedMessages( TArray<FStatMessage>& CondensedMessages )
 {
 	bWasLoaded = true;
 	TArray<FStatMessage> Messages;
-	for( int32 Index = 0; Index < InMessages.Num(); ++Index )
+	for (int32 Index = 0; Index < CondensedMessages.Num(); ++Index)
 	{
-		if( InMessages[Index].NameAndInfo.GetField<EStatOperation>() == EStatOperation::AdvanceFrameEventGameThread )
+		if (CondensedMessages[Index].NameAndInfo.GetField<EStatOperation>() == EStatOperation::AdvanceFrameEventGameThread)
 		{
 			ProcessMetaDataForLoad( Messages );
-			if( !CondensedStackHistory.Contains( CurrentGameFrame ) && Messages.Num() )
+			if (!CondensedStackHistory.Contains( CurrentGameFrame ) && Messages.Num())
 			{
 				TArray<FStatMessage>* Save = new TArray<FStatMessage>();
 				Exchange( *Save, Messages );
-				if( CondensedStackHistory.Num() >= HistoryFrames )
+				if (CondensedStackHistory.Num() >= HistoryFrames)
 				{
-					for( auto It = CondensedStackHistory.CreateIterator(); It; ++It )
+					for (auto& It : CondensedStackHistory)
 					{
-						delete It.Value();
+						delete It.Value;
 					}
 					CondensedStackHistory.Reset();
 				}
@@ -205,25 +116,26 @@ void FStatsThreadState::AddMessages( TArray<FStatMessage>& InMessages )
 			}
 		}
 
-		new (Messages)FStatMessage( InMessages[Index] );
+		new (Messages)FStatMessage( CondensedMessages[Index] );
 	}
-	bWasLoaded = false;
+	//bWasLoaded = false;
 }
 
-void FStatsThreadState::ProcessMetaDataForLoad(TArray<FStatMessage>& Data)
+// #YRX_remove 2015-11-17 
+void FStatsThreadState::ProcessMetaDataForLoad( TArray<FStatMessage>& CondensedMessages )
 {
-	check(bWasLoaded);
-	for (int32 Index = 0; Index < Data.Num() ; Index++)
+	check( bWasLoaded );
+	for (int32 Index = 0; Index < CondensedMessages.Num(); Index++)
 	{
-		FStatMessage& Item = Data[Index];
+		FStatMessage& Item = CondensedMessages[Index];
 		EStatOperation::Type Op = Item.NameAndInfo.GetField<EStatOperation>();
 		if (Op == EStatOperation::SetLongName)
 		{
-			FindOrAddMetaData(Item);
+			FindOrAddMetaData( Item );
 		}
 		else if (Op == EStatOperation::AdvanceFrameEventGameThread)
 		{
-			check(Item.NameAndInfo.GetField<EStatDataType>() == EStatDataType::ST_int64);
+			check( Item.NameAndInfo.GetField<EStatDataType>() == EStatDataType::ST_int64 );
 			if (Item.GetValue_int64() > 0)
 			{
 				CurrentGameFrame = Item.GetValue_int64();
@@ -239,30 +151,117 @@ void FStatsThreadState::ProcessMetaDataForLoad(TArray<FStatMessage>& Data)
 		}
 		else if (Op == EStatOperation::AdvanceFrameEventRenderThread)
 		{
-			check(Item.NameAndInfo.GetField<EStatDataType>() == EStatDataType::ST_int64);
+			check( Item.NameAndInfo.GetField<EStatDataType>() == EStatDataType::ST_int64 );
 			if (Item.GetValue_int64() > 0)
 			{
 				CurrentRenderFrame = Item.GetValue_int64();
-				if (CurrentGameFrame > MaxFrameSeen)
+				if (CurrentRenderFrame > MaxFrameSeen)
 				{
-					MaxFrameSeen = CurrentGameFrame;
+					MaxFrameSeen = CurrentRenderFrame;
 				}
 				if (MinFrameSeen < 0)
 				{
-					MinFrameSeen = CurrentGameFrame;
+					MinFrameSeen = CurrentRenderFrame;
 				}
 			}
 		}
 	}
+
+	//UE_LOG( LogStats, Log, TEXT( "Current GameFrame: %lli, RenderFrame: %lli" ), CurrentGameFrame, CurrentRenderFrame );
+}
+
+void FStatsThreadState::AddFrameFromCondensedMessages( TArray<FStatMessage>& CondensedMessages )
+{
+	// @see FStatsThreadState::Condense
+	AdvanceFrameForLoad( CondensedMessages );
+
+	const int64 MinFrameToKeep = CurrentGameFrame - HistoryFrames;
+
+	if (!CondensedStackHistory.Contains( CurrentGameFrame ))
+	{
+		TArray<FStatMessage>* Save = new TArray<FStatMessage>();
+		Exchange( *Save, CondensedMessages );
+
+		for (auto It = CondensedStackHistory.CreateIterator(); It; ++It)
+		{
+			const int64 ThisFrame = It.Key();
+			if (ThisFrame < MinFrameToKeep)
+			{
+				delete It.Value();
+				It.RemoveCurrent();
+			}
+		}
+		for (auto It = GoodFrames.CreateIterator(); It; ++It)
+		{
+			const int64 ThisFrame = *It;
+			if (!CondensedStackHistory.Contains( ThisFrame ))
+			{
+				It.RemoveCurrent();
+			}
+		}
+
+		CondensedStackHistory.Add( CurrentGameFrame, Save );
+		GoodFrames.Add( CurrentGameFrame );
+	}
+}
+
+void FStatsThreadState::AdvanceFrameForLoad( TArray<FStatMessage>& CondensedMessages )
+{
+	// @see FStatsThreadState::Condense
+	check( CondensedMessages.Num() >= 2 );
+	const FStatMessage& AdvanceGameThread = CondensedMessages[0];
+
+	// Sanity checks.
+	check( AdvanceGameThread.NameAndInfo.GetField<EStatOperation>() == EStatOperation::AdvanceFrameEventGameThread );
+	check( AdvanceGameThread.NameAndInfo.GetField<EStatDataType>() == EStatDataType::ST_int64 );
+	check( AdvanceGameThread.GetValue_int64() > 0 );
+
+	CurrentGameFrame = AdvanceGameThread.GetValue_int64();
+	CurrentRenderFrame = CurrentGameFrame;
+
+	if (CurrentGameFrame > MaxFrameSeen)
+	{
+		MaxFrameSeen = CurrentGameFrame;
+	}
+	if (MinFrameSeen < 0)
+	{
+		MinFrameSeen = CurrentGameFrame;
+	}
+}
+
+/*-----------------------------------------------------------------------------
+	FStastsWriteStream
+-----------------------------------------------------------------------------*/
+
+void FStatsWriteStream::WriteMetadata( FArchive& Ar )
+{
+	const FStatsThreadState& Stats = FStatsThreadState::GetLocalState();
+	for (const auto& It : Stats.ShortNameToLongName)
+	{
+		const FStatMessage& StatMessage = It.Value;
+		WriteMessage( Ar, StatMessage );
+	}
+}
+
+void FStatsWriteStream::WriteCondensedMessages( FArchive& Ar, int64 TargetFrame )
+{
+	const FStatsThreadState& Stats = FStatsThreadState::GetLocalState();
+	const TArray<FStatMessage>& Data = Stats.GetCondensedHistory( TargetFrame );
+	for (const FStatMessage& It : Data)
+	{
+		WriteMessage( Ar, It );
+	}
 }
 
 /*-----------------------------------------------------------------------------
 	IStatsWriteFile
 -----------------------------------------------------------------------------*/
 
-IStatsWriteFile::IStatsWriteFile() 
+IStatsWriteFile::IStatsWriteFile()
 	: File( nullptr )
 	, AsyncTask( nullptr )
+	, FileSize( 0 )
+	, StartTime( FPlatformTime::Seconds() )
 {
 	// Reserve 1MB.
 	CompressedData.Reserve( EStatsFileConstants::MAX_COMPRESSED_SIZE );
@@ -282,7 +281,7 @@ void IStatsWriteFile::Start( const FString& InFilename )
 	UE_LOG( LogStats, Log, TEXT( "Opening stats file: %s" ), *TempFilename );
 
 	File = IFileManager::Get().CreateFileWriter( *TempFilename );
-	if( !File )
+	if (!File)
 	{
 		UE_LOG( LogStats, Error, TEXT( "Could not open: %s" ), *TempFilename );
 	}
@@ -290,7 +289,7 @@ void IStatsWriteFile::Start( const FString& InFilename )
 	{
 		ArchiveFilename = Filename;
 		WriteHeader();
-		SetDataDelegate(true);
+		SetDataDelegate( true );
 		StatsMasterEnableAdd();
 	}
 }
@@ -298,10 +297,10 @@ void IStatsWriteFile::Start( const FString& InFilename )
 
 void IStatsWriteFile::Stop()
 {
-	if( IsValid() )
+	if (IsValid())
 	{
 		StatsMasterEnableSubtract();
-		SetDataDelegate(false);
+		SetDataDelegate( false );
 		SendTask();
 		SendTask();
 		Finalize();
@@ -311,9 +310,9 @@ void IStatsWriteFile::Stop()
 		File = nullptr;
 
 		const FString TempFilename = ArchiveFilename + STATSFILE_TEMPORARY_FILENAME_SUFFIX;
-		if (!IFileManager::Get().Move(*ArchiveFilename, *TempFilename))
+		if (!IFileManager::Get().Move( *ArchiveFilename, *TempFilename ))
 		{
-			UE_LOG(LogStats, Warning, TEXT("Could not rename stats file: %s to final name %s"), *TempFilename, *ArchiveFilename);
+			UE_LOG( LogStats, Warning, TEXT( "Could not rename stats file: %s to final name %s" ), *TempFilename, *ArchiveFilename );
 		}
 
 		UE_LOG( LogStats, Log, TEXT( "Wrote stats file: %s" ), *ArchiveFilename );
@@ -321,6 +320,14 @@ void IStatsWriteFile::Stop()
 	}
 }
 
+FText IStatsWriteFile::GetFileMetaDesc() const
+{
+	const FTimespan Duration = FTimespan( 0, 0, FPlatformTime::Seconds() - StartTime );
+
+	const FText FileMetaDesc = FText::Format( LOCTEXT( "FileMetaDesc_Fmt", "STATS FILE: Duration: {0}, Filesize: {1}" ), FText::AsTimespan( Duration ), FText::AsMemory( (SIZE_T)FileSize ) );
+	return FileMetaDesc;
+}
+
 #undef STATSFILE_TEMPORARY_FILENAME_SUFFIX
 
 void IStatsWriteFile::WriteHeader()
@@ -342,16 +349,6 @@ void IStatsWriteFile::WriteHeader()
 	Ar.Flush();
 }
 
-void IStatsWriteFile::WriteMetadata( FArchive& Ar )
-{
-	const FStatsThreadState& Stats = FStatsThreadState::GetLocalState();
-	for( const auto& It : Stats.ShortNameToLongName )
-	{
-		const FStatMessage& StatMessage = It.Value;
-		WriteMessage( Ar, StatMessage );
-	}
-	}
-
 void IStatsWriteFile::Finalize()
 {
 	FArchive& Ar = *File;
@@ -369,7 +366,7 @@ void IStatsWriteFile::Finalize()
 	const FStatsThreadState& Stats = FStatsThreadState::GetLocalState();
 
 	// Add FNames from the stats metadata.
-	for( const auto& It : Stats.ShortNameToLongName )
+	for (const auto& It : Stats.ShortNameToLongName)
 	{
 		const FStatMessage& StatMessage = It.Value;
 		FNamesSent.Add( StatMessage.NameAndInfo.GetRawName().GetComparisonIndex() );
@@ -382,9 +379,9 @@ void IStatsWriteFile::Finalize()
 	// Serialize FNames.
 	Header.FNameTableOffset = Ar.Tell();
 	Header.NumFNames = FNamesToSent.Num();
-	for( const int32 It : FNamesToSent )
+	for (const int32 It : FNamesToSent)
 	{
-		WriteFName( Ar, FStatNameAndInfo(FName(It, It, 0),false) );
+		WriteFName( Ar, FStatNameAndInfo( FName( It, It, 0 ), false ) );
 	}
 
 	// Serialize metadata messages.
@@ -399,30 +396,30 @@ void IStatsWriteFile::Finalize()
 		static TArray<FName> GetFNameArray( const TSet<int32>& NameIndices )
 		{
 			TArray<FName> Result;
-			for( const int32 NameIndex : NameIndices )
+			for (const int32 NameIndex : NameIndices)
 			{
-				new(Result) FName( NameIndex, NameIndex, 0 );
+				new(Result)FName( NameIndex, NameIndex, 0 );
 			}
 			return Result;
 		}
 	};
-	auto BMinANames = FLocal::GetFNameArray( BMinA );
+	TArray<FName> BMinANames = FLocal::GetFNameArray( BMinA );
 
 	// Seek to the position just after a magic value of the file and write out proper header.
-	Ar.Seek( sizeof(uint32) );
+	Ar.Seek( sizeof( uint32 ) );
 	Ar << Header;
 }
 
 void IStatsWriteFile::SendTask()
 {
-	if( AsyncTask )
+	if (AsyncTask)
 	{
 		SCOPE_CYCLE_COUNTER( STAT_StreamFileWaitForWrite );
 		AsyncTask->EnsureCompletion();
 		delete AsyncTask;
 		AsyncTask = nullptr;
 	}
-	if( OutData.Num() )
+	if (OutData.Num())
 	{
 		AsyncTask = new FAsyncTask<FAsyncStatsWrite>( this );
 		check( !OutData.Num() );
@@ -449,8 +446,7 @@ void FStatsWriteFile::SetDataDelegate( bool bSet )
 	}
 }
 
-
-void FStatsWriteFile::WriteFrame( int64 TargetFrame, bool bNeedFullMetadata )
+void FStatsWriteFile::WriteFrame( int64 TargetFrame )
 {
 	// #YRX_STATS: 2015-06-17 Add stat startfile -num=number of frames to capture
 
@@ -458,27 +454,17 @@ void FStatsWriteFile::WriteFrame( int64 TargetFrame, bool bNeedFullMetadata )
 
 	FMemoryWriter Ar( OutData, false, true );
 
-	if( bNeedFullMetadata )
-	{
-		WriteMetadata( Ar );
-	}
-
-	FStatsThreadState const& Stats = FStatsThreadState::GetLocalState();
-	TArray<FStatMessage> const& Data = Stats.GetCondensedHistory( TargetFrame );
-	for( auto It = Data.CreateConstIterator(); It; ++It )
-	{
-		WriteMessage( Ar, *It );
-	}
+	WriteCondensedMessages( Ar, TargetFrame );
 
 	// Get cycles for all threads, so we can use that data to generate the mini-view.
-	for( auto It = Stats.Threads.CreateConstIterator(); It; ++It )
+	const FStatsThreadState& Stats = FStatsThreadState::GetLocalState();
+	for (const auto& It : Stats.Threads)
 	{
-		const int64 Cycles = Stats.GetFastThreadFrameTime( TargetFrame, It.Key() );
-		ThreadCycles.Add( It.Key(), Cycles );
+		const int64 Cycles = Stats.GetFastThreadFrameTime( TargetFrame, It.Key );
+		ThreadCycles.Add( It.Key, Cycles );
 	}
-	
-	// Serialize thread cycles. Disabled for now.
-	//Ar << ThreadCycles;
+
+	SendTask();
 }
 
 void FStatsWriteFile::FinalizeSavingData( int64 FrameFileOffset )
@@ -494,10 +480,10 @@ void FStatsWriteFile::FinalizeSavingData( int64 FrameFileOffset )
 void FRawStatsWriteFile::SetDataDelegate( bool bSet )
 {
 	FStatsThreadState const& Stats = FStatsThreadState::GetLocalState();
-	if( bSet )
+	if (bSet)
 	{
 		DataDelegateHandle = Stats.NewRawStatPacket.AddRaw( this, &FRawStatsWriteFile::WriteRawStatPacket );
-		if( !bWrittenOffsetToData )
+		if (!bWrittenOffsetToData)
 		{
 			const int64 FrameFileOffset = File->Tell();
 			FramesInfo.Add( FStatsFrameInfo( FrameFileOffset ) );
@@ -530,7 +516,7 @@ void FRawStatsWriteFile::WriteStatPacket( FArchive& Ar, FStatPacket& StatPacket
 	// We must handle stat messages in a different way.
 	int32 NumMessages = StatPacket.StatMessages.Num();
 	Ar << NumMessages;
-	for( int32 MessageIndex = 0; MessageIndex < NumMessages; ++MessageIndex )
+	for (int32 MessageIndex = 0; MessageIndex < NumMessages; ++MessageIndex)
 	{
 		WriteMessage( Ar, StatPacket.StatMessages[MessageIndex] );
 	}
@@ -540,16 +526,16 @@ void FRawStatsWriteFile::WriteStatPacket( FArchive& Ar, FStatPacket& StatPacket
 	FAsyncRawStatsReadFile
 -----------------------------------------------------------------------------*/
 
-FAsyncRawStatsFile::FAsyncRawStatsFile( FStatsReadFile* InOwner )
+FAsyncStatsFile::FAsyncStatsFile( FStatsReadFile* InOwner )
 	: Owner( InOwner )
 {}
 
-void FAsyncRawStatsFile::DoWork()
+void FAsyncStatsFile::DoWork()
 {
 	Owner->ReadAndProcessSynchronously();
 }
 
-void FAsyncRawStatsFile::Abandon()
+void FAsyncStatsFile::Abandon()
 {
 	Owner->RequestStop();
 }
@@ -560,18 +546,6 @@ void FAsyncRawStatsFile::Abandon()
 
 const double FStatsReadFile::NumSecondsBetweenUpdates = 2.0;
 
-FStatsReadFile* FStatsReadFile::CreateReaderForRegularStats( const TCHAR* Filename )
-{
-	FStatsReadFile* StatsReadFile = new FStatsReadFile( Filename, false );
-	const bool bValid = StatsReadFile->PrepareLoading();
-	if (!bValid)
-	{
-		delete StatsReadFile;
-	}
-	return bValid ? StatsReadFile : nullptr;
-}
-
-
 void FStatsReadFile::ReadAndProcessSynchronously()
 {
 	// Read.
@@ -594,11 +568,13 @@ void FStatsReadFile::ReadAndProcessSynchronously()
 
 void FStatsReadFile::ReadAndProcessAsynchronously()
 {
-	if (bRawStatsFile)
-	{
-		AsyncWork = new FAsyncTask<FAsyncRawStatsFile>( this );
-		AsyncWork->StartBackgroundTask();
-	}
+	AsyncWork = new FAsyncTask<FAsyncStatsFile>( this );
+	AsyncWork->StartBackgroundTask();
+}
+
+void FStatsReadFile::SetHistoryFrames( int32 InHistoryFrames )
+{
+	State.HistoryFrames = InHistoryFrames;
 }
 
 FStatsReadFile::FStatsReadFile( const TCHAR* InFilename, bool bInRawStatsFile )
@@ -607,6 +583,7 @@ FStatsReadFile::FStatsReadFile( const TCHAR* InFilename, bool bInRawStatsFile )
 	, AsyncWork( nullptr )
 	, LastUpdateTime( 0.0 )
 	, Filename( InFilename )
+	, NumFrames( 0 )
 	, bRawStatsFile( bInRawStatsFile )
 {
 
@@ -652,7 +629,14 @@ bool FStatsReadFile::PrepareLoading()
 			return false;
 		}
 
-		//UE_LOG( LogStats, Warning, TEXT( "Reading a raw stats file for memory profiling: %s" ), *Filename );
+		// Verify requested stats type.
+		const bool bValidStatType = Stream.Header.bRawStatsFile == bRawStatsFile;
+		if (!bValidStatType)
+		{
+			UE_LOG( LogStats, Error, TEXT( "Could not read, invalid stats type: %s" ), *Filename );
+			bResult = false;
+			return false;
+		}
 
 		const bool bIsFinalized = Stream.Header.IsFinalized();
 		if (!bIsFinalized)
@@ -686,12 +670,12 @@ bool FStatsReadFile::PrepareLoading()
 	State.ProcessMetaDataOnly( MetadataMessages );
 
 	// Find all UObject metadata messages.
-	for (const auto& Meta : MetadataMessages)
+	for (const FStatMessage& Meta : MetadataMessages)
 	{
 		const FName EncName = Meta.NameAndInfo.GetEncodedName();
 		const FName RawName = Meta.NameAndInfo.GetRawName();
 		const FString Desc = FStatNameAndInfo::GetShortNameFrom( RawName ).GetPlainNameString();
-		const bool bContainsUObject = Desc.Contains( TEXT( "//" ) );
+		const bool bContainsUObject = Desc.Contains( TEXT( "/" ) );
 		if (bContainsUObject)
 		{
 			UObjectRawNames.Add( RawName );
@@ -708,6 +692,11 @@ bool FStatsReadFile::PrepareLoading()
 	const double TotalTime = FPlatformTime::Seconds() - StartTime;
 	UE_LOG( LogStats, Log, TEXT( "Prepare loading took %.2f sec(s)" ), TotalTime );
 
+	if (!bRawStatsFile)
+	{
+		NumFrames = Stream.FramesInfo.Num();
+	}
+
 	return true;
 }
 
@@ -734,85 +723,171 @@ FStatsReadFile::~FStatsReadFile()
 
 void FStatsReadFile::ReadStats()
 {
+	const double StartTime = FPlatformTime::Seconds();
+
 	if (bRawStatsFile)
 	{
-		const double StartTime = FPlatformTime::Seconds();
+		ReadRawStats();
+	}
+	else
+	{
+		ReadRegularStats();
+	}
 
-		// Buffer used to store the compressed and decompressed data.
-		TArray<uint8> SrcArray;
-		TArray<uint8> DestArray;
+	StageProgress.Set( 100 );
 
-		// Read all packets sequentially, forced by the memory profiler which is now a part of the raw stats.
-		// !!CAUTION!! Frame number in the raw stats is pointless, because it is time/cycles based, not frame based.
-		// Background threads usually execute time consuming operations, so the frame number won't be valid.
-		// Needs to be combined by the thread and the time, not by the frame number.
+	if (!IsProcessingStopped())
+	{
+		const double TotalTime = FPlatformTime::Seconds() - StartTime;
+		UE_LOG( LogStats, Log, TEXT( "Reading took %.2f sec(s)" ), TotalTime );
 
-		// Update stage progress once per NumSecondsBetweenUpdates(2) seconds to avoid spamming.
-		SetProcessingStage( EStatsProcessingStage::SPS_ReadStats );
+		if (bRawStatsFile)
+		{
+			UpdateCombinedHistoryStats();
+		}
+	}
+	else
+	{
+		UE_LOG( LogStats, Warning, TEXT( "Reading stopped, abandoning" ) );
 
-		while (Reader->Tell() < Reader->TotalSize())
+		if (bRawStatsFile)
 		{
-			// Read the compressed data.
-			FCompressedStatsData UncompressedData( SrcArray, DestArray );
-			*Reader << UncompressedData;
-			if (UncompressedData.HasReachedEndOfCompressedData())
-			{
-				StageProgress.Set( 100 );
-				break;
-			}
+			// Clear all data.
+			CombinedHistory.Empty();
+		}
+	}
+}
 
-			FMemoryReader MemoryReader( DestArray, true );
+void FStatsReadFile::ReadRawStats()
+{
+	// Buffer used to store the compressed and decompressed data.
+	TArray<uint8> SrcArray;
+	TArray<uint8> DestArray;
 
-			FStatPacket* StatPacket = new FStatPacket();
-			Stream.ReadStatPacket( MemoryReader, *StatPacket );
+	// Read all packets sequentially, forced by the memory profiler which is now a part of the raw stats.
+	// !!CAUTION!! Frame number in the raw stats is pointless, because it is time/cycles based, not frame based.
+	// Background threads usually execute time consuming operations, so the frame number won't be valid.
+	// Needs to be combined by the thread and the time, not by the frame number.
 
-			const int32 StatPacketFrameNum = int32(StatPacket->Frame);
-			FStatPacketArray& Frame = CombinedHistory.FindOrAdd( StatPacketFrameNum );
+	// Update stage progress once per NumSecondsBetweenUpdates(2) seconds to avoid spamming.
+	SetProcessingStage( EStatsProcessingStage::SPS_ReadStats );
 
-			// Check if we need to combine packets from the same thread.
-			FStatPacket** CombinedPacket = Frame.Packets.FindByPredicate( [&]( FStatPacket* Item ) -> bool
-			{
-				return Item->ThreadId == StatPacket->ThreadId;
-			} );
-			
-			if (CombinedPacket)
-			{
-				(*CombinedPacket)->StatMessages += StatPacket->StatMessages;
-				delete StatPacket;
-			}
-			else
-			{
-				Frame.Packets.Add( StatPacket );
-				FileInfo.MaximumPacketSize = FMath::Max<int32>( FileInfo.MaximumPacketSize, StatPacket->StatMessages.GetAllocatedSize() );
-			}
-
-			UpdateReadStageProgress();
-			if (IsProcessingStopped())
-			{
-				break;
-			}
-			
-			FileInfo.TotalPacketsNum++;
+	while (Reader->Tell() < Reader->TotalSize())
+	{
+		// Read the compressed data.
+		FCompressedStatsData UncompressedData( SrcArray, DestArray );
+		*Reader << UncompressedData;
+		if (UncompressedData.HasReachedEndOfCompressedData())
+		{
+			StageProgress.Set( 100 );
+			break;
 		}
 
-		// Generate frames array.
-		CombinedHistory.GenerateKeyArray( Frames );
-		Frames.Sort();
-		// Verify that frames are sequential.
-		check( Frames[Frames.Num() - 1] == Frames.Num() );
+		FMemoryReader MemoryReader( DestArray, true );
 
-		if (!IsProcessingStopped())
+		FStatPacket* StatPacket = new FStatPacket();
+		Stream.ReadStatPacket( MemoryReader, *StatPacket );
+
+		const int32 StatPacketFrameNum = int32( StatPacket->Frame );
+		FStatPacketArray& Frame = CombinedHistory.FindOrAdd( StatPacketFrameNum );
+
+		// Check if we need to combine packets from the same thread.
+		FStatPacket** CombinedPacket = Frame.Packets.FindByPredicate( [&]( FStatPacket* Item ) -> bool
 		{
-			const double TotalTime = FPlatformTime::Seconds() - StartTime;
-			UE_LOG( LogStats, Log, TEXT( "Reading took %.2f sec(s)" ), TotalTime );
+			return Item->ThreadId == StatPacket->ThreadId;
+		} );
 
-			UpdateCombinedHistoryStats();
+		if (CombinedPacket)
+		{
+			(*CombinedPacket)->StatMessages += StatPacket->StatMessages;
+			delete StatPacket;
 		}
 		else
 		{
-			UE_LOG( LogStats, Warning, TEXT( "Reading stopped, abandoning" ) );
-			// Clear all data.
-			CombinedHistory.Empty();
+			Frame.Packets.Add( StatPacket );
+			FileInfo.MaximumPacketSize = FMath::Max<int32>( FileInfo.MaximumPacketSize, StatPacket->StatMessages.GetAllocatedSize() );
+		}
+
+		UpdateReadStageProgress();
+		if (IsProcessingStopped())
+		{
+			break;
+		}
+
+		FileInfo.TotalPacketsNum++;
+	}
+
+	// Generate frames array.
+	CombinedHistory.GenerateKeyArray( Frames );
+	Frames.Sort();
+	// Verify that frames are sequential.
+	check( Frames[Frames.Num() - 1] == Frames.Num() );
+
+	NumFrames = Frames.Num();
+}
+
+void FStatsReadFile::ReadRegularStats()
+{
+	uint64 ReadMessages = 0;
+	TArray<FStatMessage> PendingMessages;
+
+	// Buffer used to store the compressed and decompressed data.
+	TArray<uint8> SrcArray;
+	TArray<uint8> DestArray;
+
+	const bool bHasCompressedData = Stream.Header.HasCompressedData();
+	const bool bIsFinalized = Stream.Header.IsFinalized();
+
+	// Sanity checks.
+	check( bHasCompressedData );
+	check( bIsFinalized );
+
+	// Update stage progress once per NumSecondsBetweenUpdates(2) seconds to avoid spamming.
+	SetProcessingStage( EStatsProcessingStage::SPS_ReadStats );
+
+	// Remove later.
+	//Reader;
+	//Header;
+	//Stream;
+	//State;
+
+	for (const FStatsFrameInfo& It : Stream.FramesInfo)
+	{
+		PendingMessages.Reset();
+		Reader->Seek( It.FrameFileOffset );
+
+		// Read the compressed data.
+		FCompressedStatsData UncompressedData( SrcArray, DestArray );
+		*Reader << UncompressedData;
+
+		// Read all messages from the uncompressed buffer.
+		FMemoryReader MemoryReader( DestArray, true );
+		while (MemoryReader.Tell() < MemoryReader.TotalSize())
+		{
+			// Read the message.
+			FStatMessage Message( Stream.ReadMessage( MemoryReader, bIsFinalized ) );
+			new (PendingMessages)FStatMessage( Message );
+		}
+
+		// Adds a new from from the pending messages, the pending messages will be removed after the call.
+		State.AddFrameFromCondensedMessages( PendingMessages );
+
+		// Call the derived implementation.
+		const TArray<FStatMessage>& CondensedMessages = *State.CondensedStackHistory.FindChecked( State.CurrentGameFrame );
+		ReadStatsFrame( CondensedMessages, State.CurrentGameFrame );
+
+		// TEST
+		/*{
+			// Get the stat stack root and the non frame stats
+			FRawStatStackNode Stack;
+			TArray<FStatMessage> NonFrameStats;
+			State.UncondenseStackStats( CondensedMessages, Stack, nullptr, &NonFrameStats );
+			}*/
+
+		UpdateReadStageProgress();
+		if (IsProcessingStopped())
+		{
+			break;
 		}
 	}
 }
@@ -868,7 +943,7 @@ void FStatsReadFile::ProcessStats()
 					for (int32 Index = 0; Index < NumStatMessages; Index++)
 					{
 						CurrentStatMessageIndex++;
-						
+
 						const FStatMessage& Message = Data[Index];
 						const EStatOperation::Type Op = Message.NameAndInfo.GetField<EStatOperation>();
 						const FName RawName = Message.NameAndInfo.GetRawName();
@@ -898,7 +973,7 @@ void FStatsReadFile::ProcessStats()
 									// Read OperationSequenceTag.
 									Index++; CurrentStatMessageIndex++;
 									const FStatMessage& SequenceTagMessage = Data[Index];
-									const uint32 SequenceTag = SequenceTagMessage.GetValue_int64();		
+									const uint32 SequenceTag = SequenceTagMessage.GetValue_int64();
 
 									//ThreadStats->AddMemoryMessage( GET_STATFNAME( STAT_Memory_AllocPtr ), (uint64)(UPTRINT)Ptr | (uint64)EMemoryOperation::Alloc );
 									//ThreadStats->AddMemoryMessage( GET_STATFNAME( STAT_Memory_AllocSize ), Size );
@@ -1037,7 +1112,7 @@ void FStatsReadFile::UpdateReadStageProgress()
 	{
 		const int32 PercentagePos = int32( 100.0*Reader->Tell() / Reader->TotalSize() );
 		StageProgress.Set( PercentagePos );
-		UE_LOG( LogStats, Verbose, TEXT( "%3i%%, current num frames %4i" ), PercentagePos, CombinedHistory.Num() );
+		UE_LOG( LogStats, Verbose, TEXT( "UpdateReadStageProgress: %3i%%" ), PercentagePos );
 		LastUpdateTime = CurrentSeconds;
 	}
 
@@ -1082,15 +1157,15 @@ void FStatsReadFile::UpdateCombinedHistoryStats()
 }
 
 void FStatsReadFile::UpdateProcessStageProgress( const int32 CurrentStatMessageIndex, const int32 FrameIndex, const int32 PacketIndex )
-{	
+{
 	const double CurrentSeconds = FPlatformTime::Seconds();
 	if (CurrentSeconds > LastUpdateTime + NumSecondsBetweenUpdates)
 	{
 		const int32 PercentagePos = int32( 100.0*CurrentStatMessageIndex / FileInfo.TotalStatMessagesNum );
 		StageProgress.Set( PercentagePos );
 		UE_LOG( LogStats, Verbose, TEXT( "Processing %3i%% (%10i/%10i) stat messages [Frame: %3i, Packet: %2i]" ), PercentagePos, CurrentStatMessageIndex, FileInfo.TotalStatMessagesNum, FrameIndex, PacketIndex );
-		LastUpdateTime = CurrentSeconds;	
-	}	
+		LastUpdateTime = CurrentSeconds;
+	}
 
 	// Abandon support.
 	if (bShouldStopProcessing == true)
@@ -1132,12 +1207,14 @@ void FCommandStatsFile::Stop()
 {
 	if (CurrentStatsFile)
 	{
+		StatFileActiveCounter.Decrement();
+
 		CurrentStatsFile->Stop();
 		delete CurrentStatsFile;
 		CurrentStatsFile = nullptr;
-
-		StatFileActiveCounter.Decrement();
 	}
 }
 
+#undef LOCTEXT_NAMESPACE
+
 #endif // STATS
diff --git a/Engine/Source/Runtime/Core/Public/Android/AndroidPlatform.h b/Engine/Source/Runtime/Core/Public/Android/AndroidPlatform.h
--- a/Engine/Source/Runtime/Core/Public/Android/AndroidPlatform.h
+++ b/Engine/Source/Runtime/Core/Public/Android/AndroidPlatform.h
@@ -52,6 +52,8 @@ typedef FAndroidTypes FPlatformTypes;
 #define FORCEINLINE inline __attribute__ ((always_inline))		/* Force code to be inline */
 #define FORCENOINLINE __attribute__((noinline))					/* Force code to NOT be inline */
 
+#define FUNCTION_CHECK_RETURN(...) __VA_ARGS__ __attribute__ ((warn_unused_result))	/* Wrap a function signature in this to warn that callers should not ignore the return value. */
+
 #define TEXT_HELPER(a,b)	a ## b
 #define TEXT(s)				TEXT_HELPER(L, s)
 
diff --git a/Engine/Source/Runtime/Core/Public/Delegates/Delegate.h b/Engine/Source/Runtime/Core/Public/Delegates/Delegate.h
--- a/Engine/Source/Runtime/Core/Public/Delegates/Delegate.h
+++ b/Engine/Source/Runtime/Core/Public/Delegates/Delegate.h
@@ -2,6 +2,12 @@
 
 #pragma once
 
+#if !(UE_BUILD_SHIPPING || UE_BUILD_TEST)
+	#define USE_DELEGATE_TRYGETBOUNDFUNCTIONNAME 1
+#else
+	#define USE_DELEGATE_TRYGETBOUNDFUNCTIONNAME 0
+#endif
+
 #include "SharedPointer.h"
 #include "WeakObjectPtrTemplates.h"
 #include "IDelegateInstance.h"
diff --git a/Engine/Source/Runtime/Core/Public/Delegates/DelegateBase.h b/Engine/Source/Runtime/Core/Public/Delegates/DelegateBase.h
--- a/Engine/Source/Runtime/Core/Public/Delegates/DelegateBase.h
+++ b/Engine/Source/Runtime/Core/Public/Delegates/DelegateBase.h
@@ -16,11 +16,18 @@
 #endif
 
 
+struct FWeakObjectPtr;
+
+template <typename ObjectPtrType>
+class FMulticastDelegateBase;
+
 /**
  * Base class for unicast delegates.
  */
 class FDelegateBase
 {
+	friend class FMulticastDelegateBase<FWeakObjectPtr>;
+
 public:
 
 	/**
@@ -55,6 +62,28 @@ public:
 		return *this;
 	}
 
+#if USE_DELEGATE_TRYGETBOUNDFUNCTIONNAME
+
+	/**
+	 * Tries to return the name of a bound function.  Returns NAME_None if the delegate is unbound or
+	 * a binding name is unavailable.
+	 *
+	 * Note: Only intended to be used to aid debugging of delegates.
+	 *
+	 * @return The name of the bound function, NAME_None if no name was available.
+	 */
+	FName TryGetBoundFunctionName() const
+	{
+		if (IDelegateInstance* Ptr = GetDelegateInstanceProtected())
+		{
+			return Ptr->TryGetBoundFunctionName();
+		}
+
+		return NAME_None;
+	}
+
+#endif
+
 	/**
 	 * If this is a UFunction or UObject delegate, return the UObject.
 	 *
@@ -62,12 +91,9 @@ public:
 	 */
 	inline class UObject* GetUObject( ) const
 	{
-		if (IDelegateInstance* Ptr = GetDelegateInstance())
+		if (IDelegateInstance* Ptr = GetDelegateInstanceProtected())
 		{
-			if (Ptr->GetType() == EDelegateInstanceType::UFunction || Ptr->GetType() == EDelegateInstanceType::UObjectMethod)
-			{
-				return (class UObject*)Ptr->GetRawUserObject();
-			}
+			return Ptr->GetUObject();
 		}
 
 		return nullptr;
@@ -80,7 +106,7 @@ public:
 	 */
 	inline bool IsBound( ) const
 	{
-		IDelegateInstance* Ptr = GetDelegateInstance();
+		IDelegateInstance* Ptr = GetDelegateInstanceProtected();
 
 		return Ptr && Ptr->IsSafeToExecute();
 	}
@@ -97,7 +123,7 @@ public:
 			return false;
 		}
 
-		IDelegateInstance* Ptr = GetDelegateInstance();
+		IDelegateInstance* Ptr = GetDelegateInstanceProtected();
 
 		return Ptr && Ptr->HasSameObject(InUserObject);
 	}
@@ -107,7 +133,7 @@ public:
 	 */
 	inline void Unbind( )
 	{
-		if (IDelegateInstance* Ptr = GetDelegateInstance())
+		if (IDelegateInstance* Ptr = GetDelegateInstanceProtected())
 		{
 			Ptr->~IDelegateInstance();
 			DelegateAllocator.ResizeAllocation(0, 0, sizeof(AlignedInlineDelegateType));
@@ -121,9 +147,10 @@ public:
 	 * @return The delegate instance.
 	 * @see SetDelegateInstance
 	 */
-	IDelegateInstance* GetDelegateInstance( ) const
+	DEPRECATED(4.11, "GetDelegateInstance has been deprecated - calls to IDelegateInstance::GetUObject() and IDelegateInstance::GetHandle() should call the same functions on the delegate.  Other calls should be reconsidered.")
+	FORCEINLINE IDelegateInstance* GetDelegateInstance( ) const
 	{
-		return DelegateSize ? (IDelegateInstance*)DelegateAllocator.GetAllocation() : nullptr;
+		return GetDelegateInstanceProtected();
 	}
 
 	/**
@@ -134,7 +161,7 @@ public:
 	inline FDelegateHandle GetHandle() const
 	{
 		FDelegateHandle Result;
-		if (IDelegateInstance* Ptr = GetDelegateInstance())
+		if (IDelegateInstance* Ptr = GetDelegateInstanceProtected())
 		{
 			Result = Ptr->GetHandle();
 		}
@@ -142,12 +169,24 @@ public:
 		return Result;
 	}
 
+protected:
+	/**
+	 * Gets the delegate instance.  Not intended for use by user code.
+	 *
+	 * @return The delegate instance.
+	 * @see SetDelegateInstance
+	 */
+	FORCEINLINE IDelegateInstance* GetDelegateInstanceProtected( ) const
+	{
+		return DelegateSize ? (IDelegateInstance*)DelegateAllocator.GetAllocation() : nullptr;
+	}
+
 private:
 	friend void* operator new(size_t Size, FDelegateBase& Base);
 
 	void* Allocate(int32 Size)
 	{
-		if (IDelegateInstance* CurrentInstance = GetDelegateInstance())
+		if (IDelegateInstance* CurrentInstance = GetDelegateInstanceProtected())
 		{
 			CurrentInstance->~IDelegateInstance();
 		}
diff --git a/Engine/Source/Runtime/Core/Public/Delegates/DelegateInstancesImpl.inl b/Engine/Source/Runtime/Core/Public/Delegates/DelegateInstancesImpl.inl
--- a/Engine/Source/Runtime/Core/Public/Delegates/DelegateInstancesImpl.inl
+++ b/Engine/Source/Runtime/Core/Public/Delegates/DelegateInstancesImpl.inl
@@ -105,26 +105,45 @@ public:
 
 	// IDelegateInstance interface
 
+#if USE_DELEGATE_TRYGETBOUNDFUNCTIONNAME
+
+	virtual FName TryGetBoundFunctionName() const override
+	{
+		return NAME_None;
+	}
+
+#endif
+
+	// Deprecated
 	virtual FName GetFunctionName( ) const override
 	{
 		return NAME_None;
 	}
 
+	// Deprecated
 	virtual const void* GetRawMethodPtr( ) const override
 	{
 		return GetRawMethodPtrInternal();
 	}
 
+	// Deprecated
 	virtual const void* GetRawUserObject( ) const override
 	{
 		return GetRawUserObjectInternal();
 	}
 
+	// Deprecated
 	virtual EDelegateInstanceType::Type GetType( ) const override
 	{
 		return SPMode == ESPMode::ThreadSafe ? EDelegateInstanceType::ThreadSafeSharedPointerMethod : EDelegateInstanceType::SharedPointerMethod;
 	}
 
+	virtual UObject* GetUObject( ) const override
+	{
+		return nullptr;
+	}
+
+	// Deprecated
 	virtual bool HasSameObject( const void* InUserObject ) const override
 	{
 		return UserObject.HasSameObject(InUserObject);
@@ -184,6 +203,7 @@ public:
 	}
 #endif
 
+	// Deprecated
 	virtual bool IsSameFunction( const DELEGATE_INSTANCE_INTERFACE_CLASS<FUNC_TEMPLATE_ARGS>& InOtherDelegate ) const override
 	{
 		// NOTE: Payload data is not currently considered when comparing delegate instances.
@@ -248,13 +268,6 @@ protected:
 
 private:
 
-	// Declare ourselves as a friend so we can access other template permutations in IsSameFunction().
-	template<class UserClassNoShadow, FUNC_PAYLOAD_TEMPLATE_DECL_NO_SHADOW, ESPMode SPModeNoShadow> friend class SP_METHOD_DELEGATE_INSTANCE_CLASS;
-
-	// Declare other pointer-based delegates as a friend so IsSameFunction() can compare members.
-	template<class UserClassNoShadow, FUNC_PAYLOAD_TEMPLATE_DECL_NO_SHADOW> friend class UOBJECT_METHOD_DELEGATE_INSTANCE_CLASS;
-	template<class UserClassNoShadow, FUNC_PAYLOAD_TEMPLATE_DECL_NO_SHADOW> friend class RAW_METHOD_DELEGATE_INSTANCE_CLASS;
-
 	// Weak reference to an instance of the user's class which contains a method we would like to call.
 	TWeakPtr<UserClass, SPMode> UserObject;
 
@@ -303,26 +316,45 @@ public:
 
 	// IDelegateInstance interface
 
+#if USE_DELEGATE_TRYGETBOUNDFUNCTIONNAME
+
+	virtual FName TryGetBoundFunctionName() const override
+	{
+		return NAME_None;
+	}
+
+#endif
+
+	// Deprecated
 	virtual FName GetFunctionName( ) const override
 	{
 		return NAME_None;
 	}
 
+	// Deprecated
 	virtual const void* GetRawMethodPtr( ) const override
 	{
 		return GetRawMethodPtrInternal();
 	}
 
+	// Deprecated
 	virtual const void* GetRawUserObject( ) const override
 	{
 		return GetRawUserObjectInternal();
 	}
 
+	// Deprecated
 	virtual EDelegateInstanceType::Type GetType( ) const override
 	{
 		return EDelegateInstanceType::RawMethod;
 	}
 
+	virtual UObject* GetUObject( ) const override
+	{
+		return nullptr;
+	}
+
+	// Deprecated
 	virtual bool HasSameObject( const void* InUserObject ) const override
 	{
 		return UserObject == InUserObject;
@@ -375,6 +407,7 @@ public:
 	}
 #endif
 
+	// Deprecated
 	virtual bool IsSameFunction( const DELEGATE_INSTANCE_INTERFACE_CLASS<FUNC_TEMPLATE_ARGS>& InOtherDelegate ) const override
 	{
 		// NOTE: Payload data is not currently considered when comparing delegate instances.
@@ -424,10 +457,6 @@ protected:
 
 private:
 
-	// Declare other pointer-based delegates as a friend so IsSameFunction() can compare members
-	template<class UserClassNoShadow, FUNC_PAYLOAD_TEMPLATE_DECL_NO_SHADOW, ESPMode SPModeNoShadow> friend class SP_METHOD_DELEGATE_INSTANCE_CLASS;
-	template<class UserClassNoShadow, FUNC_PAYLOAD_TEMPLATE_DECL_NO_SHADOW> friend class UOBJECT_METHOD_DELEGATE_INSTANCE_CLASS;
-
 	// Pointer to the user's class which contains a method we would like to call.
 	UserClass* UserObject;
 
@@ -477,26 +506,45 @@ public:
 
 	// IDelegateInstance interface
 
+#if USE_DELEGATE_TRYGETBOUNDFUNCTIONNAME
+
+	virtual FName TryGetBoundFunctionName() const override
+	{
+		return NAME_None;
+	}
+
+#endif
+
+	// Deprecated
 	virtual FName GetFunctionName( ) const override
 	{
 		return NAME_None;
 	}
 
+	// Deprecated
 	virtual const void* GetRawMethodPtr( ) const override
 	{
 		return GetRawMethodPtrInternal();
 	}
 
+	// Deprecated
 	virtual const void* GetRawUserObject( ) const override
 	{
 		return GetRawUserObjectInternal();
 	}
 
+	// Deprecated
 	virtual EDelegateInstanceType::Type GetType( ) const override
 	{
 		return EDelegateInstanceType::UObjectMethod;
 	}
 
+	virtual UObject* GetUObject( ) const override
+	{
+		return (UObject*)UserObject.Get();
+	}
+
+	// Deprecated
 	virtual bool HasSameObject( const void* InUserObject ) const override
 	{
 		return (UserObject.Get() == InUserObject);
@@ -559,6 +607,7 @@ public:
 	}
 #endif
 
+	// Deprecated
 	virtual bool IsSameFunction( const DELEGATE_INSTANCE_INTERFACE_CLASS<FUNC_TEMPLATE_ARGS>& InOtherDelegate ) const override
 	{
 		// NOTE: Payload data is not currently considered when comparing delegate instances.
@@ -606,10 +655,6 @@ protected:
 
 private:
 
-	// Declare other pointer-based delegates as a friend so IsSameFunction() can compare members
-	template<class UserClassNoShadow, FUNC_PAYLOAD_TEMPLATE_DECL_NO_SHADOW,ESPMode SPModeNoShadow> friend class SP_METHOD_DELEGATE_INSTANCE_CLASS;
-	template<class UserClassNoShadow, FUNC_PAYLOAD_TEMPLATE_DECL_NO_SHADOW> friend class RAW_METHOD_DELEGATE_INSTANCE_CLASS;
-
 	// Pointer to the user's class which contains a method we would like to call.
 	TWeakObjectPtr<UserClass> UserObject;
 
@@ -659,26 +704,45 @@ public:
 
 	// IDelegateInstance interface
 
+#if USE_DELEGATE_TRYGETBOUNDFUNCTIONNAME
+
+	virtual FName TryGetBoundFunctionName() const override
+	{
+		return NAME_None;
+	}
+
+#endif
+
+	// Deprecated
 	virtual FName GetFunctionName( ) const override
 	{
 		return NAME_None;
 	}
 
+	// Deprecated
 	virtual const void* GetRawMethodPtr( ) const override
 	{
 		return *(const void**)&StaticFuncPtr;
 	}
 
+	// Deprecated
 	virtual const void* GetRawUserObject( ) const override
 	{
 		return nullptr;
 	}
 
+	// Deprecated
 	virtual EDelegateInstanceType::Type GetType( ) const override
 	{
 		return EDelegateInstanceType::Raw;
 	}
 
+	virtual UObject* GetUObject( ) const override
+	{
+		return nullptr;
+	}
+
+	// Deprecated
 	virtual bool HasSameObject( const void* UserObject ) const override
 	{
 		// Raw Delegates aren't bound to an object so they can never match
@@ -722,6 +786,7 @@ public:
 	}
 #endif
 
+	// Deprecated
 	virtual bool IsSameFunction( const DELEGATE_INSTANCE_INTERFACE_CLASS<FUNC_TEMPLATE_ARGS>& InOtherDelegate ) const override
 	{
 		// NOTE: Payload data is not currently considered when comparing delegate instances.
@@ -796,11 +861,22 @@ public:
 public:
 	// IDelegateInstance interface
 
+#if USE_DELEGATE_TRYGETBOUNDFUNCTIONNAME
+
+	virtual FName TryGetBoundFunctionName() const override
+	{
+		return NAME_None;
+	}
+
+#endif
+
+	// Deprecated
 	virtual FName GetFunctionName() const override
 	{
 		return NAME_None;
 	}
 
+	// Deprecated
 	virtual const void* GetRawMethodPtr() const override
 	{
 		// casting operator() to void* is not legal C++ if it's a member function
@@ -809,6 +885,7 @@ public:
 		return nullptr;
 	}
 
+	// Deprecated
 	virtual const void* GetRawUserObject() const override
 	{
 		// returning &Functor wouldn't be particularly useful to the comparison code
@@ -817,11 +894,18 @@ public:
 		return nullptr;
 	}
 
+	// Deprecated
 	virtual EDelegateInstanceType::Type GetType() const override
 	{
 		return EDelegateInstanceType::Functor;
 	}
 
+	virtual UObject* GetUObject() const override
+	{
+		return nullptr;
+	}
+
+	// Deprecated
 	virtual bool HasSameObject(const void* UserObject) const override
 	{
 		// Functor Delegates aren't bound to a user object so they can never match
@@ -862,6 +946,7 @@ public:
 	}
 #endif
 
+	// Deprecated
 	virtual bool IsSameFunction(const DELEGATE_INSTANCE_INTERFACE_CLASS<FUNC_TEMPLATE_ARGS>& InOtherDelegate) const override
 	{
 		// There's no nice way to implement this (we don't have the type info necessary to compare against OtherDelegate's Functor)
@@ -943,26 +1028,45 @@ public:
 
 	// IDelegateInstance interface
 
+#if USE_DELEGATE_TRYGETBOUNDFUNCTIONNAME
+
+	virtual FName TryGetBoundFunctionName() const override
+	{
+		return FunctionName;
+	}
+
+#endif
+
+	// Deprecated
 	virtual FName GetFunctionName( ) const override
 	{
 		return FunctionName;
 	}
 
+	// Deprecated
 	virtual const void* GetRawMethodPtr( ) const override
 	{
 		return nullptr;
 	}
 
+	// Deprecated
 	virtual const void* GetRawUserObject( ) const override
 	{
 		return UserObjectPtr.Get();
 	}
 
+	// Deprecated
 	virtual EDelegateInstanceType::Type GetType() const override
 	{
 		return EDelegateInstanceType::UFunction;
 	}
 
+	virtual UObject* GetUObject( ) const override
+	{
+		return (UObject*)UserObjectPtr.Get();
+	}
+
+	// Deprecated
 	virtual bool HasSameObject( const void* InUserObject ) const override
 	{
 		return (UserObjectPtr.Get() == InUserObject);
@@ -1018,6 +1122,7 @@ public:
 	}
 #endif
 
+	// Deprecated
 	virtual bool IsSameFunction( const DELEGATE_INSTANCE_INTERFACE_CLASS<FUNC_TEMPLATE_ARGS>& Other ) const override
 	{
 		// NOTE: Payload data is not currently considered when comparing delegate instances.
diff --git a/Engine/Source/Runtime/Core/Public/Delegates/DelegateInstancesImpl_Variadics.inl b/Engine/Source/Runtime/Core/Public/Delegates/DelegateInstancesImpl_Variadics.inl
--- a/Engine/Source/Runtime/Core/Public/Delegates/DelegateInstancesImpl_Variadics.inl
+++ b/Engine/Source/Runtime/Core/Public/Delegates/DelegateInstancesImpl_Variadics.inl
@@ -56,26 +56,45 @@ public:
 
 	// IDelegateInstance interface
 
+#if USE_DELEGATE_TRYGETBOUNDFUNCTIONNAME
+
+	virtual FName TryGetBoundFunctionName() const override
+	{
+		return FunctionName;
+	}
+
+#endif
+
+	// Deprecated
 	virtual FName GetFunctionName( ) const override
 	{
 		return FunctionName;
 	}
 
+	// Deprecated
 	virtual const void* GetRawMethodPtr( ) const override
 	{
 		return nullptr;
 	}
 
+	// Deprecated
 	virtual const void* GetRawUserObject( ) const override
 	{
 		return UserObjectPtr.Get();
 	}
 
+	// Deprecated
 	virtual EDelegateInstanceType::Type GetType() const override
 	{
 		return EDelegateInstanceType::UFunction;
 	}
 
+	virtual UObject* GetUObject( ) const override
+	{
+		return (UObject*)UserObjectPtr.Get();
+	}
+
+	// Deprecated
 	virtual bool HasSameObject( const void* InUserObject ) const override
 	{
 		return (UserObjectPtr.Get() == InUserObject);
@@ -117,6 +136,7 @@ public:
 		return Handle;
 	}
 
+	// Deprecated
 	virtual bool IsSameFunction( const Super& Other ) const override
 	{
 		// NOTE: Payload data is not currently considered when comparing delegate instances.
@@ -223,26 +243,45 @@ public:
 
 	// IDelegateInstance interface
 
+#if USE_DELEGATE_TRYGETBOUNDFUNCTIONNAME
+
+	virtual FName TryGetBoundFunctionName() const override
+	{
+		return NAME_None;
+	}
+
+#endif
+
+	// Deprecated
 	virtual FName GetFunctionName() const override
 	{
 		return NAME_None;
 	}
 
+	// Deprecated
 	virtual const void* GetRawMethodPtr() const override
 	{
 		return GetRawMethodPtrInternal();
 	}
 
+	// Deprecated
 	virtual const void* GetRawUserObject() const override
 	{
 		return GetRawUserObjectInternal();
 	}
 
+	// Deprecated
 	virtual EDelegateInstanceType::Type GetType() const override
 	{
 		return SPMode == ESPMode::ThreadSafe ? EDelegateInstanceType::ThreadSafeSharedPointerMethod : EDelegateInstanceType::SharedPointerMethod;
 	}
 
+	virtual UObject* GetUObject() const override
+	{
+		return nullptr;
+	}
+
+	// Deprecated
 	virtual bool HasSameObject(const void* InUserObject) const override
 	{
 		return UserObject.HasSameObject(InUserObject);
@@ -287,6 +326,7 @@ public:
 		return Handle;
 	}
 
+	// Deprecated
 	virtual bool IsSameFunction( const Super& InOtherDelegate ) const override
 	{
 		// NOTE: Payload data is not currently considered when comparing delegate instances.
@@ -434,26 +474,45 @@ public:
 
 	// IDelegateInstance interface
 
+#if USE_DELEGATE_TRYGETBOUNDFUNCTIONNAME
+
+	virtual FName TryGetBoundFunctionName() const override
+	{
+		return NAME_None;
+	}
+
+#endif
+
+	// Deprecated
 	virtual FName GetFunctionName( ) const override
 	{
 		return NAME_None;
 	}
 
+	// Deprecated
 	virtual const void* GetRawMethodPtr( ) const override
 	{
 		return GetRawMethodPtrInternal();
 	}
 
+	// Deprecated
 	virtual const void* GetRawUserObject( ) const override
 	{
 		return GetRawUserObjectInternal();
 	}
 
+	// Deprecated
 	virtual EDelegateInstanceType::Type GetType( ) const override
 	{
 		return EDelegateInstanceType::RawMethod;
 	}
 
+	virtual UObject* GetUObject( ) const override
+	{
+		return nullptr;
+	}
+
+	// Deprecated
 	virtual bool HasSameObject( const void* InUserObject ) const override
 	{
 		return UserObject == InUserObject;
@@ -496,6 +555,7 @@ public:
 		return Handle;
 	}
 
+	// Deprecated
 	virtual bool IsSameFunction( const Super& InOtherDelegate ) const override
 	{
 		// NOTE: Payload data is not currently considered when comparing delegate instances.
@@ -618,26 +678,45 @@ public:
 
 	// IDelegateInstance interface
 
+#if USE_DELEGATE_TRYGETBOUNDFUNCTIONNAME
+
+	virtual FName TryGetBoundFunctionName() const override
+	{
+		return NAME_None;
+	}
+
+#endif
+
+	// Deprecated
 	virtual FName GetFunctionName( ) const override
 	{
 		return NAME_None;
 	}
 
+	// Deprecated
 	virtual const void* GetRawMethodPtr( ) const override
 	{
 		return GetRawMethodPtrInternal();
 	}
 
+	// Deprecated
 	virtual const void* GetRawUserObject( ) const override
 	{
 		return GetRawUserObjectInternal();
 	}
 
+	// Deprecated
 	virtual EDelegateInstanceType::Type GetType( ) const override
 	{
 		return EDelegateInstanceType::UObjectMethod;
 	}
 
+	virtual UObject* GetUObject( ) const override
+	{
+		return (UObject*)UserObject.Get();
+	}
+
+	// Deprecated
 	virtual bool HasSameObject( const void* InUserObject ) const override
 	{
 		return (UserObject.Get() == InUserObject);
@@ -686,6 +765,7 @@ public:
 		return Handle;
 	}
 
+	// Deprecated
 	virtual bool IsSameFunction( const Super& InOtherDelegate ) const override
 	{
 		// NOTE: Payload data is not currently considered when comparing delegate instances.
@@ -804,26 +884,45 @@ public:
 
 	// IDelegateInstance interface
 
+#if USE_DELEGATE_TRYGETBOUNDFUNCTIONNAME
+
+	virtual FName TryGetBoundFunctionName() const override
+	{
+		return NAME_None;
+	}
+
+#endif
+
+	// Deprecated
 	virtual FName GetFunctionName( ) const override
 	{
 		return NAME_None;
 	}
 
+	// Deprecated
 	virtual const void* GetRawMethodPtr( ) const override
 	{
 		return *(const void**)&StaticFuncPtr;
 	}
 
+	// Deprecated
 	virtual const void* GetRawUserObject( ) const override
 	{
 		return nullptr;
 	}
 
+	// Deprecated
 	virtual EDelegateInstanceType::Type GetType( ) const override
 	{
 		return EDelegateInstanceType::Raw;
 	}
 
+	virtual UObject* GetUObject( ) const override
+	{
+		return nullptr;
+	}
+
+	// Deprecated
 	virtual bool HasSameObject( const void* UserObject ) const override
 	{
 		// Raw Delegates aren't bound to an object so they can never match
@@ -858,6 +957,7 @@ public:
 		return Handle;
 	}
 
+	// Deprecated
 	virtual bool IsSameFunction( const Super& InOtherDelegate ) const override
 	{
 		// NOTE: Payload data is not currently considered when comparing delegate instances.
@@ -957,11 +1057,22 @@ public:
 
 	// IDelegateInstance interface
 
+#if USE_DELEGATE_TRYGETBOUNDFUNCTIONNAME
+
+	virtual FName TryGetBoundFunctionName() const override
+	{
+		return NAME_None;
+	}
+
+#endif
+
+	// Deprecated
 	virtual FName GetFunctionName() const override
 	{
 		return NAME_None;
 	}
 
+	// Deprecated
 	virtual const void* GetRawMethodPtr() const override
 	{
 		// casting operator() to void* is not legal C++ if it's a member function
@@ -970,6 +1081,7 @@ public:
 		return nullptr;
 	}
 
+	// Deprecated
 	virtual const void* GetRawUserObject() const override
 	{
 		// returning &Functor wouldn't be particularly useful to the comparison code
@@ -978,11 +1090,18 @@ public:
 		return nullptr;
 	}
 
+	// Deprecated
 	virtual EDelegateInstanceType::Type GetType() const override
 	{
 		return EDelegateInstanceType::Functor;
 	}
 
+	virtual UObject* GetUObject() const override
+	{
+		return nullptr;
+	}
+
+	// Deprecated
 	virtual bool HasSameObject(const void* UserObject) const override
 	{
 		// Functor Delegates aren't bound to a user object so they can never match
@@ -1012,6 +1131,7 @@ public:
 		return Handle;
 	}
 
+	// Deprecated
 	virtual bool IsSameFunction(const Super& InOtherDelegate) const override
 	{
 		// There's no nice way to implement this (we don't have the type info necessary to compare against OtherDelegate's Functor)
