Copying //UE4/Dev-Core to //UE4/Main

diff --git a/Engine/Source/Runtime/Core/Public/Delegates/DelegateSignatureImpl.inl b/Engine/Source/Runtime/Core/Public/Delegates/DelegateSignatureImpl.inl
--- a/Engine/Source/Runtime/Core/Public/Delegates/DelegateSignatureImpl.inl
+++ b/Engine/Source/Runtime/Core/Public/Delegates/DelegateSignatureImpl.inl
@@ -184,35 +184,45 @@ public:
 	/**
 	 * Static: Creates a raw C++ pointer global function delegate
 	 */
+	FUNCTION_CHECK_RETURN(
 	inline static DELEGATE_CLASS CreateStatic( typename FStaticDelegate::FFuncPtr InFunc )
+	)
 	{
 		DELEGATE_CLASS Result;
 		FStaticDelegate::Create( Result, InFunc );
 		return Result;
 	}
 	template< PAYLOAD_TEMPLATE_DECL_OneVar >
+	FUNCTION_CHECK_RETURN(
 	inline static DELEGATE_CLASS CreateStatic( typename TStaticDelegate_OneVar< PAYLOAD_TEMPLATE_LIST_OneVar >::FFuncPtr InFunc, PAYLOAD_TEMPLATE_ARGS_OneVar )
+	)
 	{
 		DELEGATE_CLASS Result;
 		TStaticDelegate_OneVar< PAYLOAD_TEMPLATE_LIST_OneVar >::Create( Result, InFunc, PAYLOAD_TEMPLATE_PASSIN_OneVar );
 		return Result;
 	}
 	template< PAYLOAD_TEMPLATE_DECL_TwoVars >
+	FUNCTION_CHECK_RETURN(
 	inline static DELEGATE_CLASS CreateStatic( typename TStaticDelegate_TwoVars< PAYLOAD_TEMPLATE_LIST_TwoVars >::FFuncPtr InFunc, PAYLOAD_TEMPLATE_ARGS_TwoVars )
+	)
 	{
 		DELEGATE_CLASS Result;
 		TStaticDelegate_TwoVars< PAYLOAD_TEMPLATE_LIST_TwoVars >::Create( Result, InFunc, PAYLOAD_TEMPLATE_PASSIN_TwoVars );
 		return Result;
 	}
 	template< PAYLOAD_TEMPLATE_DECL_ThreeVars >
+	FUNCTION_CHECK_RETURN(
 	inline static DELEGATE_CLASS CreateStatic( typename TStaticDelegate_ThreeVars< PAYLOAD_TEMPLATE_LIST_ThreeVars >::FFuncPtr InFunc, PAYLOAD_TEMPLATE_ARGS_ThreeVars )
+	)
 	{
 		DELEGATE_CLASS Result;
 		TStaticDelegate_ThreeVars< PAYLOAD_TEMPLATE_LIST_ThreeVars >::Create( Result, InFunc, PAYLOAD_TEMPLATE_PASSIN_ThreeVars );
 		return Result;
 	}
 	template< PAYLOAD_TEMPLATE_DECL_FourVars >
+	FUNCTION_CHECK_RETURN(
 	inline static DELEGATE_CLASS CreateStatic( typename TStaticDelegate_FourVars< PAYLOAD_TEMPLATE_LIST_FourVars >::FFuncPtr InFunc, PAYLOAD_TEMPLATE_ARGS_FourVars )
+	)
 	{
 		DELEGATE_CLASS Result;
 		TStaticDelegate_FourVars< PAYLOAD_TEMPLATE_LIST_FourVars >::Create( Result, InFunc, PAYLOAD_TEMPLATE_PASSIN_FourVars );
@@ -224,35 +234,45 @@ public:
 	 * technically this works for any functor types, but lambdas are the primary use case
 	 */
 	template<typename FunctorType>
+	FUNCTION_CHECK_RETURN(
 	inline static DELEGATE_CLASS CreateLambda(FunctorType&& InFunctor)
+	)
 	{
 		DELEGATE_CLASS Result;
 		TFunctorDelegate<typename TRemoveReference<FunctorType>::Type>::Create(Result, Forward<FunctorType>(InFunctor));
 		return Result;
 	}
 	template<typename FunctorType, PAYLOAD_TEMPLATE_DECL_OneVar>
+	FUNCTION_CHECK_RETURN(
 	inline static DELEGATE_CLASS CreateLambda(FunctorType&& InFunctor, PAYLOAD_TEMPLATE_ARGS_OneVar)
+	)
 	{
 		DELEGATE_CLASS Result;
 		TFunctorDelegate_OneVar<typename TRemoveReference<FunctorType>::Type, PAYLOAD_TEMPLATE_LIST_OneVar>::Create(Result, Forward<FunctorType>(InFunctor), PAYLOAD_TEMPLATE_PASSIN_OneVar);
 		return Result;
 	}
 	template<typename FunctorType, PAYLOAD_TEMPLATE_DECL_TwoVars>
+	FUNCTION_CHECK_RETURN(
 	inline static DELEGATE_CLASS CreateLambda(FunctorType&& InFunctor, PAYLOAD_TEMPLATE_ARGS_TwoVars)
+	)
 	{
 		DELEGATE_CLASS Result;
 		TFunctorDelegate_TwoVars<typename TRemoveReference<FunctorType>::Type, PAYLOAD_TEMPLATE_LIST_TwoVars>::Create(Result, Forward<FunctorType>(InFunctor), PAYLOAD_TEMPLATE_PASSIN_TwoVars);
 		return Result;
 	}
 	template<typename FunctorType, PAYLOAD_TEMPLATE_DECL_ThreeVars>
+	FUNCTION_CHECK_RETURN(
 	inline static DELEGATE_CLASS CreateLambda(FunctorType&& InFunctor, PAYLOAD_TEMPLATE_ARGS_ThreeVars)
+	)
 	{
 		DELEGATE_CLASS Result;
 		TFunctorDelegate_ThreeVars<typename TRemoveReference<FunctorType>::Type, PAYLOAD_TEMPLATE_LIST_ThreeVars>::Create(Result, Forward<FunctorType>(InFunctor), PAYLOAD_TEMPLATE_PASSIN_ThreeVars);
 		return Result;
 	}
 	template<typename FunctorType, PAYLOAD_TEMPLATE_DECL_FourVars>
+	FUNCTION_CHECK_RETURN(
 	inline static DELEGATE_CLASS CreateLambda(FunctorType&& InFunctor, PAYLOAD_TEMPLATE_ARGS_FourVars)
+	)
 	{
 		DELEGATE_CLASS Result;
 		TFunctorDelegate_FourVars<typename TRemoveReference<FunctorType>::Type, PAYLOAD_TEMPLATE_LIST_FourVars>::Create(Result, Forward<FunctorType>(InFunctor), PAYLOAD_TEMPLATE_PASSIN_FourVars);
@@ -266,70 +286,90 @@ public:
 	 * deleted out from underneath your delegate. Be careful when calling Execute()!
 	 */
 	template< class UserClass >
+	FUNCTION_CHECK_RETURN(
 	inline static DELEGATE_CLASS CreateRaw( UserClass* InUserObject, typename TRawMethodDelegate< UserClass >::FMethodPtr InFunc )
+	)
 	{
 		DELEGATE_CLASS Result;
 		TRawMethodDelegate< UserClass >::Create( Result, InUserObject, InFunc );
 		return Result;
 	}
 	template< class UserClass >
+	FUNCTION_CHECK_RETURN(
 	inline static DELEGATE_CLASS CreateRaw( UserClass* InUserObject, typename TRawMethodDelegate_Const< UserClass >::FMethodPtr InFunc )
+	)
 	{
 		DELEGATE_CLASS Result;
 		TRawMethodDelegate_Const< UserClass >::Create( Result, InUserObject, InFunc );
 		return Result;
 	}
 	template< class UserClass, PAYLOAD_TEMPLATE_DECL_OneVar >
+	FUNCTION_CHECK_RETURN(
 	inline static DELEGATE_CLASS CreateRaw( UserClass* InUserObject, typename TRawMethodDelegate_OneVar< UserClass, PAYLOAD_TEMPLATE_LIST_OneVar >::FMethodPtr InFunc, PAYLOAD_TEMPLATE_ARGS_OneVar )
+	)
 	{
 		DELEGATE_CLASS Result;
 		TRawMethodDelegate_OneVar< UserClass, PAYLOAD_TEMPLATE_LIST_OneVar >::Create( Result, InUserObject, InFunc, PAYLOAD_TEMPLATE_PASSIN_OneVar );
 		return Result;
 	}
 	template< class UserClass, PAYLOAD_TEMPLATE_DECL_OneVar >
+	FUNCTION_CHECK_RETURN(
 	inline static DELEGATE_CLASS CreateRaw( UserClass* InUserObject, typename TRawMethodDelegate_OneVar_Const< UserClass, PAYLOAD_TEMPLATE_LIST_OneVar >::FMethodPtr InFunc, PAYLOAD_TEMPLATE_ARGS_OneVar )
+	)
 	{
 		DELEGATE_CLASS Result;
 		TRawMethodDelegate_OneVar_Const< UserClass, PAYLOAD_TEMPLATE_LIST_OneVar >::Create( Result, InUserObject, InFunc, PAYLOAD_TEMPLATE_PASSIN_OneVar );
 		return Result;
 	}
 	template< class UserClass, PAYLOAD_TEMPLATE_DECL_TwoVars >
+	FUNCTION_CHECK_RETURN(
 	inline static DELEGATE_CLASS CreateRaw( UserClass* InUserObject, typename TRawMethodDelegate_TwoVars< UserClass, PAYLOAD_TEMPLATE_LIST_TwoVars >::FMethodPtr InFunc, PAYLOAD_TEMPLATE_ARGS_TwoVars )
+	)
 	{
 		DELEGATE_CLASS Result;
 		TRawMethodDelegate_TwoVars< UserClass, PAYLOAD_TEMPLATE_LIST_TwoVars >::Create( Result, InUserObject, InFunc, PAYLOAD_TEMPLATE_PASSIN_TwoVars );
 		return Result;
 	}
 	template< class UserClass, PAYLOAD_TEMPLATE_DECL_TwoVars >
+	FUNCTION_CHECK_RETURN(
 	inline static DELEGATE_CLASS CreateRaw( UserClass* InUserObject, typename TRawMethodDelegate_TwoVars_Const< UserClass, PAYLOAD_TEMPLATE_LIST_TwoVars >::FMethodPtr InFunc, PAYLOAD_TEMPLATE_ARGS_TwoVars )
+	)
 	{
 		DELEGATE_CLASS Result;
 		TRawMethodDelegate_TwoVars_Const< UserClass, PAYLOAD_TEMPLATE_LIST_TwoVars >::Create( Result, InUserObject, InFunc, PAYLOAD_TEMPLATE_PASSIN_TwoVars );
 		return Result;
 	}
 	template< class UserClass, PAYLOAD_TEMPLATE_DECL_ThreeVars >
+	FUNCTION_CHECK_RETURN(
 	inline static DELEGATE_CLASS CreateRaw( UserClass* InUserObject, typename TRawMethodDelegate_ThreeVars< UserClass, PAYLOAD_TEMPLATE_LIST_ThreeVars >::FMethodPtr InFunc, PAYLOAD_TEMPLATE_ARGS_ThreeVars )
+	)
 	{
 		DELEGATE_CLASS Result;
 		TRawMethodDelegate_ThreeVars< UserClass, PAYLOAD_TEMPLATE_LIST_ThreeVars >::Create( Result, InUserObject, InFunc, PAYLOAD_TEMPLATE_PASSIN_ThreeVars );
 		return Result;
 	}
 	template< class UserClass, PAYLOAD_TEMPLATE_DECL_ThreeVars >
+	FUNCTION_CHECK_RETURN(
 	inline static DELEGATE_CLASS CreateRaw( UserClass* InUserObject, typename TRawMethodDelegate_ThreeVars_Const< UserClass, PAYLOAD_TEMPLATE_LIST_ThreeVars >::FMethodPtr InFunc, PAYLOAD_TEMPLATE_ARGS_ThreeVars )
+	)
 	{
 		DELEGATE_CLASS Result;
 		TRawMethodDelegate_ThreeVars_Const< UserClass, PAYLOAD_TEMPLATE_LIST_ThreeVars >::Create( Result, InUserObject, InFunc, PAYLOAD_TEMPLATE_PASSIN_ThreeVars );
 		return Result;
 	}
 	template< class UserClass, PAYLOAD_TEMPLATE_DECL_FourVars >
+	FUNCTION_CHECK_RETURN(
 	inline static DELEGATE_CLASS CreateRaw( UserClass* InUserObject, typename TRawMethodDelegate_FourVars< UserClass, PAYLOAD_TEMPLATE_LIST_FourVars >::FMethodPtr InFunc, PAYLOAD_TEMPLATE_ARGS_FourVars )
+	)
 	{
 		DELEGATE_CLASS Result;
 		TRawMethodDelegate_FourVars< UserClass, PAYLOAD_TEMPLATE_LIST_FourVars >::Create( Result, InUserObject, InFunc, PAYLOAD_TEMPLATE_PASSIN_FourVars );
 		return Result;
 	}
 	template< class UserClass, PAYLOAD_TEMPLATE_DECL_FourVars >
+	FUNCTION_CHECK_RETURN(
 	inline static DELEGATE_CLASS CreateRaw( UserClass* InUserObject, typename TRawMethodDelegate_FourVars_Const< UserClass, PAYLOAD_TEMPLATE_LIST_FourVars >::FMethodPtr InFunc, PAYLOAD_TEMPLATE_ARGS_FourVars )
+	)
 	{
 		DELEGATE_CLASS Result;
 		TRawMethodDelegate_FourVars_Const< UserClass, PAYLOAD_TEMPLATE_LIST_FourVars >::Create( Result, InUserObject, InFunc, PAYLOAD_TEMPLATE_PASSIN_FourVars );
@@ -343,70 +383,90 @@ public:
 	 * You can use ExecuteIfBound() to call them.
 	 */
 	template< class UserClass >
+	FUNCTION_CHECK_RETURN(
 	inline static DELEGATE_CLASS CreateSP( const TSharedRef< UserClass, ESPMode::Fast >& InUserObjectRef, typename TSPMethodDelegate< UserClass >::FMethodPtr InFunc )
+	)
 	{
 		DELEGATE_CLASS Result;
 		TSPMethodDelegate< UserClass >::Create( Result, InUserObjectRef, InFunc );
 		return Result;
 	}
 	template< class UserClass >
+	FUNCTION_CHECK_RETURN(
 	inline static DELEGATE_CLASS CreateSP( const TSharedRef< UserClass, ESPMode::Fast >& InUserObjectRef, typename TSPMethodDelegate_Const< UserClass >::FMethodPtr InFunc )
+	)
 	{
 		DELEGATE_CLASS Result;
 		TSPMethodDelegate_Const< UserClass >::Create( Result, InUserObjectRef, InFunc );
 		return Result;
 	}
 	template< class UserClass, PAYLOAD_TEMPLATE_DECL_OneVar >
+	FUNCTION_CHECK_RETURN(
 	inline static DELEGATE_CLASS CreateSP( const TSharedRef< UserClass, ESPMode::Fast >& InUserObjectRef, typename TSPMethodDelegate_OneVar< UserClass, PAYLOAD_TEMPLATE_LIST_OneVar >::FMethodPtr InFunc, PAYLOAD_TEMPLATE_ARGS_OneVar )
+	)
 	{
 		DELEGATE_CLASS Result;
 		TSPMethodDelegate_OneVar< UserClass, PAYLOAD_TEMPLATE_LIST_OneVar >::Create( Result, InUserObjectRef, InFunc, PAYLOAD_TEMPLATE_PASSIN_OneVar );
 		return Result;
 	}
 	template< class UserClass, PAYLOAD_TEMPLATE_DECL_OneVar >
+	FUNCTION_CHECK_RETURN(
 	inline static DELEGATE_CLASS CreateSP( const TSharedRef< UserClass, ESPMode::Fast >& InUserObjectRef, typename TSPMethodDelegate_OneVar_Const< UserClass, PAYLOAD_TEMPLATE_LIST_OneVar >::FMethodPtr InFunc, PAYLOAD_TEMPLATE_ARGS_OneVar )
+	)
 	{
 		DELEGATE_CLASS Result;
 		TSPMethodDelegate_OneVar_Const< UserClass, PAYLOAD_TEMPLATE_LIST_OneVar >::Create( Result, InUserObjectRef, InFunc, PAYLOAD_TEMPLATE_PASSIN_OneVar );
 		return Result;
 	}
 	template< class UserClass, PAYLOAD_TEMPLATE_DECL_TwoVars >
+	FUNCTION_CHECK_RETURN(
 	inline static DELEGATE_CLASS CreateSP( const TSharedRef< UserClass, ESPMode::Fast >& InUserObjectRef, typename TSPMethodDelegate_TwoVars< UserClass, PAYLOAD_TEMPLATE_LIST_TwoVars >::FMethodPtr InFunc, PAYLOAD_TEMPLATE_ARGS_TwoVars )
+	)
 	{
 		DELEGATE_CLASS Result;
 		TSPMethodDelegate_TwoVars< UserClass, PAYLOAD_TEMPLATE_LIST_TwoVars >::Create( Result, InUserObjectRef, InFunc, PAYLOAD_TEMPLATE_PASSIN_TwoVars );
 		return Result;
 	}
 	template< class UserClass, PAYLOAD_TEMPLATE_DECL_TwoVars >
+	FUNCTION_CHECK_RETURN(
 	inline static DELEGATE_CLASS CreateSP( const TSharedRef< UserClass, ESPMode::Fast >& InUserObjectRef, typename TSPMethodDelegate_TwoVars_Const< UserClass, PAYLOAD_TEMPLATE_LIST_TwoVars >::FMethodPtr InFunc, PAYLOAD_TEMPLATE_ARGS_TwoVars )
+	)
 	{
 		DELEGATE_CLASS Result;
 		TSPMethodDelegate_TwoVars_Const< UserClass, PAYLOAD_TEMPLATE_LIST_TwoVars >::Create( Result, InUserObjectRef, InFunc, PAYLOAD_TEMPLATE_PASSIN_TwoVars );
 		return Result;
 	}
 	template< class UserClass, PAYLOAD_TEMPLATE_DECL_ThreeVars >
+	FUNCTION_CHECK_RETURN(
 	inline static DELEGATE_CLASS CreateSP( const TSharedRef< UserClass, ESPMode::Fast >& InUserObjectRef, typename TSPMethodDelegate_ThreeVars< UserClass, PAYLOAD_TEMPLATE_LIST_ThreeVars >::FMethodPtr InFunc, PAYLOAD_TEMPLATE_ARGS_ThreeVars )
+	)
 	{
 		DELEGATE_CLASS Result;
 		TSPMethodDelegate_ThreeVars< UserClass, PAYLOAD_TEMPLATE_LIST_ThreeVars >::Create( Result, InUserObjectRef, InFunc, PAYLOAD_TEMPLATE_PASSIN_ThreeVars );
 		return Result;
 	}
 	template< class UserClass, PAYLOAD_TEMPLATE_DECL_ThreeVars >
+	FUNCTION_CHECK_RETURN(
 	inline static DELEGATE_CLASS CreateSP( const TSharedRef< UserClass, ESPMode::Fast >& InUserObjectRef, typename TSPMethodDelegate_ThreeVars_Const< UserClass, PAYLOAD_TEMPLATE_LIST_ThreeVars >::FMethodPtr InFunc, PAYLOAD_TEMPLATE_ARGS_ThreeVars )
+	)
 	{
 		DELEGATE_CLASS Result;
 		TSPMethodDelegate_ThreeVars_Const< UserClass, PAYLOAD_TEMPLATE_LIST_ThreeVars >::Create( Result, InUserObjectRef, InFunc, PAYLOAD_TEMPLATE_PASSIN_ThreeVars );
 		return Result;
 	}
 	template< class UserClass, PAYLOAD_TEMPLATE_DECL_FourVars >
+	FUNCTION_CHECK_RETURN(
 	inline static DELEGATE_CLASS CreateSP( const TSharedRef< UserClass, ESPMode::Fast >& InUserObjectRef, typename TSPMethodDelegate_FourVars< UserClass, PAYLOAD_TEMPLATE_LIST_FourVars >::FMethodPtr InFunc, PAYLOAD_TEMPLATE_ARGS_FourVars )
+	)
 	{
 		DELEGATE_CLASS Result;
 		TSPMethodDelegate_FourVars< UserClass, PAYLOAD_TEMPLATE_LIST_FourVars >::Create( Result, InUserObjectRef, InFunc, PAYLOAD_TEMPLATE_PASSIN_FourVars );
 		return Result;
 	}
 	template< class UserClass, PAYLOAD_TEMPLATE_DECL_FourVars >
+	FUNCTION_CHECK_RETURN(
 	inline static DELEGATE_CLASS CreateSP( const TSharedRef< UserClass, ESPMode::Fast >& InUserObjectRef, typename TSPMethodDelegate_FourVars_Const< UserClass, PAYLOAD_TEMPLATE_LIST_FourVars >::FMethodPtr InFunc, PAYLOAD_TEMPLATE_ARGS_FourVars )
+	)
 	{
 		DELEGATE_CLASS Result;
 		TSPMethodDelegate_FourVars_Const< UserClass, PAYLOAD_TEMPLATE_LIST_FourVars >::Create( Result, InUserObjectRef, InFunc, PAYLOAD_TEMPLATE_PASSIN_FourVars );
@@ -420,52 +480,72 @@ public:
 	 * You can use ExecuteIfBound() to call them.
 	 */
 	template< class UserClass >
+	FUNCTION_CHECK_RETURN(
 	inline static DELEGATE_CLASS CreateSP( UserClass* InUserObject, typename TSPMethodDelegate< UserClass >::FMethodPtr InFunc )
+	)
 	{
 		return CreateSP( StaticCastSharedRef< UserClass >( InUserObject->AsShared() ), InFunc );
 	}
 	template< class UserClass >
+	FUNCTION_CHECK_RETURN(
 	inline static DELEGATE_CLASS CreateSP( UserClass* InUserObject, typename TSPMethodDelegate_Const< UserClass >::FMethodPtr InFunc )
+	)
 	{
 		return CreateSP( StaticCastSharedRef< UserClass >( InUserObject->AsShared() ), InFunc );
 	}
 	template< class UserClass, PAYLOAD_TEMPLATE_DECL_OneVar >
+	FUNCTION_CHECK_RETURN(
 	inline static DELEGATE_CLASS CreateSP( UserClass* InUserObject, typename TSPMethodDelegate_OneVar< UserClass, PAYLOAD_TEMPLATE_LIST_OneVar >::FMethodPtr InFunc, PAYLOAD_TEMPLATE_ARGS_OneVar )
+	)
 	{
 		return CreateSP( StaticCastSharedRef< UserClass >( InUserObject->AsShared() ), InFunc, PAYLOAD_TEMPLATE_PASSIN_OneVar );
 	}
 	template< class UserClass, PAYLOAD_TEMPLATE_DECL_OneVar >
+	FUNCTION_CHECK_RETURN(
 	inline static DELEGATE_CLASS CreateSP( UserClass* InUserObject, typename TSPMethodDelegate_OneVar_Const< UserClass, PAYLOAD_TEMPLATE_LIST_OneVar >::FMethodPtr InFunc, PAYLOAD_TEMPLATE_ARGS_OneVar )
+	)
 	{
 		return CreateSP( StaticCastSharedRef< UserClass >( InUserObject->AsShared() ), InFunc, PAYLOAD_TEMPLATE_PASSIN_OneVar );
 	}
 	template< class UserClass, PAYLOAD_TEMPLATE_DECL_TwoVars >
+	FUNCTION_CHECK_RETURN(
 	inline static DELEGATE_CLASS CreateSP( UserClass* InUserObject, typename TSPMethodDelegate_TwoVars< UserClass, PAYLOAD_TEMPLATE_LIST_TwoVars >::FMethodPtr InFunc, PAYLOAD_TEMPLATE_ARGS_TwoVars )
+	)
 	{
 		return CreateSP( StaticCastSharedRef< UserClass >( InUserObject->AsShared() ), InFunc, PAYLOAD_TEMPLATE_PASSIN_TwoVars );
 	}
 	template< class UserClass, PAYLOAD_TEMPLATE_DECL_TwoVars >
+	FUNCTION_CHECK_RETURN(
 	inline static DELEGATE_CLASS CreateSP( UserClass* InUserObject, typename TSPMethodDelegate_TwoVars_Const< UserClass, PAYLOAD_TEMPLATE_LIST_TwoVars >::FMethodPtr InFunc, PAYLOAD_TEMPLATE_ARGS_TwoVars )
+	)
 	{
 		return CreateSP( StaticCastSharedRef< UserClass >( InUserObject->AsShared() ), InFunc, PAYLOAD_TEMPLATE_PASSIN_TwoVars );
 	}
 	template< class UserClass, PAYLOAD_TEMPLATE_DECL_ThreeVars >
+	FUNCTION_CHECK_RETURN(
 	inline static DELEGATE_CLASS CreateSP( UserClass* InUserObject, typename TSPMethodDelegate_ThreeVars< UserClass, PAYLOAD_TEMPLATE_LIST_ThreeVars >::FMethodPtr InFunc, PAYLOAD_TEMPLATE_ARGS_ThreeVars )
+	)
 	{
 		return CreateSP( StaticCastSharedRef< UserClass >( InUserObject->AsShared() ), InFunc, PAYLOAD_TEMPLATE_PASSIN_ThreeVars );
 	}
 	template< class UserClass, PAYLOAD_TEMPLATE_DECL_ThreeVars >
+	FUNCTION_CHECK_RETURN(
 	inline static DELEGATE_CLASS CreateSP( UserClass* InUserObject, typename TSPMethodDelegate_ThreeVars_Const< UserClass, PAYLOAD_TEMPLATE_LIST_ThreeVars >::FMethodPtr InFunc, PAYLOAD_TEMPLATE_ARGS_ThreeVars )
+	)
 	{
 		return CreateSP( StaticCastSharedRef< UserClass >( InUserObject->AsShared() ), InFunc, PAYLOAD_TEMPLATE_PASSIN_ThreeVars );
 	}
 	template< class UserClass, PAYLOAD_TEMPLATE_DECL_FourVars >
+	FUNCTION_CHECK_RETURN(
 	inline static DELEGATE_CLASS CreateSP( UserClass* InUserObject, typename TSPMethodDelegate_FourVars< UserClass, PAYLOAD_TEMPLATE_LIST_FourVars >::FMethodPtr InFunc, PAYLOAD_TEMPLATE_ARGS_FourVars )
+	)
 	{
 		return CreateSP( StaticCastSharedRef< UserClass >( InUserObject->AsShared() ), InFunc, PAYLOAD_TEMPLATE_PASSIN_FourVars );
 	}
 	template< class UserClass, PAYLOAD_TEMPLATE_DECL_FourVars >
+	FUNCTION_CHECK_RETURN(
 	inline static DELEGATE_CLASS CreateSP( UserClass* InUserObject, typename TSPMethodDelegate_FourVars_Const< UserClass, PAYLOAD_TEMPLATE_LIST_FourVars >::FMethodPtr InFunc, PAYLOAD_TEMPLATE_ARGS_FourVars )
+	)
 	{
 		return CreateSP( StaticCastSharedRef< UserClass >( InUserObject->AsShared() ), InFunc, PAYLOAD_TEMPLATE_PASSIN_FourVars );
 	}
@@ -477,70 +557,90 @@ public:
 	 * You can use ExecuteIfBound() to call them.
 	 */
 	template< class UserClass >
+	FUNCTION_CHECK_RETURN(
 	inline static DELEGATE_CLASS CreateThreadSafeSP( const TSharedRef< UserClass, ESPMode::ThreadSafe >& InUserObjectRef, typename TThreadSafeSPMethodDelegate< UserClass >::FMethodPtr InFunc )
+	)
 	{
 		DELEGATE_CLASS Result;
 		TThreadSafeSPMethodDelegate< UserClass >::Create( Result, InUserObjectRef, InFunc );
 		return Result;
 	}
 	template< class UserClass >
+	FUNCTION_CHECK_RETURN(
 	inline static DELEGATE_CLASS CreateThreadSafeSP( const TSharedRef< UserClass, ESPMode::ThreadSafe >& InUserObjectRef, typename TThreadSafeSPMethodDelegate_Const< UserClass >::FMethodPtr InFunc )
+	)
 	{
 		DELEGATE_CLASS Result;
 		TThreadSafeSPMethodDelegate_Const< UserClass >::Create( Result, InUserObjectRef, InFunc );
 		return Result;
 	}
 	template< class UserClass, PAYLOAD_TEMPLATE_DECL_OneVar >
+	FUNCTION_CHECK_RETURN(
 	inline static DELEGATE_CLASS CreateThreadSafeSP( const TSharedRef< UserClass, ESPMode::ThreadSafe >& InUserObjectRef, typename TThreadSafeSPMethodDelegate_OneVar< UserClass, PAYLOAD_TEMPLATE_LIST_OneVar >::FMethodPtr InFunc, PAYLOAD_TEMPLATE_ARGS_OneVar )
+	)
 	{
 		DELEGATE_CLASS Result;
 		TThreadSafeSPMethodDelegate_OneVar< UserClass, PAYLOAD_TEMPLATE_LIST_OneVar >::Create( Result, InUserObjectRef, InFunc, PAYLOAD_TEMPLATE_PASSIN_OneVar );
 		return Result;
 	}
 	template< class UserClass, PAYLOAD_TEMPLATE_DECL_OneVar >
+	FUNCTION_CHECK_RETURN(
 	inline static DELEGATE_CLASS CreateThreadSafeSP( const TSharedRef< UserClass, ESPMode::ThreadSafe >& InUserObjectRef, typename TThreadSafeSPMethodDelegate_OneVar_Const< UserClass, PAYLOAD_TEMPLATE_LIST_OneVar >::FMethodPtr InFunc, PAYLOAD_TEMPLATE_ARGS_OneVar )
+	)
 	{
 		DELEGATE_CLASS Result;
 		TThreadSafeSPMethodDelegate_OneVar_Const< UserClass, PAYLOAD_TEMPLATE_LIST_OneVar >::Create( Result, InUserObjectRef, InFunc, PAYLOAD_TEMPLATE_PASSIN_OneVar );
 		return Result;
 	}
 	template< class UserClass, PAYLOAD_TEMPLATE_DECL_TwoVars >
+	FUNCTION_CHECK_RETURN(
 	inline static DELEGATE_CLASS CreateThreadSafeSP( const TSharedRef< UserClass, ESPMode::ThreadSafe >& InUserObjectRef, typename TThreadSafeSPMethodDelegate_TwoVars< UserClass, PAYLOAD_TEMPLATE_LIST_TwoVars >::FMethodPtr InFunc, PAYLOAD_TEMPLATE_ARGS_TwoVars )
+	)
 	{
 		DELEGATE_CLASS Result;
 		TThreadSafeSPMethodDelegate_TwoVars< UserClass, PAYLOAD_TEMPLATE_LIST_TwoVars >::Create( Result, InUserObjectRef, InFunc, PAYLOAD_TEMPLATE_PASSIN_TwoVars );
 		return Result;
 	}
 	template< class UserClass, PAYLOAD_TEMPLATE_DECL_TwoVars >
+	FUNCTION_CHECK_RETURN(
 	inline static DELEGATE_CLASS CreateThreadSafeSP( const TSharedRef< UserClass, ESPMode::ThreadSafe >& InUserObjectRef, typename TThreadSafeSPMethodDelegate_TwoVars_Const< UserClass, PAYLOAD_TEMPLATE_LIST_TwoVars >::FMethodPtr InFunc, PAYLOAD_TEMPLATE_ARGS_TwoVars )
+	)
 	{
 		DELEGATE_CLASS Result;
 		TThreadSafeSPMethodDelegate_TwoVars_Const< UserClass, PAYLOAD_TEMPLATE_LIST_TwoVars >::Create( Result, InUserObjectRef, InFunc, PAYLOAD_TEMPLATE_PASSIN_TwoVars );
 		return Result;
 	}
 	template< class UserClass, PAYLOAD_TEMPLATE_DECL_ThreeVars >
+	FUNCTION_CHECK_RETURN(
 	inline static DELEGATE_CLASS CreateThreadSafeSP( const TSharedRef< UserClass, ESPMode::ThreadSafe >& InUserObjectRef, typename TThreadSafeSPMethodDelegate_ThreeVars< UserClass, PAYLOAD_TEMPLATE_LIST_ThreeVars >::FMethodPtr InFunc, PAYLOAD_TEMPLATE_ARGS_ThreeVars )
+	)
 	{
 		DELEGATE_CLASS Result;
 		TThreadSafeSPMethodDelegate_ThreeVars< UserClass, PAYLOAD_TEMPLATE_LIST_ThreeVars >::Create( Result, InUserObjectRef, InFunc, PAYLOAD_TEMPLATE_PASSIN_ThreeVars );
 		return Result;
 	}
 	template< class UserClass, PAYLOAD_TEMPLATE_DECL_ThreeVars >
+	FUNCTION_CHECK_RETURN(
 	inline static DELEGATE_CLASS CreateThreadSafeSP( const TSharedRef< UserClass, ESPMode::ThreadSafe >& InUserObjectRef, typename TThreadSafeSPMethodDelegate_ThreeVars_Const< UserClass, PAYLOAD_TEMPLATE_LIST_ThreeVars >::FMethodPtr InFunc, PAYLOAD_TEMPLATE_ARGS_ThreeVars )
+	)
 	{
 		DELEGATE_CLASS Result;
 		TThreadSafeSPMethodDelegate_ThreeVars_Const< UserClass, PAYLOAD_TEMPLATE_LIST_ThreeVars >::Create( Result, InUserObjectRef, InFunc, PAYLOAD_TEMPLATE_PASSIN_ThreeVars );
 		return Result;
 	}
 	template< class UserClass, PAYLOAD_TEMPLATE_DECL_FourVars >
+	FUNCTION_CHECK_RETURN(
 	inline static DELEGATE_CLASS CreateThreadSafeSP( const TSharedRef< UserClass, ESPMode::ThreadSafe >& InUserObjectRef, typename TThreadSafeSPMethodDelegate_FourVars< UserClass, PAYLOAD_TEMPLATE_LIST_FourVars >::FMethodPtr InFunc, PAYLOAD_TEMPLATE_ARGS_FourVars )
+	)
 	{
 		DELEGATE_CLASS Result;
 		TThreadSafeSPMethodDelegate_FourVars< UserClass, PAYLOAD_TEMPLATE_LIST_FourVars >::Create( Result, InUserObjectRef, InFunc, PAYLOAD_TEMPLATE_PASSIN_FourVars );
 		return Result;
 	}
 	template< class UserClass, PAYLOAD_TEMPLATE_DECL_FourVars >
+	FUNCTION_CHECK_RETURN(
 	inline static DELEGATE_CLASS CreateThreadSafeSP( const TSharedRef< UserClass, ESPMode::ThreadSafe >& InUserObjectRef, typename TThreadSafeSPMethodDelegate_FourVars_Const< UserClass, PAYLOAD_TEMPLATE_LIST_FourVars >::FMethodPtr InFunc, PAYLOAD_TEMPLATE_ARGS_FourVars )
+	)
 	{
 		DELEGATE_CLASS Result;
 		TThreadSafeSPMethodDelegate_FourVars_Const< UserClass, PAYLOAD_TEMPLATE_LIST_FourVars >::Create( Result, InUserObjectRef, InFunc, PAYLOAD_TEMPLATE_PASSIN_FourVars );
@@ -554,52 +654,72 @@ public:
 	 * You can use ExecuteIfBound() to call them.
 	 */
 	template< class UserClass >
+	FUNCTION_CHECK_RETURN(
 	inline static DELEGATE_CLASS CreateThreadSafeSP( UserClass* InUserObject, typename TThreadSafeSPMethodDelegate< UserClass >::FMethodPtr InFunc )
+	)
 	{
 		return CreateThreadSafeSP( StaticCastSharedRef< UserClass >( InUserObject->AsShared() ), InFunc );
 	}
 	template< class UserClass >
+	FUNCTION_CHECK_RETURN(
 	inline static DELEGATE_CLASS CreateThreadSafeSP( UserClass* InUserObject, typename TThreadSafeSPMethodDelegate_Const< UserClass >::FMethodPtr InFunc )
+	)
 	{
 		return CreateThreadSafeSP( StaticCastSharedRef< UserClass >( InUserObject->AsShared() ), InFunc );
 	}
 	template< class UserClass, PAYLOAD_TEMPLATE_DECL_OneVar >
+	FUNCTION_CHECK_RETURN(
 	inline static DELEGATE_CLASS CreateThreadSafeSP( UserClass* InUserObject, typename TThreadSafeSPMethodDelegate_OneVar< UserClass, PAYLOAD_TEMPLATE_LIST_OneVar >::FMethodPtr InFunc, PAYLOAD_TEMPLATE_ARGS_OneVar )
+	)
 	{
 		return CreateThreadSafeSP( StaticCastSharedRef< UserClass >( InUserObject->AsShared() ), InFunc, PAYLOAD_TEMPLATE_PASSIN_OneVar );
 	}
 	template< class UserClass, PAYLOAD_TEMPLATE_DECL_OneVar >
+	FUNCTION_CHECK_RETURN(
 	inline static DELEGATE_CLASS CreateThreadSafeSP( UserClass* InUserObject, typename TThreadSafeSPMethodDelegate_OneVar_Const< UserClass, PAYLOAD_TEMPLATE_LIST_OneVar >::FMethodPtr InFunc, PAYLOAD_TEMPLATE_ARGS_OneVar )
+	)
 	{
 		return CreateThreadSafeSP( StaticCastSharedRef< UserClass >( InUserObject->AsShared() ), InFunc, PAYLOAD_TEMPLATE_PASSIN_OneVar );
 	}
 	template< class UserClass, PAYLOAD_TEMPLATE_DECL_TwoVars >
+	FUNCTION_CHECK_RETURN(
 	inline static DELEGATE_CLASS CreateThreadSafeSP( UserClass* InUserObject, typename TThreadSafeSPMethodDelegate_TwoVars< UserClass, PAYLOAD_TEMPLATE_LIST_TwoVars >::FMethodPtr InFunc, PAYLOAD_TEMPLATE_ARGS_TwoVars )
+	)
 	{
 		return CreateThreadSafeSP( StaticCastSharedRef< UserClass >( InUserObject->AsShared() ), InFunc, PAYLOAD_TEMPLATE_PASSIN_TwoVars );
 	}
 	template< class UserClass, PAYLOAD_TEMPLATE_DECL_TwoVars >
+	FUNCTION_CHECK_RETURN(
 	inline static DELEGATE_CLASS CreateThreadSafeSP( UserClass* InUserObject, typename TThreadSafeSPMethodDelegate_TwoVars_Const< UserClass, PAYLOAD_TEMPLATE_LIST_TwoVars >::FMethodPtr InFunc, PAYLOAD_TEMPLATE_ARGS_TwoVars )
+	)
 	{
 		return CreateThreadSafeSP( StaticCastSharedRef< UserClass >( InUserObject->AsShared() ), InFunc, PAYLOAD_TEMPLATE_PASSIN_TwoVars );
 	}
 	template< class UserClass, PAYLOAD_TEMPLATE_DECL_ThreeVars >
+	FUNCTION_CHECK_RETURN(
 	inline static DELEGATE_CLASS CreateThreadSafeSP( UserClass* InUserObject, typename TThreadSafeSPMethodDelegate_ThreeVars< UserClass, PAYLOAD_TEMPLATE_LIST_ThreeVars >::FMethodPtr InFunc, PAYLOAD_TEMPLATE_ARGS_ThreeVars )
+	)
 	{
 		return CreateThreadSafeSP( StaticCastSharedRef< UserClass >( InUserObject->AsShared() ), InFunc, PAYLOAD_TEMPLATE_PASSIN_ThreeVars );
 	}
 	template< class UserClass, PAYLOAD_TEMPLATE_DECL_ThreeVars >
+	FUNCTION_CHECK_RETURN(
 	inline static DELEGATE_CLASS CreateThreadSafeSP( UserClass* InUserObject, typename TThreadSafeSPMethodDelegate_ThreeVars_Const< UserClass, PAYLOAD_TEMPLATE_LIST_ThreeVars >::FMethodPtr InFunc, PAYLOAD_TEMPLATE_ARGS_ThreeVars )
+	)
 	{
 		return CreateThreadSafeSP( StaticCastSharedRef< UserClass >( InUserObject->AsShared() ), InFunc, PAYLOAD_TEMPLATE_PASSIN_ThreeVars );
 	}
 	template< class UserClass, PAYLOAD_TEMPLATE_DECL_FourVars >
+	FUNCTION_CHECK_RETURN(
 	inline static DELEGATE_CLASS CreateThreadSafeSP( UserClass* InUserObject, typename TThreadSafeSPMethodDelegate_FourVars< UserClass, PAYLOAD_TEMPLATE_LIST_FourVars >::FMethodPtr InFunc, PAYLOAD_TEMPLATE_ARGS_FourVars )
+	)
 	{
 		return CreateThreadSafeSP( StaticCastSharedRef< UserClass >( InUserObject->AsShared() ), InFunc, PAYLOAD_TEMPLATE_PASSIN_FourVars );
 	}
 	template< class UserClass, PAYLOAD_TEMPLATE_DECL_FourVars >
+	FUNCTION_CHECK_RETURN(
 	inline static DELEGATE_CLASS CreateThreadSafeSP( UserClass* InUserObject, typename TThreadSafeSPMethodDelegate_FourVars_Const< UserClass, PAYLOAD_TEMPLATE_LIST_FourVars >::FMethodPtr InFunc, PAYLOAD_TEMPLATE_ARGS_FourVars )
+	)
 	{
 		return CreateThreadSafeSP( StaticCastSharedRef< UserClass >( InUserObject->AsShared() ), InFunc, PAYLOAD_TEMPLATE_PASSIN_FourVars );
 	}
@@ -611,35 +731,45 @@ public:
 	 * You can use ExecuteIfBound() to call them.
 	 */
 	template< class UObjectTemplate >
+	FUNCTION_CHECK_RETURN(
 	inline static DELEGATE_CLASS CreateUFunction( UObjectTemplate* InUserObject, const FName& InFunctionName )
+	)
 	{
 		DELEGATE_CLASS Result;
 		TUFunctionDelegateBinding< UObjectTemplate >::Create( Result, InUserObject, InFunctionName );
 		return Result;
 	}
 	template< class UObjectTemplate, PAYLOAD_TEMPLATE_DECL_OneVar >
+	FUNCTION_CHECK_RETURN(
 	inline static DELEGATE_CLASS CreateUFunction( UObjectTemplate* InUserObject, const FName& InFunctionName, PAYLOAD_TEMPLATE_ARGS_OneVar )
+	)
 	{
 		DELEGATE_CLASS Result;
 		TUFunctionDelegateBinding_OneVar< UObjectTemplate, PAYLOAD_TEMPLATE_LIST_OneVar >::Create( Result, InUserObject, InFunctionName, PAYLOAD_TEMPLATE_PASSIN_OneVar );
 		return Result;
 	}
 	template< class UObjectTemplate, PAYLOAD_TEMPLATE_DECL_TwoVars >
+	FUNCTION_CHECK_RETURN(
 	inline static DELEGATE_CLASS CreateUFunction( UObjectTemplate* InUserObject, const FName& InFunctionName, PAYLOAD_TEMPLATE_ARGS_TwoVars )
+	)
 	{
 		DELEGATE_CLASS Result;
 		TUFunctionDelegateBinding_TwoVars< UObjectTemplate, PAYLOAD_TEMPLATE_LIST_TwoVars >::Create( Result, InUserObject, InFunctionName, PAYLOAD_TEMPLATE_PASSIN_TwoVars );
 		return Result;
 	}
 	template< class UObjectTemplate, PAYLOAD_TEMPLATE_DECL_ThreeVars >
+	FUNCTION_CHECK_RETURN(
 	inline static DELEGATE_CLASS CreateUFunction( UObjectTemplate* InUserObject, const FName& InFunctionName, PAYLOAD_TEMPLATE_ARGS_ThreeVars )
+	)
 	{
 		DELEGATE_CLASS Result;
 		TUFunctionDelegateBinding_ThreeVars< UObjectTemplate, PAYLOAD_TEMPLATE_LIST_ThreeVars >::Create( Result, InUserObject, InFunctionName, PAYLOAD_TEMPLATE_PASSIN_ThreeVars );
 		return Result;
 	}
 	template< class UObjectTemplate, PAYLOAD_TEMPLATE_DECL_FourVars >
+	FUNCTION_CHECK_RETURN(
 	inline static DELEGATE_CLASS CreateUFunction( UObjectTemplate* InUserObject, const FName& InFunctionName, PAYLOAD_TEMPLATE_ARGS_FourVars )
+	)
 	{
 		DELEGATE_CLASS Result;
 		TUFunctionDelegateBinding_FourVars< UObjectTemplate, PAYLOAD_TEMPLATE_LIST_FourVars >::Create( Result, InUserObject, InFunctionName, PAYLOAD_TEMPLATE_PASSIN_FourVars );
@@ -653,70 +783,90 @@ public:
 	 * You can use ExecuteIfBound() to call them.
 	 */
 	template< class UserClass >
+	FUNCTION_CHECK_RETURN(
 	inline static DELEGATE_CLASS CreateUObject( UserClass* InUserObject, typename TUObjectMethodDelegate< UserClass >::FMethodPtr InFunc )
+	)
 	{
 		DELEGATE_CLASS Result;
 		TUObjectMethodDelegate< UserClass >::Create( Result, InUserObject, InFunc );
 		return Result;
 	}
 	template< class UserClass >
+	FUNCTION_CHECK_RETURN(
 	inline static DELEGATE_CLASS CreateUObject( UserClass* InUserObject, typename TUObjectMethodDelegate_Const< UserClass >::FMethodPtr InFunc )
+	)
 	{
 		DELEGATE_CLASS Result;
 		TUObjectMethodDelegate_Const< UserClass >::Create( Result, InUserObject, InFunc );
 		return Result;
 	}
 	template< class UserClass, PAYLOAD_TEMPLATE_DECL_OneVar >
+	FUNCTION_CHECK_RETURN(
 	inline static DELEGATE_CLASS CreateUObject( UserClass* InUserObject, typename TUObjectMethodDelegate_OneVar< UserClass, PAYLOAD_TEMPLATE_LIST_OneVar >::FMethodPtr InFunc, PAYLOAD_TEMPLATE_ARGS_OneVar )
+	)
 	{
 		DELEGATE_CLASS Result;
 		TUObjectMethodDelegate_OneVar< UserClass, PAYLOAD_TEMPLATE_LIST_OneVar >::Create( Result, InUserObject, InFunc, PAYLOAD_TEMPLATE_PASSIN_OneVar );
 		return Result;
 	}
 	template< class UserClass, PAYLOAD_TEMPLATE_DECL_OneVar >
+	FUNCTION_CHECK_RETURN(
 	inline static DELEGATE_CLASS CreateUObject( UserClass* InUserObject, typename TUObjectMethodDelegate_OneVar_Const< UserClass, PAYLOAD_TEMPLATE_LIST_OneVar >::FMethodPtr InFunc, PAYLOAD_TEMPLATE_ARGS_OneVar )
+	)
 	{
 		DELEGATE_CLASS Result;
 		TUObjectMethodDelegate_OneVar_Const< UserClass, PAYLOAD_TEMPLATE_LIST_OneVar >::Create( Result, InUserObject, InFunc, PAYLOAD_TEMPLATE_PASSIN_OneVar );
 		return Result;
 	}
 	template< class UserClass, PAYLOAD_TEMPLATE_DECL_TwoVars >
+	FUNCTION_CHECK_RETURN(
 	inline static DELEGATE_CLASS CreateUObject( UserClass* InUserObject, typename TUObjectMethodDelegate_TwoVars< UserClass, PAYLOAD_TEMPLATE_LIST_TwoVars >::FMethodPtr InFunc, PAYLOAD_TEMPLATE_ARGS_TwoVars )
+	)
 	{
 		DELEGATE_CLASS Result;
 		TUObjectMethodDelegate_TwoVars< UserClass, PAYLOAD_TEMPLATE_LIST_TwoVars >::Create( Result, InUserObject, InFunc, PAYLOAD_TEMPLATE_PASSIN_TwoVars );
 		return Result;
 	}
 	template< class UserClass, PAYLOAD_TEMPLATE_DECL_TwoVars >
+	FUNCTION_CHECK_RETURN(
 	inline static DELEGATE_CLASS CreateUObject( UserClass* InUserObject, typename TUObjectMethodDelegate_TwoVars_Const< UserClass, PAYLOAD_TEMPLATE_LIST_TwoVars >::FMethodPtr InFunc, PAYLOAD_TEMPLATE_ARGS_TwoVars )
+	)
 	{
 		DELEGATE_CLASS Result;
 		TUObjectMethodDelegate_TwoVars_Const< UserClass, PAYLOAD_TEMPLATE_LIST_TwoVars >::Create( Result, InUserObject, InFunc, PAYLOAD_TEMPLATE_PASSIN_TwoVars );
 		return Result;
 	}
 	template< class UserClass, PAYLOAD_TEMPLATE_DECL_ThreeVars >
+	FUNCTION_CHECK_RETURN(
 	inline static DELEGATE_CLASS CreateUObject( UserClass* InUserObject, typename TUObjectMethodDelegate_ThreeVars< UserClass, PAYLOAD_TEMPLATE_LIST_ThreeVars >::FMethodPtr InFunc, PAYLOAD_TEMPLATE_ARGS_ThreeVars )
+	)
 	{
 		DELEGATE_CLASS Result;
 		TUObjectMethodDelegate_ThreeVars< UserClass, PAYLOAD_TEMPLATE_LIST_ThreeVars >::Create( Result, InUserObject, InFunc, PAYLOAD_TEMPLATE_PASSIN_ThreeVars );
 		return Result;
 	}
 	template< class UserClass, PAYLOAD_TEMPLATE_DECL_ThreeVars >
+	FUNCTION_CHECK_RETURN(
 	inline static DELEGATE_CLASS CreateUObject( UserClass* InUserObject, typename TUObjectMethodDelegate_ThreeVars_Const< UserClass, PAYLOAD_TEMPLATE_LIST_ThreeVars >::FMethodPtr InFunc, PAYLOAD_TEMPLATE_ARGS_ThreeVars )
+	)
 	{
 		DELEGATE_CLASS Result;
 		TUObjectMethodDelegate_ThreeVars_Const< UserClass, PAYLOAD_TEMPLATE_LIST_ThreeVars >::Create( Result, InUserObject, InFunc, PAYLOAD_TEMPLATE_PASSIN_ThreeVars );
 		return Result;
 	}
 	template< class UserClass, PAYLOAD_TEMPLATE_DECL_FourVars >
+	FUNCTION_CHECK_RETURN(
 	inline static DELEGATE_CLASS CreateUObject( UserClass* InUserObject, typename TUObjectMethodDelegate_FourVars< UserClass, PAYLOAD_TEMPLATE_LIST_FourVars >::FMethodPtr InFunc, PAYLOAD_TEMPLATE_ARGS_FourVars )
+	)
 	{
 		DELEGATE_CLASS Result;
 		TUObjectMethodDelegate_FourVars< UserClass, PAYLOAD_TEMPLATE_LIST_FourVars >::Create( Result, InUserObject, InFunc, PAYLOAD_TEMPLATE_PASSIN_FourVars );
 		return Result;
 	}
 	template< class UserClass, PAYLOAD_TEMPLATE_DECL_FourVars >
+	FUNCTION_CHECK_RETURN(
 	inline static DELEGATE_CLASS CreateUObject( UserClass* InUserObject, typename TUObjectMethodDelegate_FourVars_Const< UserClass, PAYLOAD_TEMPLATE_LIST_FourVars >::FMethodPtr InFunc, PAYLOAD_TEMPLATE_ARGS_FourVars )
+	)
 	{
 		DELEGATE_CLASS Result;
 		TUObjectMethodDelegate_FourVars_Const< UserClass, PAYLOAD_TEMPLATE_LIST_FourVars >::Create( Result, InUserObject, InFunc, PAYLOAD_TEMPLATE_PASSIN_FourVars );
@@ -767,7 +917,7 @@ public:
 		if (&Other != this)
 		{
 			// this down-cast is OK! allows for managing invocation list in the base class without requiring virtual functions
-			TDelegateInstanceInterface* OtherInstance = Other.GetDelegateInstance();
+			TDelegateInstanceInterface* OtherInstance = Other.GetDelegateInstanceProtected();
 
 			if (OtherInstance != nullptr)
 			{
@@ -792,7 +942,7 @@ public:
 		if (&Other != this)
 		{
 			// this down-cast is OK! allows for managing invocation list in the base class without requiring virtual functions
-			TDelegateInstanceInterface* OtherInstance = Other.GetDelegateInstance();
+			TDelegateInstanceInterface* OtherInstance = Other.GetDelegateInstanceProtected();
 
 			if (OtherInstance != nullptr)
 			{
@@ -1247,9 +1397,9 @@ public:
 	 *
 	 * @see ExecuteIfBound
 	 */
-	inline RetValType Execute( FUNC_PARAM_LIST ) const
+	FORCEINLINE RetValType Execute( FUNC_PARAM_LIST ) const
 	{
-		TDelegateInstanceInterface* DelegateInstance = GetDelegateInstance();
+		TDelegateInstanceInterface* DelegateInstance = GetDelegateInstanceProtected();
 
 		// If this assert goes off, Execute() was called before a function was bound to the delegate.
 		// Consider using ExecuteIfSafe() instead.
@@ -1269,7 +1419,7 @@ public:
 	{
 		if( IsBound() )
 		{
-			return GetDelegateInstance()->ExecuteIfSafe( FUNC_PARAM_PASSTHRU );
+			return GetDelegateInstanceProtected()->ExecuteIfSafe( FUNC_PARAM_PASSTHRU );
 		}
 
 		return false;
@@ -1279,9 +1429,9 @@ public:
 	/**
 	 * Returns a pointer to the correctly-typed delegate instance.
 	 */
-	TDelegateInstanceInterface* GetDelegateInstance() const
+	FORCEINLINE TDelegateInstanceInterface* GetDelegateInstanceProtected() const
 	{
-		return (TDelegateInstanceInterface*)FDelegateBase::GetDelegateInstance();
+		return (TDelegateInstanceInterface*)FDelegateBase::GetDelegateInstanceProtected();
 	}
 
 private:
@@ -1305,8 +1455,10 @@ private:
  */
 template< FUNC_TEMPLATE_DECL_TYPENAME >
 class BASE_MULTICAST_DELEGATE_CLASS
-	: public FMulticastDelegateBase<>
+	: public FMulticastDelegateBase<FWeakObjectPtr>
 {
+	typedef FMulticastDelegateBase<FWeakObjectPtr> Super;
+
 public:
 
 	/** DEPRECATED: Type definition for unicast delegate classes whose delegate instances are compatible with this delegate. */
@@ -1327,7 +1479,7 @@ public:
 	FDelegateHandle Add( FDelegate&& InNewDelegate )
 	{
 		FDelegateHandle Result;
-		if (InNewDelegate.GetDelegateInstance())
+		if (Super::GetDelegateInstanceProtectedHelper(InNewDelegate))
 		{
 			Result = AddDelegateInstance(MoveTemp(InNewDelegate));
 		}
@@ -1345,7 +1497,7 @@ public:
 	FDelegateHandle Add( const FDelegate& InNewDelegate )
 	{
 		FDelegateHandle Result;
-		if (InNewDelegate.GetDelegateInstance())
+		if (Super::GetDelegateInstanceProtectedHelper(InNewDelegate))
 		{
 			Result = AddDelegateInstance(FDelegate(InNewDelegate));
 		}
@@ -1848,12 +2000,10 @@ protected:
 			for (const FDelegateBase& OtherDelegateRef : Other.GetInvocationList())
 			{
 				// this down-cast is OK! allows for managing invocation list in the base class without requiring virtual functions
-				TDelegateInstanceInterface* OtherInstance = ((const FDelegate&)OtherDelegateRef).GetDelegateInstance();
-
-				if (OtherInstance != nullptr)
+				if (IDelegateInstance* OtherInstance = Super::GetDelegateInstanceProtectedHelper(OtherDelegateRef))
 				{
 					FDelegate Temp;
-					OtherInstance->CreateCopy(Temp);
+					((TDelegateInstanceInterface*)OtherInstance)->CreateCopy(Temp);
 					AddInternal(MoveTemp(Temp));
 				}
 			}
@@ -1873,20 +2023,6 @@ protected:
 	 */
 	FDelegateHandle AddDelegateInstance( FDelegate&& InNewDelegate )
 	{
-#if !UE_BUILD_SHIPPING && !UE_BUILD_TEST
-		// verify that the same function isn't already bound
-		for (const FDelegateBase& DelegateRef : GetInvocationList())
-		{
-			// this down-cast is OK! allows for managing invocation list in the base class without requiring virtual functions
-			TDelegateInstanceInterface* DelegateInstance = ((const FDelegate&)DelegateRef).GetDelegateInstance();
-
-			if (DelegateInstance != nullptr)
-			{
-				check(!DelegateInstance->IsSameFunction(*InNewDelegate.GetDelegateInstance()));
-			}
-		}
-#endif
-
 		return AddInternal(MoveTemp(InNewDelegate));
 	}
 
@@ -1910,8 +2046,8 @@ public:
 				// this down-cast is OK! allows for managing invocation list in the base class without requiring virtual functions
 				const FDelegate& DelegateBase = (const FDelegate&)LocalInvocationList[InvocationListIndex];
 
-				TDelegateInstanceInterface* DelegateInstanceInterface = DelegateBase.GetDelegateInstance();
-				if (DelegateInstanceInterface == nullptr || !DelegateInstanceInterface->ExecuteIfSafe(FUNC_PARAM_PASSTHRU))
+				IDelegateInstance* DelegateInstanceInterface = Super::GetDelegateInstanceProtectedHelper(DelegateBase);
+				if (DelegateInstanceInterface == nullptr || !((TDelegateInstanceInterface*)DelegateInstanceInterface)->ExecuteIfSafe(FUNC_PARAM_PASSTHRU))
 				{
 					NeedsCompaction = true;
 				}
@@ -1949,7 +2085,7 @@ protected:
 
 			// NOTE: We must do a deep compare here, not just compare delegate pointers, because multiple
 			//       delegate pointers can refer to the exact same object and method
-			IDelegateInstance* DelegateInstanceInterface = DelegateBase.GetDelegateInstance();
+			IDelegateInstance* DelegateInstanceInterface = Super::GetDelegateInstanceProtectedHelper(DelegateBase);
 			if ((DelegateInstanceInterface != nullptr) && DelegateInstanceInterface->GetHandle() == Handle)
 			{
 				DelegateBase.Unbind();
diff --git a/Engine/Source/Runtime/Core/Public/Delegates/DelegateSignatureImpl_Variadics.inl b/Engine/Source/Runtime/Core/Public/Delegates/DelegateSignatureImpl_Variadics.inl
--- a/Engine/Source/Runtime/Core/Public/Delegates/DelegateSignatureImpl_Variadics.inl
+++ b/Engine/Source/Runtime/Core/Public/Delegates/DelegateSignatureImpl_Variadics.inl
@@ -94,7 +94,9 @@ public:
 	 * Static: Creates a raw C++ pointer global function delegate
 	 */
 	template <typename... VarTypes>
+	FUNCTION_CHECK_RETURN(
 	inline static TBaseDelegate<RetValType, ParamTypes...> CreateStatic(typename TIdentity<RetValType (*)(ParamTypes..., VarTypes...)>::Type InFunc, VarTypes... Vars)
+	)
 	{
 		TBaseDelegate<RetValType, ParamTypes...> Result;
 		TBaseStaticDelegateInstance<TFuncType, VarTypes...>::Create(Result, InFunc, Vars...);
@@ -106,7 +108,9 @@ public:
 	 * technically this works for any functor types, but lambdas are the primary use case
 	 */
 	template<typename FunctorType, typename... VarTypes>
+	FUNCTION_CHECK_RETURN(
 	inline static TBaseDelegate<RetValType, ParamTypes...> CreateLambda(FunctorType&& InFunctor, VarTypes... Vars)
+	)
 	{
 		TBaseDelegate<RetValType, ParamTypes...> Result;
 		TBaseFunctorDelegateInstance<TFuncType, typename TRemoveReference<FunctorType>::Type, VarTypes...>::Create(Result, Forward<FunctorType>(InFunctor), Vars...);
@@ -120,14 +124,18 @@ public:
 	 * deleted out from underneath your delegate. Be careful when calling Execute()!
 	 */
 	template <typename UserClass, typename... VarTypes>
+	FUNCTION_CHECK_RETURN(
 	inline static TBaseDelegate<RetValType, ParamTypes...> CreateRaw(UserClass* InUserObject, typename TMemFunPtrType<false, UserClass, RetValType (ParamTypes..., VarTypes...)>::Type InFunc, VarTypes... Vars)
+	)
 	{
 		TBaseDelegate<RetValType, ParamTypes...> Result;
 		TBaseRawMethodDelegateInstance<false, UserClass, TFuncType, VarTypes...>::Create(Result, InUserObject, InFunc, Vars...);
 		return Result;
 	}
 	template <typename UserClass, typename... VarTypes>
+	FUNCTION_CHECK_RETURN(
 	inline static TBaseDelegate<RetValType, ParamTypes...> CreateRaw(UserClass* InUserObject, typename TMemFunPtrType<true, UserClass, RetValType (ParamTypes..., VarTypes...)>::Type InFunc, VarTypes... Vars)
+	)
 	{
 		TBaseDelegate<RetValType, ParamTypes...> Result;
 		TBaseRawMethodDelegateInstance<true , UserClass, TFuncType, VarTypes...>::Create(Result, InUserObject, InFunc, Vars...);
@@ -141,14 +149,18 @@ public:
 	 * You can use ExecuteIfBound() to call them.
 	 */
 	template <typename UserClass, typename... VarTypes>
+	FUNCTION_CHECK_RETURN(
 	inline static TBaseDelegate<RetValType, ParamTypes...> CreateSP(const TSharedRef<UserClass, ESPMode::Fast>& InUserObjectRef, typename TMemFunPtrType<false, UserClass, RetValType (ParamTypes..., VarTypes...)>::Type InFunc, VarTypes... Vars)
+	)
 	{
 		TBaseDelegate<RetValType, ParamTypes...> Result;
 		TBaseSPMethodDelegateInstance<false, UserClass, ESPMode::Fast, TFuncType, VarTypes...>::Create(Result, InUserObjectRef, InFunc, Vars...);
 		return Result;
 	}
 	template <typename UserClass, typename... VarTypes>
+	FUNCTION_CHECK_RETURN(
 	inline static TBaseDelegate<RetValType, ParamTypes...> CreateSP(const TSharedRef<UserClass, ESPMode::Fast>& InUserObjectRef, typename TMemFunPtrType<true, UserClass, RetValType (ParamTypes..., VarTypes...)>::Type InFunc, VarTypes... Vars)
+	)
 	{
 		TBaseDelegate<RetValType, ParamTypes...> Result;
 		TBaseSPMethodDelegateInstance<true , UserClass, ESPMode::Fast, TFuncType, VarTypes...>::Create(Result, InUserObjectRef, InFunc, Vars...);
@@ -162,12 +174,16 @@ public:
 	 * You can use ExecuteIfBound() to call them.
 	 */
 	template <typename UserClass, typename... VarTypes>
+	FUNCTION_CHECK_RETURN(
 	inline static TBaseDelegate<RetValType, ParamTypes...> CreateSP(UserClass* InUserObject, typename TMemFunPtrType<false, UserClass, RetValType (ParamTypes..., VarTypes...)>::Type InFunc, VarTypes... Vars)
+	)
 	{
 		return CreateSP(StaticCastSharedRef<UserClass>(InUserObject->AsShared()), InFunc, Vars...);
 	}
 	template <typename UserClass, typename... VarTypes>
+	FUNCTION_CHECK_RETURN(
 	inline static TBaseDelegate<RetValType, ParamTypes...> CreateSP(UserClass* InUserObject, typename TMemFunPtrType<true, UserClass, RetValType (ParamTypes..., VarTypes...)>::Type InFunc, VarTypes... Vars)
+	)
 	{
 		return CreateSP(StaticCastSharedRef<UserClass>(InUserObject->AsShared()), InFunc, Vars...);
 	}
@@ -179,14 +195,18 @@ public:
 	 * You can use ExecuteIfBound() to call them.
 	 */
 	template <typename UserClass, typename... VarTypes>
+	FUNCTION_CHECK_RETURN(
 	inline static TBaseDelegate<RetValType, ParamTypes...> CreateThreadSafeSP(const TSharedRef<UserClass, ESPMode::ThreadSafe>& InUserObjectRef, typename TMemFunPtrType<false, UserClass, RetValType (ParamTypes..., VarTypes...)>::Type InFunc, VarTypes... Vars)
+	)
 	{
 		TBaseDelegate<RetValType, ParamTypes...> Result;
 		TBaseSPMethodDelegateInstance<false, UserClass, ESPMode::ThreadSafe, TFuncType, VarTypes...>::Create(Result, InUserObjectRef, InFunc, Vars...);
 		return Result;
 	}
 	template <typename UserClass, typename... VarTypes>
+	FUNCTION_CHECK_RETURN(
 	inline static TBaseDelegate<RetValType, ParamTypes...> CreateThreadSafeSP(const TSharedRef<UserClass, ESPMode::ThreadSafe>& InUserObjectRef, typename TMemFunPtrType<true, UserClass, RetValType (ParamTypes..., VarTypes...)>::Type InFunc, VarTypes... Vars)
+	)
 	{
 		TBaseDelegate<RetValType, ParamTypes...> Result;
 		TBaseSPMethodDelegateInstance<true , UserClass, ESPMode::ThreadSafe, TFuncType, VarTypes...>::Create(Result, InUserObjectRef, InFunc, Vars...);
@@ -200,12 +220,16 @@ public:
 	 * You can use ExecuteIfBound() to call them.
 	 */
 	template <typename UserClass, typename... VarTypes>
+	FUNCTION_CHECK_RETURN(
 	inline static TBaseDelegate<RetValType, ParamTypes...> CreateThreadSafeSP(UserClass* InUserObject, typename TMemFunPtrType<false, UserClass, RetValType (ParamTypes..., VarTypes...)>::Type InFunc, VarTypes... Vars)
+	)
 	{
 		return CreateThreadSafeSP(StaticCastSharedRef<UserClass>(InUserObject->AsShared()), InFunc, Vars...);
 	}
 	template <typename UserClass, typename... VarTypes>
+	FUNCTION_CHECK_RETURN(
 	inline static TBaseDelegate<RetValType, ParamTypes...> CreateThreadSafeSP(UserClass* InUserObject, typename TMemFunPtrType<true, UserClass, RetValType (ParamTypes..., VarTypes...)>::Type InFunc, VarTypes... Vars)
+	)
 	{
 		return CreateThreadSafeSP(StaticCastSharedRef<UserClass>(InUserObject->AsShared()), InFunc, Vars...);
 	}
@@ -217,7 +241,9 @@ public:
 	 * You can use ExecuteIfBound() to call them.
 	 */
 	template <typename UObjectTemplate, typename... VarTypes>
+	FUNCTION_CHECK_RETURN(
 	inline static TBaseDelegate<RetValType, ParamTypes...> CreateUFunction(UObjectTemplate* InUserObject, const FName& InFunctionName, VarTypes... Vars)
+	)
 	{
 		TBaseDelegate<RetValType, ParamTypes...> Result;
 		TBaseUFunctionDelegateInstance<UObjectTemplate, TFuncType, VarTypes...>::Create(Result, InUserObject, InFunctionName, Vars...);
@@ -231,14 +257,18 @@ public:
 	 * You can use ExecuteIfBound() to call them.
 	 */
 	template <typename UserClass, typename... VarTypes>
+	FUNCTION_CHECK_RETURN(
 	inline static TBaseDelegate<RetValType, ParamTypes...> CreateUObject(UserClass* InUserObject, typename TMemFunPtrType<false, UserClass, RetValType (ParamTypes..., VarTypes...)>::Type InFunc, VarTypes... Vars)
+	)
 	{
 		TBaseDelegate<RetValType, ParamTypes...> Result;
 		TBaseUObjectMethodDelegateInstance<false, UserClass, TFuncType, VarTypes...>::Create(Result, InUserObject, InFunc, Vars...);
 		return Result;
 	}
 	template <typename UserClass, typename... VarTypes>
+	FUNCTION_CHECK_RETURN(
 	inline static TBaseDelegate<RetValType, ParamTypes...> CreateUObject(UserClass* InUserObject, typename TMemFunPtrType<true, UserClass, RetValType (ParamTypes..., VarTypes...)>::Type InFunc, VarTypes... Vars)
+	)
 	{
 		TBaseDelegate<RetValType, ParamTypes...> Result;
 		TBaseUObjectMethodDelegateInstance<true , UserClass, TFuncType, VarTypes...>::Create(Result, InUserObject, InFunc, Vars...);
@@ -299,7 +329,7 @@ public:
 		if (&Other != this)
 		{
 			// this down-cast is OK! allows for managing invocation list in the base class without requiring virtual functions
-			TDelegateInstanceInterface* OtherInstance = Other.GetDelegateInstance();
+			TDelegateInstanceInterface* OtherInstance = Other.GetDelegateInstanceProtected();
 
 			if (OtherInstance != nullptr)
 			{
@@ -324,7 +354,7 @@ public:
 		if (&Other != this)
 		{
 			// this down-cast is OK! allows for managing invocation list in the base class without requiring virtual functions
-			TDelegateInstanceInterface* OtherInstance = Other.GetDelegateInstance();
+			TDelegateInstanceInterface* OtherInstance = Other.GetDelegateInstanceProtected();
 
 			if (OtherInstance != nullptr)
 			{
@@ -480,9 +510,9 @@ public:
 	 *
 	 * @see ExecuteIfBound
 	 */
-	inline RetValType Execute(ParamTypes... Params) const
+	FORCEINLINE RetValType Execute(ParamTypes... Params) const
 	{
-		TDelegateInstanceInterface* LocalDelegateInstance = GetDelegateInstance();
+		TDelegateInstanceInterface* LocalDelegateInstance = GetDelegateInstanceProtected();
 
 		// If this assert goes off, Execute() was called before a function was bound to the delegate.
 		// Consider using ExecuteIfSafe() instead.
@@ -494,9 +524,9 @@ public:
 	/**
 	 * Returns a pointer to the correctly-typed delegate instance.
 	 */
-	TDelegateInstanceInterface* GetDelegateInstance() const
+	FORCEINLINE TDelegateInstanceInterface* GetDelegateInstanceProtected() const
 	{
-		return (TDelegateInstanceInterface*)FDelegateBase::GetDelegateInstance();
+		return (TDelegateInstanceInterface*)FDelegateBase::GetDelegateInstanceProtected();
 	}
 };
 
@@ -575,7 +605,7 @@ public:
 	{
 		if (Super::IsBound())
 		{
-			return Super::GetDelegateInstance()->ExecuteIfSafe(Params...);
+			return Super::GetDelegateInstanceProtected()->ExecuteIfSafe(Params...);
 		}
 
 		return false;
@@ -597,8 +627,10 @@ template <typename RetValType, typename... ParamTypes>
 class TBaseMulticastDelegate;
 
 template <typename... ParamTypes>
-class TBaseMulticastDelegate<void, ParamTypes...> : public FMulticastDelegateBase<>
+class TBaseMulticastDelegate<void, ParamTypes...> : public FMulticastDelegateBase<FWeakObjectPtr>
 {
+	typedef FMulticastDelegateBase<FWeakObjectPtr> Super;
+
 public:
 	/** DEPRECATED: Type definition for unicast delegate classes whose delegate instances are compatible with this delegate. */
 	typedef TBaseDelegate< void, ParamTypes... > FDelegate;
@@ -618,7 +650,7 @@ public:
 	FDelegateHandle Add(FDelegate&& InNewDelegate)
 	{
 		FDelegateHandle Result;
-		if (InNewDelegate.GetDelegateInstance())
+		if (Super::GetDelegateInstanceProtectedHelper(InNewDelegate))
 		{
 			Result = AddDelegateInstance(MoveTemp(InNewDelegate));
 		}
@@ -636,7 +668,7 @@ public:
 	FDelegateHandle Add(const FDelegate& InNewDelegate)
 	{
 		FDelegateHandle Result;
-		if (InNewDelegate.GetDelegateInstance())
+		if (Super::GetDelegateInstanceProtectedHelper(InNewDelegate))
 		{
 			Result = AddDelegateInstance(FDelegate(InNewDelegate));
 		}
@@ -839,13 +871,10 @@ protected:
 
 			for (const FDelegateBase& OtherDelegateRef : Other.GetInvocationList())
 			{
-				// this down-cast is OK! allows for managing invocation list in the base class without requiring virtual functions
-				TDelegateInstanceInterface* OtherInstance = ((const FDelegate&)OtherDelegateRef).GetDelegateInstance();
-
-				if (OtherInstance != nullptr)
+				if (IDelegateInstance* OtherInstance = Super::GetDelegateInstanceProtectedHelper(OtherDelegateRef))
 				{
 					FDelegate Temp;
-					OtherInstance->CreateCopy(Temp);
+					((TDelegateInstanceInterface*)OtherInstance)->CreateCopy(Temp);
 					AddInternal(MoveTemp(Temp));
 				}
 			}
@@ -865,20 +894,6 @@ protected:
 	 */
 	FDelegateHandle AddDelegateInstance(FDelegate&& InNewDelegate)
 	{
-#if !UE_BUILD_SHIPPING && !UE_BUILD_TEST
-		// verify that the same function isn't already bound
-		for (const FDelegateBase& DelegateRef : GetInvocationList())
-		{
-			// this down-cast is OK! allows for managing invocation list in the base class without requiring virtual functions
-			TDelegateInstanceInterface* DelegateInstance = ((const FDelegate&)DelegateRef).GetDelegateInstance();
-
-			if (DelegateInstance != nullptr)
-			{
-				check(!DelegateInstance->IsSameFunction(*InNewDelegate.GetDelegateInstance()));
-			}
-		}
-#endif
-
 		return AddInternal(MoveTemp(InNewDelegate));
 	}
 
@@ -902,8 +917,8 @@ public:
 				// this down-cast is OK! allows for managing invocation list in the base class without requiring virtual functions
 				const FDelegate& DelegateBase = (const FDelegate&)LocalInvocationList[InvocationListIndex];
 
-				TDelegateInstanceInterface* DelegateInstanceInterface = DelegateBase.GetDelegateInstance();
-				if (DelegateInstanceInterface == nullptr || !DelegateInstanceInterface->ExecuteIfSafe(Params...))
+				IDelegateInstance* DelegateInstanceInterface = Super::GetDelegateInstanceProtectedHelper(DelegateBase);
+				if (DelegateInstanceInterface == nullptr || !((TDelegateInstanceInterface*)DelegateInstanceInterface)->ExecuteIfSafe(Params...))
 				{
 					NeedsCompaction = true;
 				}
@@ -936,7 +951,7 @@ protected:
 			// TODO: This is horrible, can we get the base class to do it?
 			FDelegateBase& DelegateBase = const_cast<FDelegateBase&>(LocalInvocationList[InvocationListIndex]);
 
-			IDelegateInstance* DelegateInstanceInterface = DelegateBase.GetDelegateInstance();
+			IDelegateInstance* DelegateInstanceInterface = Super::GetDelegateInstanceProtectedHelper(DelegateBase);
 			if ((DelegateInstanceInterface != nullptr) && DelegateInstanceInterface->GetHandle() == Handle)
 			{
 				DelegateBase.Unbind();
diff --git a/Engine/Source/Runtime/Core/Public/Delegates/IDelegateInstance.h b/Engine/Source/Runtime/Core/Public/Delegates/IDelegateInstance.h
--- a/Engine/Source/Runtime/Core/Public/Delegates/IDelegateInstance.h
+++ b/Engine/Source/Runtime/Core/Public/Delegates/IDelegateInstance.h
@@ -93,17 +93,41 @@ private:
 class IDelegateInstance
 {
 public:
+#if USE_DELEGATE_TRYGETBOUNDFUNCTIONNAME
+
+	/**
+	 * Tries to return the name of a bound function.  Returns NAME_None if the delegate is unbound or
+	 * a binding name is unavailable.
+	 *
+	 * Note: Only intended to be used to aid debugging of delegates.
+	 *
+	 * @return The name of the bound function, NAME_None if no name was available.
+	 */
+	virtual FName TryGetBoundFunctionName() const = 0;
+
+#endif
 
 	/**
 	 * Returns the name of the UFunction that this delegate instance is bound to.
 	 *
+	 * Deprecated.
+	 *
 	 * @return Name of the function, or NAME_None if not bound to a UFunction.
 	 */
 	virtual FName GetFunctionName( ) const = 0;
 
+	/**
+	 * Returns the UObject that this delegate instance is bound to.
+	 *
+	 * @return Pointer to the UObject, or nullptr if not bound to a UObject.
+	 */
+	virtual UObject* GetUObject( ) const = 0;
+
 	/**
 	 * Returns raw pointer to the delegate method.
 	 *
+	 * Deprecated.
+	 *
 	 * @return Raw pointer to the delegate method.
 	 */
 	virtual const void* GetRawMethodPtr( ) const = 0;
@@ -111,6 +135,8 @@ public:
 	/**
 	 * Returns raw pointer to UserObject,
 	 *
+	 * Deprecated.
+	 *
 	 * @return Raw pointer to UserObject.
 	 */
 	virtual const void* GetRawUserObject( ) const = 0;
@@ -118,6 +144,8 @@ public:
 	/**
 	 * Returns the type of delegate instance
 	 *
+	 * Deprecated.
+	 *
 	 * @return Delegate instance type
 	 */
 	virtual EDelegateInstanceType::Type GetType( ) const = 0;
@@ -125,6 +153,8 @@ public:
 	/**
 	 * Returns true if this delegate is bound to the specified UserObject,
 	 *
+	 * Deprecated.
+	 *
 	 * @param InUserObject
 	 *
 	 * @return True if delegate is bound to the specified UserObject
diff --git a/Engine/Source/Runtime/Core/Public/Delegates/MulticastDelegateBase.h b/Engine/Source/Runtime/Core/Public/Delegates/MulticastDelegateBase.h
--- a/Engine/Source/Runtime/Core/Public/Delegates/MulticastDelegateBase.h
+++ b/Engine/Source/Runtime/Core/Public/Delegates/MulticastDelegateBase.h
@@ -8,7 +8,7 @@
  */
 typedef TArray<FDelegateBase, TInlineAllocator<1> > TInvocationList;
 
-template<typename ObjectPtrType = FWeakObjectPtr>
+template<typename ObjectPtrType>
 class FMulticastDelegateBase
 {
 public:
@@ -37,7 +37,7 @@ public:
 	{
 		for (const FDelegateBase& DelegateBaseRef : InvocationList)
 		{
-			if (DelegateBaseRef.GetDelegateInstance())
+			if (DelegateBaseRef.GetDelegateInstanceProtected())
 			{
 				return true;
 			}
@@ -54,7 +54,7 @@ public:
 	{
 		for (const FDelegateBase& DelegateBaseRef : InvocationList)
 		{
-			IDelegateInstance* DelegateInstance = DelegateBaseRef.GetDelegateInstance();
+			IDelegateInstance* DelegateInstance = DelegateBaseRef.GetDelegateInstanceProtected();
 			if ((DelegateInstance != nullptr) && DelegateInstance->HasSameObject(InUserObject))
 			{
 				return true;
@@ -76,7 +76,7 @@ public:
 		{
 			FDelegateBase& DelegateBaseRef = InvocationList[InvocationListIndex];
 
-			IDelegateInstance* DelegateInstance = DelegateBaseRef.GetDelegateInstance();
+			IDelegateInstance* DelegateInstance = DelegateBaseRef.GetDelegateInstanceProtected();
 			if ((DelegateInstance != nullptr) && DelegateInstance->HasSameObject(InUserObject))
 			{
 				DelegateBaseRef.Unbind();
@@ -124,7 +124,7 @@ protected:
 		{
 			FDelegateBase& DelegateBaseRef = InvocationList[InvocationListIndex];
 
-			IDelegateInstance* DelegateInstance = DelegateBaseRef.GetDelegateInstance();
+			IDelegateInstance* DelegateInstance = DelegateBaseRef.GetDelegateInstanceProtected();
 
 			if ((DelegateInstance == nullptr) || DelegateInstance->IsCompactable())
 			{
@@ -157,6 +157,15 @@ protected:
 		--InvocationListLockCount;
 	}
 
+protected:
+	/**
+	 * Helper function for derived classes of FMulticastDelegateBase to get at the delegate instance.
+	 */
+	static FORCEINLINE IDelegateInstance* GetDelegateInstanceProtectedHelper(const FDelegateBase& Base)
+	{
+		return Base.GetDelegateInstanceProtected();
+	}
+
 private:
 
 	/** Holds the collection of delegate instances to invoke. */
diff --git a/Engine/Source/Runtime/Core/Public/GenericPlatform/GenericPlatformFile.h b/Engine/Source/Runtime/Core/Public/GenericPlatform/GenericPlatformFile.h
--- a/Engine/Source/Runtime/Core/Public/GenericPlatform/GenericPlatformFile.h
+++ b/Engine/Source/Runtime/Core/Public/GenericPlatform/GenericPlatformFile.h
@@ -266,6 +266,8 @@ public:
 	/////////// Generally, these do not need to be implemented per platform.
 	/////////////////////////////////////////////////////////////////////////////////////////////////////////
 
+	virtual void GetTimeStampPair(const TCHAR* PathA, const TCHAR* PathB, FDateTime& OutTimeStampA, FDateTime& OutTimeStampB);
+
 	/** 
 	 * Call the Visit function of the visitor once for each file or directory in a directory tree. This function explores subdirectories.
 	 * @param Directory		The directory to iterate the contents of, recursively.
diff --git a/Engine/Source/Runtime/Core/Public/HAL/Allocators/CachedOSPageAllocator.h b/Engine/Source/Runtime/Core/Public/HAL/Allocators/CachedOSPageAllocator.h
new file mode 100644
--- /dev/null
+++ b/Engine/Source/Runtime/Core/Public/HAL/Allocators/CachedOSPageAllocator.h
@@ -0,0 +1,47 @@
+// Copyright 1998-2015 Epic Games, Inc. All Rights Reserved.
+
+#pragma once
+
+struct FCachedOSPageAllocator
+{
+protected:
+	struct FFreePageBlock
+	{
+		void*  Ptr;
+		SIZE_T ByteSize;
+
+		FFreePageBlock() 
+		{
+			Ptr      = nullptr;
+			ByteSize = 0;
+		}
+	};
+
+	void* AllocateImpl(SIZE_T Size, FFreePageBlock* First, FFreePageBlock* Last, uint32& FreedPageBlocksNum, uint32& CachedTotal);
+	void FreeImpl(void* Ptr, SIZE_T Size, uint32 NumCacheBlocks, uint32 CachedByteLimit, FFreePageBlock* First, uint32& FreedPageBlocksNum, uint32& CachedTotal);
+};
+
+template <uint32 NumCacheBlocks, uint32 CachedByteLimit>
+struct TCachedOSPageAllocator : private FCachedOSPageAllocator
+{
+	TCachedOSPageAllocator()
+		: FreedPageBlocksNum(0)
+		, CachedTotal       (0)
+	{
+	}
+
+	FORCEINLINE void* Allocate(SIZE_T Size)
+	{
+		return AllocateImpl(Size, FreedPageBlocks, FreedPageBlocks + FreedPageBlocksNum, FreedPageBlocksNum, CachedTotal);
+	}
+
+	void Free(void* Ptr, SIZE_T Size)
+	{
+		return FreeImpl(Ptr, Size, NumCacheBlocks, CachedByteLimit, FreedPageBlocks, FreedPageBlocksNum, CachedTotal);
+	}
+
+private:
+	FFreePageBlock FreedPageBlocks[NumCacheBlocks];
+	uint32         FreedPageBlocksNum;
+	uint32         CachedTotal;
+};
diff --git a/Engine/Source/Runtime/Core/Public/HAL/FileManager.h b/Engine/Source/Runtime/Core/Public/HAL/FileManager.h
--- a/Engine/Source/Runtime/Core/Public/HAL/FileManager.h
+++ b/Engine/Source/Runtime/Core/Public/HAL/FileManager.h
@@ -166,6 +166,11 @@ public:
 	 */
 	virtual FDateTime GetTimeStamp( const TCHAR* Path ) = 0;
 
+	/**
+	* @return the modification time of the given file (or FDateTime::MinValue() on failure)
+	*/
+	virtual void GetTimeStampPair(const TCHAR* PathA, const TCHAR* PathB, FDateTime& OutTimeStampA, FDateTime& OutTimeStampB) = 0;
+
 	/** 
 	 * Sets the modification time of the given file
 	 */
diff --git a/Engine/Source/Runtime/Core/Public/HAL/FileManagerGeneric.h b/Engine/Source/Runtime/Core/Public/HAL/FileManagerGeneric.h
--- a/Engine/Source/Runtime/Core/Public/HAL/FileManagerGeneric.h
+++ b/Engine/Source/Runtime/Core/Public/HAL/FileManagerGeneric.h
@@ -67,6 +67,7 @@ public:
 	double GetFileAgeSeconds( const TCHAR* Filename ) override;
 	FDateTime GetTimeStamp( const TCHAR* Filename ) override;
 	FDateTime GetAccessTimeStamp( const TCHAR* Filename ) override;
+	void GetTimeStampPair(const TCHAR* PathA, const TCHAR* PathB, FDateTime& OutTimeStampA, FDateTime& OutTimeStampB);
 	bool SetTimeStamp( const TCHAR* Filename, FDateTime Timestamp ) override;
 	virtual FString GetFilenameOnDisk(const TCHAR* Filename) override;
 
diff --git a/Engine/Source/Runtime/Core/Public/HAL/MallocBinned2.h b/Engine/Source/Runtime/Core/Public/HAL/MallocBinned2.h
--- a/Engine/Source/Runtime/Core/Public/HAL/MallocBinned2.h
+++ b/Engine/Source/Runtime/Core/Public/HAL/MallocBinned2.h
@@ -2,15 +2,13 @@
 
 #pragma once
 
-#define BINNED2_CACHE_FREED_OS_ALLOCS
-
-#if defined BINNED2_CACHE_FREED_OS_ALLOCS
-	#define BINNED2_MAX_CACHED_OS_FREES (64)
-	#if PLATFORM_64BITS
-		#define BINNED2_MAX_CACHED_OS_FREES_BYTE_LIMIT (64*1024*1024)
-	#else
-		#define BINNED2_MAX_CACHED_OS_FREES_BYTE_LIMIT (16*1024*1024)
-	#endif
+#include "Allocators/CachedOSPageAllocator.h"
+
+#define BINNED2_MAX_CACHED_OS_FREES (64)
+#if PLATFORM_64BITS
+	#define BINNED2_MAX_CACHED_OS_FREES_BYTE_LIMIT (64*1024*1024)
+#else
+	#define BINNED2_MAX_CACHED_OS_FREES_BYTE_LIMIT (16*1024*1024)
 #endif
 
 #if STATS
@@ -43,21 +41,6 @@ private:
 	struct FPoolInfo;
 	struct PoolHashBucket;
 
-#ifdef BINNED2_CACHE_FREED_OS_ALLOCS
-	/**  */
-	struct FFreePageBlock
-	{
-		void*				Ptr;
-		SIZE_T				ByteSize;
-
-		FFreePageBlock() 
-		{
-			Ptr = nullptr;
-			ByteSize = 0;
-		}
-	};
-#endif
-
 	/** Pool table. */
 	struct FPoolTable
 	{
@@ -108,53 +91,65 @@ private:
 		}
 	};
 
-	uint64 TableAddressLimit;
+	const uint32 PageSize;
 
-	// PageSize dependent constants
-	uint64 MaxHashBuckets; 
-	uint64 MaxHashBucketBits;
-	uint64 MaxHashBucketWaste;
-	uint64 MaxBookKeepingOverhead;
 	/** Shift to get the reference from the indirect tables */
-	uint64 PoolBitShift;
-	uint64 IndirectPoolBitShift;
-	uint64 IndirectPoolBlockSize;
+	const uint64 PoolBitShift;
+	const uint64 IndirectPoolBlockSize;
+	const uint64 IndirectPoolBitShift;
+
+	// PageSize dependent constants
+	const uint64 MaxHashBuckets;
+
 	/** Shift required to get required hash table key. */
-	uint64 HashKeyShift;
+	const uint64 HashKeyShift;
+
 	/** Used to mask off the bits that have been used to lookup the indirect table */
-	uint64 PoolMask;
-	uint64 BinnedSizeLimit;
-	uint64 BinnedOSTableIndex;
+	const uint64 PoolMask;
+	const uint64 BinnedSizeLimit;
+	const uint64 BinnedOSTableIndex;
 
 	// Variables.
 	FPoolTable  PoolTable[POOL_COUNT];
-	FPoolTable	OsTable;
-	FPoolTable	PagePoolTable[EXTENDED_PAGE_POOL_ALLOCATION_COUNT];
-	FPoolTable* MemSizeToPoolTable[MAX_POOLED_ALLOCATION_SIZE+EXTENDED_PAGE_POOL_ALLOCATION_COUNT];
+	FPoolTable  PagePoolTable[EXTENDED_PAGE_POOL_ALLOCATION_COUNT];
+	FPoolTable* MemSizeToPoolTable[MAX_POOLED_ALLOCATION_SIZE + EXTENDED_PAGE_POOL_ALLOCATION_COUNT];
 
 	PoolHashBucket* HashBuckets;
 	PoolHashBucket* HashBucketFreeList;
 
-	uint32		PageSize;
-
-#ifdef BINNED2_CACHE_FREED_OS_ALLOCS
-	FFreePageBlock	FreedPageBlocks[BINNED2_MAX_CACHED_OS_FREES];
-	uint32			FreedPageBlocksNum;
-	uint32			CachedTotal;
-#endif
+	TCachedOSPageAllocator<BINNED2_MAX_CACHED_OS_FREES, BINNED2_MAX_CACHED_OS_FREES_BYTE_LIMIT> CachedOSPageAllocator;
 
 #if STATS
-	BINNED2_STAT		OsCurrent;
-	BINNED2_STAT		OsPeak;
-	BINNED2_STAT		WasteCurrent;
-	BINNED2_STAT		WastePeak;
-	BINNED2_STAT		UsedCurrent;
-	BINNED2_STAT		UsedPeak;
-	BINNED2_STAT		CurrentAllocs;
-	BINNED2_STAT		TotalAllocs;
-	/** OsCurrent - WasteCurrent - UsedCurrent. */
-	BINNED2_STAT		SlackCurrent;
-	double		MemTime;
+	struct FStats
+	{
+		BINNED2_STAT		OsCurrent;
+		BINNED2_STAT		OsPeak;
+		BINNED2_STAT		WasteCurrent;
+		BINNED2_STAT		WastePeak;
+		BINNED2_STAT		UsedCurrent;
+		BINNED2_STAT		UsedPeak;
+		BINNED2_STAT		CurrentAllocs;
+		BINNED2_STAT		TotalAllocs;
+		/** OsCurrent - WasteCurrent - UsedCurrent. */
+		BINNED2_STAT		SlackCurrent;
+		double				MemTime;
+
+		FStats()
+			: OsCurrent    (0)
+			, OsPeak       (0)
+			, WasteCurrent (0)
+			, WastePeak    (0)
+			, UsedCurrent  (0)
+			, UsedPeak     (0)
+			, CurrentAllocs(0)
+			, TotalAllocs  (0)
+			, SlackCurrent (0)
+			, MemTime      (0.0)
+		{
+		}
+	};
+
+	FStats Stats;
 #endif
 
 public:
diff --git a/Engine/Source/Runtime/Core/Public/HAL/PThreadCriticalSection.h b/Engine/Source/Runtime/Core/Public/HAL/PThreadCriticalSection.h
--- a/Engine/Source/Runtime/Core/Public/HAL/PThreadCriticalSection.h
+++ b/Engine/Source/Runtime/Core/Public/HAL/PThreadCriticalSection.h
@@ -55,4 +55,8 @@ public:
 	{
 		pthread_mutex_unlock(&Mutex);
 	}
+
+private:
+	FPThreadsCriticalSection(const FPThreadsCriticalSection&);
+	FPThreadsCriticalSection& operator=(const FPThreadsCriticalSection&);
 };
diff --git a/Engine/Source/Runtime/Core/Public/HAL/Platform.h b/Engine/Source/Runtime/Core/Public/HAL/Platform.h
--- a/Engine/Source/Runtime/Core/Public/HAL/Platform.h
+++ b/Engine/Source/Runtime/Core/Public/HAL/Platform.h
@@ -313,25 +313,28 @@
 #define PLATFORM_VTABLE_AT_END_OF_CLASS 0 
 
 #ifndef VARARGS
-	#define VARARGS					/* Functions with variable arguments */
+	#define VARARGS									/* Functions with variable arguments */
 #endif
 #ifndef CDECL
-	#define CDECL	    			/* Standard C function */
+	#define CDECL	    							/* Standard C function */
 #endif
 #ifndef STDCALL
-	#define STDCALL					/* Standard calling convention */
+	#define STDCALL									/* Standard calling convention */
 #endif
 #ifndef FORCEINLINE
-	#define FORCEINLINE 			/* Force code to be inline */
+	#define FORCEINLINE 							/* Force code to be inline */
 #endif
 #ifndef FORCENOINLINE
-	#define FORCENOINLINE 			/* Force code to NOT be inline */
+	#define FORCENOINLINE 							/* Force code to NOT be inline */
 #endif
 #ifndef RESTRICT
-	#define RESTRICT __restrict		/* no alias hint */
+	#define RESTRICT __restrict						/* no alias hint */
+#endif
+#ifndef FUNCTION_CHECK_RETURN
+	#define FUNCTION_CHECK_RETURN(...) __VA_ARGS__	/* Wrap a function signature in this to warn that callers should not ignore the return value */
 #endif
 
-#ifndef ASSUME						/* Hints compiler that expression is true; generally restricted to comparisons against constants */
+#ifndef ASSUME										/* Hints compiler that expression is true; generally restricted to comparisons against constants */
 	#define ASSUME(...) 
 #endif
 
diff --git a/Engine/Source/Runtime/Core/Public/HTML5/HTML5Platform.h b/Engine/Source/Runtime/Core/Public/HTML5/HTML5Platform.h
--- a/Engine/Source/Runtime/Core/Public/HTML5/HTML5Platform.h
+++ b/Engine/Source/Runtime/Core/Public/HTML5/HTML5Platform.h
@@ -58,9 +58,11 @@ typedef FHTML5Types FPlatformTypes;
 #define FORCEINLINE _forceinline
 #define FORCENOINLINE __declspec(noinline)	/* Force code to NOT be inline */
 #define CONSTEXPR    
+#define FUNCTION_CHECK_RETURN(...) __declspec("SAL_checkReturn") __VA_ARGS__	/* Wrap a function signature in this to warn that callers should not ignore the return value. */
 #else
 #define FORCEINLINE		inline __attribute__((__always_inline__))		/* Force code to be inline */
 #define FORCENOINLINE	__attribute__((noinline))			/* Force code to NOT be inline */
+#define FUNCTION_CHECK_RETURN(...) __VA_ARGS__ __attribute__ ((warn_unused_result))	/* Wrap a function signature in this to warn that callers should not ignore the return value. */
 #endif
 
 // Optimization macros
diff --git a/Engine/Source/Runtime/Core/Public/HTML5/HTML5PlatformProcess.h b/Engine/Source/Runtime/Core/Public/HTML5/HTML5PlatformProcess.h
--- a/Engine/Source/Runtime/Core/Public/HTML5/HTML5PlatformProcess.h
+++ b/Engine/Source/Runtime/Core/Public/HTML5/HTML5PlatformProcess.h
@@ -62,6 +62,10 @@ public:
 	FORCEINLINE void Unlock(void)
 	{
 	}
+
+private:
+	FHTML5CriticalSection(const FHTML5CriticalSection&);
+	FHTML5CriticalSection& operator=(const FHTML5CriticalSection&);
 };
 
 typedef FHTML5CriticalSection FCriticalSection;
diff --git a/Engine/Source/Runtime/Core/Public/IOS/IOSPlatform.h b/Engine/Source/Runtime/Core/Public/IOS/IOSPlatform.h
--- a/Engine/Source/Runtime/Core/Public/IOS/IOSPlatform.h
+++ b/Engine/Source/Runtime/Core/Public/IOS/IOSPlatform.h
@@ -47,11 +47,12 @@ typedef FIOSPlatformTypes FPlatformTypes;
 #define FORCE_ANSI_ALLOCATOR 1
 
 // Function type macros.
-#define VARARGS														/* Functions with variable arguments */
-#define CDECL														/* Standard C function */
-#define STDCALL														/* Standard calling convention */
-#define FORCEINLINE inline __attribute__ ((always_inline))			/* Force code to be inline */
-#define FORCENOINLINE __attribute__((noinline))						/* Force code to NOT be inline */
+#define VARARGS																		/* Functions with variable arguments */
+#define CDECL																		/* Standard C function */
+#define STDCALL																		/* Standard calling convention */
+#define FORCEINLINE inline __attribute__ ((always_inline))							/* Force code to be inline */
+#define FORCENOINLINE __attribute__((noinline))										/* Force code to NOT be inline */
+#define FUNCTION_CHECK_RETURN(...) __VA_ARGS__ __attribute__ ((warn_unused_result))	/* Wrap a function signature in this to warn that callers should not ignore the return value. */
 
 #define TEXT_HELPER(a,b)	a ## b
 #define TEXT(s)				TEXT_HELPER(L, s)
diff --git a/Engine/Source/Runtime/Core/Public/Linux/LinuxPlatform.h b/Engine/Source/Runtime/Core/Public/Linux/LinuxPlatform.h
--- a/Engine/Source/Runtime/Core/Public/Linux/LinuxPlatform.h
+++ b/Engine/Source/Runtime/Core/Public/Linux/LinuxPlatform.h
@@ -59,11 +59,12 @@ typedef FLinuxPlatformTypes FPlatformTypes;
 #endif // defined(_M_IX86) || defined(__i386__) || defined(_M_X64) || defined(__x86_64__) || defined (__amd64__) 
 
 // Function type macros.
-#define VARARGS													/* Functions with variable arguments */
-#define CDECL													/* Standard C function */
-#define STDCALL													/* Standard calling convention */
-#define FORCEINLINE inline __attribute__ ((always_inline))		/* Force code to be inline */
-#define FORCENOINLINE __attribute__((noinline))					/* Force code to NOT be inline */
+#define VARARGS																		/* Functions with variable arguments */
+#define CDECL																		/* Standard C function */
+#define STDCALL																		/* Standard calling convention */
+#define FORCEINLINE inline __attribute__ ((always_inline))							/* Force code to be inline */
+#define FORCENOINLINE __attribute__((noinline))										/* Force code to NOT be inline */
+#define FUNCTION_CHECK_RETURN(...) __VA_ARGS__ __attribute__ ((warn_unused_result))	/* Wrap a function signature in this to warn that callers should not ignore the return value. */
 
 #define TEXT_HELPER(a,b)	a ## b
 #define TEXT(s)				TEXT_HELPER(L, s)
diff --git a/Engine/Source/Runtime/Core/Public/Mac/MacPlatform.h b/Engine/Source/Runtime/Core/Public/Mac/MacPlatform.h
--- a/Engine/Source/Runtime/Core/Public/Mac/MacPlatform.h
+++ b/Engine/Source/Runtime/Core/Public/Mac/MacPlatform.h
@@ -41,11 +41,12 @@ typedef FMacPlatformTypes FPlatformTypes;
 #define PLATFORM_SUPPORTS_STACK_SYMBOLS				1
 
 // Function type macros.
-#define VARARGS														/* Functions with variable arguments */
-#define CDECL														/* Standard C function */
-#define STDCALL														/* Standard calling convention */
-#define FORCEINLINE inline __attribute__ ((always_inline))			/* Force code to be inline */
-#define FORCENOINLINE __attribute__((noinline))						/* Force code to NOT be inline */
+#define VARARGS																		/* Functions with variable arguments */
+#define CDECL																		/* Standard C function */
+#define STDCALL																		/* Standard calling convention */
+#define FORCEINLINE inline __attribute__ ((always_inline))							/* Force code to be inline */
+#define FORCENOINLINE __attribute__((noinline))										/* Force code to NOT be inline */
+#define FUNCTION_CHECK_RETURN(...) __attribute__ ((warn_unused_result))	__VA_ARGS__	/* Wrap a function signature in this to warn that callers should not ignore the return value. */
 
 #define TEXT_HELPER(a,b)	a ## b
 #define TEXT(s)				TEXT_HELPER(L, s)
diff --git a/Engine/Source/Runtime/Core/Public/Math/UnrealMathUtility.h b/Engine/Source/Runtime/Core/Public/Math/UnrealMathUtility.h
--- a/Engine/Source/Runtime/Core/Public/Math/UnrealMathUtility.h
+++ b/Engine/Source/Runtime/Core/Public/Math/UnrealMathUtility.h
@@ -218,9 +218,10 @@ struct FMath : public FPlatformMath
 	 *	@param Value	Number to check
 	 *	@return			true if Value is a power of two
 	 */
-	static FORCEINLINE bool IsPowerOfTwo( uint32 Value )
+	template <typename T>
+	static FORCEINLINE bool IsPowerOfTwo( T Value )
 	{
-		return ((Value & (Value - 1)) == 0);
+		return ((Value & (Value - 1)) == (T)0);
 	}
 
 
diff --git a/Engine/Source/Runtime/Core/Public/Misc/CallbackDevice.h b/Engine/Source/Runtime/Core/Public/Misc/CallbackDevice.h
--- a/Engine/Source/Runtime/Core/Public/Misc/CallbackDevice.h
+++ b/Engine/Source/Runtime/Core/Public/Misc/CallbackDevice.h
@@ -54,14 +54,11 @@ public:
 	// delegate type for prompting the pak system to mount a new pak
 	DECLARE_DELEGATE_RetVal_OneParam(bool, FOnUnmountPak, const FString&);
 
-	// Callback for PER_MODULE_BOILERPLATE macro's GSerialNumberBlocksForDebugVisualizers
-	DECLARE_DELEGATE_RetVal(int32***, FGetSerialNumberBlocksForDebugVisualizersDelegate);
-
 	/** delegate type for opening a modal message box ( Params: EAppMsgType::Type MessageType, const FText& Text, const FText& Title ) */
 	DECLARE_DELEGATE_RetVal_ThreeParams(EAppReturnType::Type, FOnModalMessageBox, EAppMsgType::Type, const FText&, const FText&);
 
 	// Callback for PER_MODULE_BOILERPLATE macro's GObjectArrayForDebugVisualizers
-	DECLARE_DELEGATE_RetVal(class UObjectBase***, FObjectArrayForDebugVisualizersDelegate);
+	DECLARE_DELEGATE_RetVal(class FFixedUObjectArray*, FObjectArrayForDebugVisualizersDelegate);
 
 	// Called in PER_MODULE_BOILERPLATE macro.
 	static FObjectArrayForDebugVisualizersDelegate& GetObjectArrayForDebugVisualizersDelegate();
@@ -132,9 +129,6 @@ public:
 	// Called when before the application is exiting.
 	static FSimpleMulticastDelegate OnPreExit;
 
-	// Called in PER_MODULE_BOILERPLATE macro.
-	static FGetSerialNumberBlocksForDebugVisualizersDelegate& GetSerialNumberBlocksForDebugVisualizersDelegate();
-
 	/** Color picker color has changed, please refresh as needed*/
 	static FSimpleMulticastDelegate ColorPickerChanged;
 
diff --git a/Engine/Source/Runtime/Core/Public/Misc/CoreMisc.h b/Engine/Source/Runtime/Core/Public/Misc/CoreMisc.h
--- a/Engine/Source/Runtime/Core/Public/Misc/CoreMisc.h
+++ b/Engine/Source/Runtime/Core/Public/Misc/CoreMisc.h
@@ -122,6 +122,18 @@ struct CORE_API FCommandLine
 	 */
 	static void Parse(const TCHAR* CmdLine, TArray<FString>& Tokens, TArray<FString>& Switches);
 private:
+#if WANTS_COMMANDLINE_WHITELIST
+	/** Filters both the original and current command line list for approved only args */
+	static void WhitelistCommandLines();
+	/** Filters any command line args that aren't on the approved list */
+	static TArray<FString> FilterCommandLine(TCHAR* CommandLine);
+	/** Rebuilds the command line using the filtered args */
+	static void BuildWhitelistCommandLine(TCHAR* CommandLine, uint32 Length, const TArray<FString>& FilteredArgs);
+	static TArray<FString> ApprovedArgs;
+#else
+#define WhitelistCommandLines()
+#endif
+
 	/** Flag to check if the commandline has been initialized or not. */
 	static bool bIsInitialized;
 	/** character buffer containing the command line */
@@ -169,7 +181,7 @@ struct CORE_API FFileHelper
 	/**
 	 * Load a binary file to a dynamic array.
 	*/
-	static bool LoadFileToArray( TArray<uint8>& Result, const TCHAR* Filename,uint32 Flags = 0 );
+	static bool LoadFileToArray( TArray<uint8>& Result, const TCHAR* Filename, uint32 Flags = 0 );
 
 	/**
 	 * Load a text file to an FString.
@@ -183,13 +195,13 @@ struct CORE_API FFileHelper
 	/**
 	 * Save a binary array to a file.
 	 */
-	static bool SaveArrayToFile( const TArray<uint8>& Array, const TCHAR* Filename, IFileManager* FileManager=&IFileManager::Get() );
+	static bool SaveArrayToFile( const TArray<uint8>& Array, const TCHAR* Filename, IFileManager* FileManager=&IFileManager::Get(), uint32 WriteFlags = 0 );
 
 	/**
 	 * Write the FString to a file.
 	 * Supports all combination of ANSI/Unicode files and platforms.
 	 */
-	static bool SaveStringToFile( const FString& String, const TCHAR* Filename, EEncodingOptions::Type EncodingOptions=EEncodingOptions::AutoDetect, IFileManager* FileManager=&IFileManager::Get() );
+	static bool SaveStringToFile( const FString& String, const TCHAR* Filename, EEncodingOptions::Type EncodingOptions=EEncodingOptions::AutoDetect, IFileManager* FileManager=&IFileManager::Get(), uint32 WriteFlags = 0 );
 
 	/**
 	 * Saves a 24/32Bit BMP file to disk
diff --git a/Engine/Source/Runtime/Core/Public/Modules/Boilerplate/ModuleBoilerplate.h b/Engine/Source/Runtime/Core/Public/Modules/Boilerplate/ModuleBoilerplate.h
--- a/Engine/Source/Runtime/Core/Public/Modules/Boilerplate/ModuleBoilerplate.h
+++ b/Engine/Source/Runtime/Core/Public/Modules/Boilerplate/ModuleBoilerplate.h
@@ -33,7 +33,6 @@
 #define PER_MODULE_BOILERPLATE \
 	TArray<FNameEntry const*>* GFNameTableForDebuggerVisualizers = FName::GetNameTableForDebuggerVisualizers_ST(); \
 	FNameEntry*** GFNameTableForDebuggerVisualizers_MT = FName::GetNameTableForDebuggerVisualizers_MT(); \
-	int32*** GSerialNumberBlocksForDebugVisualizers = FCoreDelegates::GetSerialNumberBlocksForDebugVisualizersDelegate().IsBound() ? FCoreDelegates::GetSerialNumberBlocksForDebugVisualizersDelegate().Execute() : NULL; \
-	UObjectBase*** GObjectArrayForDebugVisualizers = FCoreDelegates::GetObjectArrayForDebugVisualizersDelegate().IsBound() ? FCoreDelegates::GetObjectArrayForDebugVisualizersDelegate().Execute() : NULL; \
+	FFixedUObjectArray* GObjectArrayForDebugVisualizers = FCoreDelegates::GetObjectArrayForDebugVisualizersDelegate().IsBound() ? FCoreDelegates::GetObjectArrayForDebugVisualizersDelegate().Execute() : NULL; \
 	bool GFNameDebuggerVisualizersIsUE3=false; \
 	REPLACEMENT_OPERATOR_NEW_AND_DELETE
diff --git a/Engine/Source/Runtime/Core/Public/Serialization/Archive.h b/Engine/Source/Runtime/Core/Public/Serialization/Archive.h
--- a/Engine/Source/Runtime/Core/Public/Serialization/Archive.h
+++ b/Engine/Source/Runtime/Core/Public/Serialization/Archive.h
@@ -198,7 +198,7 @@ public:
 		ArIsPersistent	= bIsPersistent;
 	}
 
-	// with this method it's possible to attach data behind some serialized data 
+	/** With this method it's possible to attach data behind some serialized data. */
 	void SetLimitSize(int64 NewLimitSize)
 	{
 		LimitSize = NewLimitSize;
diff --git a/Engine/Source/Runtime/Core/Public/Serialization/ArchiveBase.h b/Engine/Source/Runtime/Core/Public/Serialization/ArchiveBase.h
--- a/Engine/Source/Runtime/Core/Public/Serialization/ArchiveBase.h
+++ b/Engine/Source/Runtime/Core/Public/Serialization/ArchiveBase.h
@@ -1187,17 +1187,20 @@ public:
 
 	virtual FArchive& operator<<(class FName& Value) override
 	{
-		return InnerArchive << Value;
+		InnerArchive << Value;
+		return *this;
 	}
 
 	virtual FArchive& operator<<(class UObject*& Value) override
 	{
-		return InnerArchive << Value;
+		InnerArchive << Value;
+		return *this;
 	}
 
 	virtual FArchive& operator<<(struct FStringAssetReference& Value) override
 	{
-		return InnerArchive << Value;
+		InnerArchive << Value;
+		return *this;
 	}
 
 	virtual void Serialize(void* V, int64 Length) override
diff --git a/Engine/Source/Runtime/Core/Public/Stats/Stats2.h b/Engine/Source/Runtime/Core/Public/Stats/Stats2.h
--- a/Engine/Source/Runtime/Core/Public/Stats/Stats2.h
+++ b/Engine/Source/Runtime/Core/Public/Stats/Stats2.h
@@ -91,8 +91,7 @@ MS_ALIGN(8)
 struct TStatIdData
 {
 	FORCEINLINE TStatIdData()
-		: Name(NameToMinimalName(NAME_None))
-		, AnsiString(0)
+		: AnsiString(0)
 		, WideString(0)
 	{
 	}
@@ -2071,5 +2070,6 @@ DECLARE_STATS_GROUP(TEXT("User"),STATGROUP_User, STATCAT_Advanced);
 
 DECLARE_CYCLE_STAT_EXTERN(TEXT("FrameTime"),STAT_FrameTime,STATGROUP_Engine, CORE_API);
 DECLARE_FNAME_STAT_EXTERN(TEXT("NamedMarker"),STAT_NamedMarker,STATGROUP_StatSystem, CORE_API);
+DECLARE_FLOAT_COUNTER_STAT_EXTERN(TEXT("Seconds Per Cycle"),STAT_SecondsPerCycle,STATGROUP_Engine, CORE_API);
 
 #endif
diff --git a/Engine/Source/Runtime/Core/Public/Stats/StatsData.h b/Engine/Source/Runtime/Core/Public/Stats/StatsData.h
--- a/Engine/Source/Runtime/Core/Public/Stats/StatsData.h
+++ b/Engine/Source/Runtime/Core/Public/Stats/StatsData.h
@@ -24,8 +24,8 @@ struct CORE_API FStatConstants
 	static const char* ThreadGroupName;
 	static const FName NAME_ThreadGroup;
 
-	/** Stat short name for seconds per cycle. */
-	static const FName NAME_SecondsPerCycle;
+	/** Stat raw name for seconds per cycle. */
+	static const FName RAW_SecondsPerCycle;
 
 	/** Special case category, when we want to Stat to appear at the root of the menu (leaving the category blank omits it from the menu entirely) */
 	static const FName NAME_NoCategory;
@@ -491,7 +491,7 @@ class CORE_API FStatsThreadState
 	void ScanForAdvance(FStatPacketArray& NewData);
 
 public:
-	/** Internal method to add meta data packets to the data structures. **/
+	/** Internal method to update the internal metadata. **/
 	void ProcessMetaDataOnly(TArray<FStatMessage>& Data);
 
 	/** Toggles tracking the most memory expensive stats. */
@@ -639,14 +639,23 @@ public:
 		return Frame;
 	}
 
+	/** Gets the old-skool flat grouped inclusive stats. These ignore recursion, merge threads, etc and so generally the condensed callstack is less confusing. **/
+	void GetInclusiveAggregateStackStats( const TArray<FStatMessage>& CondensedMessages, TArray<FStatMessage>& OutStats, IItemFiler* Filter = nullptr, bool bAddNonStackStats = true, TMap<FName, TArray<FStatMessage>>* OptionalOutThreadBreakdownMap = nullptr ) const;
+
 	/** Gets the old-skool flat grouped inclusive stats. These ignore recursion, merge threads, etc and so generally the condensed callstack is less confusing. **/
 	void GetInclusiveAggregateStackStats(int64 TargetFrame, TArray<FStatMessage>& OutStats, IItemFiler* Filter = nullptr, bool bAddNonStackStats = true, TMap<FName, TArray<FStatMessage>>* OptionalOutThreadBreakdownMap = nullptr) const;
 
+	/** Gets the old-skool flat grouped exclusive stats. These merge threads, etc and so generally the condensed callstack is less confusing. **/
+	void GetExclusiveAggregateStackStats( const TArray<FStatMessage>& CondensedMessages, TArray<FStatMessage>& OutStats, IItemFiler* Filter = nullptr, bool bAddNonStackStats = true ) const;
+
 	/** Gets the old-skool flat grouped exclusive stats. These merge threads, etc and so generally the condensed callstack is less confusing. **/
 	void GetExclusiveAggregateStackStats(int64 TargetFrame, TArray<FStatMessage>& OutStats, IItemFiler* Filter = nullptr, bool bAddNonStackStats = true) const;
 
 	/** Used to turn the condensed version of stack stats back into a tree for easier handling. **/
-	void UncondenseStackStats(int64 TargetFrame, FRawStatStackNode& Out, IItemFiler* Filter = nullptr, TArray<FStatMessage>* OutNonStackStats = nullptr) const;
+	void UncondenseStackStats( const TArray<FStatMessage>& CondensedMessages, FRawStatStackNode& Root, IItemFiler* Filter = nullptr, TArray<FStatMessage>* OutNonStackStats = nullptr ) const;
+
+	/** Used to turn the condensed version of stack stats back into a tree for easier handling. **/
+	void UncondenseStackStats(int64 TargetFrame, FRawStatStackNode& Root, IItemFiler* Filter = nullptr, TArray<FStatMessage>* OutNonStackStats = nullptr) const;
 
 	/** Adds missing stats to the group so it doesn't jitter. **/
 	void AddMissingStats(TArray<FStatMessage>& Dest, TSet<FName> const& EnabledItems) const;
@@ -660,11 +669,13 @@ public:
 		#YRX_Stats: 2015-07-07 Maybe move to FStatsLoadedState
 	-----------------------------------------------------------------------------*/
 public:
+	friend struct FStatsReadFile;
+
 	/** Constructor to load stats from a file **/
 	FStatsThreadState( FString const& Filename );
 
 	/** Adds a frame worth of messages */
-	void AddMessages( TArray<FStatMessage>& InMessages );
+	void AddCondensedMessages( TArray<FStatMessage>& CondensedMessages );
 
 	/** Marks this stats state as loaded. */
 	void MarkAsLoaded()
@@ -673,8 +684,18 @@ public:
 	}
 
 protected:
-	/** Internal method to scan the messages to accumulate any non-frame stats. **/
-	void ProcessMetaDataForLoad( TArray<FStatMessage>& Data );
+	/** Internal method to update the internal metadata. **/
+	void ProcessMetaDataForLoad( TArray<FStatMessage>& CondensedMessages );
+
+	/** 
+	  * Internal method to place the data into the history, 
+	  * maintains the history based on the requested number of frames to keep in the history. 
+	  * The condensed messages are emplaced in the condensed history.
+	  */
+	void AddFrameFromCondensedMessages( TArray<FStatMessage>& CondensedMessages );
+
+	/** Internal method to scan the messages to find the current game/render thread frame. */
+	void AdvanceFrameForLoad( TArray<FStatMessage>& CondensedMessages );
 
 	/** Largest frame seen. Loaded stats only. **/
 	int64 MaxFrameSeen;
diff --git a/Engine/Source/Runtime/Core/Public/Stats/StatsFile.h b/Engine/Source/Runtime/Core/Public/Stats/StatsFile.h
--- a/Engine/Source/Runtime/Core/Public/Stats/StatsFile.h
+++ b/Engine/Source/Runtime/Core/Public/Stats/StatsFile.h
@@ -281,6 +281,10 @@
 	}
 };
 
+/*-----------------------------------------------------------------------------
+	FStatsFrameInfo
+-----------------------------------------------------------------------------*/
+
 /**
  * Contains basic information about one frame of the stats.
  * This data is used to generate ultra-fast preview of the stats history without the requirement of reading the whole file.
@@ -317,8 +321,89 @@
 	TMap<uint32, int64> ThreadCycles;
 };
 
+/*-----------------------------------------------------------------------------
+	FStatsWriteStream
+-----------------------------------------------------------------------------*/
+
+/** Struct used to send a stream of stat messages. */
+struct CORE_API FStatsWriteStream
+{
+protected:
+	/** Writes metadata messages into the stream. */
+	void WriteMetadata( FArchive& Ar );
+
+	/** Writes condensed messages into the stream. */
+	void WriteCondensedMessages( FArchive& Ar, int64 TargetFrame );
+
+	/** Sends an FName, and the string it represents if we have not sent that string before. **/
+	FORCEINLINE_STATS void WriteFName( FArchive& Ar, FStatNameAndInfo NameAndInfo )
+	{
+		FName RawName = NameAndInfo.GetRawName();
+		bool bSendFName = !FNamesSent.Contains( RawName.GetComparisonIndex() );
+		int32 Index = RawName.GetComparisonIndex();
+		Ar << Index;
+		int32 Number = NameAndInfo.GetRawNumber();
+		if (bSendFName)
+		{
+			FNamesSent.Add( RawName.GetComparisonIndex() );
+			Number |= EStatMetaFlags::SendingFName << (EStatMetaFlags::Shift + EStatAllFields::StartShift);
+		}
+		Ar << Number;
+		if (bSendFName)
+		{
+			FString Name = RawName.ToString();
+			Ar << Name;
+		}
+	}
+
+	/** Write a stat message. **/
+	FORCEINLINE_STATS void WriteMessage( FArchive& Ar, FStatMessage const& Item )
+	{
+		WriteFName( Ar, Item.NameAndInfo );
+		switch (Item.NameAndInfo.GetField<EStatDataType>())
+		{
+			case EStatDataType::ST_int64:
+			{
+				int64 Payload = Item.GetValue_int64();
+				Ar << Payload;
+				break;
+			}
+
+			case EStatDataType::ST_double:
+			{
+				double Payload = Item.GetValue_double();
+				Ar << Payload;
+				break;
+			}
+
+			case EStatDataType::ST_FName:
+			{
+				WriteFName( Ar, FStatNameAndInfo( Item.GetValue_FName(), false ) );
+				break;
+			}
+
+			case EStatDataType::ST_Ptr:
+			{
+				uint64 Payload = Item.GetValue_Ptr();
+				Ar << Payload;
+				break;
+			}
+		}
+	}
+
+	/** Set of names already sent. */
+	TSet<int32> FNamesSent;
+
+	/** Data to write. */
+	TArray<uint8> OutData;
+};
+
+/*-----------------------------------------------------------------------------
+	IStatsWriteFile
+-----------------------------------------------------------------------------*/
+
 /** Interface for writing stats data. Can be only used in the stats thread. */
-struct CORE_API IStatsWriteFile
+struct CORE_API IStatsWriteFile : public FStatsWriteStream
 {
 	friend class FAsyncStatsWrite;
 
@@ -332,15 +417,9 @@
 	/** Stats stream header. */
 	FStatsStreamHeader Header;
 
-	/** Set of names already sent. */
-	TSet<int32> FNamesSent;
-
 	/** Async task used to offload saving the capture data. */
 	FAsyncTask<FAsyncStatsWrite>* AsyncTask;
 
-	/** Data to write through the async task. **/
-	TArray<uint8> OutData;
-
 	/** Buffer to store the compressed data, used by the FAsyncStatsWrite. */
 	TArray<uint8> CompressedData;
 
@@ -351,6 +430,15 @@
 	 */
 	TArray<FStatsFrameInfo> FramesInfo;
 
+	/** NewFrame delegate handle  */
+	FDelegateHandle DataDelegateHandle;
+
+	/** The size of the stats file, sets by the async writing thread. */
+	int64 FileSize;
+
+	/** Start time of the 'stat startfile', in seconds. */
+	double StartTime;
+
 protected:
 	/** Default constructor. */
 	IStatsWriteFile();
@@ -360,17 +448,17 @@
 	virtual ~IStatsWriteFile()
 	{}
 
-protected:
-	/** NewFrame delegate handle  */
-	FDelegateHandle DataDelegateHandle;
-
-public:
 	/** Creates a file writer and registers for the data delegate. */
 	void Start( const FString& InFilename );
 
 	/** Finalizes writing the stats data and unregisters the data delegate. */
 	void Stop();
 
+	/**
+	 * @return the stats file metadata description like the current size and the duration of the stats session.
+	 */
+	FText GetFileMetaDesc() const;
+
 protected:
 	/** Sets the data delegate used to receive a stats data. */
 	virtual void SetDataDelegate( bool bSet ) = 0;
@@ -384,85 +472,15 @@
 		return !!File;
 	}
 
-public:
-	const TArray<uint8>& GetOutData() const
-	{
-		return OutData;
-	}
-
-	void ResetData()
-	{
-		OutData.Reset();
-	}
-
 	/** Writes magic value, dummy header and initial metadata. */
 	void WriteHeader();
 
 protected:
-	/** Writes metadata messages into the stream. */
-	void WriteMetadata( FArchive& Ar );
-
 	/**	Finalizes writing to the file. */
 	void Finalize();
 
 	/** Sends the data to the file via async task. */
 	void SendTask();
-
-	/** Sends an FName, and the string it represents if we have not sent that string before. **/
-	FORCEINLINE_STATS void WriteFName( FArchive& Ar, FStatNameAndInfo NameAndInfo )
-	{
-		FName RawName = NameAndInfo.GetRawName();
-		bool bSendFName = !FNamesSent.Contains( RawName.GetComparisonIndex() );
-		int32 Index = RawName.GetComparisonIndex();
-		Ar << Index;
-		int32 Number = NameAndInfo.GetRawNumber();
-		if( bSendFName )
-		{
-			FNamesSent.Add( RawName.GetComparisonIndex() );
-			Number |= EStatMetaFlags::SendingFName << (EStatMetaFlags::Shift + EStatAllFields::StartShift);
-		}
-		Ar << Number;
-		if( bSendFName )
-		{
-			FString Name = RawName.ToString();
-			Ar << Name;
-		}
-	}
-
-	/** Write a stat message. **/
-	FORCEINLINE_STATS void WriteMessage( FArchive& Ar, FStatMessage const& Item )
-	{
-		WriteFName( Ar, Item.NameAndInfo );
-		switch( Item.NameAndInfo.GetField<EStatDataType>() )
-		{
-			case EStatDataType::ST_int64:
-			{
-				int64 Payload = Item.GetValue_int64();
-				Ar << Payload;
-				break;
-			}
-				
-			case EStatDataType::ST_double:
-			{
-				double Payload = Item.GetValue_double();
-				Ar << Payload;
-				break;
-			}
-
-			case EStatDataType::ST_FName:
-			{
-				WriteFName( Ar, FStatNameAndInfo( Item.GetValue_FName(), false ) );
-				break;
-			}
-
-			case EStatDataType::ST_Ptr:
-			{
-				uint64 Payload = Item.GetValue_Ptr();
-				Ar << Payload;
-				break;
-			}
-		}
-	}
 };
 
 /** Helper struct used to write regular stats to the file. */
@@ -486,23 +504,11 @@
 
 	virtual void FinalizeSavingData( int64 FrameFileOffset ) override;
 
-public:
 	/**
-	 *	Grabs a frame from the local FStatsThreadState and writes it to the array.
-	 *	Only used by the ProfilerServiceManager.
-	 */
-	void WriteFrame( int64 TargetFrame, bool bNeedFullMetadata );
-
-protected:
-	/**
 	 *	Grabs a frame from the local FStatsThreadState and adds it to the output.
 	 *	Called from the stats thread, but the data is saved using the the FAsyncStatsWrite. 
 	 */
-	void WriteFrame( int64 TargetFrame )
-	{
-		WriteFrame( TargetFrame, false );
-		SendTask();
-	}
+	void WriteFrame( int64 TargetFrame );
 };
 
 /** Helper struct used to write raw stats to the file. */
@@ -799,16 +805,16 @@
 struct FStatsReadFile;
 
 /**
- *	Helper class used to read and process raw stats file on the async task.
+ *	Helper class used to read and process stats file on the async task.
  */
-class FAsyncRawStatsFile
+class FAsyncStatsFile
 {
 	/** Pointer to FStatsReadFile to call for async work. */
 	FStatsReadFile* Owner;
 
 public:
 	/** Initialization constructor. */
-	FAsyncRawStatsFile( FStatsReadFile* InOwner );
+	FAsyncStatsFile( FStatsReadFile* InOwner );
 
 	/** Call DoWork on the parent */
 	void DoWork();
@@ -855,10 +861,10 @@
 -----------------------------------------------------------------------------*/
 
 template<typename T>
-struct FCreateStatsReader
+struct FStatsReader
 {
-	/** Creates a new reader for raw stats file. Will be nullptr for invalid files. */	
-	static T* ForRawStats( const TCHAR* Filename )
+	/** Creates a new reader for raw/regular stats file based on the T. Will be nullptr for invalid files. */	
+	static T* Create( const TCHAR* Filename )
 	{
 		T* StatsReadFile = new T( Filename );
 		const bool bValid = StatsReadFile->PrepareLoading();
@@ -873,14 +879,12 @@
 /** Struct used to read from ue4stats/ue4statsraw files, initializes all metadata and starts a process of reading the file asynchronously. */
 struct CORE_API FStatsReadFile
 {
-	friend class FAsyncRawStatsFile;
+	friend class FAsyncStatsFile;
+	friend struct FStatsReader<FStatsReadFile>;
 
 	/** Number of seconds between updating the current stage. */
 	static const double NumSecondsBetweenUpdates;
 
-	/** Creates a new reader for regular stats file. Will be nullptr for invalid files. */
-	static FStatsReadFile* CreateReaderForRegularStats( const TCHAR* Filename );
-
 public:
 	/** Reads and processes the file on the current thread. This is a blocking operation. */
 	void ReadAndProcessSynchronously();
@@ -888,6 +892,19 @@
 	/** Reads and processes the file using the async tasks on the pool thread. The read data is sent to the game thread using the task graph. This is a non-blocking operation. */
 	void ReadAndProcessAsynchronously();
 
+	/** Sets the number of frame to be stored for future use. */
+	void SetHistoryFrames( int32 InHistoryFrames );
+
+	/**
+	 * @return number of frames in the file.
+	 * For regular stat file valid after creating a reader.
+	 * For raw stat file reading whole file.
+	 */
+	int32 GetNumFrames() const
+	{
+		return NumFrames;
+	}
+
 protected:
 	/** Initialization constructor. */
 	FStatsReadFile( const TCHAR* InFilename, bool bInRawStatsFile );
@@ -903,15 +920,25 @@
 	 */
 	bool PrepareLoading();
 
-	/** Reads stats from the file into combined history. */
+	/** Reads stats from the file. */
 	void ReadStats();
 
-	/** Called before started processing combined history. */
+	/** Reads raw stats. */
+	void ReadRawStats();
+
+	/** Reads regular stats. */
+	void ReadRegularStats();
+
+	/** Called before started processing. */
 	virtual void PreProcessStats();
 
 	/** Processes combined history using the internal functionality and provided overloaded Process*Operation methods. */
 	void ProcessStats();
 
+	/** Called every each frame has been read from the file. */
+	virtual void ReadStatsFrame( const TArray<FStatMessage>& CondensedMessages, const int64 Frame )
+	{}
+
 	/** Called after finished processing combined history. */
 	virtual void PostProcessStats();
 
@@ -1076,7 +1103,7 @@
 	FArchive* Reader;
 
 	/** Async task. */
-	FAsyncTask<FAsyncRawStatsFile>* AsyncWork;
+	FAsyncTask<FAsyncStatsFile>* AsyncWork;
 
 	/** Basic information about the stats file. */
 	FRawStatsFileInfo FileInfo;
@@ -1105,6 +1132,9 @@
 	/** Filename. */
 	const FString Filename;
 
+	/** Number of frames. */
+	int32 NumFrames;
+
 	/** Whether this stats file uses raw data. */
 	const bool bRawStatsFile;
 };
@@ -1134,11 +1164,26 @@
 	/** Stat StopFile. */
 	void Stop();
 
-	bool IsStatFileActive()
+	/**
+	 * @return true, if any stats write command is active.
+	 */
+	bool IsStatFileActive() const
 	{
 		return StatFileActiveCounter.GetValue() > 0;
 	}
 
+	/**
+	* @return the stats file metadata description like the current size and the duration of the stats session.
+	*/
+	FText GetFileMetaDesc() const
+	{
+		if (IsStatFileActive() && CurrentStatsFile != nullptr)
+		{
+			return CurrentStatsFile->GetFileMetaDesc();
+		}
+		return FText();
+	}
+
 	/** Filename of the last saved stats file. */
 	FString LastFileSaved;
 
