Removal of #pragma once from the .cpp file.

diff --git a/Engine/Source/Runtime/Core/Private/HAL/MallocVerify.cpp b/Engine/Source/Runtime/Core/Private/HAL/MallocVerify.cpp
--- a/Engine/Source/Runtime/Core/Private/HAL/MallocVerify.cpp
+++ b/Engine/Source/Runtime/Core/Private/HAL/MallocVerify.cpp
@@ -4,8 +4,6 @@
 	MallocVerify.cpp: Helper class to track memory allocations
 =============================================================================*/
 
-#pragma once
-
 #include "CorePrivatePCH.h"
 #include "MemoryMisc.h"
 #include "MallocVerify.h"
