Fixing CIS warnings on the Mac

diff --git a/Engine/Source/Runtime/Core/Private/HAL/PlatformFileManager.cpp b/Engine/Source/Runtime/Core/Private/HAL/PlatformFileManager.cpp
--- a/Engine/Source/Runtime/Core/Private/HAL/PlatformFileManager.cpp
+++ b/Engine/Source/Runtime/Core/Private/HAL/PlatformFileManager.cpp
@@ -82,7 +82,7 @@ IPlatformFile* FPlatformFileManager::GetPlatformFile(const TCHAR* Name)
 		static TScopedPointer<IPlatformFile> AutoDestroySingleton(new FCachedReadPlatformFile());
 		PlatformFile = AutoDestroySingleton.GetOwnedPointer();
 	}
-	else
+	else if (FModuleManager::Get().ModuleExists(Name))
 	{
 		// Try to load a module containing the platform file.
 		class IPlatformFileModule* PlatformFileModule = FModuleManager::LoadModulePtr<IPlatformFileModule>(Name);
