FromString overload for FString moved to UnrealString.h.

diff --git a/Engine/Source/Programs/CrashReporter/CrashReportClient/Private/CrashDescription.h b/Engine/Source/Programs/CrashReporter/CrashReportClient/Private/CrashDescription.h
--- a/Engine/Source/Programs/CrashReporter/CrashReportClient/Private/CrashDescription.h
+++ b/Engine/Source/Programs/CrashReporter/CrashReportClient/Private/CrashDescription.h
@@ -81,11 +81,6 @@ namespace LexicalConversion
 	{
 		FEngineVersion::Parse( Buffer, OutValue );
 	}
-
-	inline void FromString( FString& OutValue, const TCHAR* Buffer )
-	{
-		OutValue = Buffer;
-	}
 }
 
 /** Simple crash property. Only for string values. */
diff --git a/Engine/Source/Runtime/Core/Public/Containers/UnrealString.h b/Engine/Source/Runtime/Core/Public/Containers/UnrealString.h
--- a/Engine/Source/Runtime/Core/Public/Containers/UnrealString.h
+++ b/Engine/Source/Runtime/Core/Public/Containers/UnrealString.h
@@ -1812,6 +1812,7 @@ namespace LexicalConversion
 	inline void FromString(float& OutValue,		const TCHAR* Buffer)	{	OutValue = FCString::Atof(Buffer);		}
 	inline void FromString(double& OutValue, 	const TCHAR* Buffer)	{	OutValue = FCString::Atod(Buffer);		}
 	inline void FromString(bool& OutValue, 		const TCHAR* Buffer)	{	OutValue = FCString::ToBool(Buffer);	}
+	inline void FromString(FString& OutValue, 	const TCHAR* Buffer)	{	OutValue = Buffer;						}
 
 	/** Convert numeric types to a string */
 	template<typename T>
