Fix for change to MountPak execute

diff --git a/Engine/Source/Runtime/Core/Private/IOS/IOSChunkInstaller.cpp b/Engine/Source/Runtime/Core/Private/IOS/IOSChunkInstaller.cpp
--- a/Engine/Source/Runtime/Core/Private/IOS/IOSChunkInstaller.cpp
+++ b/Engine/Source/Runtime/Core/Private/IOS/IOSChunkInstaller.cpp
@@ -158,7 +158,7 @@ bool FIOSChunkInstall::PrioritizeChunk( uint32 ChunkID, EChunkPriority::Type Pri
                 FString PakFile = FString::Printf(TEXT("pakchunk%i-ios.pak"), BlockID);
                 NSString* ResourcePath = [Bundle pathForResource:PakFile.GetNSString() ofType:nil];
                 NSLog(@"ResourcePath: %@", ResourcePath);
-                bool bMounted = FCoreDelegates::OnMountPak.Execute(FString(ResourcePath), 0);
+                bool bMounted = FCoreDelegates::OnMountPak.Execute(FString(ResourcePath), 0, nullptr);
                 if (bMounted)
                 {
                     MountedChunks.Add(BlockID);
