Copying //UE4/Dev-Rendering to Dev-Main (//UE4/Dev-Main)

==========================
MAJOR FEATURES + CHANGES
==========================

Change 2821445 on 2016/01/08 by [REDACTED_PERSON]
	More vertex factory improvements, storing off particle vertex factories on the scene proxy instead of the dynamic data to avoid recreating all the time; saves up to 2ms render thread time according to QA's testing.

Change 2821520 on 2016/01/08 by [REDACTED_PERSON]

	Coloring subuv modules green for easier visual ID

Change 2823479 on 2016/01/11 by [REDACTED_PERSON]

	Updated Lightmass HLOD logic to avoid HLODs shadowing non-related meshes. Duplicated CL#2823104 from Dev-General.

Change 2823570 on 2016/01/11 by [REDACTED_PERSON]

	Introduce multiplier that controls decal fade speed.

	#1777

Change 2823615 on 2016/01/11 by [REDACTED_PERSON]

	Fixed stencil ref multithreading issue.
	Fixed state caching when depth range is enabled.

Change 2823652 on 2016/01/11 by [REDACTED_PERSON]

	Rename FadeSpeedScale -> FadeDurationScale to be logically more consistent.

Change 2824065 on 2016/01/11 by [REDACTED_PERSON]

	Fixed last viewrect motion blur bug. Enabled new motion blur algorithm for default.

Change 2825432 on 2016/01/12 by [REDACTED_PERSON]

	Store off view matrices at at time of freezing and base lod selection useing relevant matrices, thus allows lods to be frozen. #OR-10918

Change 2825971 on 2016/01/12 by [REDACTED_PERSON]

	New motion blur enabled.

Change 2825974 on 2016/01/12 by [REDACTED_PERSON]

	Fixed refraction check value. 1 does nothing not 0.

Change 2825975 on 2016/01/12 by [REDACTED_PERSON]

	Cloth gets skylight for movable sky.

Change 2827519 on 2016/01/13 by [REDACTED_PERSON]

	ALLOW_UAV_CONDITION did not have a safe fallback when not SM5.0 && COMPILER_SUPPORTS_ATTRIBUTES.

Change 2830172 on 2016/01/15 by [REDACTED_PERSON]

	DR - Minor cleanup
	- Renamed Vertex Factories' struct Data to struct FData
	- Removed Data type on FVertexFactory

Change 2830242 on 2016/01/15 by [REDACTED_PERSON]

	DR - Prep cleanup for gpu morph targets
	- Split common code for  GPU skin cache into a base class
	- Moved some local static arrays from UpdateMorphVertexBuffer() to static members

Change 2830455 on 2016/01/15 by [REDACTED_PERSON]

	DR - Compile fix from bad merge

Change 2832023 on 2016/01/18 by [REDACTED_PERSON]

	DR - Removed TangentZDelta_DEPRECATED from FVertexAnimDelta

Change 2832067 on 2016/01/18 by [REDACTED_PERSON]

	UE4 - Changed PC to default to parallel rendering when not in editor. Fixed lack of a stall on texture locks and unlocks coming from texture streamer. Fixed a few cases where stuff was being added to rhicommandlists even when we were bypassed.

Change 2834379 on 2016/01/19 by [REDACTED_PERSON]

	UE4 - fix perf regression related to cvar

Change 2834864 on 2016/01/19 by [REDACTED_PERSON]

	Fixing potential crash with auto-kill trail emitters, fixing use of the wrong flag to auto-deactivate

Change 2835777 on 2016/01/20 by [REDACTED_PERSON]

	EyeAdaptation - using a screen center focus in the weights

	related to: UE-15509.  This is adding the ability to focus the basic eye-adaptation region in the center of the screen, and cvar functionality for paragon testing on ps4

Change 2835778 on 2016/01/20 by [REDACTED_PERSON]

	EyeAdapation - DefaultFeature for method
	Adding a default feature cvar for eye adaptation method

Change 2837410 on 2016/01/20 by [REDACTED_PERSON]

	OR-13213   SetupPerObjectProjection()
	#test:PC

diff --git a/Engine/Config/BaseEditorPerProjectUserSettings.ini b/Engine/Config/BaseEditorPerProjectUserSettings.ini
--- a/Engine/Config/BaseEditorPerProjectUserSettings.ini
+++ b/Engine/Config/BaseEditorPerProjectUserSettings.ini
@@ -535,6 +535,8 @@ ModuleColor_Spawn_Unselected=(R=200,G=100,B=100)
 ModuleColor_Spawn_Selected=(R=255,G=50,B=50)
 ModuleColor_Light_Unselected=(B=49,G=40,R=90)
 ModuleColor_Light_Selected=(B=0,G=100,R=255)
+ModuleColor_SubUV_Unselected=(B=49,G=90,R=40)
+ModuleColor_SubUV_Selected=(B=100,G=200,R=50)
 ModuleColor_Required_Unselected=(R=200,G=200,B=100)
 ModuleColor_Required_Selected=(R=255,G=225,B=50)
 ModuleColor_Event_Unselected=(R=64,G=64,B=255)
diff --git a/Engine/Plugins/2D/Paper2D/Source/Paper2D/Private/PaperRenderSceneProxy.cpp b/Engine/Plugins/2D/Paper2D/Source/Paper2D/Private/PaperRenderSceneProxy.cpp
--- a/Engine/Plugins/2D/Paper2D/Source/Paper2D/Private/PaperRenderSceneProxy.cpp
+++ b/Engine/Plugins/2D/Paper2D/Source/Paper2D/Private/PaperRenderSceneProxy.cpp
@@ -80,7 +80,7 @@ void FPaperSpriteVertexFactory::Init(const FPaperSpriteVertexBuffer* InVertexBuf
 	if (IsInRenderingThread())
 	{
 		// Initialize the vertex factory's stream components.
-		DataType NewData;
+		FDataType NewData;
 		NewData.PositionComponent = STRUCTMEMBER_VERTEXSTREAMCOMPONENT(InVertexBuffer, FPaperSpriteVertex, Position, VET_Float3);
 		NewData.TangentBasisComponents[0] = STRUCTMEMBER_VERTEXSTREAMCOMPONENT(InVertexBuffer, FPaperSpriteVertex, TangentX, VET_PackedNormal);
 		NewData.TangentBasisComponents[1] = STRUCTMEMBER_VERTEXSTREAMCOMPONENT(InVertexBuffer, FPaperSpriteVertex, TangentZ, VET_PackedNormal);
@@ -120,7 +120,7 @@ class FPaperSpriteVertexFactoryDummy : public FLocalVertexFactory
 public:
 	void AllocateStuff()
 	{
-		FLocalVertexFactory::DataType VertData;
+		FLocalVertexFactory::FDataType VertData;
 
 		VertData.PositionComponent = FVertexStreamComponent(
 			&GDummyMaterialSpriteVertexBuffer,
diff --git a/Engine/Plugins/Experimental/CodeEditor/Source/CodeEditor/Private/DirectoryScanner.cpp b/Engine/Plugins/Experimental/CodeEditor/Source/CodeEditor/Private/DirectoryScanner.cpp
--- a/Engine/Plugins/Experimental/CodeEditor/Source/CodeEditor/Private/DirectoryScanner.cpp
+++ b/Engine/Plugins/Experimental/CodeEditor/Source/CodeEditor/Private/DirectoryScanner.cpp
@@ -39,7 +39,7 @@ struct FDirectoryScannerCommand : public IQueuedWork
 		class FDirectoryEnumerator : public IPlatformFile::FDirectoryVisitor
 		{
 		public:
-			FDirectoryEnumerator(TLockFreePointerListUnordered<FDirectoryResult>& InFoundFiles)
+			FDirectoryEnumerator(TLockFreePointerListUnordered<FDirectoryResult, PLATFORM_CACHE_LINE_SIZE>& InFoundFiles)
 				: FoundFiles(InFoundFiles)
 			{
 			}
@@ -58,7 +58,7 @@ struct FDirectoryScannerCommand : public IQueuedWork
 				return true;
 			}
 
-			TLockFreePointerListUnordered<FDirectoryResult>& FoundFiles;
+			TLockFreePointerListUnordered<FDirectoryResult, PLATFORM_CACHE_LINE_SIZE>& FoundFiles;
 		};
 
 		FDirectoryEnumerator DirectoryEnumerator(FoundFiles);
@@ -73,7 +73,7 @@ struct FDirectoryScannerCommand : public IQueuedWork
 
 	FOnDirectoryScanned OnDirectoryScanned;
 
-	TLockFreePointerListUnordered<FDirectoryResult> FoundFiles;
+	TLockFreePointerListUnordered<FDirectoryResult, PLATFORM_CACHE_LINE_SIZE> FoundFiles;
 
 	volatile int32 bExecuted;
 };
diff --git a/Engine/Plugins/Runtime/CableComponent/Source/CableComponent/Private/CableComponent.cpp b/Engine/Plugins/Runtime/CableComponent/Source/CableComponent/Private/CableComponent.cpp
--- a/Engine/Plugins/Runtime/CableComponent/Source/CableComponent/Private/CableComponent.cpp
+++ b/Engine/Plugins/Runtime/CableComponent/Source/CableComponent/Private/CableComponent.cpp
@@ -48,7 +48,7 @@ public:
 		if(IsInRenderingThread())
 		{
 			// Initialize the vertex factory's stream components.
-			DataType NewData;
+			FDataType NewData;
 			NewData.PositionComponent = STRUCTMEMBER_VERTEXSTREAMCOMPONENT(VertexBuffer,FDynamicMeshVertex,Position,VET_Float3);
 			NewData.TextureCoordinates.Add(
 				FVertexStreamComponent(VertexBuffer,STRUCT_OFFSET(FDynamicMeshVertex,TextureCoordinate),sizeof(FDynamicMeshVertex),VET_Float2)
@@ -65,7 +65,7 @@ public:
 				const FCableVertexBuffer*,VertexBuffer,VertexBuffer,
 			{
 				// Initialize the vertex factory's stream components.
-				DataType NewData;
+				FDataType NewData;
 				NewData.PositionComponent = STRUCTMEMBER_VERTEXSTREAMCOMPONENT(VertexBuffer,FDynamicMeshVertex,Position,VET_Float3);
 				NewData.TextureCoordinates.Add(
 					FVertexStreamComponent(VertexBuffer,STRUCT_OFFSET(FDynamicMeshVertex,TextureCoordinate),sizeof(FDynamicMeshVertex),VET_Float2)
diff --git a/Engine/Plugins/Runtime/CustomMeshComponent/Source/CustomMeshComponent/Private/CustomMeshComponent.cpp b/Engine/Plugins/Runtime/CustomMeshComponent/Source/CustomMeshComponent/Private/CustomMeshComponent.cpp
--- a/Engine/Plugins/Runtime/CustomMeshComponent/Source/CustomMeshComponent/Private/CustomMeshComponent.cpp
+++ b/Engine/Plugins/Runtime/CustomMeshComponent/Source/CustomMeshComponent/Private/CustomMeshComponent.cpp
@@ -58,7 +58,7 @@ public:
 	{
 		check(IsInRenderingThread());
 
-		DataType NewData;
+		FDataType NewData;
 		NewData.PositionComponent = STRUCTMEMBER_VERTEXSTREAMCOMPONENT(VertexBuffer, FDynamicMeshVertex, Position, VET_Float3);
 		NewData.TextureCoordinates.Add(
 			FVertexStreamComponent(VertexBuffer, STRUCT_OFFSET(FDynamicMeshVertex, TextureCoordinate), sizeof(FDynamicMeshVertex), VET_Float2)
diff --git a/Engine/Plugins/Runtime/ProceduralMeshComponent/Source/Private/ProceduralMeshComponent.cpp b/Engine/Plugins/Runtime/ProceduralMeshComponent/Source/Private/ProceduralMeshComponent.cpp
--- a/Engine/Plugins/Runtime/ProceduralMeshComponent/Source/Private/ProceduralMeshComponent.cpp
+++ b/Engine/Plugins/Runtime/ProceduralMeshComponent/Source/Private/ProceduralMeshComponent.cpp
@@ -83,7 +83,7 @@ public:
 		check(IsInRenderingThread());
 
 		// Initialize the vertex factory's stream components.
-		DataType NewData;
+		FDataType NewData;
 		NewData.PositionComponent = STRUCTMEMBER_VERTEXSTREAMCOMPONENT(VertexBuffer, FDynamicMeshVertex, Position, VET_Float3);
 		NewData.TextureCoordinates.Add(
 			FVertexStreamComponent(VertexBuffer, STRUCT_OFFSET(FDynamicMeshVertex, TextureCoordinate), sizeof(FDynamicMeshVertex), VET_Float2)
diff --git a/Engine/Shaders/Definitions.usf b/Engine/Shaders/Definitions.usf
--- a/Engine/Shaders/Definitions.usf
+++ b/Engine/Shaders/Definitions.usf
@@ -208,6 +208,9 @@
 #ifndef FLATTEN
 #define FLATTEN
 #endif
+#ifndef ALLOW_UAV_CONDITION	
+#define ALLOW_UAV_CONDITION	
+#endif
 #endif
 
 #define ISOLATE
diff --git a/Engine/Shaders/PostProcessEyeAdaptation.usf b/Engine/Shaders/PostProcessEyeAdaptation.usf
--- a/Engine/Shaders/PostProcessEyeAdaptation.usf
+++ b/Engine/Shaders/PostProcessEyeAdaptation.usf
@@ -10,36 +10,52 @@
 
 // prior frame's eye adaptation exposure scale
 Texture2D EyeAdaptationTexture;
+uint2 EyeAdaptionSrcExtent;
 
 
 // Compute the average alpha stored in a 2D texture.
-float ComputeTextureAverageAlpha(Texture2D InBuffer)
+float ComputeWeightedTextureAverageAlpha(Texture2D InBuffer, float slope)
 {
-	uint Imax;
-	uint Jmax;
-	InBuffer.GetDimensions(Imax, Jmax);
+	uint Xmax = EyeAdaptionSrcExtent.x;
+	uint Ymax = EyeAdaptionSrcExtent.y;
+
+	float InvXmax = 1.f / float(Xmax);
+	float InvYmax = 1.f / float(Ymax);
+
+	// use product of linear weight in x and y.  
 
 	float Average = 0.0f;
-	for (uint i = 0; i < Imax; ++i)
+	float WeightTotal = 0.0f;
+
+
+	for (uint i = 0; i < Xmax; ++i)
 	{
-		for (uint j = 0; j < Jmax; ++j)
+		float IWeight = AdaptationWeight( i * InvXmax, slope);  
+		for (uint j = 0; j < Ymax; ++j)
 		{
+			float JWeight = AdaptationWeight( j * InvYmax, slope);
+			float Weight = max(IWeight * JWeight, 0.05f);
+			
+			WeightTotal += Weight;
+
 			//accumulate values from alpha channel
-			Average += InBuffer.Load(int3(i, j, 0)).w;
+			float Sample = InBuffer.Load(int3(i, j, 0)).w;
+			Average += Weight * Sample;
 		}
 	}
 
-	Average /= float(Imax * Jmax);
+	Average /= WeightTotal;
 	return Average;
 }
 
+
 // @param FrameTime in seconds
 // @return smoothed exposure
 float ComputeEyeAdaptation(float OldExposure, float TargetExposure, float FrameTime)
 {
 	float Diff = TargetExposure - OldExposure;
 
-	const float EyeAdaptionSpeedUp = EyeAdaptationParams[1].z;
+	const float EyeAdaptionSpeedUp   = EyeAdaptationParams[1].z;
 	const float EyeAdaptionSpeedDown = EyeAdaptationParams[1].w;
 
     float AdaptionSpeed = (Diff > 0) ? EyeAdaptionSpeedUp : EyeAdaptionSpeedDown;
@@ -117,9 +133,16 @@ void MainLogLuminance2ExposureScalePS(
 	)
 {
 	float2 UV = UVAndScreenPos.xy;
-
+	
+	// Weighting is a triangle function focused at the center of the screen
+	// with the product of linear decay in x and y.  
+	// WeightSlope: 0 is uniform weighting.
+	//              1 decays to the edges.
+	//              >1 focuses more in the center of the screen.
+	float WeightSlope = EyeAdaptationParams[2].w;
+	
 	// Compute scaled Log Luminance Average 
-	float LogLumAve = ComputeTextureAverageAlpha(PostprocessInput0);
+	float LogLumAve = ComputeWeightedTextureAverageAlpha(PostprocessInput0, WeightSlope);
 
 	// Correct for [0,1] scaling
 	LogLumAve = (LogLumAve - EyeAdaptationParams[2].y) / EyeAdaptationParams[2].x;
diff --git a/Engine/Shaders/PostProcessHistogramCommon.usf b/Engine/Shaders/PostProcessHistogramCommon.usf
--- a/Engine/Shaders/PostProcessHistogramCommon.usf
+++ b/Engine/Shaders/PostProcessHistogramCommon.usf
@@ -131,4 +131,14 @@ float ComputeEyeAdaptationExposure(Texture2D HistogramTexture)
 	float ClampedAdaptedLuminance = clamp(UnclampedAdaptedLuminance, EyeAdaptationParams[0].z, EyeAdaptationParams[0].w);
 
 	return ClampedAdaptedLuminance;
-}
\ No newline at end of file
+}
+
+// Computation of triangle function on normalized space
+// used in Basic EyeAdaptation in weighting the center of the screen
+// @return range of 0..1 with 1 nx = 0.5
+// @param nx in 0..1 
+// @param slope assumed possitive
+float AdaptationWeight(float nx, float slope)
+{
+	return max(1.f - 2. * slope * abs(nx - 0.5f), 0.0f);
+}
diff --git a/Engine/Shaders/PostProcessTonemap.usf b/Engine/Shaders/PostProcessTonemap.usf
--- a/Engine/Shaders/PostProcessTonemap.usf
+++ b/Engine/Shaders/PostProcessTonemap.usf
@@ -29,6 +29,11 @@ float2 TonemapperParams;
 float3 GrainRandomFull;
 float3 GrainScaleBiasJitter;
 
+#if EYEADAPTATION_EXPOSURE_FIX != 1
+ // Default value used when eye adaption is not avalible
+float DefaultEyeExposure;
+#endif
+
 half GrainFromUV(float2 GrainUV) 
 {
 	half Grain = frac(sin(GrainUV.x + GrainUV.y * 543.31) *  493013.0);
@@ -101,12 +106,12 @@ void MainVS(
 {
 	DrawRectangle(InPosition, InTexCoord, OutPosition, OutTexCoord);
 
-#if FEATURE_LEVEL >= FEATURE_LEVEL_SM5
+#if EYEADAPTATION_EXPOSURE_FIX == 1
 	// texture can be GWhiteTexture which is 1x1. It's important we don't read outside bounds.
 	OutExposureScaleVignette.x = EyeAdaptation.Load(int3(0, 0, 0)).r;
 #else
-	// Eye adaptation is not yet supported
-	OutExposureScaleVignette.x = 1.0f;
+	// Eye adaptation is not yet supported, use the default value.
+	OutExposureScaleVignette.x = DefaultEyeExposure;
 #endif
 	
 	// -1..1 1..-1 in the viewport
@@ -280,10 +285,6 @@ void MainPS(
 			LinearColor += CombinedBloom.rgb * (ColorScale1.rgb + BloomDirtMaskColor); 
 		#endif
 
-		#if NO_EYEADAPTATION_EXPOSURE_FIX
-			ExposureScale = BloomDirtMask.Tint.w;
-		#endif
-
 		LinearColor *= ExposureScale;
 
 		#if USE_VIGNETTE
diff --git a/Engine/Shaders/PostProcessVisualizeHDR.usf b/Engine/Shaders/PostProcessVisualizeHDR.usf
--- a/Engine/Shaders/PostProcessVisualizeHDR.usf
+++ b/Engine/Shaders/PostProcessVisualizeHDR.usf
@@ -129,6 +129,11 @@ void MainPS(noperspective float4 UVAndScreenPos : TEXCOORD0, float4 SvPosition :
 	float LuminanceVal = max(SceneColor.r, max(SceneColor.g, SceneColor.b));
 	OutColor = float4(Colorize(ComputeHistogramPositionFromLuminance(LuminanceVal)), 1.0f);
 
+	// Compute and apply weight value at this location 
+	float WeightSlope = EyeAdaptationParams[2].w;
+	float2 NdPos = ViewportSize.zw * (PixelPos.xy - ViewportRect.xy);
+	float weight = AdaptationWeight(NdPos.x, WeightSlope) * AdaptationWeight(NdPos.y, WeightSlope);
+	OutColor *= max(weight, 0.05f);  
 
 	// crosshair
 	{
diff --git a/Engine/Shaders/SkyLighting.usf b/Engine/Shaders/SkyLighting.usf
--- a/Engine/Shaders/SkyLighting.usf
+++ b/Engine/Shaders/SkyLighting.usf
@@ -92,6 +92,11 @@ void SkyLightDiffusePS(
 			DiffuseColor = PI * HairShading( GBuffer, L, V, N, 1, 0, 0.2, uint2(0,0) );
 		}
 
+		BRANCH if( GBuffer.ShadingModelID == SHADINGMODELID_CLOTH )
+		{
+			DiffuseColor +=  GBuffer.CustomData.rgb * GBuffer.CustomData.a;
+		}
+
 		// Compute the preconvolved incoming lighting with the bent normal direction
 		float3 DiffuseLookup = GetSkySHDiffuse(SkyLightingNormal) * Frame.SkyLightColor.rgb;
 
diff --git a/Engine/Source/Developer/MaterialUtilities/Private/MeshRendering.cpp b/Engine/Source/Developer/MaterialUtilities/Private/MeshRendering.cpp
--- a/Engine/Source/Developer/MaterialUtilities/Private/MeshRendering.cpp
+++ b/Engine/Source/Developer/MaterialUtilities/Private/MeshRendering.cpp
@@ -162,7 +162,7 @@ public:
 	/** Default constructor. */
 	FMeshVertexFactory()
 	{
-		FLocalVertexFactory::DataType Data;
+		FLocalVertexFactory::FDataType Data;
 
 		// position
 		Data.PositionComponent = FVertexStreamComponent(
@@ -218,7 +218,7 @@ public:
 		ENQUEUE_UNIQUE_RENDER_COMMAND_TWOPARAMETER(
 			FMeshVertexFactoryConstructor,
 			FMeshVertexFactory*, FactoryParam, this,
-			FLocalVertexFactory::DataType, DataParam, Data,
+			FLocalVertexFactory::FDataType, DataParam, Data,
 			{
 				FactoryParam->SetData(DataParam);
 			}
diff --git a/Engine/Source/Editor/Cascade/Private/CascadeEmitterCanvasClient.cpp b/Engine/Source/Editor/Cascade/Private/CascadeEmitterCanvasClient.cpp
--- a/Engine/Source/Editor/Cascade/Private/CascadeEmitterCanvasClient.cpp
+++ b/Engine/Source/Editor/Cascade/Private/CascadeEmitterCanvasClient.cpp
@@ -95,7 +95,8 @@ FCascadeEmitterCanvasClient::FCascadeEmitterCanvasClient(TWeakPtr<FCascade> InCa
 													EditorOptions->ModuleColor_Spawn_Unselected,
 													EditorOptions->ModuleColor_Required_Unselected,
 													EditorOptions->ModuleColor_Event_Unselected,
-													EditorOptions->ModuleColor_Light_Unselected
+													EditorOptions->ModuleColor_Light_Unselected,
+													EditorOptions->ModuleColor_SubUV_Unselected
 												};
 
 	FColor ColorOptionsSelected[EPMT_MAX] =	{	EditorOptions->ModuleColor_General_Selected, 
@@ -105,7 +106,8 @@ FCascadeEmitterCanvasClient::FCascadeEmitterCanvasClient(TWeakPtr<FCascade> InCa
 												EditorOptions->ModuleColor_Spawn_Selected,
 												EditorOptions->ModuleColor_Required_Selected,
 												EditorOptions->ModuleColor_Event_Selected,
-												EditorOptions->ModuleColor_Light_Selected
+												EditorOptions->ModuleColor_Light_Selected,
+												EditorOptions->ModuleColor_SubUV_Selected
 											};
 
 	for (int32 i = 0; i < EPMT_MAX; ++i)
diff --git a/Engine/Source/Editor/UnrealEd/Classes/Preferences/CascadeOptions.h b/Engine/Source/Editor/UnrealEd/Classes/Preferences/CascadeOptions.h
--- a/Engine/Source/Editor/UnrealEd/Classes/Preferences/CascadeOptions.h
+++ b/Engine/Source/Editor/UnrealEd/Classes/Preferences/CascadeOptions.h
@@ -78,6 +78,12 @@ class UNREALED_API UCascadeOptions : public UObject
 	UPROPERTY(EditAnywhere, config, Category=Options)
 	FColor ModuleColor_Light_Selected;
 
+	UPROPERTY(EditAnywhere, config, Category = Options)
+	FColor ModuleColor_SubUV_Unselected;
+
+	UPROPERTY(EditAnywhere, config, Category = Options)
+	FColor ModuleColor_SubUV_Selected;
+
 	UPROPERTY(EditAnywhere, config, Category=Options)
 	FColor ModuleColor_Required_Unselected;
 
diff --git a/Engine/Source/Editor/UnrealEd/Private/Lightmass/LightmassLandscapeRender.cpp b/Engine/Source/Editor/UnrealEd/Private/Lightmass/LightmassLandscapeRender.cpp
--- a/Engine/Source/Editor/UnrealEd/Private/Lightmass/LightmassLandscapeRender.cpp
+++ b/Engine/Source/Editor/UnrealEd/Private/Lightmass/LightmassLandscapeRender.cpp
@@ -51,7 +51,7 @@ public:
 	/** Default constructor. */
 	FLightmassLandscapeVertexFactory()
 	{
-		FLocalVertexFactory::DataType Data;
+		FLocalVertexFactory::FDataType Data;
 		// Position
 		Data.PositionComponent = STRUCTMEMBER_VERTEXSTREAMCOMPONENT(&LightmassLandscapeVertexBuffer, FLightmassLandscapeVertex, Position, VET_Float3);
 		// Tangents
diff --git a/Engine/Source/Programs/UnrealLightmass/Private/Lighting/Embree.cpp b/Engine/Source/Programs/UnrealLightmass/Private/Lighting/Embree.cpp
--- a/Engine/Source/Programs/UnrealLightmass/Private/Lighting/Embree.cpp
+++ b/Engine/Source/Programs/UnrealLightmass/Private/Lighting/Embree.cpp
@@ -429,12 +429,13 @@ bool EmbreeMemoryMonitor(const ssize_t bytes, const bool post)
 FEmbreeAggregateMesh::FEmbreeAggregateMesh(const FScene& InScene):
 	FStaticLightingAggregateMesh(InScene),
 	EmbreeDevice(NULL),
-	EmbreeScene(NULL)
+	EmbreeScene(NULL),
+	TotalNumTriangles(0)
 {
 	rtcDeviceSetMemoryMonitorFunction(InScene.EmbreeDevice, EmbreeMemoryMonitor);
 
 	EmbreeDevice = InScene.EmbreeDevice;
-	EmbreeScene = rtcDeviceNewScene(InScene.EmbreeDevice, RTC_SCENE_STATIC | RTC_SCENE_COHERENT, RTC_INTERSECT1);
+	EmbreeScene = rtcDeviceNewScene(InScene.EmbreeDevice, RTC_SCENE_STATIC, RTC_INTERSECT1);
 	check(rtcDeviceGetError(EmbreeDevice) == RTC_NO_ERROR);
 }
 
@@ -470,6 +471,8 @@ void FEmbreeAggregateMesh::AddMesh(const FStaticLightingMesh* Mesh, const FStati
 		// Sum the total triangle area of everything in the aggregate mesh
 		SceneSurfaceArea += Geo->SurfaceArea;
 		SceneSurfaceAreaWithinImportanceVolume += Geo->SurfaceAreaWithinImportanceVolume;
+
+		TotalNumTriangles += Mesh->NumTriangles;
 	}
 }
 
@@ -492,10 +495,14 @@ void FEmbreeAggregateMesh::DumpStats() const
 		LightmapUV += Geo->LightmapUVs.GetAllocatedSize();
 	}
 
-	UE_LOG(LogLightmass, Log, TEXT("Embree MeshInfos      : %7.1fMb"), MeshInfoSize / 1048576.0f);
-	UE_LOG(LogLightmass, Log, TEXT("Embree UVs            : %7.1fMb"), UVSize / 1048576.0f);
-	UE_LOG(LogLightmass, Log, TEXT("Embree LightmapUVs    : %7.1fMb"), LightmapUV / 1048576.0f);
-	UE_LOG(LogLightmass, Log, TEXT("Embree Allocations    : %7.1fMb"), GEmbreeAllocatedSpace / 1048576.0f);
+	UE_LOG(LogLightmass, Log, TEXT("\n"));
+	UE_LOG(LogLightmass, Log, TEXT("Collision Mesh Overview:"));
+	UE_LOG(LogLightmass, Log, TEXT("Num Triangles         : %d"), TotalNumTriangles);
+	UE_LOG(LogLightmass, Log, TEXT("MeshInfos             : %7.1fMb"), MeshInfoSize / 1048576.0f);
+	UE_LOG(LogLightmass, Log, TEXT("UVs                   : %7.1fMb"), UVSize / 1048576.0f);
+	UE_LOG(LogLightmass, Log, TEXT("LightmapUVs           : %7.1fMb"), LightmapUV / 1048576.0f);
+	UE_LOG(LogLightmass, Log, TEXT("Embree Used Memory    : %7.1fMb"), GEmbreeAllocatedSpace / 1048576.0f);
+	UE_LOG(LogLightmass, Log, TEXT("\n"));
 }
 
 bool FEmbreeAggregateMesh::IntersectLightRay(
diff --git a/Engine/Source/Programs/UnrealLightmass/Private/Lighting/Embree.h b/Engine/Source/Programs/UnrealLightmass/Private/Lighting/Embree.h
--- a/Engine/Source/Programs/UnrealLightmass/Private/Lighting/Embree.h
+++ b/Engine/Source/Programs/UnrealLightmass/Private/Lighting/Embree.h
@@ -124,9 +124,13 @@ namespace Lightmass
 		/** Information about the meshes used in the kDOP tree. */
 		TArray<const FEmbreeGeometry*> MeshInfos;
 
-		/*Embree scene */
+		/** Embree scene */
 		RTCDevice EmbreeDevice;
 		RTCScene EmbreeScene;
+		
+		/** Total number of triangles in the shadow mesh */
+		int32 TotalNumTriangles;
+
 	};
 
 	class FEmbreeVerifyAggregateMesh final : public FStaticLightingAggregateMesh
diff --git a/Engine/Source/Programs/UnrealLightmass/Private/LightmassCore/Math/LMMathSSE.h b/Engine/Source/Programs/UnrealLightmass/Private/LightmassCore/Math/LMMathSSE.h
--- a/Engine/Source/Programs/UnrealLightmass/Private/LightmassCore/Math/LMMathSSE.h
+++ b/Engine/Source/Programs/UnrealLightmass/Private/LightmassCore/Math/LMMathSSE.h
@@ -531,6 +531,15 @@ FORCEINLINE void VectorMatrixMultiply( void *Result, const void* Matrix1, const
  */
 #define VectorBitwiseOR( Vec1, Vec2 )	_mm_or_ps( (Vec1), (Vec2) )
 
+/**
+ * Returns the bitwise XOR.
+ *
+ * @param	Vec1	Vector to XOR
+ * @param	Vec2	Vector to XOR
+ * @return	bitwise per component XOR operation.
+ */
+#define VectorBitwiseXOR( Vec1, Vec2 )	_mm_xor_ps( (Vec1), (Vec2) )
+
 /**
  * Returns the bitwise AND NOT.
  *
diff --git a/Engine/Source/Programs/UnrealLightmass/Private/LightmassCore/Math/LMkDOP.h b/Engine/Source/Programs/UnrealLightmass/Private/LightmassCore/Math/LMkDOP.h
--- a/Engine/Source/Programs/UnrealLightmass/Private/LightmassCore/Math/LMkDOP.h
+++ b/Engine/Source/Programs/UnrealLightmass/Private/LightmassCore/Math/LMkDOP.h
@@ -211,10 +211,19 @@ static const VectorRegister GSmallNumber = { 0.0001f, 0.0001f, 0.0001f, 0.0001f
 //extern const VectorRegister GSmallNumber;
 
 static const VectorRegister GZeroVectorRegister = { 0, 0, 0, 0 };
+static const VectorRegister GFullMaskVectorRegister = MakeVectorRegister((uint32)-1, (uint32)-1, (uint32)-1, (uint32)-1);
 
 static const VectorRegister IndexNoneVectorRegister = MakeVectorRegister((uint32)INDEX_NONE, (uint32)INDEX_NONE, (uint32)INDEX_NONE, (uint32)INDEX_NONE);
 static const VectorRegister VectorNegativeOne = MakeVectorRegister( -1.0f, -1.0f, -1.0f, -1.0f );
 
+// LOD masks
+static const VectorRegister HLODTreeIndexMask = MakeVectorRegister((uint32)0xFFFF0000, (uint32)0xFFFF0000, (uint32)0xFFFF0000, (uint32)0xFFFF0000);
+static const VectorRegister HLODRangeStartIndexMask = MakeVectorRegister((uint32)0xFFFF, (uint32)0xFFFF, (uint32)0xFFFF, (uint32)0xFFFF);
+static const VectorRegister HLODRangeEndIndexMask = MakeVectorRegister((uint32)0xFFFF0000, (uint32)0xFFFF0000, (uint32)0xFFFF0000, (uint32)0xFFFF0000);
+static const VectorRegister HLODTreeIndexNoneRegister = MakeVectorRegister(0xFFFF0000 & (uint32)INDEX_NONE, 0xFFFF0000 & (uint32)INDEX_NONE, 0xFFFF0000 & (uint32)INDEX_NONE, 0xFFFF0000 & (uint32)INDEX_NONE);
+
+static const VectorRegister LODIndexMask = MakeVectorRegister((uint32)0xFFFF, (uint32)0xFFFF, (uint32)0xFFFF, (uint32)0xFFFF);
+
 /**
  * Line vs triangle intersection test. Tests 1 line against 4 triangles at once.
  *
@@ -312,13 +321,6 @@ FORCEINLINE int32 appLineCheckTriangleSOA(const FVector3SOA& Start, const FVecto
 		}
 	}
 
-	// HLOD masks
-	static const VectorRegister HLODTreeIndexMask = MakeVectorRegister((uint32)0xFFFF0000,(uint32)0xFFFF0000,(uint32)0xFFFF0000,(uint32)0xFFFF0000);
-	static const VectorRegister HLODRangeStartIndexMask = MakeVectorRegister((uint32)0xFFFF,(uint32)0xFFFF,(uint32)0xFFFF,(uint32)0xFFFF);
-	static const VectorRegister HLODRangeEndIndexMask = MakeVectorRegister((uint32)0xFFFF0000,(uint32)0xFFFF0000,(uint32)0xFFFF0000,(uint32)0xFFFF0000);
-
-	static const VectorRegister TreeIndexNoneRegister = VectorShiftRight(VectorBitwiseAND(HLODTreeIndexMask, IndexNoneVectorRegister), 16);
-
 	// Unpack HLOD Data
 	VectorRegister HLODTreeIndexRay = VectorShiftRight(VectorBitwiseAND(HLODTreeIndexMask, LODIndices), 16);
 	VectorRegister HLODTreeIndexTri = VectorShiftRight(VectorBitwiseAND(HLODTreeIndexMask, Triangle4.LODIndices), 16);
@@ -328,7 +330,7 @@ FORCEINLINE int32 appLineCheckTriangleSOA(const FVector3SOA& Start, const FVecto
 	VectorRegister HLODRangeEndTri = VectorShiftRight(VectorBitwiseAND(HLODRangeEndIndexMask, Triangle4.HLODRange), 16);
 
 	VectorRegister IsDifferentHLOD = VectorMask_NE(HLODTreeIndexRay, HLODTreeIndexTri);
-	VectorRegister IsTriNotAHLOD = VectorBitwiseOR(VectorMask_EQ(HLODTreeIndexTri, TreeIndexNoneRegister), VectorMask_EQ(HLODTreeIndexTri, GZeroVectorRegister));
+	VectorRegister IsTriNotAHLOD = VectorBitwiseOR(VectorMask_EQ(HLODTreeIndexTri, HLODTreeIndexNoneRegister), VectorMask_EQ(HLODTreeIndexTri, GZeroVectorRegister));
 
 	// Ignore children of this HLOD node
 	VectorRegister IsRayOutsideRange = VectorBitwiseOR(VectorMask_LT(HLODRangeStartRay, HLODRangeStartTri), VectorMask_GT(HLODRangeStartRay, HLODRangeEndTri));
@@ -339,14 +341,17 @@ FORCEINLINE int32 appLineCheckTriangleSOA(const FVector3SOA& Start, const FVecto
 	VectorRegister HLODMask = VectorBitwiseAND(IsRayOutsideRange, IsTriOutsideRange);
 	HLODMask = VectorBitwiseOR(HLODMask, VectorBitwiseAND(IsRayLeafNode, IsTriLeafNode));
 
-	// Ignore interactions between unrelated HLODs
-	HLODMask = VectorBitwiseOR(HLODMask, IsDifferentHLOD);
-	HLODMask = VectorBitwiseOR(HLODMask, IsTriNotAHLOD);
-
 	// Allow if exact same mesh
 	VectorRegister SameHLODMesh = VectorMask_EQ(HLODRange, Triangle4.HLODRange);
 	HLODMask = VectorBitwiseOR(HLODMask, SameHLODMesh);
 
+	// Ignore interactions between unrelated HLODs and non-HLODs
+	HLODMask = VectorBitwiseOR(HLODMask, IsDifferentHLOD);
+	HLODMask = VectorBitwiseOR(HLODMask, IsTriNotAHLOD);
+
+	VectorRegister IsHLODNonLeafAndOtherMesh = VectorBitwiseAND(VectorMask_NE(HLODRangeStartTri, HLODRangeEndTri), IsDifferentHLOD);
+	HLODMask = VectorBitwiseAND(HLODMask, VectorBitwiseXOR(IsHLODNonLeafAndOtherMesh, GFullMaskVectorRegister));
+	
 	TriangleMask = VectorBitwiseAND(TriangleMask, HLODMask);
 	if (VectorMaskBits(TriangleMask) == 0)
 	{
@@ -354,7 +359,6 @@ FORCEINLINE int32 appLineCheckTriangleSOA(const FVector3SOA& Start, const FVecto
 	}
 
 	// Only allow intersections with the base LOD of other meshes and the LOD that is initiating the trace of the current mesh
-	static const VectorRegister LODIndexMask = MakeVectorRegister((uint32)0xFFFF, (uint32)0xFFFF, (uint32)0xFFFF, (uint32)0xFFFF);
 	VectorRegister LODIndexRay = VectorBitwiseAND(LODIndices, LODIndexMask);
 	VectorRegister LODIndexTri = VectorBitwiseAND(Triangle4.LODIndices, LODIndexMask);
 
diff --git a/Engine/Source/Runtime/Core/Private/Async/TaskGraph.cpp b/Engine/Source/Runtime/Core/Private/Async/TaskGraph.cpp
--- a/Engine/Source/Runtime/Core/Private/Async/TaskGraph.cpp
+++ b/Engine/Source/Runtime/Core/Private/Async/TaskGraph.cpp
@@ -11,12 +11,123 @@ DEFINE_STAT(STAT_FTriggerEventGraphTask);
 DEFINE_STAT(STAT_ParallelFor);
 DEFINE_STAT(STAT_ParallelForTask);
 
+#if PLATFORM_XBOXONE || PLATFORM_PS4
+#define CREATE_HIPRI_TASK_THREADS (1)
+#define CREATE_BACKGROUND_TASK_THREADS (1)
+#else
+#define CREATE_HIPRI_TASK_THREADS (0)
+#define CREATE_BACKGROUND_TASK_THREADS (0)
+#endif
+
 namespace ENamedThreads
 {
 	CORE_API Type RenderThread = ENamedThreads::GameThread; // defaults to game and is set and reset by the render thread itself
 	CORE_API Type RenderThread_Local = ENamedThreads::GameThread_Local; // defaults to game local and is set and reset by the render thread itself
+	CORE_API int32 bHasBackgroundThreads = CREATE_BACKGROUND_TASK_THREADS;
+	CORE_API int32 bHasHighPriorityThreads = CREATE_HIPRI_TASK_THREADS;
 }
 
+static int32 GNumWorkerThreadsToIgnore = 0;
+static FAutoConsoleVariableRef CVarNumWorkerThreadsToIgnore(
+	TEXT("TaskGraph.NumWorkerThreadsToIgnore"),
+	GNumWorkerThreadsToIgnore,
+	TEXT("Used to tune the number of task threads. Generally once you have found the right value, PlatformMisc::NumberOfWorkerThreadsToSpawn() should be hardcoded."),
+	ECVF_Cheat
+	);
+
+static int32 GFastScheduler = 0;
+#if USE_NEW_LOCK_FREE_LISTS
+static int32 GFastSchedulerLatched = 1;
+#else
+static int32 GFastSchedulerLatched = 0;
+#endif
+static FAutoConsoleVariableRef CVarFastScheduler(
+	TEXT("TaskGraph.FastScheduler"),
+	GFastScheduler,
+	TEXT("If > 0, then use the new experimental anythread task scheduler."),
+	ECVF_Cheat
+	);
+
+// move to platform abstraction
+#if PLATFORM_XBOXONE || PLATFORM_PS4
+#define PLATFORM_OK_TO_BURN_CPU (1)
+#else
+#define PLATFORM_OK_TO_BURN_CPU (0)
+#endif
+
+#define USE_INTRUSIVE_TASKQUEUES (USE_NEW_LOCK_FREE_LISTS)
+
+
+#if CREATE_HIPRI_TASK_THREADS || CREATE_BACKGROUND_TASK_THREADS
+	static int32 GConsoleSpinMode = 0; // when we have multiple task thread banks, we don't want to do spin mode as a high priority thread spinning would interfere with lower priority work
+	static void ThreadSwitchForABTest(const TArray<FString>& Args)
+	{
+		if (Args.Num() == 2)
+		{
+#if CREATE_HIPRI_TASK_THREADS
+			ENamedThreads::bHasHighPriorityThreads = !!FCString::Atoi(*Args[0]);
+#endif
+#if CREATE_BACKGROUND_TASK_THREADS
+			ENamedThreads::bHasBackgroundThreads = !!FCString::Atoi(*Args[1]);
+#endif
+		}
+		else
+		{
+			UE_LOG(LogConsoleResponse, Display, TEXT("This command requires two arguments, both 0 or 1 to control the use of high priority and background priority threads, respectively."));
+		}
+		UE_LOG(LogConsoleResponse, Display, TEXT("High priority task threads: %d    Bacxkground priority threads: %d"), ENamedThreads::bHasHighPriorityThreads, ENamedThreads::bHasBackgroundThreads);
+	}
+
+	static FAutoConsoleCommand ThreadSwitchForABTestCommand(
+		TEXT("TaskGraph.ABTestThreads"),
+		TEXT("Takes two 0/1 arguments. Equivalent to setting TaskGraph.UseHiPriThreads and TaskGraph.UseBackgroundThreads, respectively. Packages as one command for use with the abtest command."),
+		FConsoleCommandWithArgsDelegate::CreateStatic(&ThreadSwitchForABTest)
+		);
+
+#else
+	static int32 GConsoleSpinMode = PLATFORM_OK_TO_BURN_CPU * 2;
+	static FAutoConsoleVariableRef CVarConsoleSpinMode(
+		TEXT("TaskGraph.ConsoleSpinMode"),
+		GConsoleSpinMode,
+		TEXT("If > 0, then we never allow all task threads to idle; one will always be spinning looking for work (if mode == 2, then we sleep0 when spinning); this relieve the calling thread from the buden of starting threads. Only active with TaskGraph.FastScheduler 1."),
+		ECVF_Cheat
+		);
+#endif 
+static int32 GConsoleSpinModeLatched = 0;
+
+#if CREATE_BACKGROUND_TASK_THREADS
+static FAutoConsoleVariableRef CVarUseBackgroundThreads(
+	TEXT("TaskGraph.UseBackgroundThreads"),
+	ENamedThreads::bHasBackgroundThreads,
+	TEXT("If > 0, then use background threads, otherwise run background tasks on normal priority task threads. Used for performance tuning."),
+	ECVF_Cheat
+	);
+#endif
+
+#if CREATE_HIPRI_TASK_THREADS
+static FAutoConsoleVariableRef CVarUseHiPriThreads(
+	TEXT("TaskGraph.UseHiPriThreads"),
+	ENamedThreads::bHasHighPriorityThreads,
+	TEXT("If > 0, then use hi priority task threads, otherwise run background tasks on normal priority task threads. Used for performance tuning."),
+	ECVF_Cheat
+	);
+#endif
+
+
+
+#if USE_NEW_LOCK_FREE_LISTS
+static int32 GMaxTasksToStartOnDequeue = 1;
+#else
+static int32 GMaxTasksToStartOnDequeue = 2;
+#endif
+
+static FAutoConsoleVariableRef CVarMaxTasksToStartOnDequeue(
+	TEXT("TaskGraph.MaxTasksToStartOnDequeue"),
+	GMaxTasksToStartOnDequeue,
+	TEXT("Performance tweak, controls how many additional task threads a task thread starts when it grabs a list of new tasks. This only applies in TaskGraph.ConsoleSpinMode"),
+	ECVF_Cheat
+	);
+
 #define PROFILE_TASKGRAPH (0)
 #if PROFILE_TASKGRAPH
 	struct FProfileRec
@@ -190,57 +301,97 @@ FORCEINLINE void TestRandomizedThreads()
 
 #endif
 
-static int32 GNumWorkerThreadsToIgnore = 0;
-static FAutoConsoleVariableRef CVarNumWorkerThreadsToIgnore(
-	TEXT("TaskGraph.NumWorkerThreadsToIgnore"),
-	GNumWorkerThreadsToIgnore,
-	TEXT("Used to tune the number of task threads. Generally once you have found the right value, PlatformMisc::NumberOfWorkerThreadsToSpawn() should be hardcoded."),
-	ECVF_Cheat
-	);
-
-static int32 GFastScheduler = 0;
-#if USE_NEW_LOCK_FREE_LISTS
-static int32 GFastSchedulerLatched = 1;
-#else
-static int32 GFastSchedulerLatched = 0;
-#endif
-static FAutoConsoleVariableRef CVarFastScheduler(
-	TEXT("TaskGraph.FastScheduler"),
-	GFastScheduler,
-	TEXT("If > 0, then use the new experimental anythread task scheduler."),
-	ECVF_Cheat
-	);
-
-// move to platform abstraction
-#if PLATFORM_XBOXONE || PLATFORM_PS4
-	#define PLATFORM_OK_TO_BURN_CPU (1)
-#else
-	#define PLATFORM_OK_TO_BURN_CPU (0)
-#endif
-
-#define USE_INTRUSIVE_TASKQUEUES (USE_NEW_LOCK_FREE_LISTS)
+static FString ThreadPriorityToName(ENamedThreads::Type Priority)
+{
+	if (Priority == ENamedThreads::NormalThreadPriority)
+	{
+		return FString(TEXT("Normal"));
+	}
+	if (Priority == ENamedThreads::HighThreadPriority)
+	{
+		return FString(TEXT("High"));
+	}
+	if (Priority == ENamedThreads::BackgroundThreadPriority)
+	{
+		return FString(TEXT("Background"));
+	}
+	return FString(TEXT("??Unknown??"));
+}
 
-static int32 GConsoleSpinMode = PLATFORM_OK_TO_BURN_CPU * 2;
-static int32 GConsoleSpinModeLatched = 0;
-static FAutoConsoleVariableRef CVarConsoleSpinMode(
-	TEXT("TaskGraph.ConsoleSpinMode"),
-	GConsoleSpinMode,
-	TEXT("If > 0, then we never allow all task threads to idle; one will always be spinning looking for work (if mode == 2, then we sleep0 when spinning); this relieve the calling thread from the buden of starting threads. Only active with TaskGraph.FastScheduler 1."),
-	ECVF_Cheat
-	);
+static FString TaskPriorityToName(ENamedThreads::Type Priority)
+{
+	if (Priority == ENamedThreads::NormalTaskPriority)
+	{
+		return FString(TEXT("Normal"));
+	}
+	if (Priority == ENamedThreads::HighTaskPriority)
+	{
+		return FString(TEXT("High"));
+	}
+	return FString(TEXT("??Unknown??"));
+}
 
-#if USE_NEW_LOCK_FREE_LISTS
-static int32 GMaxTasksToStartOnDequeue = 1;
-#else
-static int32 GMaxTasksToStartOnDequeue = 2;
-#endif
+void FAutoConsoleTaskPriority::CommandExecute(const TArray<FString>& Args)
+{
+	if (Args.Num() > 0)
+	{
+		if (Args[0].Compare(ThreadPriorityToName(ENamedThreads::NormalThreadPriority), ESearchCase::IgnoreCase) == 0)
+		{
+			ThreadPriority = ENamedThreads::NormalThreadPriority;
+		}
+		else if (Args[0].Compare(ThreadPriorityToName(ENamedThreads::HighThreadPriority), ESearchCase::IgnoreCase) == 0)
+		{
+			ThreadPriority = ENamedThreads::HighThreadPriority;
+		}
+		else if (Args[0].Compare(ThreadPriorityToName(ENamedThreads::BackgroundThreadPriority), ESearchCase::IgnoreCase) == 0)
+		{
+			ThreadPriority = ENamedThreads::BackgroundThreadPriority;
+		}
+		else
+		{
+			UE_LOG(LogConsoleResponse, Display, TEXT("Could not parse thread priority %s"), *Args[0]);
+		}
+	}
+	if (Args.Num() > 1)
+	{
+		if (Args[1].Compare(TaskPriorityToName(ENamedThreads::NormalTaskPriority), ESearchCase::IgnoreCase) == 0)
+		{
+			TaskPriority = ENamedThreads::NormalTaskPriority;
+		}
+		else if (Args[1].Compare(TaskPriorityToName(ENamedThreads::HighTaskPriority), ESearchCase::IgnoreCase) == 0)
+		{
+			TaskPriority = ENamedThreads::HighTaskPriority;
+		}
+		else
+		{
+			UE_LOG(LogConsoleResponse, Display, TEXT("Could not parse task priority %s"), *Args[1]);
+		}
+	}
+	if (Args.Num() > 2)
+	{
+		if (Args[2].Compare(TaskPriorityToName(ENamedThreads::NormalTaskPriority), ESearchCase::IgnoreCase) == 0)
+		{
+			TaskPriorityIfForcedToNormalThreadPriority = ENamedThreads::NormalTaskPriority;
+		}
+		else if (Args[2].Compare(TaskPriorityToName(ENamedThreads::HighTaskPriority), ESearchCase::IgnoreCase) == 0)
+		{
+			TaskPriorityIfForcedToNormalThreadPriority = ENamedThreads::HighTaskPriority;
+		}
+		else
+		{
+			UE_LOG(LogConsoleResponse, Display, TEXT("Could not parse task priority %s"), *Args[2]);
+		}
+	}
+	if (ThreadPriority == ENamedThreads::NormalThreadPriority)
+	{
+		UE_LOG(LogConsoleResponse, Display, TEXT("%s - thread priority:%s   task priority:%s"), *CommandName, *ThreadPriorityToName(ThreadPriority), *TaskPriorityToName(TaskPriority));
+	}
+	else
+	{
+		UE_LOG(LogConsoleResponse, Display, TEXT("%s - thread priority:%s   task priority:%s  %s (when forced to normal)"), *CommandName, *ThreadPriorityToName(ThreadPriority), *TaskPriorityToName(TaskPriority), *TaskPriorityToName(this->TaskPriorityIfForcedToNormalThreadPriority));
+	}
+}
 
-static FAutoConsoleVariableRef CVarMaxTasksToStartOnDequeue(
-	TEXT("TaskGraph.MaxTasksToStartOnDequeue"),
-	GMaxTasksToStartOnDequeue,
-	TEXT("Performance tweak, controls how many additional task threads a task thread starts when it grabs a list of new tasks. This only applies in TaskGraph.ConsoleSpinMode"),
-	ECVF_Cheat
-	);
 
 /** 
  *	FTaskQueue
@@ -556,7 +707,7 @@ public:
 					{
 						break;
 					}
-					if (ENamedThreads::GetPriority(NewTask->ThreadToExecuteOn))
+					if (ENamedThreads::GetTaskPriority(NewTask->ThreadToExecuteOn))
 					{
 						Queue(QueueIndex).OutstandingHiPriTasks.Decrement();
 						Task = NewTask;
@@ -701,7 +852,7 @@ public:
 						for (int32 Index = NewTasks.Num() - 1; Index >= 0 ; Index--) // reverse the order since PopAll is implicitly backwards
 						{
 							FBaseGraphTask* NewTask = NewTasks[Index];
-							if (ENamedThreads::GetPriority(NewTask->ThreadToExecuteOn))
+							if (ENamedThreads::GetTaskPriority(NewTask->ThreadToExecuteOn))
 							{
 								Queue(QueueIndex).PrivateQueueHiPri.Enqueue(NewTask);
 								Queue(QueueIndex).OutstandingHiPriTasks.Decrement();
@@ -746,7 +897,7 @@ public:
 					{
 						FBaseGraphTask* NewTask = NewTasks[Index];
 						checkThreadGraph(NewTask);
-						if (ENamedThreads::GetPriority(NewTask->ThreadToExecuteOn))
+						if (ENamedThreads::GetTaskPriority(NewTask->ThreadToExecuteOn))
 						{
 							Queue(QueueIndex).PrivateQueueHiPri.Enqueue(NewTask);
 							Queue(QueueIndex).OutstandingHiPriTasks.Decrement();
@@ -795,7 +946,7 @@ public:
 	virtual void EnqueueFromThisThread(int32 QueueIndex, FBaseGraphTask* Task) override
 	{
 		checkThreadGraph(Task && Queue(QueueIndex).StallRestartEvent); // make sure we are started up
-		if (ENamedThreads::GetPriority(Task->ThreadToExecuteOn))
+		if (ENamedThreads::GetTaskPriority(Task->ThreadToExecuteOn))
 		{
 			Queue(QueueIndex).PrivateQueueHiPri.Enqueue(Task);
 		}
@@ -845,7 +996,7 @@ public:
 		checkThreadGraph(Task && Queue(QueueIndex).StallRestartEvent); // make sure we are started up
 
 		bool bWasReopenedByMe;
-		bool bHiPri = !!ENamedThreads::GetPriority(Task->ThreadToExecuteOn);
+		bool bHiPri = !!ENamedThreads::GetTaskPriority(Task->ThreadToExecuteOn);
 		{
 			TASKGRAPH_SCOPE_CYCLE_COUNTER(0, STAT_TaskGraph_EnqueueFromOtherThread_ReopenIfClosedAndPush);
 			bWasReopenedByMe = Queue(QueueIndex).IncomingQueue.ReopenIfClosedAndPush(Task);
@@ -876,20 +1027,26 @@ private:
 		FTaskQueue											PrivateQueue;
 		/** A non-threas safe queue for the thread locked tasks of a named thread. These are high priority.**/
 		FTaskQueue											PrivateQueueHiPri;
+
+		uint8 PadToAvoidContention1[PLATFORM_CACHE_LINE_SIZE];
 		/** Used to signal pending hi pri tasks. **/
 		FThreadSafeCounter									OutstandingHiPriTasks;
+		uint8 PadToAvoidContention2[PLATFORM_CACHE_LINE_SIZE];
+
+
 		/** 
 		 *	For named threads, this is a queue of thread locked tasks coming from other threads. They are not stealable.
 		 *	For unnamed thread this is the public queue, subject to stealing.
 		 *	In either case this queue is closely related to the stall event. Other threads that reopen the incoming queue must trigger the stall event to allow the thread to run.
 		**/
 #if USE_NEW_LOCK_FREE_LISTS && USE_INTRUSIVE_TASKQUEUES
-		FCloseableLockFreePointerQueueBaseSingleBaseConsumerIntrusive<FBaseGraphTask>		IncomingQueue;
+		FCloseableLockFreePointerQueueBaseSingleBaseConsumerIntrusive<FBaseGraphTask, PLATFORM_CACHE_LINE_SIZE>		IncomingQueue;
 #elif USE_NEW_LOCK_FREE_LISTS
-		TReopenableLockFreePointerListFIFOSingleConsumer<FBaseGraphTask>		IncomingQueue;
+		TReopenableLockFreePointerListFIFOSingleConsumer<FBaseGraphTask, PLATFORM_CACHE_LINE_SIZE>		IncomingQueue;
 #else
 		TReopenableLockFreePointerListLIFO<FBaseGraphTask>		IncomingQueue;
 #endif
+		uint8 PadToAvoidContention3[PLATFORM_CACHE_LINE_SIZE];
 		/** Used to signal the thread to quit when idle. **/
 		FThreadSafeCounter									QuitWhenIdle;
 		/** We need to disallow reentry of the processing loop **/
@@ -970,6 +1127,10 @@ private:
 class FTaskThreadAnyThread : public FTaskThreadBase
 {
 public:
+	FTaskThreadAnyThread(int32 InPriorityIndex)
+		: PriorityIndex(InPriorityIndex)
+	{
+	}
 	/** Used for named threads to start processing tasks until the thread is idle and RequestQuit has been called. **/
 	virtual void ProcessTasksUntilQuit(int32 QueueIndex) override
 	{
@@ -1155,6 +1316,8 @@ private:
 
 	/** Array of queues, only the first one is used for unnamed threads. **/
 	FThreadTaskQueue Queue;
+
+	int32 PriorityIndex;
 };
 
 /** 
@@ -1206,16 +1369,27 @@ public:
 	 *	Constructor - initializes the data structures, sets the singleton pointer and creates the internal threads.
 	 *	@param InNumThreads; total number of threads in the system, including named threads, unnamed threads, internal threads and external threads. Must be at least 1 + the number of named threads.
 	**/
-	FTaskGraphImplementation(int32 InNumThreads)
+	FTaskGraphImplementation(int32)
 	{
+		bCreatedHiPriorityThreads = !!ENamedThreads::bHasHighPriorityThreads;
+		bCreatedBackgroundPriorityThreads = !!ENamedThreads::bHasBackgroundThreads;
+
+		int32 MaxTaskThreads = MAX_THREADS;
+		int32 NumTaskThreads = FPlatformMisc::NumberOfWorkerThreadsToSpawn();
+
 		// if we don't want any performance-based threads, then force the task graph to not create any worker threads, and run in game thread
 		if (!FPlatformProcess::SupportsMultithreading())
 		{
 			// this is the logic that used to be spread over a couple of places, that will make the rest of this function disable a worker thread
 			// @todo: it could probably be made simpler/clearer
-			InNumThreads = 1;
 			// this - 1 tells the below code there is no rendering thread
+			MaxTaskThreads = 1;
+			NumTaskThreads = 1;
 			LastExternalThread = (ENamedThreads::Type)(ENamedThreads::ActualRenderingThread - 1);
+			bCreatedHiPriorityThreads = false;
+			bCreatedBackgroundPriorityThreads = false;
+			ENamedThreads::bHasBackgroundThreads = 0;
+			ENamedThreads::bHasHighPriorityThreads = 0;
 		}
 		else
 		{
@@ -1223,25 +1397,35 @@ public:
 		}
 		
 		NumNamedThreads = LastExternalThread + 1;
-		NumThreads = FMath::Max<int32>(FMath::Min<int32>(InNumThreads + NumNamedThreads,MAX_THREADS),NumNamedThreads + 1);
+
+		NumTaskThreadSets = 1 + bCreatedHiPriorityThreads + bCreatedBackgroundPriorityThreads;
+
+		// if we don't have enough threads to allow all of the sets asked for, then we can't create what was asked for.
+		check(NumTaskThreadSets == 1 || FMath::Min<int32>(NumTaskThreads * NumTaskThreadSets + NumNamedThreads, MAX_THREADS) == NumTaskThreads * NumTaskThreadSets + NumNamedThreads);
+		NumThreads = FMath::Max<int32>(FMath::Min<int32>(NumTaskThreads * NumTaskThreadSets + NumNamedThreads, MAX_THREADS), NumNamedThreads + 1);
+
 		// Cap number of extra threads to the platform worker thread count
-		NumThreads = FMath::Min(NumThreads, NumNamedThreads + FPlatformMisc::NumberOfWorkerThreadsToSpawn());
-		UE_LOG(LogTaskGraph, Log, TEXT("Started task graph with %d named threads and %d total threads."), NumNamedThreads, NumThreads);
+		// if we don't have enough threads to allow all of the sets asked for, then we can't create what was asked for.
+		check(NumTaskThreadSets == 1 || FMath::Min(NumThreads, NumNamedThreads + NumTaskThreads * NumTaskThreadSets) == NumThreads);
+		NumThreads = FMath::Min(NumThreads, NumNamedThreads + NumTaskThreads * NumTaskThreadSets);
+
+		NumTaskThreadsPerSet = (NumThreads - NumNamedThreads) / NumTaskThreadSets;
+		check((NumThreads - NumNamedThreads) % NumTaskThreadSets == 0); // should be equal numbers of threads per priority set
+
+		UE_LOG(LogTaskGraph, Log, TEXT("Started task graph with %d named threads and %d total threads with %d sets of task threads."), NumNamedThreads, NumThreads, NumTaskThreadSets);
 		check(NumThreads - NumNamedThreads >= 1);  // need at least one pure worker thread
 		check(NumThreads <= MAX_THREADS);
-		check(!NextStealFromThread.GetValue()); // reentrant?
-		NextStealFromThread.Increment(); // just checking for reentrancy
+		check(!ReentrancyCheck.GetValue()); // reentrant?
+		ReentrancyCheck.Increment(); // just checking for reentrancy
 		PerThreadIDTLSSlot = FPlatformTLS::AllocTlsSlot();
 
-		NextUnnamedThreadMod = NumThreads - NumNamedThreads;
-
 		for (int32 ThreadIndex = 0; ThreadIndex < NumThreads; ThreadIndex++)
 		{
 			check(!WorkerThreads[ThreadIndex].bAttached); // reentrant?
 			bool bAnyTaskThread = ThreadIndex >= NumNamedThreads;
 			if (bAnyTaskThread)
 			{
-				WorkerThreads[ThreadIndex].TaskGraphWorker = new FTaskThreadAnyThread;
+				WorkerThreads[ThreadIndex].TaskGraphWorker = new FTaskThreadAnyThread(ThreadIndexToPriorityIndex(ThreadIndex));
 			}
 			else
 			{
@@ -1254,9 +1438,26 @@ public:
 
 		for (int32 ThreadIndex = LastExternalThread + 1; ThreadIndex < NumThreads; ThreadIndex++)
 		{
-			FString Name = FString::Printf(TEXT("TaskGraphThread %d"), ThreadIndex - (LastExternalThread + 1));
-			uint32 StackSize = 256 * 1024;
-			WorkerThreads[ThreadIndex].RunnableThread = FRunnableThread::Create(&Thread(ThreadIndex), *Name, StackSize, TPri_BelowNormal, FPlatformAffinity::GetTaskGraphThreadMask()); // these are below normal threads? so that they sleep when the named threads are active
+			FString Name;
+			int32 Priority = ThreadIndexToPriorityIndex(ThreadIndex);
+			EThreadPriority ThreadPri;
+			if (Priority == 1)
+			{
+				Name = FString::Printf(TEXT("TaskGraphThreadHP %d"), ThreadIndex - (LastExternalThread + 1));
+				ThreadPri = TPri_SlightlyBelowNormal; // we want even hi priority tasks below the normal threads
+			}
+			else if (Priority == 2)
+			{
+				Name = FString::Printf(TEXT("TaskGraphThreadBP %d"), ThreadIndex - (LastExternalThread + 1));
+				ThreadPri = TPri_Lowest;
+			}
+			else
+			{
+				Name = FString::Printf(TEXT("TaskGraphThreadNP %d"), ThreadIndex - (LastExternalThread + 1));
+				ThreadPri = TPri_BelowNormal; // we want normal tasks below normal threads like the game thread
+			}
+			uint32 StackSize = 64 * 1024;
+			WorkerThreads[ThreadIndex].RunnableThread = FRunnableThread::Create(&Thread(ThreadIndex), *Name, StackSize, ThreadPri, FPlatformAffinity::GetTaskGraphThreadMask()); // these are below normal threads so that they sleep when the named threads are active
 			WorkerThreads[ThreadIndex].bAttached = true;
 		}
 	}
@@ -1286,9 +1487,12 @@ public:
 			WorkerThreads[ThreadIndex].bAttached = false;
 		}
 		TaskGraphImplementationSingleton = NULL;
-		NextUnnamedThreadMod = 0;
+		NumTaskThreadsPerSet = 0;
 		TArray<FTaskThreadBase*> NotProperlyUnstalled;
-		StalledUnnamedThreads.PopAll(NotProperlyUnstalled);
+		for (int32 PriorityIndex = 0; PriorityIndex < NumTaskThreadSets; PriorityIndex++)
+		{
+			StalledUnnamedThreads[PriorityIndex].PopAll(NotProperlyUnstalled);
+		}
 		FPlatformTLS::FreeTlsSlot(PerThreadIDTLSSlot);
 	}
 
@@ -1305,7 +1509,7 @@ public:
 		TASKGRAPH_SCOPE_CYCLE_COUNTER(2, STAT_TaskGraph_QueueTask);
 
 		TestRandomizedThreads();
-		checkThreadGraph(NextUnnamedThreadMod);
+		checkThreadGraph(NumTaskThreadsPerSet);
 		if (GFastSchedulerLatched != GFastScheduler && IsInGameThread())
 		{
 #if USE_NEW_LOCK_FREE_LISTS
@@ -1315,8 +1519,11 @@ public:
 			{
 				if (GFastScheduler)
 				{
-					AtomicForConsoleApproach = FAtomicStateBitfield();
-					AtomicForConsoleApproach.Stalled = (1 << GetNumWorkerThreads()) - 1; // everyone is stalled
+					for (int32 PriorityIndex = 0; PriorityIndex < NumTaskThreadSets; PriorityIndex++)
+					{
+						AtomicForConsoleApproach[PriorityIndex] = FAtomicStateBitfield();
+						AtomicForConsoleApproach[PriorityIndex].Stalled = (1 << GetNumWorkerThreads()) - 1; // everyone is stalled
+					}
 				}
 				// this is all kinda of sketchy, but runtime switching just has to barely work.
 				FPlatformProcess::Sleep(.1f); // hopefully all task threads are idle
@@ -1331,7 +1538,14 @@ public:
 				{
 					for (int32 Index = 0; Index < GetNumWorkerThreads(); Index++)
 					{
-						StartTaskThread(Index);
+						for (int32 Priority = 0; Priority < ENamedThreads::NumThreadPriorities; Priority++)
+						{
+							if (Priority == (ENamedThreads::NormalThreadPriority >> ENamedThreads::ThreadPriorityShift) ||
+								(Priority == (ENamedThreads::HighThreadPriority >> ENamedThreads::ThreadPriorityShift) && bCreatedHiPriorityThreads))
+							{
+								StartTaskThread(Priority, Index);
+							}
+						}
 					}
 				}
 			}
@@ -1342,7 +1556,14 @@ public:
 			GConsoleSpinModeLatched = GConsoleSpinMode;
 			if (bStartTask)
 			{
-				StartTaskThreadFastMode();
+				for (int32 Priority = 0; Priority < ENamedThreads::NumThreadPriorities; Priority++)
+				{
+					if (Priority == (ENamedThreads::NormalThreadPriority >> ENamedThreads::ThreadPriorityShift) ||
+						(Priority == (ENamedThreads::HighThreadPriority >> ENamedThreads::ThreadPriorityShift) && bCreatedHiPriorityThreads))
+					{
+						StartTaskThreadFastMode(Priority, -1);
+					}
+				}
 			}
 		}
 
@@ -1351,22 +1572,37 @@ public:
 			TASKGRAPH_SCOPE_CYCLE_COUNTER(3, STAT_TaskGraph_QueueTask_AnyThread);
 			if (FPlatformProcess::SupportsMultithreading())
 			{
+				int32 TaskPriority = ENamedThreads::GetTaskPriority(Task->ThreadToExecuteOn);
+				int32 Priority = ENamedThreads::GetThreadPriorityIndex(Task->ThreadToExecuteOn);
+				if (Priority == (ENamedThreads::BackgroundThreadPriority >> ENamedThreads::ThreadPriorityShift) && (!bCreatedBackgroundPriorityThreads || !ENamedThreads::bHasBackgroundThreads))
+				{
+					Priority = ENamedThreads::NormalThreadPriority >> ENamedThreads::ThreadPriorityShift; // we don't have background threads, promote to normal
+					TaskPriority = ENamedThreads::NormalTaskPriority >> ENamedThreads::TaskPriorityShift; // demote to normal task pri
+				}
+				else if (Priority == (ENamedThreads::HighThreadPriority >> ENamedThreads::ThreadPriorityShift) && (!bCreatedHiPriorityThreads || !ENamedThreads::bHasHighPriorityThreads))
+				{
+					Priority = ENamedThreads::NormalThreadPriority >> ENamedThreads::ThreadPriorityShift; // we don't have hi priority threads, demote to normal
+					TaskPriority = ENamedThreads::HighTaskPriority >> ENamedThreads::TaskPriorityShift; // promote to hi task pri
+				}
+				check(Priority >= 0 && Priority < MAX_THREAD_PRIORITIES);
+
 				{
 					TASKGRAPH_SCOPE_CYCLE_COUNTER(4, STAT_TaskGraph_QueueTask_IncomingAnyThreadTasks_Push);
-					if (ENamedThreads::GetPriority(Task->ThreadToExecuteOn))
+					if (TaskPriority)
 					{
-						IncomingAnyThreadTasksHiPri.Push(Task);
+						IncomingAnyThreadTasksHiPri[Priority].Push(Task);
 					}
 					else
 					{
-						IncomingAnyThreadTasks.Push(Task);
+						IncomingAnyThreadTasks[Priority].Push(Task);
 					}
 				}
 				if (GFastSchedulerLatched)
 				{
-					if (!GConsoleSpinModeLatched)
+					if (!GConsoleSpinModeLatched || 
+						Priority == (ENamedThreads::BackgroundThreadPriority >> ENamedThreads::ThreadPriorityShift)) // background never spins
 					{
-						StartTaskThreadFastMode();
+						StartTaskThreadFastMode(Priority, -1);
 					}
 				}
 				else
@@ -1384,8 +1620,8 @@ public:
 					FTaskThreadBase* TempTarget;
 					{
 						TASKGRAPH_SCOPE_CYCLE_COUNTER(5, STAT_TaskGraph_QueueTask_StalledUnnamedThreads_Pop);
-						TempTarget = StalledUnnamedThreads.Pop(); //@todo it is possible that a thread is in the process of stalling and we just missed it, non-fatal, but we could lose a whole task of potential parallelism.
-						if (TempTarget && GNumWorkerThreadsToIgnore && (TempTarget->GetThreadId() - NumNamedThreads) >= GetNumWorkerThreads())
+						TempTarget = StalledUnnamedThreads[Priority].Pop(); //@todo it is possible that a thread is in the process of stalling and we just missed it, non-fatal, but we could lose a whole task of potential parallelism.
+						if (TempTarget && GNumWorkerThreadsToIgnore && (TempTarget->GetThreadId() - NumNamedThreads) % NumTaskThreadsPerSet >= GetNumWorkerThreads())
 						{
 							TempTarget = nullptr;
 						}
@@ -1396,11 +1632,12 @@ public:
 					}
 					else
 					{
-						check(NextUnnamedThreadMod - GNumWorkerThreadsToIgnore > 0); // can't tune it to zero task threads
-						ThreadToExecuteOn = ENamedThreads::Type((uint32(NextUnnamedThreadForTaskFromUnknownThread.Increment()) % uint32(NextUnnamedThreadMod - GNumWorkerThreadsToIgnore)) + NumNamedThreads);
+						check(NumTaskThreadsPerSet - GNumWorkerThreadsToIgnore > 0); // can't tune it to zero task threads
+						ThreadToExecuteOn = ENamedThreads::Type(
+							(uint32(NextUnnamedThreadForTaskFromUnknownThread[Priority].Increment()) % uint32(NumTaskThreadsPerSet - GNumWorkerThreadsToIgnore)) + Priority * NumTaskThreadsPerSet + NumNamedThreads);
 					}
 					FTaskThreadBase* Target = &Thread(ThreadToExecuteOn);
-					if (ThreadToExecuteOn != CurrentThreadIfKnown)
+					if (ENamedThreads::GetThreadIndex(ThreadToExecuteOn) != ENamedThreads::GetThreadIndex(CurrentThreadIfKnown))
 					{
 						Target->WakeUp();
 					}
@@ -1420,13 +1657,13 @@ public:
 		else
 		{
 			CurrentThreadIfKnown = ENamedThreads::GetThreadIndex(InCurrentThreadIfKnown);
-			checkThreadGraph(CurrentThreadIfKnown == GetCurrentThread());
+			checkThreadGraph(CurrentThreadIfKnown == ENamedThreads::GetThreadIndex(GetCurrentThread()));
 		}
 		{
 			int32 QueueToExecuteOn = ENamedThreads::GetQueueIndex(ThreadToExecuteOn);
 			ThreadToExecuteOn = ENamedThreads::GetThreadIndex(ThreadToExecuteOn);
 			FTaskThreadBase* Target = &Thread(ThreadToExecuteOn);
-			if (ThreadToExecuteOn == CurrentThreadIfKnown)
+			if (ThreadToExecuteOn == ENamedThreads::GetThreadIndex(CurrentThreadIfKnown))
 			{
 				Target->EnqueueFromThisThread(QueueToExecuteOn, Task);
 			}
@@ -1440,7 +1677,7 @@ public:
 
 	virtual	int32 GetNumWorkerThreads() final override
 	{
-		int32 Result = NumThreads - NumNamedThreads - GNumWorkerThreadsToIgnore;
+		int32 Result = (NumThreads - NumNamedThreads) / NumTaskThreadSets - GNumWorkerThreadsToIgnore;
 		check(Result > 0); // can't tune it to zero task threads
 		return Result;
 	}
@@ -1448,7 +1685,7 @@ public:
 	virtual ENamedThreads::Type GetCurrentThreadIfKnown(bool bLocalQueue) final override
 	{
 		ENamedThreads::Type Result = GetCurrentThread();
-		if (bLocalQueue && Result >= 0 && Result < NumNamedThreads)
+		if (bLocalQueue && ENamedThreads::GetThreadIndex(Result) >= 0 && ENamedThreads::GetThreadIndex(Result) < NumNamedThreads)
 		{
 			Result = ENamedThreads::Type(int32(Result) | int32(ENamedThreads::LocalQueue));
 		}
@@ -1468,7 +1705,7 @@ public:
 	virtual void AttachToThread(ENamedThreads::Type CurrentThread) final override
 	{
 		CurrentThread = ENamedThreads::GetThreadIndex(CurrentThread);
-		check(NextUnnamedThreadMod);
+		check(NumTaskThreadsPerSet);
 		check(CurrentThread >= 0 && CurrentThread < NumNamedThreads);
 		check(!WorkerThreads[CurrentThread].bAttached);
 		Thread(CurrentThread).InitializeForCurrentThread();
@@ -1505,20 +1742,23 @@ public:
 		ENamedThreads::Type CurrentThread = CurrentThreadIfKnown;
 		if (ENamedThreads::GetThreadIndex(CurrentThreadIfKnown) == ENamedThreads::AnyThread)
 		{
-			bool bIsHiPri = !!ENamedThreads::GetPriority(CurrentThreadIfKnown);
+			bool bIsHiPri = !!ENamedThreads::GetTaskPriority(CurrentThreadIfKnown);
+			int32 Priority = ENamedThreads::GetThreadPriorityIndex(CurrentThreadIfKnown);
 			check(!ENamedThreads::GetQueueIndex(CurrentThreadIfKnown));
-			CurrentThreadIfKnown = GetCurrentThread();
-			CurrentThread = bIsHiPri ? ENamedThreads::HiPri(CurrentThreadIfKnown) : CurrentThreadIfKnown;
+			CurrentThreadIfKnown = ENamedThreads::GetThreadIndex(GetCurrentThread());
+			CurrentThread = ENamedThreads::SetPriorities(CurrentThreadIfKnown, Priority, bIsHiPri);
 		}
 		else
 		{
 			CurrentThreadIfKnown = ENamedThreads::GetThreadIndex(CurrentThreadIfKnown);
-			check(CurrentThreadIfKnown == GetCurrentThread());
+			check(CurrentThreadIfKnown == ENamedThreads::GetThreadIndex(GetCurrentThread()));
 			// we don't modify CurrentThread here because it might be a local queue
 		}
 
 		if (CurrentThreadIfKnown != ENamedThreads::AnyThread && CurrentThreadIfKnown < NumNamedThreads && !IsThreadProcessingTasks(CurrentThread))
 		{
+			if (Tasks.Num() > 8) // don't bother to check for completion if there are lots of prereqs...too expensive to check
+			{
 			bool bAnyPending = false;
 			for (int32 Index = 0; Index < Tasks.Num(); Index++)
 			{
@@ -1532,6 +1772,7 @@ public:
 			{
 				return;
 			}
+			}
 			// named thread process tasks while we wait
 			TGraphTask<FReturnGraphTask>::CreateTask(&Tasks, CurrentThread).ConstructAndDispatchWhenReady(CurrentThread);
 			ProcessThreadUntilRequestReturn(CurrentThread);
@@ -1547,13 +1788,17 @@ public:
 	virtual void TriggerEventWhenTasksComplete(FEvent* InEvent, const FGraphEventArray& Tasks, ENamedThreads::Type CurrentThreadIfKnown = ENamedThreads::AnyThread) final override
 	{
 		check(InEvent);
-		bool bAnyPending = false;
-		for (int32 Index = 0; Index < Tasks.Num(); Index++)
+		bool bAnyPending = true;
+		if (Tasks.Num() > 8) // don't bother to check for completion if there are lots of prereqs...too expensive to check
 		{
-			if (!Tasks[Index]->IsComplete())
+			bAnyPending = false;
+			for (int32 Index = 0; Index < Tasks.Num(); Index++)
 			{
-				bAnyPending = true;
-				break;
+				if (!Tasks[Index]->IsComplete())
+				{
+					bAnyPending = true;
+					break;
+				}
 			}
 		}
 		if (!bAnyPending)
@@ -1573,44 +1818,68 @@ public:
 
 	// Scheduling utilities
 
-	void StartTaskThread(int32 IndexToStart)
+	void StartTaskThread(int32 Priority, int32 IndexToStart)
 	{
-		ENamedThreads::Type ThreadToWake = ENamedThreads::Type(IndexToStart + NumNamedThreads);
+		ENamedThreads::Type ThreadToWake = ENamedThreads::Type(IndexToStart + Priority * NumTaskThreadsPerSet + NumNamedThreads);
 		((FTaskThreadAnyThread&)Thread(ThreadToWake)).FTaskThreadAnyThread::WakeUp();
 	}
 
-	bool StartTaskThreadFastMode(int32 MyIndex = -1)
+
+	void StartTaskThreadFastMode(int32 Priority, int32 MyIndex)
 	{
+		check(Priority >= 0 && Priority < MAX_THREAD_PRIORITIES);
+		int32 LocalNumWorkingThread = GetNumWorkerThreads();
 		while (true)
 		{
+			FAtomicStateBitfield LocalAtomicForConsoleApproachInner = AtomicForConsoleApproach[Priority];
 			TestRandomizedThreads();
-			FAtomicStateBitfield LocalAtomicForConsoleApproachInner = AtomicForConsoleApproach;
 			check(MyIndex < 0 || !(LocalAtomicForConsoleApproachInner.Stalled & (1 << MyIndex))); // if I was stalled, I would not be here
+			int32 NumStalled = LocalAtomicForConsoleApproachInner.NumberOfStalledThreads();
+			int32 NumSpinning = LocalNumWorkingThread - NumStalled - LocalAtomicForConsoleApproachInner.NumberOfWorkingThreads();
 
-			if (!LocalAtomicForConsoleApproachInner.Stalled)
+			if (NumStalled)
 			{
-				break;
-			}
+				uint32 IndexToStart = 31 - FMath::CountLeadingZeros((uint32)LocalAtomicForConsoleApproachInner.Stalled);
+				check(IndexToStart != MyIndex);
+				check(((uint32)LocalAtomicForConsoleApproachInner.Stalled) & (1 << IndexToStart));
 
-			uint32 IndexToStart =  31 - FMath::CountLeadingZeros((uint32)LocalAtomicForConsoleApproachInner.Stalled);
-			check(IndexToStart != MyIndex);
-			check(((uint32)LocalAtomicForConsoleApproachInner.Stalled) & (1 << IndexToStart));
+				FAtomicStateBitfield NewAtomicForConsoleApproachInner = LocalAtomicForConsoleApproachInner;
+				NewAtomicForConsoleApproachInner.Stalled = LocalAtomicForConsoleApproachInner.Stalled & ~(1 << IndexToStart);
+				check(!(((uint32)NewAtomicForConsoleApproachInner.Stalled) & (1 << IndexToStart)));
+				check(LocalAtomicForConsoleApproachInner != NewAtomicForConsoleApproachInner);
+				if (FAtomicStateBitfield::InterlockedCompareExchange(&AtomicForConsoleApproach[Priority], NewAtomicForConsoleApproachInner, LocalAtomicForConsoleApproachInner) == LocalAtomicForConsoleApproachInner)
+				{
+					TestRandomizedThreads();
+					StartTaskThread(Priority, IndexToStart);
+					return;
+				}
+			}
+			else if (NumSpinning)
+			{
+				uint32 IndexToPreventStall = 31 - FMath::CountLeadingZeros(uint32((1 << FAtomicStateBitfield_MAX_THREADS) - 1) ^ uint32(LocalAtomicForConsoleApproachInner.Working | LocalAtomicForConsoleApproachInner.Stalled));
+				check(IndexToPreventStall != MyIndex);
+				check(!(((uint32)LocalAtomicForConsoleApproachInner.Working) & (1 << IndexToPreventStall)));
+				check(!(((uint32)LocalAtomicForConsoleApproachInner.Stalled) & (1 << IndexToPreventStall)));
 
-			FAtomicStateBitfield NewAtomicForConsoleApproachInner = LocalAtomicForConsoleApproachInner;
-			NewAtomicForConsoleApproachInner.Stalled = LocalAtomicForConsoleApproachInner.Stalled & ~(1 << IndexToStart);
-			check(!(((uint32)NewAtomicForConsoleApproachInner.Stalled) & (1 << IndexToStart)));
-			check(LocalAtomicForConsoleApproachInner != NewAtomicForConsoleApproachInner);
-			if (FAtomicStateBitfield::InterlockedCompareExchange(&AtomicForConsoleApproach, NewAtomicForConsoleApproachInner, LocalAtomicForConsoleApproachInner) == LocalAtomicForConsoleApproachInner)
+				FAtomicStateBitfield NewAtomicForConsoleApproachInner = LocalAtomicForConsoleApproachInner;
+				NewAtomicForConsoleApproachInner.Working = LocalAtomicForConsoleApproachInner.Working | (1 << IndexToPreventStall);
+				check((((uint32)NewAtomicForConsoleApproachInner.Working) & (1 << IndexToPreventStall)));
+				check(LocalAtomicForConsoleApproachInner != NewAtomicForConsoleApproachInner);
+				if (FAtomicStateBitfield::InterlockedCompareExchange(&AtomicForConsoleApproach[Priority], NewAtomicForConsoleApproachInner, LocalAtomicForConsoleApproachInner) == LocalAtomicForConsoleApproachInner)
+				{
+					// by changing it to working, it will be unable to stall without rechecking the queues
+					return;
+				}
+			}
+			else
 			{
-				TestRandomizedThreads();
-				StartTaskThread(IndexToStart);
-				return true;
+				return;
 			}
 		}
-		return false;
+		return;
 	}
 
-	FORCEINLINE void SetWorking(uint32 MyIndex, int32 LocalNumWorkingThread)
+	FORCEINLINE void SetWorking(int32 Priority, uint32 MyIndex, int32 LocalNumWorkingThread)
 	{
 		if (!GConsoleSpinModeLatched)
 		{
@@ -1619,22 +1888,18 @@ public:
 		}
 		while (true)
 		{
+			FAtomicStateBitfield LocalAtomicForConsoleApproachInner = AtomicForConsoleApproach[Priority];
 			TestRandomizedThreads();
-			FPlatformMisc::MemoryBarrier();
-			FAtomicStateBitfield LocalAtomicForConsoleApproachInner = AtomicForConsoleApproach;
 			check(!(LocalAtomicForConsoleApproachInner.Stalled & (1 << MyIndex))); // if I was stalled, I would not be here
-			bool CurrentlyWorking = !!((LocalAtomicForConsoleApproachInner.Working) & (1 << MyIndex));
-			if (CurrentlyWorking)
-			{
-				break;
-			}
+			check(!((LocalAtomicForConsoleApproachInner.Working) & (1 << MyIndex))); // I am already marked working
 
 			FAtomicStateBitfield NewAtomicForConsoleApproachInner = LocalAtomicForConsoleApproachInner;
 			NewAtomicForConsoleApproachInner.Working = LocalAtomicForConsoleApproachInner.Working | (1 << MyIndex);
 			int32 NumStalled = LocalAtomicForConsoleApproachInner.NumberOfStalledThreads();
+			int32 NumSpinning = LocalNumWorkingThread - NumStalled - NewAtomicForConsoleApproachInner.NumberOfWorkingThreads();
 
 			// we are going to do work, so if everyone else is either working or stalled, we need to start a thread
-			if (NumStalled && LocalAtomicForConsoleApproachInner.NumberOfStalledThreads() + LocalAtomicForConsoleApproachInner.NumberOfWorkingThreads() + 1 == LocalNumWorkingThread)
+			if (NumStalled && !NumSpinning)
 			{
 				// start one while we flip ourselves to working.
 				check(LocalAtomicForConsoleApproachInner.Stalled);
@@ -1646,15 +1911,15 @@ public:
 				NewAtomicForConsoleApproachInner.Stalled = LocalAtomicForConsoleApproachInner.Stalled & ~(1 << IndexToStart);
 
 				check(LocalAtomicForConsoleApproachInner != NewAtomicForConsoleApproachInner);
-				if (FAtomicStateBitfield::InterlockedCompareExchange(&AtomicForConsoleApproach, NewAtomicForConsoleApproachInner, LocalAtomicForConsoleApproachInner) == LocalAtomicForConsoleApproachInner)
+				if (FAtomicStateBitfield::InterlockedCompareExchange(&AtomicForConsoleApproach[Priority], NewAtomicForConsoleApproachInner, LocalAtomicForConsoleApproachInner) == LocalAtomicForConsoleApproachInner)
 				{
 					TestRandomizedThreads();
-					StartTaskThread(IndexToStart);
+					StartTaskThread(Priority, IndexToStart);
 					// if we have another stalled thread and it looks like there are more tasks, then start an extra thread to get fan-out on the thread starts
 #if USE_NEW_LOCK_FREE_LISTS
-					if (GMaxTasksToStartOnDequeue > 1 && NumStalled > 1 && (!IncomingAnyThreadTasks.IsEmptyFast() || !IncomingAnyThreadTasksHiPri.IsEmptyFast()))
+					if (GMaxTasksToStartOnDequeue > 1 && NumStalled > 1 && (!IncomingAnyThreadTasks[Priority].IsEmptyFast() || !IncomingAnyThreadTasksHiPri[Priority].IsEmptyFast()))
 					{
-						StartTaskThreadFastMode(MyIndex);
+						StartTaskThreadFastMode(Priority, MyIndex);
 					}
 #endif
 					return;
@@ -1664,7 +1929,7 @@ public:
 			{
 				// we don't need to wake anyone either because there are other unstalled idle threads or everyone is busy
 				check(LocalAtomicForConsoleApproachInner != NewAtomicForConsoleApproachInner);
-				if (FAtomicStateBitfield::InterlockedCompareExchange(&AtomicForConsoleApproach, NewAtomicForConsoleApproachInner, LocalAtomicForConsoleApproachInner) == LocalAtomicForConsoleApproachInner)
+				if (FAtomicStateBitfield::InterlockedCompareExchange(&AtomicForConsoleApproach[Priority], NewAtomicForConsoleApproachInner, LocalAtomicForConsoleApproachInner) == LocalAtomicForConsoleApproachInner)
 				{
 					TestRandomizedThreads();
 					return;
@@ -1676,255 +1941,124 @@ public:
 	FBaseGraphTask* FindWork(ENamedThreads::Type ThreadInNeed)
 	{
 		int32 LocalNumWorkingThread = GetNumWorkerThreads();
-		uint32 MyIndex = uint32(ThreadInNeed) - NumNamedThreads;
+		uint32 MyIndex = (uint32(ThreadInNeed) - NumNamedThreads) % NumTaskThreadsPerSet;
+		int32 Priority = (uint32(ThreadInNeed) - NumNamedThreads) / NumTaskThreadsPerSet;
 		check(MyIndex >= 0 && (int32)MyIndex < LocalNumWorkingThread && 
 			LocalNumWorkingThread <= FAtomicStateBitfield_MAX_THREADS); // this is an atomic on 32 bits currently; we use two fields
+		check(Priority >= 0 && Priority < ENamedThreads::NumThreadPriorities); 
+		bool bDoSpinMode = !!GConsoleSpinModeLatched && Priority != (ENamedThreads::BackgroundThreadPriority >> ENamedThreads::ThreadPriorityShift);
 		while (true)
 		{
 			if (!GFastSchedulerLatched)
 			{
 				return nullptr;
 			}
-			bool bOkWithFalseEmpties = !!GConsoleSpinModeLatched;
+			FAtomicStateBitfield LocalAtomicForConsoleApproachInner = AtomicForConsoleApproach[Priority];
+			TestRandomizedThreads();
+			check(!(LocalAtomicForConsoleApproachInner.Stalled & (1 << MyIndex))); // if I was stalled, I would not be here
+			bool CurrentlyWorking = !!((LocalAtomicForConsoleApproachInner.Working) & (1 << MyIndex));
+
 			{
-				FBaseGraphTask* Task = IncomingAnyThreadTasksHiPri.Pop(bOkWithFalseEmpties);
+				FBaseGraphTask* Task = IncomingAnyThreadTasksHiPri[Priority].Pop(CurrentlyWorking); // if we are marked working, this is still preliminary so we do a faster pop
 				if (!Task)
 				{
-					Task = IncomingAnyThreadTasks.Pop(bOkWithFalseEmpties);
+					Task = IncomingAnyThreadTasks[Priority].Pop(CurrentlyWorking); // if we are marked working, this is still preliminary so we do a faster pop
 				}
 				if (Task)
 				{
-					SetWorking(MyIndex, LocalNumWorkingThread);
+					if (!CurrentlyWorking)
+					{
+						SetWorking(Priority, MyIndex, LocalNumWorkingThread);
+					}
 					return Task;
 				}
 			}
 
 			TestRandomizedThreads();
-			FPlatformMisc::MemoryBarrier();
-			FAtomicStateBitfield LocalAtomicForConsoleApproachInner = AtomicForConsoleApproach;
-			check(!(LocalAtomicForConsoleApproachInner.Stalled & (1 << MyIndex))); // if I was stalled, I would not be here
-
 
-			FAtomicStateBitfield NewAtomicForConsoleApproachInner = LocalAtomicForConsoleApproachInner;
-			bool CurrentlyWorking = !!((LocalAtomicForConsoleApproachInner.Working) & (1 << MyIndex));
-			NewAtomicForConsoleApproachInner.Working = NewAtomicForConsoleApproachInner.Working & ~(1 << MyIndex);
-			int32 NumWorking = LocalAtomicForConsoleApproachInner.NumberOfWorkingThreads() - !!CurrentlyWorking;
-			bool bAttemptStall = (!GConsoleSpinModeLatched || LocalAtomicForConsoleApproachInner.NumberOfStalledThreads() + NumWorking < LocalNumWorkingThread - 1);
-			if (bAttemptStall)
+			if (CurrentlyWorking)
 			{
-				NewAtomicForConsoleApproachInner.Stalled = LocalAtomicForConsoleApproachInner.Stalled | (1 << MyIndex);
-				check((((uint32)NewAtomicForConsoleApproachInner.Stalled) & (1 << MyIndex)));
+				FAtomicStateBitfield NewAtomicForConsoleApproachInner = LocalAtomicForConsoleApproachInner;
+				NewAtomicForConsoleApproachInner.Working = NewAtomicForConsoleApproachInner.Working & ~(1 << MyIndex);
 
-				// if we are going to attempt to stall, we better be sure there really are not tasks...unless we are already sure
-				if (bOkWithFalseEmpties)
+				check(LocalAtomicForConsoleApproachInner != NewAtomicForConsoleApproachInner);
+				FAtomicStateBitfield::InterlockedCompareExchange(&AtomicForConsoleApproach[Priority], NewAtomicForConsoleApproachInner, LocalAtomicForConsoleApproachInner);
+			}
+			else
+			{
+				int32 NumSpinning = LocalNumWorkingThread  - LocalAtomicForConsoleApproachInner.NumberOfStalledThreads() - LocalAtomicForConsoleApproachInner.NumberOfWorkingThreads();
+				bool bAttemptStall = (!bDoSpinMode || NumSpinning > 1);
+				if (bAttemptStall)
 				{
-					FBaseGraphTask* Task = IncomingAnyThreadTasksHiPri.Pop();
+					FBaseGraphTask* Task = IncomingAnyThreadTasksHiPri[Priority].Pop();
 					if (!Task)
 					{
-						Task = IncomingAnyThreadTasks.Pop();
+						Task = IncomingAnyThreadTasks[Priority].Pop();
 					}
 					if (Task)
 					{
-						SetWorking(MyIndex, LocalNumWorkingThread);
+						SetWorking(Priority, MyIndex, LocalNumWorkingThread);
 						return Task;
 					}
-				}
-			}
-			if (bAttemptStall || CurrentlyWorking) // otherwise we have no reason to change the state
-			{
-				check(LocalAtomicForConsoleApproachInner != NewAtomicForConsoleApproachInner);
-				if (FAtomicStateBitfield::InterlockedCompareExchange(&AtomicForConsoleApproach, NewAtomicForConsoleApproachInner, LocalAtomicForConsoleApproachInner) == LocalAtomicForConsoleApproachInner)
-				{
-					if (bAttemptStall)
+					FAtomicStateBitfield NewAtomicForConsoleApproachInner = LocalAtomicForConsoleApproachInner;
+					NewAtomicForConsoleApproachInner.Stalled = LocalAtomicForConsoleApproachInner.Stalled | (1 << MyIndex);
+					check((((uint32)NewAtomicForConsoleApproachInner.Stalled) & (1 << MyIndex)));
+
+					check(LocalAtomicForConsoleApproachInner != NewAtomicForConsoleApproachInner);
+					if (FAtomicStateBitfield::InterlockedCompareExchange(&AtomicForConsoleApproach[Priority], NewAtomicForConsoleApproachInner, LocalAtomicForConsoleApproachInner) == LocalAtomicForConsoleApproachInner)
 					{
 						TestRandomizedThreads();
 						return nullptr; // we did stall
 					}
+
 				}
-			}
-			else if (GConsoleSpinModeLatched == 2)
-			{
-				FPlatformProcess::SleepNoStats(0.0f);
-			}
-		}
-	}
-#else
-	FBaseGraphTask* FindWorkConsole(ENamedThreads::Type ThreadInNeed)
-	{
-		int32 LocalNumWorkingThread = GetNumWorkerThreads();
-		uint32 MyIndex = uint32(ThreadInNeed) - NumNamedThreads;
-		check(MyIndex >= 0 && (int32)MyIndex < LocalNumWorkingThread && 
-			LocalNumWorkingThread <= FAtomicStateBitfield_MAX_THREADS); // this is an atomic on 32 bits currently; we use two fields
-		while (true)
-		{
-			if (!GFastSchedulerLatched)
-			{
-				return nullptr;
-			}
-			{
-				FBaseGraphTask* Task = SortedAnyThreadTasksHiPri.Pop();
-				if (Task)
-				{
-					SetWorking(MyIndex, LocalNumWorkingThread);
-					return Task;
-				}
-			}
-			TestRandomizedThreads();
-			FPlatformMisc::MemoryBarrier();
-			FAtomicStateBitfield LocalAtomicForConsoleApproach = AtomicForConsoleApproach;
-			check(!(LocalAtomicForConsoleApproach.Stalled & (1 << MyIndex))); // if I was stalled, I would not be here
-			if (!LocalAtomicForConsoleApproach.QueueOwned)
-			{
-				if (IncomingAnyThreadTasksHiPri.IsEmpty())
+				else if (GConsoleSpinModeLatched == 2)
 				{
-					FBaseGraphTask* Task = SortedAnyThreadTasks.Pop();
-					if (Task)
-					{
-						SetWorking(MyIndex, LocalNumWorkingThread);
-						return Task;
-					}
-				}
-				FAtomicStateBitfield NewAtomicForConsoleApproach = LocalAtomicForConsoleApproach;
-				// Queue is unowned, lets grab it and try to empty them
-				NewAtomicForConsoleApproach.QueueOwned = 1;
-				NewAtomicForConsoleApproach.QueueOwnerIndex = MyIndex;
-				NewAtomicForConsoleApproach.Working = LocalAtomicForConsoleApproach.Working & ~(1 << MyIndex);
-
-				check(LocalAtomicForConsoleApproach != NewAtomicForConsoleApproach);
-
-				if (FAtomicStateBitfield::InterlockedCompareExchange(&AtomicForConsoleApproach, NewAtomicForConsoleApproach, LocalAtomicForConsoleApproach) == LocalAtomicForConsoleApproach)
-				{
-					TestRandomizedThreads();
-					// we win
-					// dequeue any pending things
-					static TArray<FBaseGraphTask*> NewTasks;
-					int32 TotalTasks = 0;
-					bool bWeWereBeat = false;
-					if (SortedAnyThreadTasksHiPri.IsEmpty()) // otherwise someone beat us to it, so we need to drain the queue first
-					{
-						NewTasks.Reset();
-						IncomingAnyThreadTasksHiPri.PopAll(NewTasks);
-						if (NewTasks.Num())
-						{
-							TotalTasks += NewTasks.Num();
-							TLockFreePointerListLIFO<FBaseGraphTask> TempSortedAnyThreadTasks;
-							for (int32 Index = 0 ; Index < NewTasks.Num(); Index++) // we are going to take the last one for ourselves
-							{
-								TempSortedAnyThreadTasks.Push(NewTasks[Index]);
-							}
-							verify(SortedAnyThreadTasksHiPri.ReplaceListIfEmpty(TempSortedAnyThreadTasks));
-						}
-					}
-					else
-					{
-						bWeWereBeat = true;
-					}
-					if (SortedAnyThreadTasks.IsEmpty()) // otherwise someone beat us to it, so we need to drain the queue first
-					{
-						NewTasks.Reset();
-						IncomingAnyThreadTasks.PopAll(NewTasks);
-						if (NewTasks.Num())
-						{
-							TotalTasks += NewTasks.Num();
-							TLockFreePointerListLIFO<FBaseGraphTask> TempSortedAnyThreadTasks;
-							for (int32 Index = 0 ; Index < NewTasks.Num(); Index++) // we are going to take the last one for ourselves
-							{
-								TempSortedAnyThreadTasks.Push(NewTasks[Index]);
-							}
-							verify(SortedAnyThreadTasks.ReplaceListIfEmpty(TempSortedAnyThreadTasks));
-						}
-					}
-					else
-					{
-						bWeWereBeat = true;
-					}
-					// release the ownership
-					while (true)
-					{
-						TestRandomizedThreads();
-						FAtomicStateBitfield LocalAtomicForConsoleApproachInner = AtomicForConsoleApproach;
-						check(!(LocalAtomicForConsoleApproachInner.Stalled & (1 << MyIndex))); // if I was stalled, I would not be here
-						check(!GConsoleSpinModeLatched || !(LocalAtomicForConsoleApproachInner.Working & (1 << MyIndex))); // if I was working, I would not be here
-
-
-						FAtomicStateBitfield NewAtomicForConsoleApproachInner = LocalAtomicForConsoleApproachInner;
-						check(LocalAtomicForConsoleApproachInner.QueueOwned && LocalAtomicForConsoleApproachInner.QueueOwnerIndex == MyIndex);
-						NewAtomicForConsoleApproachInner.QueueOwned = false;
-						NewAtomicForConsoleApproachInner.QueueOwnerIndex = 0;
-						bool bAttemptStall = !TotalTasks && !bWeWereBeat &&
-							(!GConsoleSpinModeLatched || LocalAtomicForConsoleApproachInner.NumberOfStalledThreads() + LocalAtomicForConsoleApproachInner.NumberOfWorkingThreads() < LocalNumWorkingThread - 1);
-						if (bAttemptStall)
-						{
-							NewAtomicForConsoleApproachInner.Stalled = LocalAtomicForConsoleApproachInner.Stalled | (1 << MyIndex);
-							NewAtomicForConsoleApproachInner.Working = LocalAtomicForConsoleApproachInner.Working & ~(1 << MyIndex); // this is redundant in most cases, but might as well track the correct state
-							check((((uint32)NewAtomicForConsoleApproachInner.Stalled) & (1 << MyIndex)));
-						}
-
-						check(LocalAtomicForConsoleApproachInner != NewAtomicForConsoleApproachInner);
-						if (FAtomicStateBitfield::InterlockedCompareExchange(&AtomicForConsoleApproach, NewAtomicForConsoleApproachInner, LocalAtomicForConsoleApproachInner) == LocalAtomicForConsoleApproachInner)
-						{
-							TestRandomizedThreads();
-							if (bAttemptStall)
-							{
-								return nullptr; // we did stall
-							}
-							break;
-						}
-					}
-					if (GConsoleSpinModeLatched)
-					{
-						// -- here because if we actually do one of these tasks, we will start another thread by transitioning into working
-						TotalTasks--;
-						int32 ThreadsToStart = FMath::Min<int32>(FMath::Min<int32>(TotalTasks, GMaxTasksToStartOnDequeue), LocalNumWorkingThread);
-						while (ThreadsToStart > 0)
-						{
-							TestRandomizedThreads();
-							if (!StartTaskThreadFastMode(MyIndex))
-							{
-								break;
-							}
-							ThreadsToStart--;
-						}
-					}
+					FPlatformProcess::SleepNoStats(0.0f);
 				}
 			}
 		}
 	}
+#else
 	// API used by FWorkerThread's
 
 	/** 
-	 *	Attempt to steal some work from another thread.
 	 *	@param	ThreadInNeed; Id of the thread requesting work.
 	 *	@return Task that was stolen if any was found.
 	**/
 	FBaseGraphTask* FindWork(ENamedThreads::Type ThreadInNeed)
 	{
+		int32 LocalNumWorkingThread = GetNumWorkerThreads();
+		uint32 MyIndex = (uint32(ThreadInNeed) - NumNamedThreads) % NumTaskThreadsPerSet;
+		int32 Priority = (uint32(ThreadInNeed) - NumNamedThreads) / NumTaskThreadsPerSet;
+		check(MyIndex >= 0 && (int32)MyIndex < LocalNumWorkingThread); 
+		check(Priority >= 0 && Priority < ENamedThreads::NumThreadPriorities && 
+			(ENamedThreads::bHasHighPriorityThreads || (Priority << ENamedThreads::ThreadPriorityShift) != ENamedThreads::HighThreadPriority) &&
+			(ENamedThreads::bHasBackgroundThreads || (Priority << ENamedThreads::ThreadPriorityShift) != ENamedThreads::BackgroundThreadPriority)
+			);
 		TestRandomizedThreads();
-		if (GFastSchedulerLatched)
-		{
-			return FindWorkConsole(ThreadInNeed);
-		}
+		check(!GFastSchedulerLatched); // the fast scheduler is no longer supported without the new lock free lists.
 		{
-			FBaseGraphTask* Task = SortedAnyThreadTasksHiPri.Pop();
+			FBaseGraphTask* Task = SortedAnyThreadTasksHiPri[Priority].Pop();
 			if (Task)
 			{
 				return Task;
 			}
 		}
-		if (!IncomingAnyThreadTasksHiPri.IsEmpty())
+		if (!IncomingAnyThreadTasksHiPri[Priority].IsEmpty())
 		{
 			do
 			{
-				FScopeLock ScopeLock(&CriticalSectionForSortingIncomingAnyThreadTasksHiPri);
+				FScopeLock ScopeLock(&CriticalSectionForSortingIncomingAnyThreadTasksHiPri[Priority]);
 				if (GFastSchedulerLatched)
 				{
 					return nullptr;
 				}
-				if (!IncomingAnyThreadTasksHiPri.IsEmpty() && SortedAnyThreadTasksHiPri.IsEmpty())
+				if (!IncomingAnyThreadTasksHiPri[Priority].IsEmpty() && SortedAnyThreadTasksHiPri[Priority].IsEmpty())
 				{
 					static TArray<FBaseGraphTask*> NewTasks;
 					NewTasks.Reset();
-					IncomingAnyThreadTasksHiPri.PopAll(NewTasks);
+					IncomingAnyThreadTasksHiPri[Priority].PopAll(NewTasks);
 					check(NewTasks.Num());
 
 					if (NewTasks.Num() > 1)
@@ -1934,21 +2068,21 @@ public:
 						{
 							TempSortedAnyThreadTasks.Push(NewTasks[Index]);
 						}
-						verify(SortedAnyThreadTasksHiPri.ReplaceListIfEmpty(TempSortedAnyThreadTasks));
+						verify(SortedAnyThreadTasksHiPri[Priority].ReplaceListIfEmpty(TempSortedAnyThreadTasks));
 					}
 					return NewTasks[NewTasks.Num() - 1];
 				}
 				{
-					FBaseGraphTask* Task = SortedAnyThreadTasksHiPri.Pop();
+					FBaseGraphTask* Task = SortedAnyThreadTasksHiPri[Priority].Pop();
 					if (Task)
 					{
 						return Task;
 					}
 				}
-			} while (!IncomingAnyThreadTasksHiPri.IsEmpty() || !SortedAnyThreadTasksHiPri.IsEmpty());
+			} while (!IncomingAnyThreadTasksHiPri[Priority].IsEmpty() || !SortedAnyThreadTasksHiPri[Priority].IsEmpty());
 		}
 		{
-			FBaseGraphTask* Task = SortedAnyThreadTasks.Pop();
+			FBaseGraphTask* Task = SortedAnyThreadTasks[Priority].Pop();
 			if (Task)
 			{
 				return Task;
@@ -1956,37 +2090,44 @@ public:
 		}
 		do
 		{
-			FScopeLock ScopeLock(&CriticalSectionForSortingIncomingAnyThreadTasks);
+			FScopeLock ScopeLock(&CriticalSectionForSortingIncomingAnyThreadTasks[Priority]);
 			if (GFastSchedulerLatched)
 			{
 				return nullptr;
 			}
-			if (!IncomingAnyThreadTasks.IsEmpty() && SortedAnyThreadTasks.IsEmpty())
+			if (!IncomingAnyThreadTasks[Priority].IsEmpty() && SortedAnyThreadTasks[Priority].IsEmpty())
 			{
-				static TArray<FBaseGraphTask*> NewTasks;
+				struct FPadArray
+				{
+					uint8 Pad1[PLATFORM_CACHE_LINE_SIZE / 2];
+					TArray<FBaseGraphTask*> NewTaskArray;
+					uint8 Pad2[PLATFORM_CACHE_LINE_SIZE / 2];
+				};
+				static FPadArray PrioTasks[MAX_THREAD_PRIORITIES];
+				TArray<FBaseGraphTask*>& NewTasks = PrioTasks[Priority].NewTaskArray;
 				NewTasks.Reset();
-				IncomingAnyThreadTasks.PopAll(NewTasks);
+				IncomingAnyThreadTasks[Priority].PopAll(NewTasks);
 				check(NewTasks.Num());
 
 				if (NewTasks.Num() > 1)
 				{
 					TLockFreePointerListLIFO<FBaseGraphTask> TempSortedAnyThreadTasks;
-					for (int32 Index = 0 ; Index < NewTasks.Num() - 1; Index++) // we are going to take the last one for ourselves
+					for (int32 Index = 0; Index < NewTasks.Num() - 1; Index++) // we are going to take the last one for ourselves
 					{
 						TempSortedAnyThreadTasks.Push(NewTasks[Index]);
 					}
-					verify(SortedAnyThreadTasks.ReplaceListIfEmpty(TempSortedAnyThreadTasks));
+					verify(SortedAnyThreadTasks[Priority].ReplaceListIfEmpty(TempSortedAnyThreadTasks));
 				}
 				return NewTasks[NewTasks.Num() - 1];
 			}
 			{
-				FBaseGraphTask* Task = SortedAnyThreadTasks.Pop();
+				FBaseGraphTask* Task = SortedAnyThreadTasks[Priority].Pop();
 				if (Task)
 				{
 					return Task;
 				}
 			}
-		} while (!IncomingAnyThreadTasks.IsEmpty() || !SortedAnyThreadTasks.IsEmpty());
+		} while (!IncomingAnyThreadTasks[Priority].IsEmpty() || !SortedAnyThreadTasks[Priority].IsEmpty());
 		return nullptr;
 	}
 
@@ -2000,12 +2141,21 @@ public:
 	{
 		if (StallingThread >= NumNamedThreads && !GFastSchedulerLatched)
 		{
-			StalledUnnamedThreads.Push(&Thread(StallingThread));
+			int32 LocalNumWorkingThread = GetNumWorkerThreads();
+			uint32 MyIndex = (uint32(StallingThread) - NumNamedThreads) % NumTaskThreadsPerSet;
+			int32 Priority = (uint32(StallingThread) - NumNamedThreads) / NumTaskThreadsPerSet;
+			check(MyIndex >= 0 && (int32)MyIndex < LocalNumWorkingThread);
+			check(Priority >= 0 && Priority < ENamedThreads::NumThreadPriorities &&
+				(ENamedThreads::bHasHighPriorityThreads || (Priority << ENamedThreads::ThreadPriorityShift) != ENamedThreads::HighThreadPriority) &&
+				(ENamedThreads::bHasBackgroundThreads || (Priority << ENamedThreads::ThreadPriorityShift) != ENamedThreads::BackgroundThreadPriority)
+				);
+			StalledUnnamedThreads[Priority].Push(&Thread(StallingThread));
 		}
 	}
 
 	void SetTaskThreadPriorities(EThreadPriority Pri)
 	{
+		check(NumTaskThreadSets == 1); // otherwise tuning this doesn't make a lot of sense
 		for (int32 ThreadIndex = 0; ThreadIndex < NumThreads; ThreadIndex++)
 		{
 			if (ThreadIndex > LastExternalThread)
@@ -2042,17 +2192,36 @@ private:
 		if (TLSPointer)
 		{
 			checkThreadGraph(TLSPointer - WorkerThreads >= 0 && TLSPointer - WorkerThreads < NumThreads);
-			CurrentThreadIfKnown = ENamedThreads::Type(TLSPointer - WorkerThreads);
-			checkThreadGraph(Thread(CurrentThreadIfKnown).GetThreadId() == CurrentThreadIfKnown);
+			int32 ThreadIndex = TLSPointer - WorkerThreads;
+			checkThreadGraph(Thread(ThreadIndex).GetThreadId() == ThreadIndex);
+			if (ThreadIndex < NumNamedThreads)
+			{
+				CurrentThreadIfKnown = ENamedThreads::Type(ThreadIndex);
+			}
+			else
+			{
+				int32 Priority = (ThreadIndex - NumNamedThreads) / NumTaskThreadsPerSet;
+				CurrentThreadIfKnown = ENamedThreads::SetPriorities(ENamedThreads::Type(ThreadIndex), Priority, false);
+			}
 		}
 		return CurrentThreadIfKnown;
 	}
 
+	int32 ThreadIndexToPriorityIndex(int32 ThreadIndex)
+	{
+		check(ThreadIndex >= NumNamedThreads && ThreadIndex < NumThreads);
+		int32 Result = (ThreadIndex - NumNamedThreads) / NumTaskThreadsPerSet;
+		check(Result >= 0 && Result < NumTaskThreadSets);
+		return Result;
+	}
+
+
 
 	enum
 	{
 		/** Compile time maximum number of threads. @todo Didn't really need to be a compile time constant. **/
-		MAX_THREADS=12
+		MAX_THREADS=24,
+		MAX_THREAD_PRIORITIES = 3
 	};
 
 	/** Per thread data. **/
@@ -2061,42 +2230,45 @@ private:
 	int32				NumThreads;
 	/** Number of named threads actually in use. **/
 	int32				NumNamedThreads;
-	/** 
+	/** Number of tasks thread sets for priority **/
+	int32				NumTaskThreadSets;
+	/** Number of tasks threads per priority set **/
+	int32				NumTaskThreadsPerSet;
+	bool				bCreatedHiPriorityThreads;
+	bool				bCreatedBackgroundPriorityThreads;
+	/**
 	 * "External Threads" are not created, the thread is created elsewhere and makes an explicit call to run 
 	 * Here all of the named threads are external but that need not be the case.
 	 * All unnamed threads must be internal
 	**/
 	ENamedThreads::Type LastExternalThread;
 	/** Counter used to distribute new jobs to "any thread". **/
-	FThreadSafeCounter	NextUnnamedThreadForTaskFromUnknownThread;
-	/** Number of unnamed threads. **/
-	int32					NextUnnamedThreadMod;
-	/** Counter used to determine next thread to attempt to steal from. **/
-	FThreadSafeCounter	NextStealFromThread;
+	FThreadSafeCounter	NextUnnamedThreadForTaskFromUnknownThread[MAX_THREAD_PRIORITIES];
+	FThreadSafeCounter	ReentrancyCheck;
 	/** Index of TLS slot for FWorkerThread* pointer. **/
 	uint32				PerThreadIDTLSSlot;
 	/** Thread safe list of stalled thread "Hints". **/
-	TLockFreePointerListUnordered<FTaskThreadBase>		StalledUnnamedThreads; 
+	TLockFreePointerListUnordered<FTaskThreadBase, PLATFORM_CACHE_LINE_SIZE>		StalledUnnamedThreads[MAX_THREAD_PRIORITIES];
 
 	/** Array of callbacks to call before shutdown. **/
 	TArray<TFunction<void()> > ShutdownCallbacks;
 
 #if USE_NEW_LOCK_FREE_LISTS
-	#if USE_INTRUSIVE_TASKQUEUES
-		FLockFreePointerListFIFOIntrusive<FBaseGraphTask>		IncomingAnyThreadTasks;
-		FLockFreePointerListFIFOIntrusive<FBaseGraphTask>		IncomingAnyThreadTasksHiPri;
+#if USE_INTRUSIVE_TASKQUEUES
+	FLockFreePointerListFIFOIntrusive<FBaseGraphTask, PLATFORM_CACHE_LINE_SIZE>		IncomingAnyThreadTasks[MAX_THREAD_PRIORITIES];
+	FLockFreePointerListFIFOIntrusive<FBaseGraphTask, PLATFORM_CACHE_LINE_SIZE>		IncomingAnyThreadTasksHiPri[MAX_THREAD_PRIORITIES];
 	#else
-		TLockFreePointerListFIFO<FBaseGraphTask>		IncomingAnyThreadTasks;
-		TLockFreePointerListFIFO<FBaseGraphTask>		IncomingAnyThreadTasksHiPri;
+	TLockFreePointerListFIFO<FBaseGraphTask, PLATFORM_CACHE_LINE_SIZE>		IncomingAnyThreadTasks[MAX_THREAD_PRIORITIES];
+	TLockFreePointerListFIFO<FBaseGraphTask, PLATFORM_CACHE_LINE_SIZE>		IncomingAnyThreadTasksHiPri[MAX_THREAD_PRIORITIES];
 	#endif
 #else
-	TLockFreePointerListLIFO<FBaseGraphTask>		IncomingAnyThreadTasks;
-	TLockFreePointerListLIFO<FBaseGraphTask>		IncomingAnyThreadTasksHiPri;
-	TLockFreePointerListLIFO<FBaseGraphTask>		SortedAnyThreadTasks;
-	TLockFreePointerListLIFO<FBaseGraphTask>		SortedAnyThreadTasksHiPri;
+	TLockFreePointerListLIFO<FBaseGraphTask>		IncomingAnyThreadTasks[MAX_THREAD_PRIORITIES];
+	TLockFreePointerListLIFO<FBaseGraphTask>		IncomingAnyThreadTasksHiPri[MAX_THREAD_PRIORITIES];
+	TLockFreePointerListLIFO<FBaseGraphTask>		SortedAnyThreadTasks[MAX_THREAD_PRIORITIES];
+	TLockFreePointerListLIFO<FBaseGraphTask>		SortedAnyThreadTasksHiPri[MAX_THREAD_PRIORITIES];
 #endif
-	FCriticalSection CriticalSectionForSortingIncomingAnyThreadTasks;
-	FCriticalSection CriticalSectionForSortingIncomingAnyThreadTasksHiPri;
+	FCriticalSection CriticalSectionForSortingIncomingAnyThreadTasks[MAX_THREAD_PRIORITIES];
+	FCriticalSection CriticalSectionForSortingIncomingAnyThreadTasksHiPri[MAX_THREAD_PRIORITIES];
 
 	// we use a single atomic to control the state of the anythread queues. This limits the maximum number of threads.
 	struct FAtomicStateBitfield
@@ -2105,9 +2277,12 @@ private:
 		uint32 QueueOwned: 1;
 		uint32 QueueOwnerIndex: 4;
 #endif
-		uint32 Stalled: FAtomicStateBitfield_MAX_THREADS;
+		uint32 Stalled : FAtomicStateBitfield_MAX_THREADS;
 		uint32 Working: FAtomicStateBitfield_MAX_THREADS;
 
+		uint32 Unused;
+		uint8 PadToAvoidContention[PLATFORM_CACHE_LINE_SIZE - 2 * sizeof(uint32)];
+
 		FAtomicStateBitfield()
 			: 
 #if !USE_NEW_LOCK_FREE_LISTS
@@ -2117,6 +2292,7 @@ private:
 			 Stalled(0)
 			, Working(0)
 		{
+			static_assert(offsetof(FAtomicStateBitfield, Unused) == sizeof(uint32), "The bitfields in FAtomicStateBitfield must fit into a single uint32 and be the first member");
 		}
 
 		static FORCEINLINE FAtomicStateBitfield InterlockedCompareExchange(FAtomicStateBitfield* Dest, FAtomicStateBitfield Exchange, FAtomicStateBitfield Comparand)
@@ -2194,7 +2370,8 @@ private:
 		}
 	};
 
-	FAtomicStateBitfield AtomicForConsoleApproach;
+	uint8 PadToAvoidContention6[PLATFORM_CACHE_LINE_SIZE];
+	FAtomicStateBitfield AtomicForConsoleApproach[MAX_THREAD_PRIORITIES];
 };
 
 
@@ -2270,8 +2447,8 @@ static FTaskLockFreeListCounter GraphEventsWithInlineStorageInUse;
 
 #if USE_TLS_FOR_EVENTS
 
-static TPointerSet_TLSCacheBase<FGraphEvent, TLockFreePointerListUnordered<FGraphEvent>, FTaskLockFreeListCounter> TheGraphEventAllocator;
-static TPointerSet_TLSCacheBase<FGraphEventAndSmallTaskStorage, TLockFreePointerListUnordered<FGraphEventAndSmallTaskStorage>, FTaskLockFreeListCounter> TheGraphEventWithInlineStorageAllocator;
+static TPointerSet_TLSCacheBase<FGraphEvent, TLockFreePointerListUnordered<FGraphEvent, PLATFORM_CACHE_LINE_SIZE>, FTaskLockFreeListCounter> TheGraphEventAllocator;
+static TPointerSet_TLSCacheBase<FGraphEventAndSmallTaskStorage, TLockFreePointerListUnordered<FGraphEventAndSmallTaskStorage, PLATFORM_CACHE_LINE_SIZE>, FTaskLockFreeListCounter> TheGraphEventWithInlineStorageAllocator;
 
 static FThreadSafeCounter GraphEventsAllocated;
 static FThreadSafeCounter GraphEventsWithInlineStorageAllocated;
@@ -2380,7 +2557,7 @@ void FGraphEvent::Reset()
 }
 
 #else
-static TLockFreeClassAllocator_TLSCache<FGraphEvent> TheGraphEventAllocator;
+static TLockFreeClassAllocator_TLSCache<FGraphEvent, PLATFORM_CACHE_LINE_SIZE> TheGraphEventAllocator;
 
 FGraphEventRef FGraphEvent::CreateGraphEvent()
 {
@@ -2393,7 +2570,6 @@ void FGraphEvent::Recycle(FGraphEvent* ToRecycle)
 }
 #endif
 
-
 void FGraphEvent::DispatchSubsequents(TArray<FBaseGraphTask*>& NewTasks, ENamedThreads::Type CurrentThreadIfKnown)
 {
 	if (EventsToWaitFor.Num())
@@ -2406,7 +2582,7 @@ void FGraphEvent::DispatchSubsequents(TArray<FBaseGraphTask*>& NewTasks, ENamedT
 			STAT_FNullGraphTask_DontCompleteUntil,
 			STATGROUP_TaskGraphTasks);
 
-		TGraphTask<FNullGraphTask>::CreateTask(FGraphEventRef(this), &TempEventsToWaitFor, CurrentThreadIfKnown).ConstructAndDispatchWhenReady(GET_STATID(STAT_FNullGraphTask_DontCompleteUntil), ENamedThreads::HiPri(ENamedThreads::AnyThread));
+		TGraphTask<FNullGraphTask>::CreateTask(FGraphEventRef(this), &TempEventsToWaitFor, CurrentThreadIfKnown).ConstructAndDispatchWhenReady(GET_STATID(STAT_FNullGraphTask_DontCompleteUntil), ENamedThreads::AnyHiPriThreadHiPriTask);
 		return;
 	}
 
diff --git a/Engine/Source/Runtime/Core/Private/GenericPlatform/GenericPlatformMemory.cpp b/Engine/Source/Runtime/Core/Private/GenericPlatform/GenericPlatformMemory.cpp
--- a/Engine/Source/Runtime/Core/Private/GenericPlatform/GenericPlatformMemory.cpp
+++ b/Engine/Source/Runtime/Core/Private/GenericPlatform/GenericPlatformMemory.cpp
@@ -30,7 +30,7 @@ struct FGenericStatsUpdater
 	/** Called once per second, enqueues stats update. */
 	static bool EnqueueUpdateStats( float /*InDeltaTime*/ )
 	{
-		AsyncTask( ENamedThreads::AnyThread, []()
+		AsyncTask( ENamedThreads::AnyBackgroundThreadNormalTask, []()
 		{
 			DoUpdateStats();
 		} );
diff --git a/Engine/Source/Runtime/Core/Private/HAL/PThreadRunnableThread.h b/Engine/Source/Runtime/Core/Private/HAL/PThreadRunnableThread.h
--- a/Engine/Source/Runtime/Core/Private/HAL/PThreadRunnableThread.h
+++ b/Engine/Source/Runtime/Core/Private/HAL/PThreadRunnableThread.h
@@ -45,6 +45,8 @@ protected:
 			case TPri_AboveNormal: return 25;
 			case TPri_Normal: return 15;
 			case TPri_BelowNormal: return 5;
+			case TPri_Lowest: return 1;
+			case TPri_SlightlyBelowNormal: return 14;
 			default: UE_LOG(LogHAL, Fatal, TEXT("Unknown Priority passed to FRunnableThreadPThread::TranslateThreadPriority()"));
 		}
 	}
diff --git a/Engine/Source/Runtime/Core/Private/HAL/UnrealMemory.cpp b/Engine/Source/Runtime/Core/Private/HAL/UnrealMemory.cpp
--- a/Engine/Source/Runtime/Core/Private/HAL/UnrealMemory.cpp
+++ b/Engine/Source/Runtime/Core/Private/HAL/UnrealMemory.cpp
@@ -29,7 +29,7 @@ class FMallocPurgatoryProxy : public FMalloc
 	uint32 LastCheckFrame; // unitialized; won't matter
 	FThreadSafeCounter OutstandingSizeInKB;
 	FThreadSafeCounter NextOversizeClear;
-	TLockFreePointerListUnordered<void> Purgatory[PURGATORY_STOMP_CHECKS_FRAMES];
+	TLockFreePointerListUnordered<void, PLATFORM_CACHE_LINE_SIZE> Purgatory[PURGATORY_STOMP_CHECKS_FRAMES];
 
 public:
 	/**
diff --git a/Engine/Source/Runtime/Core/Private/Misc/EventPool.h b/Engine/Source/Runtime/Core/Private/Misc/EventPool.h
--- a/Engine/Source/Runtime/Core/Private/Misc/EventPool.h
+++ b/Engine/Source/Runtime/Core/Private/Misc/EventPool.h
@@ -85,5 +85,5 @@ public:
 private:
 
 	/** Holds the collection of recycled events. */
-	TLockFreePointerListUnordered<FEvent> Pool;
+	TLockFreePointerListUnordered<FEvent, PLATFORM_CACHE_LINE_SIZE> Pool;
 };
diff --git a/Engine/Source/Runtime/Core/Private/Misc/MemStack.cpp b/Engine/Source/Runtime/Core/Private/Misc/MemStack.cpp
--- a/Engine/Source/Runtime/Core/Private/Misc/MemStack.cpp
+++ b/Engine/Source/Runtime/Core/Private/Misc/MemStack.cpp
@@ -6,8 +6,8 @@ DECLARE_MEMORY_STAT(TEXT("MemStack Large Block"), STAT_MemStackLargeBLock,STATGR
 DECLARE_MEMORY_STAT(TEXT("PageAllocator Free"), STAT_PageAllocatorFree, STATGROUP_Memory);
 DECLARE_MEMORY_STAT(TEXT("PageAllocator Used"), STAT_PageAllocatorUsed, STATGROUP_Memory);
 
-TLockFreeFixedSizeAllocator<FPageAllocator::PageSize, FThreadSafeCounter> FPageAllocator::TheAllocator;
-TLockFreeFixedSizeAllocator<FPageAllocator::SmallPageSize, FThreadSafeCounter> FPageAllocator::TheSmallAllocator;
+FPageAllocator::TPageAllocator FPageAllocator::TheAllocator;
+FPageAllocator::TSmallPageAllocator FPageAllocator::TheSmallAllocator;
 
 
 #define USE_MEMSTACK_PUGATORY (0)
diff --git a/Engine/Source/Runtime/Core/Private/Stats/StatsData.cpp b/Engine/Source/Runtime/Core/Private/Stats/StatsData.cpp
--- a/Engine/Source/Runtime/Core/Private/Stats/StatsData.cpp
+++ b/Engine/Source/Runtime/Core/Private/Stats/StatsData.cpp
@@ -439,9 +439,9 @@ void FRawStatStackNode::Encode(TArray<FStatMessage>& OutStats) const
 	}
 }
 
-TLockFreeFixedSizeAllocator<sizeof(FRawStatStackNode)>& GetRawStatStackNodeAllocator()
+TLockFreeFixedSizeAllocator<sizeof(FRawStatStackNode), PLATFORM_CACHE_LINE_SIZE>& GetRawStatStackNodeAllocator()
 {
-	static TLockFreeFixedSizeAllocator<sizeof(FRawStatStackNode)> TheAllocator;
+	static TLockFreeFixedSizeAllocator<sizeof(FRawStatStackNode), PLATFORM_CACHE_LINE_SIZE> TheAllocator;
 	return TheAllocator;
 }
 
@@ -460,9 +460,9 @@ void FRawStatStackNode::operator delete(void *RawMemory)
 	FComplexRawStatStackNode
 -----------------------------------------------------------------------------*/
 
-TLockFreeFixedSizeAllocator<sizeof(FComplexRawStatStackNode)>& GetRawStatStackNodeAllocatorEx()
+TLockFreeFixedSizeAllocator<sizeof(FComplexRawStatStackNode), PLATFORM_CACHE_LINE_SIZE>& GetRawStatStackNodeAllocatorEx()
 {	
-	static TLockFreeFixedSizeAllocator<sizeof(FComplexRawStatStackNode)> TheAllocatorComplex;
+	static TLockFreeFixedSizeAllocator<sizeof(FComplexRawStatStackNode), PLATFORM_CACHE_LINE_SIZE> TheAllocatorComplex;
 	return TheAllocatorComplex;
 }
 
diff --git a/Engine/Source/Runtime/Core/Private/WinRT/WinRTRunnableThread.h b/Engine/Source/Runtime/Core/Private/WinRT/WinRTRunnableThread.h
--- a/Engine/Source/Runtime/Core/Private/WinRT/WinRTRunnableThread.h
+++ b/Engine/Source/Runtime/Core/Private/WinRT/WinRTRunnableThread.h
@@ -108,6 +108,20 @@ public:
 			Kill(true);
 		}
 	}
+
+	static int TranslateThreadPriority(EThreadPriority Priority)
+	{
+		switch (Priority)
+		{
+		case TPri_AboveNormal: return THREAD_PRIORITY_ABOVE_NORMAL;
+		case TPri_Normal: return THREAD_PRIORITY_NORMAL;
+		case TPri_BelowNormal: return THREAD_PRIORITY_BELOW_NORMAL;
+		case TPri_Highest: return THREAD_PRIORITY_HIGHEST;
+		case TPri_Lowest: return THREAD_PRIORITY_LOWEST;
+		case TPri_SlightlyBelowNormal: return THREAD_PRIORITY_NORMAL - 1;
+		default: UE_LOG(LogHAL, Fatal, TEXT("Unknown Priority passed to TranslateThreadPriority()")); return TPri_Normal;
+		}
+	}
 	
 	virtual void SetThreadPriority(EThreadPriority NewPriority) override
 	{
@@ -116,10 +130,7 @@ public:
 		{
 			ThreadPriority = NewPriority;
 			// Change the priority on the thread
-			ThreadEmulation::SetThreadPriority(Thread,
-				ThreadPriority == TPri_AboveNormal ? THREAD_PRIORITY_ABOVE_NORMAL :
-				ThreadPriority == TPri_BelowNormal ? THREAD_PRIORITY_BELOW_NORMAL :
-				THREAD_PRIORITY_NORMAL);
+			ThreadEmulation::SetThreadPriority(Thread, TranslateThreadPriority(ThreadPriority));
 		}
 	}
 
diff --git a/Engine/Source/Runtime/Core/Private/Windows/WindowsRunnableThread.h b/Engine/Source/Runtime/Core/Private/Windows/WindowsRunnableThread.h
--- a/Engine/Source/Runtime/Core/Private/Windows/WindowsRunnableThread.h
+++ b/Engine/Source/Runtime/Core/Private/Windows/WindowsRunnableThread.h
@@ -90,6 +90,20 @@ public:
 		}
 	}
 	
+	static int TranslateThreadPriority(EThreadPriority Priority)
+	{
+		switch (Priority)
+		{
+		case TPri_AboveNormal: return THREAD_PRIORITY_ABOVE_NORMAL;
+		case TPri_Normal: return THREAD_PRIORITY_NORMAL;
+		case TPri_BelowNormal: return THREAD_PRIORITY_BELOW_NORMAL;
+		case TPri_Highest: return THREAD_PRIORITY_HIGHEST;
+		case TPri_Lowest: return THREAD_PRIORITY_LOWEST;
+		case TPri_SlightlyBelowNormal: return THREAD_PRIORITY_NORMAL - 1;
+		default: UE_LOG(LogHAL, Fatal, TEXT("Unknown Priority passed to TranslateThreadPriority()")); return TPri_Normal;
+		}
+	}
+
 	virtual void SetThreadPriority( EThreadPriority NewPriority ) override
 	{
 		// Don't bother calling the OS if there is no need
@@ -97,10 +111,7 @@ public:
 		{
 			ThreadPriority = NewPriority;
 			// Change the priority on the thread
-			::SetThreadPriority(Thread,
-				ThreadPriority == TPri_AboveNormal ? THREAD_PRIORITY_ABOVE_NORMAL :
-				ThreadPriority == TPri_BelowNormal ? THREAD_PRIORITY_BELOW_NORMAL :
-				THREAD_PRIORITY_NORMAL);
+			::SetThreadPriority(Thread, TranslateThreadPriority(ThreadPriority));
 		}
 	}
 
diff --git a/Engine/Source/Runtime/Core/Public/Async/ParallelFor.h b/Engine/Source/Runtime/Core/Public/Async/ParallelFor.h
--- a/Engine/Source/Runtime/Core/Public/Async/ParallelFor.h
+++ b/Engine/Source/Runtime/Core/Public/Async/ParallelFor.h
@@ -74,7 +74,7 @@ public:
 	}
 	static FORCEINLINE ENamedThreads::Type GetDesiredThread()
 	{
-		return ENamedThreads::HiPri(ENamedThreads::AnyThread);
+		return ENamedThreads::AnyHiPriThreadHiPriTask;
 	}
 	static FORCEINLINE ESubsequentsMode::Type GetSubsequentsMode() 
 	{ 
