Copying //UE4/Dev-Core to //UE4/Main

diff --git a/Engine/Source/Runtime/Core/Private/HAL/MallocBinned2.cpp b/Engine/Source/Runtime/Core/Private/HAL/MallocBinned2.cpp
--- a/Engine/Source/Runtime/Core/Private/HAL/MallocBinned2.cpp
+++ b/Engine/Source/Runtime/Core/Private/HAL/MallocBinned2.cpp
@@ -9,71 +9,25 @@
 #include "MallocBinned2.h"
 #include "MemoryMisc.h"
 #include "HAL/PlatformAtomics.h"
-#include "MallocVerify.h"
-
-//when modifying the global allocator stats, if we are using COARSE locks, then all callsites for stat modification are covered by the allocator-wide access guard. Thus the stats can be modified directly.
-//If we are using FINE locks, then we must modify the stats through atomics as the locks are either not actually covering the stat callsites, or are locking specific table locks which is not sufficient for stats.
-#if STATS
-#	define BINNED2_INCREMENT_STATCOUNTER(counter) (FPlatformAtomics::InterlockedIncrement(&(counter)))
-#	define BINNED2_DECREMENT_STATCOUNTER(counter) (FPlatformAtomics::InterlockedDecrement(&(counter)))
-#	define BINNED2_ADD_STATCOUNTER(counter, value) (FPlatformAtomics::InterlockedAdd(&counter, (value)))
-#	define BINNED2_PEAK_STATCOUNTER(PeakCounter, CompareVal)	{																												\
-																volatile TDecay<decltype(PeakCounter)>::Type NewCompare;													\
-																volatile TDecay<decltype(PeakCounter)>::Type NewPeak;														\
-																do																											\
-																{																											\
-																	NewCompare = (PeakCounter);																				\
-																	NewPeak = FMath::Max((PeakCounter), (CompareVal));														\
-																}																											\
-																while (FPlatformAtomics::InterlockedCompareExchange(&(PeakCounter), NewPeak, NewCompare) != NewCompare);	\
-															}
-#else
-#	define BINNED2_INCREMENT_STATCOUNTER(counter)
-#	define BINNED2_DECREMENT_STATCOUNTER(counter)
-#	define BINNED2_ADD_STATCOUNTER(counter, value)
-#	define BINNED2_PEAK_STATCOUNTER(PeakCounter, CompareVal)
-#endif
-
-/** Malloc binned allocator specific stats. */
-DECLARE_MEMORY_STAT_EXTERN(TEXT("Binned Os Current"),		STAT_Binned2_OsCurrent,STATGROUP_MemoryAllocator, CORE_API);
-DECLARE_MEMORY_STAT_EXTERN(TEXT("Binned Os Peak"),			STAT_Binned2_OsPeak,STATGROUP_MemoryAllocator, CORE_API);
-DECLARE_MEMORY_STAT_EXTERN(TEXT("Binned Waste Current"),	STAT_Binned2_WasteCurrent,STATGROUP_MemoryAllocator, CORE_API);
-DECLARE_MEMORY_STAT_EXTERN(TEXT("Binned Waste Peak"),		STAT_Binned2_WastePeak,STATGROUP_MemoryAllocator, CORE_API);
-DECLARE_MEMORY_STAT_EXTERN(TEXT("Binned Used Current"),		STAT_Binned2_UsedCurrent,STATGROUP_MemoryAllocator, CORE_API);
-DECLARE_MEMORY_STAT_EXTERN(TEXT("Binned Used Peak"),		STAT_Binned2_UsedPeak,STATGROUP_MemoryAllocator, CORE_API);
-DECLARE_DWORD_COUNTER_STAT_EXTERN(TEXT("Binned Current Allocs"),	STAT_Binned2_CurrentAllocs,STATGROUP_MemoryAllocator, CORE_API);
-DECLARE_DWORD_COUNTER_STAT_EXTERN(TEXT("Binned Total Allocs"),		STAT_Binned2_TotalAllocs,STATGROUP_MemoryAllocator, CORE_API);
-DECLARE_MEMORY_STAT_EXTERN(TEXT("Binned Slack Current"),	STAT_Binned2_SlackCurrent,STATGROUP_MemoryAllocator, CORE_API);
-
-/** Malloc binned allocator specific stats. */
-DEFINE_STAT(STAT_Binned2_OsCurrent);
-DEFINE_STAT(STAT_Binned2_OsPeak);
-DEFINE_STAT(STAT_Binned2_WasteCurrent);
-DEFINE_STAT(STAT_Binned2_WastePeak);
-DEFINE_STAT(STAT_Binned2_UsedCurrent);
-DEFINE_STAT(STAT_Binned2_UsedPeak);
-DEFINE_STAT(STAT_Binned2_CurrentAllocs);
-DEFINE_STAT(STAT_Binned2_TotalAllocs);
-DEFINE_STAT(STAT_Binned2_SlackCurrent);
-
 
 // Block sizes are based around getting the maximum amount of allocations per pool, with as little alignment waste as possible.
-// Block sizes should be close to even divisors of the POOL_SIZE, and well distributed. They must be 16-byte aligned as well.
-const uint32 GMallocBinned2BlockSizes[] =
+// Block sizes should be close to even divisors of the system page size, and well distributed.
+// They must be 16-byte aligned as well.
+const uint32 GMallocBinned2SmallBlockSizes[] =
 {
-	8,		16,		32,		48,		64,		80,		96,		112,
-	128,	160,	192,	224,	256,	288,	320,	384,
-	448,	512,	576,	640,	704,	768,	896,	1024,
-	1168,	1360,	1632,	2048,	2336,	2720,	3264,	4096,
-	4672,	5456,	6544,	8192,	9360,	10912,	13104,	16384,
-	21840,	32768
+	16,		32,		48,		64,		80,		96,		112,	128,
+	160,	192,	224,	256,	288,	320,	384,	448,
+	512,	576,	640,	704,	768,	896,	1024,	1168,
+	1360,	1632,	2048,	2336,	2720,	3264,	4096,	4672,
+	5456,	6544,	8192,	9360,	10912,	13104,	16384,	21840,
+	32768
 };
 
-/** Information about a piece of free memory. 8 bytes */
+/** Information about a piece of free memory. */
 struct FMallocBinned2::FFreeMem
 {
 	/** Next or MemLastPool[], always in order by pool. */
-	FFreeMem*	Next;
+	FFreeMem*	NextFreePool;
 	/** Number of consecutive free blocks here, at least 1. */
 	uint32		NumFreeBlocks;
 };
@@ -83,11 +37,10 @@ struct FMallocBinned2::FPoolInfo
 {
 	/** Number of allocated elements in this pool, when counts down to zero can free the entire pool. */
 	uint16			Taken;		// 2
-	/** Index of pool. Index into MemSizeToPoolTable[]. Valid when < MAX_POOLED_ALLOCATION_SIZE, MAX_POOLED_ALLOCATION_SIZE is OsTable.
-		When AllocSize is 0, this is the number of pages to step back to find the base address of an allocation. See FindPoolInfo()
-	*/
-	uint16			TableIndex; // 4		
-	/** Number of bytes allocated */
+	/** Index of pool. Index into MemSizeToPoolTable[]. Valid when < MAX_SMALL_POOLED_ALLOCATION_SIZE, -1 is OsTable.
+		When AllocSize is 0, this is the number of pages to step back to find the base address of an allocation. See FindPoolInfo() */
+	uint16			TableIndex; // 4
+	/** Number of bytes allocated, or 0, this is a trailing pool (i.e. a pool used to find its way back to the base pool) */
 	uint32			AllocSize;	// 8
 	/** Pointer to first free memory in this pool or the OS Allocation Size in bytes if this allocation is not binned*/
 	FFreeMem*		FirstMem;   // 12/16
@@ -97,27 +50,51 @@ struct FMallocBinned2::FPoolInfo
 	/** Explicit padding for 32 bit builds */
 	uint8 Padding[12]; // 32
 #endif
+	uint32 GetTrailingPoolOffset() const
+	{
+		checkSlow(IsTrailingPool());
+		return TableIndex;
+	}
 
-	void SetAllocationSizes( uint32 InBytes, UPTRINT InOsBytes, uint32 InTableIndex, uint32 SmallAllocLimt )
+	void SetTrailingPoolOffset(uint32 InOffset)
 	{
-		TableIndex=InTableIndex;
-		AllocSize=InBytes;
-		if (TableIndex == SmallAllocLimt)
-		{
-			FirstMem=(FFreeMem*)InOsBytes;
-		}
+		TableIndex = InOffset;
+		AllocSize  = 0;
 	}
 
-	UPTRINT GetOsBytes( uint32 InPageSize, uint32 SmallAllocLimt ) const
+	void SetPoolAllocationSizes(uint32 InBytes, uint32 InTableIndex)
 	{
-		if (TableIndex == SmallAllocLimt)
-		{
-			return (UPTRINT)FirstMem;
-		}
-		else
-		{
-			return Align(AllocSize, InPageSize);
-		}
+		// Shouldn't be pooling zero byte allocations
+		checkSlow(InBytes != 0);
+
+		TableIndex = InTableIndex;
+		AllocSize  = InBytes;
+	}
+
+	void SetOSAllocationSizes(uint32 InBytes, UPTRINT InOsBytes)
+	{
+		// Shouldn't be pooling zero byte allocations
+		checkSlow(InBytes != 0);
+
+		TableIndex = (uint16)-1;
+		AllocSize  = InBytes;
+		FirstMem   = (FFreeMem*)InOsBytes;
+	}
+
+	UPTRINT GetOsBytes() const
+	{
+		check(IsOSAllocation());
+		return (UPTRINT)FirstMem;
+	}
+
+	FORCEINLINE bool IsTrailingPool() const
+	{
+		return AllocSize == 0;
+	}
+
+	FORCEINLINE bool IsOSAllocation() const
+	{
+		return TableIndex == (uint16)-1;
 	}
 
 	void Link( FPoolInfo*& Before )
@@ -141,20 +118,26 @@ struct FMallocBinned2::FPoolInfo
 	}
 };
 
+template <>
+struct TIsZeroConstructType<FMallocBinned2::FPoolInfo>
+{
+	enum { Value = true };
+};
+
 /** Hash table struct for retrieving allocation book keeping information */
 struct FMallocBinned2::PoolHashBucket
 {
-	UPTRINT         Key;
+	UPTRINT         BucketIndex;
 	FPoolInfo*      FirstPool;
 	PoolHashBucket* Prev;
 	PoolHashBucket* Next;
 
 	PoolHashBucket()
 	{
-		Key      = 0;
-		FirstPool= nullptr;
-		Prev     = this;
-		Next     = this;
+		BucketIndex = 0;
+		FirstPool   = nullptr;
+		Prev        = this;
+		Next        = this;
 	}
 
 	void Link(PoolHashBucket* After)
@@ -176,19 +159,11 @@ struct FMallocBinned2::PoolHashBucket
 
 struct FMallocBinned2::Private
 {
-	static_assert(ARRAY_COUNT(GMallocBinned2BlockSizes) == POOL_COUNT, "Block size array size must match POOL_COUNT");
-	static_assert(sizeof(FMallocBinned2::FPoolInfo)     == 32,         "sizeof(FPoolInfo) is expected to be 32 bytes");
+	static_assert(ARRAY_COUNT(GMallocBinned2SmallBlockSizes) == SMALL_POOL_COUNT, "Small block size array size must match SMALL_POOL_COUNT");
+	static_assert(sizeof(FMallocBinned2::FPoolInfo)          == 32,               "sizeof(FPoolInfo) is expected to be 32 bytes");
 
 	// Default alignment for binned allocator
-	enum { DEFAULT_BINNED_ALLOCATOR_ALIGNMENT = sizeof(FFreeMem) };
-
-	// There is internal limit on page size of 64k
-	enum { PAGE_SIZE_LIMIT = 65536 };
-
-	// BINNED_ALLOC_POOL_SIZE can be increased beyond 64k to cause binned malloc to allocate
-	// the small size bins in bigger chunks. If OS Allocation is slow, increasing
-	// this number *may* help performance but YMMV.
-	enum { BINNED_ALLOC_POOL_SIZE = 65536 }; 
+	enum { DEFAULT_BINNED_ALLOCATOR_ALIGNMENT = 16 };
 
 	// Implementation. 
 	static CA_NO_RETURN void OutOfMemory(uint64 Size, uint32 Alignment=0)
@@ -197,82 +172,60 @@ struct FMallocBinned2::Private
 		FPlatformMemory::OnOutOfMemory(Size, Alignment);
 	}
 
-	static FORCEINLINE void TrackStats(FPoolTable* Table, SIZE_T Size)
-	{
-#if STATS
-		// keep track of memory lost to padding
-		Table->TotalWaste += Table->BlockSize - Size;
-		Table->TotalRequests++;
-		Table->ActiveRequests++;
-		Table->MaxActiveRequests = FMath::Max(Table->MaxActiveRequests, Table->ActiveRequests);
-		Table->MaxRequest = Size > Table->MaxRequest ? Size : Table->MaxRequest;
-		Table->MinRequest = Size < Table->MinRequest ? Size : Table->MinRequest;
-#endif
-	}
-
 	/** 
-	 * Create a 64k page of FPoolInfo structures for tracking allocations
+	 * Creates an array of FPoolInfo structures for tracking allocations.
 	 */
-	static FPoolInfo* CreateIndirect(FMallocBinned2& Allocator)
+	static FPoolInfo* CreatePoolArray(uint64 NumPools)
 	{
-		uint64 IndirectPoolBlockSizeBytes = Allocator.IndirectPoolBlockSize * sizeof(FPoolInfo);
+		uint64 PoolArraySize = NumPools * sizeof(FPoolInfo);
 
-		checkSlow(IndirectPoolBlockSizeBytes <= Allocator.PageSize);
-		FPoolInfo* Indirect = (FPoolInfo*)FPlatformMemory::BinnedAllocFromOS(IndirectPoolBlockSizeBytes);
-		if( !Indirect )
+		void* Result = FPlatformMemory::BinnedAllocFromOS(PoolArraySize);
+		if (!Result)
 		{
-			OutOfMemory(IndirectPoolBlockSizeBytes);
+			OutOfMemory(PoolArraySize);
 		}
-		FMemory::Memset(Indirect, 0, IndirectPoolBlockSizeBytes);
-
-		BINNED2_PEAK_STATCOUNTER(Allocator.Stats.OsPeak,    BINNED2_ADD_STATCOUNTER(Allocator.Stats.OsCurrent,    (int64)(Align(IndirectPoolBlockSizeBytes, Allocator.PageSize))));
-		BINNED2_PEAK_STATCOUNTER(Allocator.Stats.WastePeak, BINNED2_ADD_STATCOUNTER(Allocator.Stats.WasteCurrent, (int64)(Align(IndirectPoolBlockSizeBytes, Allocator.PageSize))));
 
-		return Indirect;
+		DefaultConstructItems<FPoolInfo>(Result, NumPools);
+		return (FPoolInfo*)Result;
 	}
 
-	/** 
-	 * Gets the FPoolInfo for a memory address. If no valid info exists one is created. 
-	 * NOTE: This function requires a mutex across threads, but it's the caller's responsibility to 
-	 * acquire the mutex before calling
+	/**
+	 * Gets the FPoolInfo for a memory address. If no valid info exists one is created.
 	 */
-	static /*FORCEINLINE*/ FPoolInfo* GetPoolInfo(FMallocBinned2& Allocator, void* InPtr)
+	static FPoolInfo* GetOrCreatePoolInfo(FMallocBinned2& Allocator, void* InPtr)
 	{
-		UPTRINT Ptr = (UPTRINT)InPtr;
-
-		UPTRINT Key       = Ptr >> Allocator.HashKeyShift;
-		UPTRINT Hash      = Key & (Allocator.MaxHashBuckets - 1);
-		UPTRINT PoolIndex = (Ptr >> Allocator.PoolBitShift) & Allocator.PoolMask;
+		UPTRINT BucketIndex;
+		uint32  PoolIndex;
+		Allocator.PtrToPoolMapping.GetHashBucketAndPoolIndices(InPtr, BucketIndex, PoolIndex);
 
-		PoolHashBucket* Collision = &Allocator.HashBuckets[Hash];
+		PoolHashBucket* FirstBucket = &Allocator.HashBuckets[BucketIndex];
+		PoolHashBucket* Collision   = FirstBucket;
 		do
 		{
 			if (!Collision->FirstPool)
 			{
-				Collision->Key       = Key;
-				Collision->FirstPool = CreateIndirect(Allocator);
+				Collision->BucketIndex = BucketIndex;
+				Collision->FirstPool   = CreatePoolArray(Allocator.NumPoolsPerPage);
 
 				return &Collision->FirstPool[PoolIndex];
 			}
 
-			if (Collision->Key == Key)
+			if (Collision->BucketIndex == BucketIndex)
 			{
 				return &Collision->FirstPool[PoolIndex];
 			}
 
 			Collision = Collision->Next;
-		} while (Collision != &Allocator.HashBuckets[Hash]);
+		}
+		while (Collision != FirstBucket);
 
 		// Create a new hash bucket entry
-		if (!Allocator.HashBucketFreeList) 
+		if (!Allocator.HashBucketFreeList)
 		{
 			const uint32 PageSize = Allocator.PageSize;
 
 			Allocator.HashBucketFreeList = (PoolHashBucket*)FPlatformMemory::BinnedAllocFromOS(PageSize);
 
-			BINNED2_PEAK_STATCOUNTER(Allocator.Stats.OsPeak,    BINNED2_ADD_STATCOUNTER(Allocator.Stats.OsCurrent,    PageSize));
-			BINNED2_PEAK_STATCOUNTER(Allocator.Stats.WastePeak, BINNED2_ADD_STATCOUNTER(Allocator.Stats.WasteCurrent, PageSize));
-			
 			for (UPTRINT i = 0, n = PageSize / sizeof(PoolHashBucket); i < n; ++i)
 			{
 				Allocator.HashBucketFreeList->Link(new (Allocator.HashBucketFreeList + i) PoolHashBucket());
@@ -293,12 +246,12 @@ struct FMallocBinned2::Private
 
 		if (!NewBucket->FirstPool)
 		{
-			NewBucket->FirstPool = CreateIndirect(Allocator);
+			NewBucket->FirstPool = CreatePoolArray(Allocator.NumPoolsPerPage);
 		}
 
-		NewBucket->Key = Key;
+		NewBucket->BucketIndex = BucketIndex;
 
-		Allocator.HashBuckets[Hash].Link(NewBucket);
+		FirstBucket->Link(NewBucket);
 
 		return &NewBucket->FirstPool[PoolIndex];
 	}
@@ -307,32 +260,33 @@ struct FMallocBinned2::Private
 	{
 		const UPTRINT PageSize = (UPTRINT)Allocator.PageSize;
 
-		UPTRINT Ptr = (UPTRINT)InPtr &~ (PageSize - 1);
-		for (uint32 Repeat = (BINNED_ALLOC_POOL_SIZE / PageSize) + 1; Repeat; --Repeat)
+		uint8* Ptr = AlignDown((uint8*)InPtr, PageSize);
+		for (uint32 Repeat = 2; Repeat; --Repeat)
 		{
-			uint32 Key       = Ptr >> Allocator.HashKeyShift;
-			uint32 Hash      = Key & (Allocator.MaxHashBuckets - 1);
-			uint32 PoolIndex = (Ptr >> Allocator.PoolBitShift) & Allocator.PoolMask;
+			UPTRINT BucketIndex;
+			uint32  PoolIndex;
+			Allocator.PtrToPoolMapping.GetHashBucketAndPoolIndices(Ptr, BucketIndex, PoolIndex);
 
-			PoolHashBucket* Collision = &Allocator.HashBuckets[Hash];
+			PoolHashBucket* FirstBucket = &Allocator.HashBuckets[BucketIndex];
+			PoolHashBucket* Collision   = FirstBucket;
 			for (;;)
 			{
-				if (Collision->Key == Key)
+				if (Collision->BucketIndex == BucketIndex)
 				{
-					if (Collision->FirstPool[PoolIndex].AllocSize)
+					if (!Collision->FirstPool[PoolIndex].IsTrailingPool())
 					{
-						AllocationBase = (void*)Ptr;
+						AllocationBase = Ptr;
 						return &Collision->FirstPool[PoolIndex];
 					}
 
-					Ptr = ((Ptr - (PageSize * Collision->FirstPool[PoolIndex].TableIndex)) - 1) &~ (PageSize - 1);
+					Ptr -= PageSize * (Collision->FirstPool[PoolIndex].GetTrailingPoolOffset());
 					break;
 				}
 
 				Collision = Collision->Next;
-				if (Collision == &Allocator.HashBuckets[Hash])
+				if (Collision == FirstBucket)
 				{
-					Ptr = (Ptr - 1) &~ (PageSize - 1);
+					Ptr -= PageSize;
 					break;
 				}
 			}
@@ -347,71 +301,20 @@ struct FMallocBinned2::Private
 		return FindPoolInfo(Allocator, Ptr, BasePtr);
 	}
 
-	static FPoolInfo* AllocatePoolMemory(FMallocBinned2& Allocator, FPoolTable& Table, uint32 PoolSize, uint16 TableIndex)
-	{
-		const uint32 PageSize  = Allocator.PageSize;
-		const uint32 BlockSize = Table.BlockSize;
-
-		// Must create a new pool.
-		uint32 Blocks   = PoolSize / BlockSize;
-		uint32 Bytes    = Blocks * BlockSize;
-		UPTRINT OsBytes = Align(Bytes, PageSize);
-
-		checkSlow(Blocks >= 1);
-		checkSlow(Blocks * BlockSize <= Bytes && PoolSize >= Bytes);
-
-		// Allocate memory.
-		FFreeMem* Free = (FFreeMem*)Allocator.CachedOSPageAllocator.Allocate(OsBytes);
-		if( !Free )
-		{
-			OutOfMemory(OsBytes);
-		}
-
-		checkSlow(IsAligned(Free, PageSize));
-
-		// Create pool in the indirect table.
-		FPoolInfo* Pool = GetPoolInfo(Allocator, Free);
-		for (UPTRINT i = (UPTRINT)PageSize, Offset = 0; i < OsBytes; i += PageSize, ++Offset)
-		{
-			FPoolInfo* TrailingPool = GetPoolInfo(Allocator, (void*)(((UPTRINT)Free) + i));
-			check(TrailingPool);
-
-			//Set trailing pools to point back to first pool
-			TrailingPool->SetAllocationSizes(0, 0, Offset, Allocator.BinnedOSTableIndex);
-		}
-
-		BINNED2_PEAK_STATCOUNTER(Allocator.Stats.OsPeak,    BINNED2_ADD_STATCOUNTER(Allocator.Stats.OsCurrent,    OsBytes));
-		BINNED2_PEAK_STATCOUNTER(Allocator.Stats.WastePeak, BINNED2_ADD_STATCOUNTER(Allocator.Stats.WasteCurrent, (OsBytes - Bytes)));
-
-		// Init pool.
-		Pool->Link( Table.FirstPool );
-		Pool->SetAllocationSizes(Bytes, OsBytes, TableIndex, Allocator.BinnedOSTableIndex);
-		Pool->Taken		 = 0;
-		Pool->FirstMem   = Free;
-
-#if STATS
-		Table.NumActivePools++;
-		Table.MaxActivePools = FMath::Max(Table.MaxActivePools, Table.NumActivePools);
-#endif
-		// Create first free item.
-		Free->NumFreeBlocks = Blocks;
-		Free->Next          = nullptr;
-
-		return Pool;
-	}
-
-	static /*FORCEINLINE*/ FFreeMem* AllocateBlockFromPool(FMallocBinned2& Allocator, FPoolTable& Table, FPoolInfo* Pool, uint32 Alignment)
+	static void* AllocateBlockFromPool(FPoolTable& Table, FPoolInfo* Pool)
 	{
-		// Pick first available block and unlink it.
-		Pool->Taken++;
-		checkSlow(Pool->TableIndex < Allocator.BinnedOSTableIndex); // if this is false, FirstMem is actually a size not a pointer
+		checkSlow(!Pool->IsOSAllocation());
 		checkSlow(Pool->FirstMem);
 		checkSlow(Pool->FirstMem->NumFreeBlocks > 0);
-		checkSlow(Pool->FirstMem->NumFreeBlocks < PAGE_SIZE_LIMIT);
-		FFreeMem* Free = (FFreeMem*)((uint8*)Pool->FirstMem + --Pool->FirstMem->NumFreeBlocks * Table.BlockSize);
-		if( !Pool->FirstMem->NumFreeBlocks )
+
+		// Pick first available block and unlink it.
+		++Pool->Taken;
+		--Pool->FirstMem->NumFreeBlocks;
+
+		void* Free = (uint8*)Pool->FirstMem + Pool->FirstMem->NumFreeBlocks * Table.BlockSize;
+		if (Pool->FirstMem->NumFreeBlocks == 0)
 		{
-			Pool->FirstMem = Pool->FirstMem->Next;
+			Pool->FirstMem = Pool->FirstMem->NextFreePool;
 			if( !Pool->FirstMem )
 			{
 				// Move to exhausted list.
@@ -419,116 +322,49 @@ struct FMallocBinned2::Private
 				Pool->Link( Table.ExhaustedPool );
 			}
 		}
-		BINNED2_PEAK_STATCOUNTER(Allocator.Stats.UsedPeak, BINNED2_ADD_STATCOUNTER(Allocator.Stats.UsedCurrent, Table.BlockSize));
-		return Align(Free, Alignment);
-	}
 
-#if	STATS
-	static void UpdateSlackStat(FMallocBinned2& Allocator)
-	{
-		SIZE_T LocalWaste = Allocator.Stats.WasteCurrent;
-		double Waste      = 0.0;
-		for (FPoolTable& Table : Allocator.PoolTable)
-		{
-			Waste += ( ( double )Table.TotalWaste / ( double )Table.TotalRequests ) * ( double )Table.ActiveRequests;
-			Waste += Table.NumActivePools * ( BINNED_ALLOC_POOL_SIZE - ( ( BINNED_ALLOC_POOL_SIZE / Table.BlockSize ) * Table.BlockSize ) );
-		}
-		LocalWaste += ( uint32 )Waste;
-		Allocator.Stats.SlackCurrent = Allocator.Stats.OsCurrent - LocalWaste - Allocator.Stats.UsedCurrent;
+		return Free;
 	}
-#endif // STATS
 };
 
-void FMallocBinned2::GetAllocatorStats( FGenericMemoryStats& out_Stats )
-{
-	FMalloc::GetAllocatorStats( out_Stats );
-
-#if	STATS
-	Private::UpdateSlackStat(*this);
-
-	// Malloc binned stats.
-	out_Stats.Add( GET_STATDESCRIPTION( STAT_Binned2_OsCurrent ),     Stats.OsCurrent );
-	out_Stats.Add( GET_STATDESCRIPTION( STAT_Binned2_OsPeak ),        Stats.OsPeak );
-	out_Stats.Add( GET_STATDESCRIPTION( STAT_Binned2_WasteCurrent ),  Stats.WasteCurrent );
-	out_Stats.Add( GET_STATDESCRIPTION( STAT_Binned2_WastePeak ),     Stats.WastePeak );
-	out_Stats.Add( GET_STATDESCRIPTION( STAT_Binned2_UsedCurrent ),   Stats.UsedCurrent );
-	out_Stats.Add( GET_STATDESCRIPTION( STAT_Binned2_UsedPeak ),      Stats.UsedPeak );
-	out_Stats.Add( GET_STATDESCRIPTION( STAT_Binned2_CurrentAllocs ), Stats.CurrentAllocs );
-	out_Stats.Add( GET_STATDESCRIPTION( STAT_Binned2_TotalAllocs ),   Stats.TotalAllocs );
-	out_Stats.Add( GET_STATDESCRIPTION( STAT_Binned2_SlackCurrent ),  Stats.SlackCurrent );
-#endif // STATS
-}
-
-void FMallocBinned2::InitializeStatsMetadata()
-{
-	FMalloc::InitializeStatsMetadata();
-
-	// Initialize stats metadata here instead of UpdateStats.
-	// Mostly to avoid dead-lock when stats malloc profiler is enabled.
-	GET_STATFNAME(STAT_Binned2_OsCurrent);
-	GET_STATFNAME(STAT_Binned2_OsPeak);
-	GET_STATFNAME(STAT_Binned2_WasteCurrent);
-	GET_STATFNAME(STAT_Binned2_WastePeak);
-	GET_STATFNAME(STAT_Binned2_UsedCurrent);
-	GET_STATFNAME(STAT_Binned2_UsedPeak);
-	GET_STATFNAME(STAT_Binned2_CurrentAllocs);
-	GET_STATFNAME(STAT_Binned2_TotalAllocs);
-	GET_STATFNAME(STAT_Binned2_SlackCurrent);
-}
-
 FMallocBinned2::FMallocBinned2(uint32 InPageSize, uint64 AddressLimit)
-	: PageSize             (InPageSize)
-	, PoolBitShift         (FPlatformMath::CeilLogTwo(InPageSize))
-	, IndirectPoolBlockSize(InPageSize / sizeof(FPoolInfo))
-	, IndirectPoolBitShift (FPlatformMath::CeilLogTwo(IndirectPoolBlockSize))
-	, MaxHashBuckets       (AddressLimit >> (IndirectPoolBitShift + PoolBitShift))
-	, HashKeyShift         (PoolBitShift + IndirectPoolBitShift)
-	, PoolMask             ((1ull << (HashKeyShift - PoolBitShift)) - 1)
-	, BinnedSizeLimit      (Private::PAGE_SIZE_LIMIT / 2)
-	, BinnedOSTableIndex   (BinnedSizeLimit + EXTENDED_PAGE_POOL_ALLOCATION_COUNT)
-	, HashBucketFreeList   (nullptr)
+	: PageSize          (InPageSize)
+	, NumPoolsPerPage   (PageSize / sizeof(FPoolInfo))
+	, PtrToPoolMapping  (PageSize, NumPoolsPerPage)
+	, HashBucketFreeList(nullptr)
 {
 	check(FMath::IsPowerOfTwo(PageSize));
 	check(FMath::IsPowerOfTwo(AddressLimit));
-	check(FMath::IsPowerOfTwo(BinnedSizeLimit));
-	check(PageSize <= Private::PAGE_SIZE_LIMIT);
 	check(AddressLimit > PageSize); // Check to catch 32 bit overflow in AddressLimit
+	check(GMallocBinned2SmallBlockSizes[SMALL_POOL_COUNT - 1] == MAX_SMALL_POOLED_ALLOCATION_SIZE);
 
-	// Init tables.
-
-	/** The following options are not valid for page sizes less than 64k. They are here to reduce waste */
-	PagePoolTable[0].BlockSize = (PageSize == Private::PAGE_SIZE_LIMIT) ? BinnedSizeLimit + (BinnedSizeLimit / 2) : 0;
-	PagePoolTable[1].BlockSize = (PageSize == Private::PAGE_SIZE_LIMIT) ? PageSize        +  BinnedSizeLimit      : 0;
-
-	for (uint32 Index = 0; Index != POOL_COUNT; ++Index)
+	// Init pool tables.
+	for (uint32 Index = 0; Index != SMALL_POOL_COUNT; ++Index)
 	{
-		PoolTable[Index].BlockSize  = GMallocBinned2BlockSizes[Index];
-#if STATS
-		PoolTable[Index].MinRequest = GMallocBinned2BlockSizes[Index];
-#endif
+		checkf(Index == 0 || GMallocBinned2SmallBlockSizes[Index - 1] < GMallocBinned2SmallBlockSizes[Index], TEXT("Small block sizes must be strictly increasing"));
+		checkf(GMallocBinned2SmallBlockSizes[Index] <= PageSize,                                              TEXT("Small block size must be small enough to fit into a page"));
+		checkf(GMallocBinned2SmallBlockSizes[Index] % Private::DEFAULT_BINNED_ALLOCATOR_ALIGNMENT == 0,       TEXT("Small block size must be a multiple of DEFAULT_BINNED_ALLOCATOR_ALIGNMENT"));
+
+		SmallPoolTables[Index].BlockSize = GMallocBinned2SmallBlockSizes[Index];
 	}
 
-	for (uint32 Index = 0; Index != MAX_POOLED_ALLOCATION_SIZE; ++Index)
+	// Set up pool mappings
+	FPoolTable** TableEntry = MemSizeToPoolTable;
+	FPoolTable*  Pool       = SmallPoolTables;
+	for (uint32 Index = 0; Index != MAX_SMALL_POOLED_ALLOCATION_SIZE; ++Index)
 	{
-		FPoolTable* Pool = PoolTable;
 		while (Pool->BlockSize < Index)
 		{
 			++Pool;
-			checkSlow(Pool != PoolTable + POOL_COUNT);
+			check(Pool != SmallPoolTables + SMALL_POOL_COUNT);
 		}
-		MemSizeToPoolTable[Index] = Pool;
+		*TableEntry++ = Pool;
 	}
 
-	MemSizeToPoolTable[BinnedSizeLimit    ] = &PagePoolTable[0];
-	MemSizeToPoolTable[BinnedSizeLimit + 1] = &PagePoolTable[1];
+	uint64 MaxHashBuckets = PtrToPoolMapping.GetMaxHashBuckets(AddressLimit);
 
 	HashBuckets = (PoolHashBucket*)FPlatformMemory::BinnedAllocFromOS(Align(MaxHashBuckets * sizeof(PoolHashBucket), PageSize));
-	for (uint32 i = 0; i < MaxHashBuckets; ++i) 
-	{
-		new (HashBuckets + i) PoolHashBucket();
-	}
-
-	check(MAX_POOLED_ALLOCATION_SIZE - 1 == PoolTable[POOL_COUNT - 1].BlockSize);
+	DefaultConstructItems<PoolHashBucket>(HashBuckets, MaxHashBuckets);
 }
 
 FMallocBinned2::~FMallocBinned2()
@@ -542,66 +378,47 @@ bool FMallocBinned2::IsInternallyThreadSafe() const
 
 void* FMallocBinned2::Malloc(SIZE_T Size, uint32 Alignment)
 {
-	// Handle DEFAULT_ALIGNMENT for binned allocator.
-	if (Alignment == DEFAULT_ALIGNMENT)
-	{
-		Alignment = Private::DEFAULT_BINNED_ALLOCATOR_ALIGNMENT;
-	}
-
 	Alignment = FMath::Max<uint32>(Alignment, Private::DEFAULT_BINNED_ALLOCATOR_ALIGNMENT);
-	SIZE_T SpareBytesCount = FMath::Min<SIZE_T>(Private::DEFAULT_BINNED_ALLOCATOR_ALIGNMENT, Size);
-	Size = FMath::Max<SIZE_T>(PoolTable[0].BlockSize, Size + (Alignment - SpareBytesCount));
+	Size      = Align(Size, Alignment);
 
-	BINNED2_INCREMENT_STATCOUNTER(Stats.CurrentAllocs);
-	BINNED2_INCREMENT_STATCOUNTER(Stats.TotalAllocs);
+	checkf(FMath::IsPowerOfTwo(Alignment), TEXT("Alignment must be a power of two"));
+	checkf(Alignment <= PageSize, TEXT("Alignment cannot exceed the system page size"));
 
-	if (Size < BinnedSizeLimit)
+	// Only allocate from the small pools if the size is small enough and the alignment isn't crazy large.
+	// With large alignments, we'll waste a lot of memory allocating an entire page, but such alignments are highly unlikely in practice.
+	if (Size < MAX_SMALL_POOLED_ALLOCATION_SIZE && Alignment <= Private::DEFAULT_BINNED_ALLOCATOR_ALIGNMENT)
 	{
-		// Allocate from pool.
+		// Allocate from small object pool.
 		FPoolTable* Table = MemSizeToPoolTable[Size];
-
-		checkSlow(Size <= Table->BlockSize);
-
-		Private::TrackStats(Table, Size);
-
-		FPoolInfo* Pool = Table->FirstPool;
+		FPoolInfo*  Pool  = Table->FirstPool;
 		if (!Pool)
 		{
-			Pool = Private::AllocatePoolMemory(*this, *Table, Private::BINNED_ALLOC_POOL_SIZE/*PageSize*/, Size);
-		}
-
-		FFreeMem* Result = Private::AllocateBlockFromPool(*this, *Table, Pool, Alignment);
-#if MALLOC_VERIFY
-		FMallocVerify::Get().Malloc(Result);
-#endif
-		return Result;
-	}
-
-	if (Size <= PagePoolTable[0].BlockSize || (Size > PageSize && Size <= PagePoolTable[1].BlockSize))
-	{
-		// Bucket in a pool of 3*PageSize or 6*PageSize
-		uint32      BinType   = Size < PageSize ? 0 : 1;
-		uint32      PageCount = 3 * BinType + 3;
-		FPoolTable* Table     = &PagePoolTable[BinType];
+			// Allocate memory.
+			FFreeMem* Free = (FFreeMem*)CachedOSPageAllocator.Allocate(PageSize);
+			if (!Free)
+			{
+				Private::OutOfMemory(PageSize);
+			}
 
-		checkSlow(Size <= Table->BlockSize);
+			// Create pool in the pool page.
+			Pool = Private::GetOrCreatePoolInfo(*this, Free);
 
-		Private::TrackStats(Table, Size);
+			// Init pool.
+			Pool->Link(Table->FirstPool);
+			Pool->SetPoolAllocationSizes(PageSize, Size);
+			Pool->Taken    = 0;
+			Pool->FirstMem = Free;
 
-		FPoolInfo* Pool = Table->FirstPool;
-		if (!Pool)
-		{
-			Pool = Private::AllocatePoolMemory(*this, *Table, PageCount*PageSize, BinnedSizeLimit+BinType);
+			// Create first free item.
+			Free->NumFreeBlocks = PageSize / Table->BlockSize;
+			Free->NextFreePool  = nullptr;
 		}
 
-		FFreeMem* Result = Private::AllocateBlockFromPool(*this, *Table, Pool, Alignment);
-#if MALLOC_VERIFY
-		FMallocVerify::Get().Malloc(Result);
-#endif
-		return Result;
+		void* Result = Private::AllocateBlockFromPool(*Table, Pool);
+		return Align(Result, Alignment);
 	}
 
-	// Use OS for large allocations.
+	// Use OS for non-pooled allocations.
 	UPTRINT AlignedSize = Align(Size, PageSize);
 	FFreeMem* Result = (FFreeMem*)CachedOSPageAllocator.Allocate(AlignedSize);
 	if (!Result)
@@ -614,109 +431,75 @@ void* FMallocBinned2::Malloc(SIZE_T Size, uint32 Alignment)
 
 	void* AlignedResult = Align(Result, Alignment);
 
-	// Create indirect.
-	FPoolInfo* Pool = Private::GetPoolInfo(*this, Result);
+	// Create pool.
+	FPoolInfo* Pool = Private::GetOrCreatePoolInfo(*this, Result);
 	if (Result != AlignedResult)
 	{
 		// Mark the FPoolInfo for AlignedResult to jump back to the FPoolInfo for ptr.
-		for (UPTRINT i = (UPTRINT)PageSize, Offset = 0; i < AlignedSize; i += PageSize, ++Offset)
+		for (UPTRINT i = (UPTRINT)PageSize, Offset = 1; i < AlignedSize; i += PageSize, ++Offset)
 		{
-			FPoolInfo* TrailingPool = Private::GetPoolInfo(*this, (void*)((UPTRINT)Result + i));
-			check(TrailingPool);
+			FPoolInfo* TrailingPool = Private::GetOrCreatePoolInfo(*this, (uint8*)Result + i);
 
-			//Set trailing pools to point back to first pool
-			TrailingPool->SetAllocationSizes(0, 0, Offset, BinnedOSTableIndex);
+			// Set trailing pools to point back to first pool
+			TrailingPool->SetTrailingPoolOffset(Offset);
 		}
 	}
 
-	Pool->SetAllocationSizes(Size, AlignedSize, BinnedOSTableIndex, BinnedOSTableIndex);
-
-	BINNED2_PEAK_STATCOUNTER(Stats.OsPeak,    BINNED2_ADD_STATCOUNTER(Stats.OsCurrent,    AlignedSize));
-	BINNED2_PEAK_STATCOUNTER(Stats.UsedPeak,  BINNED2_ADD_STATCOUNTER(Stats.UsedCurrent,  Size));
-	BINNED2_PEAK_STATCOUNTER(Stats.WastePeak, BINNED2_ADD_STATCOUNTER(Stats.WasteCurrent, (int64)(AlignedSize - Size)));
+	Pool->SetOSAllocationSizes(Size, AlignedSize);
 
-#if MALLOC_VERIFY
-	FMallocVerify::Get().Malloc(AlignedResult);
-#endif
 	return AlignedResult;
 }
 
 void* FMallocBinned2::Realloc(void* Ptr, SIZE_T NewSize, uint32 Alignment)
-{	
-	// Handle DEFAULT_ALIGNMENT for binned allocator.
-	Alignment = FMath::Max<uint32>(Alignment, Private::DEFAULT_BINNED_ALLOCATOR_ALIGNMENT);
-
-	if (!Ptr)
-	{
-		void* Result = FMallocBinned2::Malloc(NewSize, Alignment);
-		return Result;
-	}
-
-	//ValidateHeap();
-
+{
 	if (NewSize == 0)
 	{
 		FMallocBinned2::Free(Ptr);
 		return nullptr;
 	}
 
-	const uint32 NewSizeUnmodified = NewSize;
-	SIZE_T SpareBytesCount = FMath::Min<SIZE_T>(Private::DEFAULT_BINNED_ALLOCATOR_ALIGNMENT, NewSize);
-	if (NewSize)
+	if (!Ptr)
 	{
-		NewSize = FMath::Max<SIZE_T>(PoolTable[0].BlockSize, NewSize + (Alignment - SpareBytesCount));
+		void* Result = FMallocBinned2::Malloc(NewSize, Alignment);
+		return Result;
 	}
 
+	Alignment = FMath::Max<uint32>(Alignment, Private::DEFAULT_BINNED_ALLOCATOR_ALIGNMENT);
+	NewSize   = Align(NewSize, Alignment);
+
+	checkf(FMath::IsPowerOfTwo(Alignment), TEXT("Alignment must be a power of two"));
+	checkf(Alignment <= PageSize, TEXT("Alignment cannot exceed the system page size"));
+
 	FPoolInfo* Pool = Private::FindPoolInfo(*this, Ptr);
-	if (Pool->TableIndex < BinnedOSTableIndex)
+	if (!Pool->IsOSAllocation())
 	{
-		// Allocated from pool, so grow or shrink if necessary.
-		check(Pool->TableIndex > 0); // it isn't possible to allocate a size of 0, Malloc will increase the size to DEFAULT_BINNED_ALLOCATOR_ALIGNMENT
-		if (NewSizeUnmodified > MemSizeToPoolTable[Pool->TableIndex]->BlockSize || NewSizeUnmodified <= MemSizeToPoolTable[Pool->TableIndex - 1]->BlockSize)
+		// Reallocate to a smaller/bigger pool if necessary
+		FPoolTable* Table = MemSizeToPoolTable[Pool->TableIndex];
+		if (NewSize > Table->BlockSize || (Pool->TableIndex != 0 && NewSize <= Table[-1].BlockSize) || !IsAligned(Ptr, Alignment))
 		{
-			void* Result = FMallocBinned2::Malloc(NewSizeUnmodified, Alignment);
-			FMemory::Memcpy(Result, Ptr, FMath::Min<SIZE_T>(NewSizeUnmodified, MemSizeToPoolTable[Pool->TableIndex]->BlockSize - (Alignment - SpareBytesCount)));
+			// Reallocate and copy the data across
+			void* Result = FMallocBinned2::Malloc(NewSize, Alignment);
+			FMemory::Memcpy(Result, Ptr, FMath::Min<SIZE_T>(NewSize, Table->BlockSize));
 			FMallocBinned2::Free(Ptr);
 			return Result;
 		}
 
-		if (!IsAligned(Ptr, Alignment))
-		{
-			void* Result = Align(Ptr, Alignment);
-			FMemory::Memmove(Result, Ptr, NewSize);
-#if MALLOC_VERIFY
-			FMallocVerify::Get().Realloc(Ptr, Result);
-#endif
-			return Result;
-		}
-
 		return Ptr;
 	}
 
 	// Allocated from OS.
-	UPTRINT PoolOsBytes = Pool->GetOsBytes(PageSize, BinnedOSTableIndex);
+	UPTRINT PoolOsBytes = Pool->GetOsBytes();
 	if( NewSize > PoolOsBytes || NewSize * 3 < PoolOsBytes * 2 )
 	{
 		// Grow or shrink.
-		void* Result = FMallocBinned2::Malloc(NewSizeUnmodified, Alignment);
-		FMemory::Memcpy(Result, Ptr, FMath::Min<SIZE_T>(NewSizeUnmodified, Pool->AllocSize));
+		void* Result = FMallocBinned2::Malloc(NewSize, Alignment);
+		FMemory::Memcpy(Result, Ptr, FMath::Min<SIZE_T>(NewSize, Pool->AllocSize));
 		FMallocBinned2::Free(Ptr);
 		return Result;
 	}
 
-	//need a lock to cover the SetAllocationSizes()
-	int32 UsedChange = (NewSize - Pool->AllocSize);
-
-	// Keep as-is, reallocation isn't worth the overhead.
-	BINNED2_ADD_STATCOUNTER (Stats.UsedCurrent,  UsedChange);
-	BINNED2_PEAK_STATCOUNTER(Stats.UsedPeak,     Stats.UsedCurrent);
-	BINNED2_ADD_STATCOUNTER (Stats.WasteCurrent, (Pool->AllocSize - NewSize));
+	Pool->SetOSAllocationSizes(NewSize, PoolOsBytes);
 
-	Pool->SetAllocationSizes(NewSizeUnmodified, PoolOsBytes, BinnedOSTableIndex, BinnedOSTableIndex);
-
-#if MALLOC_VERIFY
-	FMallocVerify::Get().Realloc(Ptr, Ptr);
-#endif
 	return Ptr;
 }
 
@@ -727,12 +510,6 @@ void FMallocBinned2::Free( void* Ptr )
 		return;
 	}
 
-#if MALLOC_VERIFY
-	FMallocVerify::Get().Free(Ptr);
-#endif
-
-	BINNED2_DECREMENT_STATCOUNTER(Stats.CurrentAllocs);
-
 	void* BasePtr;
 	FPoolInfo* Pool = Private::FindPoolInfo(*this, Ptr, BasePtr);
 #if PLATFORM_IOS || PLATFORM_MAC
@@ -743,13 +520,11 @@ void FMallocBinned2::Free( void* Ptr )
 	}
 #endif
 	checkSlow(Pool);
-	checkSlow(Pool->AllocSize != 0);
-	if (Pool->TableIndex < BinnedOSTableIndex)
+	checkSlow(!Pool->IsTrailingPool());
+	if (!Pool->IsOSAllocation())
 	{
 		FPoolTable* Table = MemSizeToPoolTable[Pool->TableIndex];
-#if STATS
-		Table->ActiveRequests--;
-#endif
+
 		// If this pool was exhausted, move to available list.
 		if (!Pool->FirstMem)
 		{
@@ -757,45 +532,32 @@ void FMallocBinned2::Free( void* Ptr )
 			Pool->Link(Table->FirstPool);
 		}
 
-		check((UPTRINT)BasePtr <= (UPTRINT)Ptr);
-
 		uint32 AlignOffset = ((UPTRINT)Ptr - (UPTRINT)BasePtr) % Table->BlockSize;
 
 		// Patch pointer to include previously applied alignment.
-		Ptr = (void*)((PTRINT)Ptr - (PTRINT)AlignOffset);
+		Ptr = (uint8*)Ptr - AlignOffset;
 
 		// Free a pooled allocation.
-		FFreeMem* Free		= (FFreeMem*)Ptr;
-		Free->NumFreeBlocks	= 1;
-		Free->Next			= Pool->FirstMem;
-		Pool->FirstMem		= Free;
-		BINNED2_ADD_STATCOUNTER(Stats.UsedCurrent, -(int64)Table->BlockSize);
+		FFreeMem* Free      = (FFreeMem*)Ptr;
+		Free->NumFreeBlocks = 1;
+		Free->NextFreePool  = Pool->FirstMem;
+		Pool->FirstMem      = Free;
 
 		// Free this pool.
 		checkSlow(Pool->Taken >= 1);
 		if( --Pool->Taken == 0 )
 		{
-#if STATS
-			Table->NumActivePools--;
-#endif
 			// Free the OS memory.
-			SIZE_T OsBytes = Pool->GetOsBytes(PageSize, BinnedOSTableIndex);
-			BINNED2_ADD_STATCOUNTER(Stats.OsCurrent,    -(int64)OsBytes);
-			BINNED2_ADD_STATCOUNTER(Stats.WasteCurrent, -(int64)(OsBytes - Pool->AllocSize));
 			Pool->Unlink();
-			Pool->SetAllocationSizes(0, 0, 0, BinnedOSTableIndex);
-			CachedOSPageAllocator.Free(BasePtr, OsBytes);
+			CachedOSPageAllocator.Free(BasePtr, PageSize);
 		}
 	}
 	else
 	{
 		// Free an OS allocation.
 		checkSlow(IsAligned(Ptr, PageSize));
-		SIZE_T OsBytes = Pool->GetOsBytes(PageSize, BinnedOSTableIndex);
+		SIZE_T OsBytes = Pool->GetOsBytes();
 
-		BINNED2_ADD_STATCOUNTER(Stats.UsedCurrent,  -(int64)Pool->AllocSize);
-		BINNED2_ADD_STATCOUNTER(Stats.OsCurrent,    -(int64)OsBytes);
-		BINNED2_ADD_STATCOUNTER(Stats.WasteCurrent, -(int64)(OsBytes - Pool->AllocSize));
 		CachedOSPageAllocator.Free(BasePtr, OsBytes);
 	}
 }
@@ -808,20 +570,20 @@ bool FMallocBinned2::GetAllocationSize(void* Original, SIZE_T& SizeOut)
 	}
 
 	FPoolInfo* Pool = Private::FindPoolInfo(*this, Original);
-	SizeOut = Pool->TableIndex < BinnedOSTableIndex ? MemSizeToPoolTable[Pool->TableIndex]->BlockSize : Pool->AllocSize;
+	SizeOut = Pool->IsOSAllocation() ? Pool->AllocSize : MemSizeToPoolTable[Pool->TableIndex]->BlockSize;
 	return true;
 }
 
 bool FMallocBinned2::ValidateHeap()
 {
-	for (FPoolTable& Table : PoolTable)
+	for (FPoolTable& Table : SmallPoolTables)
 	{
 		for( FPoolInfo** PoolPtr = &Table.FirstPool; *PoolPtr; PoolPtr = &(*PoolPtr)->Next )
 		{
 			FPoolInfo* Pool = *PoolPtr;
 			check(Pool->PrevLink == PoolPtr);
 			check(Pool->FirstMem);
-			for( FFreeMem* Free = Pool->FirstMem; Free; Free = Free->Next )
+			for( FFreeMem* Free = Pool->FirstMem; Free; Free = Free->NextFreePool )
 			{
 				check(Free->NumFreeBlocks > 0);
 			}
@@ -834,114 +596,10 @@ bool FMallocBinned2::ValidateHeap()
 		}
 	}
 
-	return( true );
-}
-
-void FMallocBinned2::UpdateStats()
-{
-	FMalloc::UpdateStats();
-#if STATS
-
-	Private::UpdateSlackStat(*this);
-
-	SET_MEMORY_STAT( STAT_Binned2_OsCurrent,     Stats.OsCurrent );
-	SET_MEMORY_STAT( STAT_Binned2_OsPeak,        Stats.OsPeak );
-	SET_MEMORY_STAT( STAT_Binned2_WasteCurrent,  Stats.WasteCurrent );
-	SET_MEMORY_STAT( STAT_Binned2_WastePeak,     Stats.WastePeak );
-	SET_MEMORY_STAT( STAT_Binned2_UsedCurrent,   Stats.UsedCurrent );
-	SET_MEMORY_STAT( STAT_Binned2_UsedPeak,      Stats.UsedPeak );
-	SET_DWORD_STAT ( STAT_Binned2_CurrentAllocs, Stats.CurrentAllocs );
-	SET_DWORD_STAT ( STAT_Binned2_TotalAllocs,   Stats.TotalAllocs );
-	SET_MEMORY_STAT( STAT_Binned2_SlackCurrent,  Stats.SlackCurrent );
-#endif
-}
-
-void FMallocBinned2::DumpAllocatorStats( class FOutputDevice& Ar )
-{
-	FBufferedOutputDevice BufferedOutput;
-	{
-		ValidateHeap();
-#if STATS
-		Private::UpdateSlackStat(*this);
-#if !NO_LOGGING
-		// This is all of the memory including stuff too big for the pools
-		BufferedOutput.CategorizedLogf( LogMemory.GetCategoryName(), ELogVerbosity::Log, TEXT( "Allocator Stats for %s:" ), GetDescriptiveName() );
-		// Waste is the total overhead of the memory system
-		BufferedOutput.CategorizedLogf( LogMemory.GetCategoryName(), ELogVerbosity::Log, TEXT( "Current Memory %.2f MB used, plus %.2f MB waste" ), Stats.UsedCurrent / (1024.0f * 1024.0f), (Stats.OsCurrent - Stats.UsedCurrent) / (1024.0f * 1024.0f) );
-		BufferedOutput.CategorizedLogf( LogMemory.GetCategoryName(), ELogVerbosity::Log, TEXT( "Peak Memory %.2f MB used, plus %.2f MB waste" ), Stats.UsedPeak / (1024.0f * 1024.0f), (Stats.OsPeak - Stats.UsedPeak) / (1024.0f * 1024.0f) );
-
-		BufferedOutput.CategorizedLogf( LogMemory.GetCategoryName(), ELogVerbosity::Log, TEXT( "Current OS Memory %.2f MB, peak %.2f MB" ), Stats.OsCurrent / (1024.0f * 1024.0f), Stats.OsPeak / (1024.0f * 1024.0f) );
-		BufferedOutput.CategorizedLogf( LogMemory.GetCategoryName(), ELogVerbosity::Log, TEXT( "Current Waste %.2f MB, peak %.2f MB" ), Stats.WasteCurrent / (1024.0f * 1024.0f), Stats.WastePeak / (1024.0f * 1024.0f) );
-		BufferedOutput.CategorizedLogf( LogMemory.GetCategoryName(), ELogVerbosity::Log, TEXT( "Current Used %.2f MB, peak %.2f MB" ), Stats.UsedCurrent / (1024.0f * 1024.0f), Stats.UsedPeak / (1024.0f * 1024.0f) );
-		BufferedOutput.CategorizedLogf( LogMemory.GetCategoryName(), ELogVerbosity::Log, TEXT( "Current Slack %.2f MB" ), Stats.SlackCurrent / (1024.0f * 1024.0f) );
-
-		BufferedOutput.CategorizedLogf( LogMemory.GetCategoryName(), ELogVerbosity::Log, TEXT( "Allocs      % 6i Current / % 6i Total" ), Stats.CurrentAllocs, Stats.TotalAllocs );
-
-		// This is the memory tracked inside individual allocation pools
-		BufferedOutput.CategorizedLogf( LogMemory.GetCategoryName(), ELogVerbosity::Log, TEXT( "" ) );
-		BufferedOutput.CategorizedLogf( LogMemory.GetCategoryName(), ELogVerbosity::Log, TEXT( "Block Size Num Pools Max Pools Cur Allocs Total Allocs Min Req Max Req Mem Used Mem Slack Mem Waste Efficiency" ) );
-		BufferedOutput.CategorizedLogf( LogMemory.GetCategoryName(), ELogVerbosity::Log, TEXT( "---------- --------- --------- ---------- ------------ ------- ------- -------- --------- --------- ----------" ) );
-
-		uint32 TotalMemory = 0;
-		uint32 TotalWaste = 0;
-		uint32 TotalActiveRequests = 0;
-		uint32 TotalTotalRequests = 0;
-		uint32 TotalPools = 0;
-		uint32 TotalSlack = 0;
-
-		FPoolTable* Table = nullptr;
-		for( int32 i = 0; i < BinnedSizeLimit + EXTENDED_PAGE_POOL_ALLOCATION_COUNT; i++ )
-		{
-			if( Table == MemSizeToPoolTable[i] || MemSizeToPoolTable[i]->BlockSize == 0 )
-				continue;
-
-			Table = MemSizeToPoolTable[i];
-
-			uint32 TableAllocSize = (Table->BlockSize > BinnedSizeLimit ? (((3 * (i - BinnedSizeLimit)) + 3)*Private::BINNED_ALLOC_POOL_SIZE) : Private::BINNED_ALLOC_POOL_SIZE);
-			// The amount of memory allocated from the OS
-			uint32 MemAllocated = (Table->NumActivePools * TableAllocSize) / 1024;
-			// Amount of memory actually in use by allocations
-			uint32 MemUsed = (Table->BlockSize * Table->ActiveRequests) / 1024;
-			// Wasted memory due to pool size alignment
-			uint32 PoolMemWaste = Table->NumActivePools * (TableAllocSize - ((TableAllocSize / Table->BlockSize) * Table->BlockSize)) / 1024;
-			// Wasted memory due to individual allocation alignment. This is an estimate.
-			uint32 MemWaste = (uint32)(((double)Table->TotalWaste / (double)Table->TotalRequests) * (double)Table->ActiveRequests) / 1024 + PoolMemWaste;
-			// Memory that is reserved in active pools and ready for future use
-			uint32 MemSlack = MemAllocated - MemUsed - PoolMemWaste;
-
-			BufferedOutput.CategorizedLogf( LogMemory.GetCategoryName(), ELogVerbosity::Log, TEXT( "% 10i % 9i % 9i % 10i % 12i % 7i % 7i % 7iK % 8iK % 8iK % 9.2f%%" ),
-										  Table->BlockSize,
-										  Table->NumActivePools,
-										  Table->MaxActivePools,
-										  Table->ActiveRequests,
-										  (uint32)Table->TotalRequests,
-										  Table->MinRequest,
-										  Table->MaxRequest,
-										  MemUsed,
-										  MemSlack,
-										  MemWaste,
-										  MemAllocated ? 100.0f * (MemAllocated - MemWaste) / MemAllocated : 100.0f );
-
-			TotalMemory += MemAllocated;
-			TotalWaste += MemWaste;
-			TotalSlack += MemSlack;
-			TotalActiveRequests += Table->ActiveRequests;
-			TotalTotalRequests += Table->TotalRequests;
-			TotalPools += Table->NumActivePools;
-		}
-
-		BufferedOutput.CategorizedLogf( LogMemory.GetCategoryName(), ELogVerbosity::Log, TEXT( "" ) );
-		BufferedOutput.CategorizedLogf( LogMemory.GetCategoryName(), ELogVerbosity::Log, TEXT( "%iK allocated in pools (with %iK slack and %iK waste). Efficiency %.2f%%" ), TotalMemory, TotalSlack, TotalWaste, TotalMemory ? 100.0f * (TotalMemory - TotalWaste) / TotalMemory : 100.0f );
-		BufferedOutput.CategorizedLogf( LogMemory.GetCategoryName(), ELogVerbosity::Log, TEXT( "Allocations %i Current / %i Total (in %i pools)" ), TotalActiveRequests, TotalTotalRequests, TotalPools );
-		BufferedOutput.CategorizedLogf( LogMemory.GetCategoryName(), ELogVerbosity::Log, TEXT( "" ) );
-#endif
-#endif
-	}
-
-	BufferedOutput.RedirectTo( Ar );
+	return true;
 }
 
 const TCHAR* FMallocBinned2::GetDescriptiveName()
 {
-	return TEXT("binned");
+	return TEXT("binned2");
 }
diff --git a/Engine/Source/Runtime/Core/Private/HAL/MallocTBB.cpp b/Engine/Source/Runtime/Core/Private/HAL/MallocTBB.cpp
--- a/Engine/Source/Runtime/Core/Private/HAL/MallocTBB.cpp
+++ b/Engine/Source/Runtime/Core/Private/HAL/MallocTBB.cpp
@@ -6,7 +6,6 @@
 
 #include "CorePrivatePCH.h"
 #include "MallocTBB.h"
-#include "MallocVerify.h"
 
 // Only use for supported platforms
 #if PLATFORM_SUPPORTS_TBB && TBB_ALLOCATOR_ALLOWED
@@ -54,9 +53,6 @@ void* FMallocTBB::Malloc( SIZE_T Size, uint32 Alignment )
 	{
 		FMemory::Memset(NewPtr, DEBUG_FILL_NEW, Size); 
 	}
-#endif
-#if MALLOC_VERIFY
-	FMallocVerify::Get().Malloc(NewPtr);
 #endif
 	MEM_TIME(MemTime += FPlatformTime::Seconds());
 	return NewPtr;
@@ -98,9 +94,6 @@ void* FMallocTBB::Realloc( void* Ptr, SIZE_T NewSize, uint32 Alignment )
 	{
 		OutOfMemory(NewSize, Alignment);
 	}
-#if MALLOC_VERIFY
-	FMallocVerify::Get().Realloc(Ptr, NewPtr);
-#endif
 	MEM_TIME(MemTime += FPlatformTime::Seconds())
 	return NewPtr;
 }
@@ -112,9 +105,6 @@ void FMallocTBB::Free( void* Ptr )
 		return;
 	}
 	MEM_TIME(MemTime -= FPlatformTime::Seconds())
-#if MALLOC_VERIFY
-	FMallocVerify::Get().Free(Ptr);
-#endif
 #if UE_BUILD_DEBUG || UE_BUILD_DEVELOPMENT
 	FMemory::Memset(Ptr, DEBUG_FILL_FREED, scalable_msize(Ptr)); 
 #endif
diff --git a/Engine/Source/Runtime/Core/Private/HAL/MallocVerify.h b/Engine/Source/Runtime/Core/Private/HAL/MallocVerify.h
--- a/Engine/Source/Runtime/Core/Private/HAL/MallocVerify.h
+++ b/Engine/Source/Runtime/Core/Private/HAL/MallocVerify.h
@@ -13,7 +13,7 @@
 #if MALLOC_VERIFY
 
 /**
-/* Maintains a list of all pointers to currently allocated memory.
+ * Maintains a list of all pointers to currently allocated memory.
  */
 class FMallocVerify
 {
@@ -31,12 +31,83 @@ public:
 
 	/** Removes allocated pointer from list */
 	void Free(void* Ptr);
+};
+
+/**
+ * A verifying proxy malloc that takes a malloc to be used and checks that the caller
+ * is passing valid pointers.
+ */
+class FMallocVerifyProxy : public FMalloc
+{
+private:
+	/** Malloc we're based on, aka using under the hood */
+	FMalloc* UsedMalloc;
+
+	/* Verifier object */
+	FMallocVerify Verify;
+
+public:
+	explicit FMallocVerifyProxy(FMalloc* InMalloc)
+		: UsedMalloc(InMalloc)
+	{
+	}
+
+	virtual void* Malloc(SIZE_T Size, uint32 Alignment) override
+	{
+		void* Result = UsedMalloc->Malloc(Size, Alignment);
+		Verify.Malloc(Result);
+		return Result;
+	}
+
+	virtual void* Realloc(void* Ptr, SIZE_T NewSize, uint32 Alignment) override
+	{
+		void* Result = UsedMalloc->Realloc(Ptr, NewSize, Alignment);
+		Verify.Realloc(Ptr, Result);
+		return Result;
+	}
+
+	virtual void Free(void* Ptr) override
+	{
+		if (Ptr)
+		{
+			Verify.Free(Ptr);
+			UsedMalloc->Free(Ptr);
+		}
+	}
+
+	virtual void InitializeStatsMetadata() override
+	{
+		UsedMalloc->InitializeStatsMetadata();
+	}
+
+	virtual void GetAllocatorStats(FGenericMemoryStats& OutStats) override
+	{
+		UsedMalloc->GetAllocatorStats(OutStats);
+	}
+
+	virtual void DumpAllocatorStats(FOutputDevice& Ar) override
+	{
+		UsedMalloc->DumpAllocatorStats(Ar);
+	}
 
-	/** Returns singleton object */
-	static FMallocVerify& Get()
+	virtual bool ValidateHeap() override
 	{
-		static FMallocVerify Instance;
-		return Instance;
+		return UsedMalloc->ValidateHeap();
+	}
+
+	virtual bool Exec(UWorld* InWorld, const TCHAR* Cmd, FOutputDevice& Ar) override
+	{
+		return UsedMalloc->Exec(InWorld, Cmd, Ar);
+	}
+
+	virtual bool GetAllocationSize(void* Original, SIZE_T& OutSize) override
+	{
+		return UsedMalloc->GetAllocationSize(Original, OutSize);
+	}
+
+	virtual const TCHAR* GetDescriptiveName() override
+	{ 
+		return UsedMalloc->GetDescriptiveName();
 	}
 };
 
diff --git a/Engine/Source/Runtime/Core/Private/HAL/UnrealMemory.cpp b/Engine/Source/Runtime/Core/Private/HAL/UnrealMemory.cpp
--- a/Engine/Source/Runtime/Core/Private/HAL/UnrealMemory.cpp
+++ b/Engine/Source/Runtime/Core/Private/HAL/UnrealMemory.cpp
@@ -13,6 +13,7 @@
 #include "MallocDebug.h"
 #include "MallocProfiler.h"
 #include "MallocThreadSafeProxy.h"
+#include "MallocVerify.h"
 #include "PlatformMallocCrash.h"
 
 
@@ -238,6 +239,11 @@ void GCreateMalloc()
 	{
 		GMalloc = new FMallocThreadSafeProxy( GMalloc );
 	}
+
+#if MALLOC_VERIFY
+	// Add the verifier
+	GMalloc = new FMallocVerifyProxy( GMalloc );
+#endif
 }
 
 
diff --git a/Engine/Source/Runtime/Core/Private/Linux/LinuxPlatformStackWalk.cpp b/Engine/Source/Runtime/Core/Private/Linux/LinuxPlatformStackWalk.cpp
--- a/Engine/Source/Runtime/Core/Private/Linux/LinuxPlatformStackWalk.cpp
+++ b/Engine/Source/Runtime/Core/Private/Linux/LinuxPlatformStackWalk.cpp
@@ -753,7 +753,7 @@ void FLinuxPlatformStackWalk::CaptureStackBackTrace( uint64* BackTrace, uint32 M
 
 		if (LinuxContext->BacktraceSymbols == NULL)
 		{
-			// @TODO yrx 2014-09-29 Replace with backtrace_symbols_fd due to malloc()
+			// #CrashReport: 2014-09-29 Replace with backtrace_symbols_fd due to malloc()
 			LinuxContext->BacktraceSymbols = backtrace_symbols(reinterpret_cast< void** >( BackTrace ), MaxDepth);
 		}
 	}
diff --git a/Engine/Source/Runtime/Core/Private/Misc/Parse.cpp b/Engine/Source/Runtime/Core/Private/Misc/Parse.cpp
--- a/Engine/Source/Runtime/Core/Private/Misc/Parse.cpp
+++ b/Engine/Source/Runtime/Core/Private/Misc/Parse.cpp
@@ -101,7 +101,8 @@ void ConsoleCommandLibrary_DumpLibrary(UWorld* InWorld, FExec& SubSystem, const
 
 void ConsoleCommandLibrary_DumpLibraryHTML(UWorld* InWorld, FExec& SubSystem, const FString& OutPath)
 {
-	ConsoleCommandLibrary LocalConsoleCommandLibrary(FString(TEXT("*")));
+	const FString& Pattern(TEXT("*"));
+	ConsoleCommandLibrary LocalConsoleCommandLibrary(Pattern);
 
 	FOutputDeviceNull Null;
 
diff --git a/Engine/Source/Runtime/Core/Private/Serialization/AsyncIOSystemBase.cpp b/Engine/Source/Runtime/Core/Private/Serialization/AsyncIOSystemBase.cpp
--- a/Engine/Source/Runtime/Core/Private/Serialization/AsyncIOSystemBase.cpp
+++ b/Engine/Source/Runtime/Core/Private/Serialization/AsyncIOSystemBase.cpp
@@ -483,7 +483,7 @@ int32 FAsyncIOSystemBase::CancelRequests( uint64* RequestIndices, int32 NumIndic
 		for( int32 TheRequestIndex=0; TheRequestIndex<NumIndices; TheRequestIndex++ )
 		{
 			// Look for matching request index in queue.
-			const FAsyncIORequest IORequest = OutstandingRequests[OutstandingIndex];
+			const FAsyncIORequest& IORequest = OutstandingRequests[OutstandingIndex];
 			if( IORequest.RequestIndex == RequestIndices[TheRequestIndex] )
 			{
 				INC_DWORD_STAT( STAT_AsyncIO_CanceledReadCount );
diff --git a/Engine/Source/Runtime/Core/Private/Stats/StatsCommand.cpp b/Engine/Source/Runtime/Core/Private/Stats/StatsCommand.cpp
--- a/Engine/Source/Runtime/Core/Private/Stats/StatsCommand.cpp
+++ b/Engine/Source/Runtime/Core/Private/Stats/StatsCommand.cpp
@@ -138,7 +138,7 @@ struct FStatParams
 	 *	Maximum number of frames to be included in the history. 
 	 *	-maxhistoryframes=[20:20-120]
 	 */
-	// #YRX_STATS 2014-08-21 Replace with TParsedValueWithDefaultAndRange
+	// #Stats: 2014-08-21 Replace with TParsedValueWithDefaultAndRange
 	TParsedValueWithDefault<int32> MaxHistoryFrames;
 
 	/**
diff --git a/Engine/Source/Runtime/Core/Private/Stats/StatsData.cpp b/Engine/Source/Runtime/Core/Private/Stats/StatsData.cpp
--- a/Engine/Source/Runtime/Core/Private/Stats/StatsData.cpp
+++ b/Engine/Source/Runtime/Core/Private/Stats/StatsData.cpp
@@ -540,7 +540,6 @@ void FComplexRawStatStackNode::CullByCycles( int64 MinCycles )
 			// Don't accumulate if we have just one child.
 			if (NumChildren > 1)
 			{	
-				// #YRX_STATS: 2015-06-09 Accumulate over complex stats
 				delete Child;
 				It.RemoveCurrent();
 			}
diff --git a/Engine/Source/Runtime/Core/Private/Stats/StatsFile.cpp b/Engine/Source/Runtime/Core/Private/Stats/StatsFile.cpp
--- a/Engine/Source/Runtime/Core/Private/Stats/StatsFile.cpp
+++ b/Engine/Source/Runtime/Core/Private/Stats/StatsFile.cpp
@@ -394,7 +394,7 @@ void FStatsWriteFile::SetDataDelegate( bool bSet )
 
 void FStatsWriteFile::WriteFrame( int64 TargetFrame )
 {
-	// #YRX_STATS: 2015-06-17 Add stat startfile -num=number of frames to capture
+	// #Stats: 2015-06-17 Add stat startfile -num=number of frames to capture
 
 	SCOPE_CYCLE_COUNTER( STAT_StreamFile );
 
diff --git a/Engine/Source/Runtime/Core/Private/Windows/WindowsPlatformCrashContext.cpp b/Engine/Source/Runtime/Core/Private/Windows/WindowsPlatformCrashContext.cpp
--- a/Engine/Source/Runtime/Core/Private/Windows/WindowsPlatformCrashContext.cpp
+++ b/Engine/Source/Runtime/Core/Private/Windows/WindowsPlatformCrashContext.cpp
@@ -31,7 +31,7 @@ static int32 ReportCrashCallCount = 0;
  * @return Success or failure
  */
 
-// @TODO yrx 2014-10-08 Move to FWindowsPlatformCrashContext
+// #CrashReport: 2014-10-08 Move to FWindowsPlatformCrashContext
 bool WriteMinidump( const TCHAR* Path, LPEXCEPTION_POINTERS ExceptionInfo, bool bIsEnsure )
 {
 	// Try to create file for minidump.
@@ -379,7 +379,7 @@ int32 ReportCrashUsingCrashReportClient(EXCEPTION_POINTERS* ExceptionInfo, const
 static FCriticalSection EnsureLock;
 static bool bReentranceGuard = false;
 
-// #YRX_Crash: 2015-05-28 This should be named EngineEnsureHandler
+// #CrashReport: 2015-05-28 This should be named EngineEnsureHandler
 /** 
  * Report an ensure to the crash reporting system
  */
@@ -433,7 +433,7 @@ void NewReportEnsure( const TCHAR* ErrorMessage )
 #include "AllowWindowsPlatformTypes.h"
 void CreateExceptionInfoString(EXCEPTION_RECORD* ExceptionRecord)
 {
-	// @TODO yrx 2014-08-18 Fix FString usage?
+	// #CrashReport: 2014-08-18 Fix FString usage?
 	FString ErrorString = TEXT("Unhandled Exception: ");
 
 #define HANDLE_CASE(x) case x: ErrorString += TEXT(#x); break;
@@ -471,7 +471,7 @@ void CreateExceptionInfoString(EXCEPTION_RECORD* ExceptionRecord)
 }
 #include "HideWindowsPlatformTypes.h"
 
-// #YRX_Crash: 2015-05-28 THis should be named EngineCrashHandler
+// #CrashReport: 2015-05-28 THis should be named EngineCrashHandler
 int32 ReportCrash( LPEXCEPTION_POINTERS ExceptionInfo )
 {
 	// Only create a minidump the first time this function is called.
diff --git a/Engine/Source/Runtime/Core/Private/Windows/WindowsPlatformOutputDevices.cpp b/Engine/Source/Runtime/Core/Private/Windows/WindowsPlatformOutputDevices.cpp
--- a/Engine/Source/Runtime/Core/Private/Windows/WindowsPlatformOutputDevices.cpp
+++ b/Engine/Source/Runtime/Core/Private/Windows/WindowsPlatformOutputDevices.cpp
@@ -77,7 +77,14 @@ void FOutputDeviceWindowsError::Serialize( const TCHAR* Msg, ELogVerbosity::Type
 		ErrorBuffer[0] = 0;
 
 		// Windows error.
-		UE_LOG( LogWindows, Error, TEXT( "Windows GetLastError: %s (%i)" ), FPlatformMisc::GetSystemErrorMessage( ErrorBuffer, 1024, LastError ), LastError );
+		if (LastError == 0)
+		{
+			UE_LOG(LogWindows, Log, TEXT("Windows GetLastError: %s (%i)"), FPlatformMisc::GetSystemErrorMessage(ErrorBuffer, 1024, LastError), LastError);
+		}
+		else
+		{
+			UE_LOG(LogWindows, Error, TEXT("Windows GetLastError: %s (%i)"), FPlatformMisc::GetSystemErrorMessage(ErrorBuffer, 1024, LastError), LastError);
+		}
 	}
 	else
 	{
diff --git a/Engine/Source/Runtime/Core/Private/Windows/WindowsPlatformStackWalk.cpp b/Engine/Source/Runtime/Core/Private/Windows/WindowsPlatformStackWalk.cpp
--- a/Engine/Source/Runtime/Core/Private/Windows/WindowsPlatformStackWalk.cpp
+++ b/Engine/Source/Runtime/Core/Private/Windows/WindowsPlatformStackWalk.cpp
@@ -176,7 +176,7 @@ void FWindowsPlatformStackWalk::StackWalkAndDump( ANSICHAR* HumanReadableString,
 	InitStackWalking();
 	FGenericPlatformStackWalk::StackWalkAndDump(HumanReadableString, HumanReadableStringSize, IgnoreCount, Context);
 }
-// @TODO yrx 2014-09-05 Switch to TArray<uint64,TFixedAllocator<100>>>
+// #CrashReport: 2014-09-05 Switch to TArray<uint64,TFixedAllocator<100>>>
 /**
  * Capture a stack backtrace and optionally use the passed in exception pointers.
  *
diff --git a/Engine/Source/Runtime/Core/Public/Containers/Array.h b/Engine/Source/Runtime/Core/Public/Containers/Array.h
--- a/Engine/Source/Runtime/Core/Public/Containers/Array.h
+++ b/Engine/Source/Runtime/Core/Public/Containers/Array.h
@@ -1673,8 +1673,8 @@ public:
 		// only reallocate if we need to, I don't trust realloc to the same size to work
 		if (ArrayMax != Slack)
 		{
+			AllocatorInstance.ResizeAllocation(0, Slack, sizeof(ElementType));
 			ArrayMax = Slack;
-			AllocatorInstance.ResizeAllocation(0, ArrayMax, sizeof(ElementType));
 		}
 	}
 
@@ -2468,7 +2468,7 @@ public:
 	/** 
 	 * Adds a new element to the heap.
 	 *
-	 * @param InIntem Item to be added.
+	 * @param InItem Item to be added.
 	 * @param Predicate Predicate class instance.
 	 *
 	 * @return The index of the new element.
@@ -2492,7 +2492,7 @@ public:
 	 * Adds a new element to the heap. Assumes < operator is defined for the
 	 * template type.
 	 *
-	 * @param InIntem Item to be added.
+	 * @param InItem Item to be added.
 	 *
 	 * @return The index of the new element.
 	 */
@@ -3428,7 +3428,7 @@ public:
 	 *
 	 * @returns Number of bytes allocated by this container.
 	 */
-	uint32 GetAllocatedSize() const
+	SIZE_T GetAllocatedSize() const
 	{
 		return Array.Max() * sizeof(T*) + Array.Num() * sizeof(T);
 	}
diff --git a/Engine/Source/Runtime/Core/Public/Containers/BitArray.h b/Engine/Source/Runtime/Core/Public/Containers/BitArray.h
--- a/Engine/Source/Runtime/Core/Public/Containers/BitArray.h
+++ b/Engine/Source/Runtime/Core/Public/Containers/BitArray.h
@@ -797,7 +797,8 @@ private:
 	void FindFirstSetBit()
 	{
 		const uint32* ArrayData      = Array.GetData();
-		const int32   LastDWORDIndex = (Array.Num() - 1) / NumBitsPerDWORD;
+		const int32   ArrayNum       = Array.Num();
+		const int32   LastDWORDIndex = (ArrayNum - 1) / NumBitsPerDWORD;
 
 		// Advance to the next non-zero uint32.
 		uint32 RemainingBitMask = ArrayData[this->DWORDIndex] & UnvisitedBitMask;
@@ -808,7 +809,7 @@ private:
 			if (this->DWORDIndex > LastDWORDIndex)
 			{
 				// We've advanced past the end of the array.
-				CurrentBitIndex = Array.Num();
+				CurrentBitIndex = ArrayNum;
 				return;
 			}
 
@@ -825,6 +826,13 @@ private:
 
 		// If the Nth bit was the lowest set bit of BitMask, then this gives us N
 		CurrentBitIndex = BaseBitIndex + NumBitsPerDWORD - 1 - FMath::CountLeadingZeros(this->Mask);
+
+		// If we've accidentally iterated off the end of an array but still within the same DWORD
+		// then set the index to the last index of the array
+		if (CurrentBitIndex > ArrayNum)
+		{
+			CurrentBitIndex = ArrayNum;
+		}
 	}
 };
 
diff --git a/Engine/Source/Runtime/Core/Public/Containers/ChunkedArray.h b/Engine/Source/Runtime/Core/Public/Containers/ChunkedArray.h
--- a/Engine/Source/Runtime/Core/Public/Containers/ChunkedArray.h
+++ b/Engine/Source/Runtime/Core/Public/Containers/ChunkedArray.h
@@ -110,7 +110,7 @@ public:
 		return NumElements; 
 	}
 
-	uint32 GetAllocatedSize( void ) const
+	SIZE_T GetAllocatedSize( void ) const
 	{
 		return Chunks.GetAllocatedSize();
 	}
diff --git a/Engine/Source/Runtime/Core/Public/CoreGlobals.h b/Engine/Source/Runtime/Core/Public/CoreGlobals.h
--- a/Engine/Source/Runtime/Core/Public/CoreGlobals.h
+++ b/Engine/Source/Runtime/Core/Public/CoreGlobals.h
@@ -22,7 +22,7 @@ extern CORE_API FOutputDeviceConsole* GLogConsole;
 
 extern CORE_API TCHAR GErrorHist[16384];
 
-// @TODO yrx 2014-08-19 Combine into one.
+// #crashReport: 2014-08-19 Combine into one, refactor.
 extern CORE_API TCHAR GErrorExceptionDescription[4096];
 extern CORE_API TCHAR GErrorMessage[4096];
 
diff --git a/Engine/Source/Runtime/Core/Public/GenericPlatform/GenericPlatformMisc.h b/Engine/Source/Runtime/Core/Public/GenericPlatform/GenericPlatformMisc.h
--- a/Engine/Source/Runtime/Core/Public/GenericPlatform/GenericPlatformMisc.h
+++ b/Engine/Source/Runtime/Core/Public/GenericPlatform/GenericPlatformMisc.h
@@ -277,7 +277,7 @@ struct CORE_API FGenericPlatformMisc
 	 */
 	static FString GetUniqueDeviceId();
 
-	// @TODO yrx 2015-02-24 Remove
+	// #CrashReport: 2015-02-24 Remove
 	/** Submits a crash report to a central server (release builds only) */
 	static void SubmitErrorReport( const TCHAR* InErrorHist, EErrorReportMode::Type InMode );
 
diff --git a/Engine/Source/Runtime/Core/Public/HAL/ExceptionHandling.h b/Engine/Source/Runtime/Core/Public/HAL/ExceptionHandling.h
--- a/Engine/Source/Runtime/Core/Public/HAL/ExceptionHandling.h
+++ b/Engine/Source/Runtime/Core/Public/HAL/ExceptionHandling.h
@@ -13,17 +13,17 @@ extern CORE_API bool GUseCrashReportClient;
 
 extern CORE_API TCHAR MiniDumpFilenameW[1024];
 
-// @TODO yrx 2014-09-11 Move to PlatformExceptionHandling
+// #CrashReport: 2014-09-11 Move to PlatformExceptionHandling
 #if PLATFORM_WINDOWS
-// @TODO yrx 2014-10-09 These methods are specific to windows, remove from here.
+// #CrashReport: 2014-10-09 These methods are specific to windows, remove from here.
 extern CORE_API int32 ReportCrash( LPEXCEPTION_POINTERS ExceptionInfo );
 extern CORE_API void NewReportEnsure( const TCHAR* ErrorMessage );
 #elif PLATFORM_XBOXONE
-// @TODO yrx 2014-10-09 Should be move to another file
+// #CrashReport: 2014-10-09 Should be move to another file
 extern CORE_API int32 ReportCrash( int ExceptionCode, LPEXCEPTION_POINTERS ExceptionInfo );
 extern CORE_API void NewReportEnsure( const TCHAR* ErrorMessage );
 #elif PLATFORM_MAC
-// @TODO yrx 2014-10-09 Should be move to another file
+// #CrashReport: 2014-10-09 Should be move to another file
 extern CORE_API int32 ReportCrash( ucontext_t *Context, int32 Signal, struct __siginfo* Info );
 extern CORE_API void NewReportEnsure( const TCHAR* ErrorMessage );
 #elif PLATFORM_LINUX
diff --git a/Engine/Source/Runtime/Core/Public/HAL/MallocBinned2.h b/Engine/Source/Runtime/Core/Public/HAL/MallocBinned2.h
--- a/Engine/Source/Runtime/Core/Public/HAL/MallocBinned2.h
+++ b/Engine/Source/Runtime/Core/Public/HAL/MallocBinned2.h
@@ -11,14 +11,6 @@
 	#define BINNED2_MAX_CACHED_OS_FREES_BYTE_LIMIT (16*1024*1024)
 #endif
 
-#if STATS
-#	if PLATFORM_64BITS
-#		define BINNED2_STAT volatile int64
-#	else
-#		define BINNED2_STAT volatile int32
-#	endif
-#endif
-
 //
 // Optimized virtual memory allocator.
 //
@@ -27,13 +19,12 @@ class FMallocBinned2 : public FMalloc
 	struct Private;
 
 private:
-
 	// Counts.
-	enum { POOL_COUNT = 42 };
+	enum { SMALL_POOL_COUNT = 41 };
 
 	/** Maximum allocation for the pooled allocator */
-	enum { EXTENDED_PAGE_POOL_ALLOCATION_COUNT = 2 };
-	enum { MAX_POOLED_ALLOCATION_SIZE   = 32768+1 };
+	enum { MAX_SMALL_POOLED_ALLOCATION_SIZE = 32768 };
+	enum { MAX_LARGE_POOLED_PAGE_FRACTIONS  = 2 };
 
 	// Forward declares.
 	struct FFreeMem;
@@ -44,114 +35,67 @@ private:
 	/** Pool table. */
 	struct FPoolTable
 	{
-		FPoolInfo*			FirstPool;
-		FPoolInfo*			ExhaustedPool;
-		uint32				BlockSize;
-#if STATS
-		/** Number of currently active pools */
-		uint32				NumActivePools;
-
-		/** Largest number of pools simultaneously active */
-		uint32				MaxActivePools;
-
-		/** Number of requests currently active */
-		uint32				ActiveRequests;
+		FPoolInfo* FirstPool;
+		FPoolInfo* ExhaustedPool;
+		uint32     BlockSize;
 
-		/** High watermark of requests simultaneously active */
-		uint32				MaxActiveRequests;
-
-		/** Minimum request size (in bytes) */
-		uint32				MinRequest;
-
-		/** Maximum request size (in bytes) */
-		uint32				MaxRequest;
-
-		/** Total number of requests ever */
-		uint64				TotalRequests;
-
-		/** Total waste from all allocs in this table */
-		uint64				TotalWaste;
-#endif
 		FPoolTable()
 			: FirstPool(nullptr)
 			, ExhaustedPool(nullptr)
 			, BlockSize(0)
-#if STATS
-			, NumActivePools(0)
-			, MaxActivePools(0)
-			, ActiveRequests(0)
-			, MaxActiveRequests(0)
-			, MinRequest(0)
-			, MaxRequest(0)
-			, TotalRequests(0)
-			, TotalWaste(0)
-#endif
 		{
-
 		}
 	};
 
 	const uint32 PageSize;
+	const uint64 NumPoolsPerPage;
+
+	struct FPtrToPoolMapping
+	{
+		explicit FPtrToPoolMapping(uint32 InPageSize, uint64 InNumPoolsPerPage)
+		{
+			uint64 PoolPageToPoolBitShift = FPlatformMath::CeilLogTwo(InNumPoolsPerPage);
+
+			PtrToPoolPageBitShift = FPlatformMath::CeilLogTwo(InPageSize);
+			HashKeyShift          = PtrToPoolPageBitShift + PoolPageToPoolBitShift;
+			PoolMask              = (1ull << PoolPageToPoolBitShift) - 1;
+		}
+
+		FORCEINLINE void GetHashBucketAndPoolIndices(const void* InPtr, UPTRINT& OutBucketIndex, uint32& OutPoolIndex) const
+		{
+			OutBucketIndex = (UPTRINT)InPtr >> HashKeyShift;
+			OutPoolIndex   = ((UPTRINT)InPtr >> PtrToPoolPageBitShift) & PoolMask;
+		}
 
-	/** Shift to get the reference from the indirect tables */
-	const uint64 PoolBitShift;
-	const uint64 IndirectPoolBlockSize;
-	const uint64 IndirectPoolBitShift;
+		FORCEINLINE uint64 GetMaxHashBuckets(uint64 AddressLimit) const
+		{
+			return AddressLimit >> HashKeyShift;
+		}
+
+	private:
+		/** Shift to apply to a pointer to get the reference from the indirect tables */
+		uint64 PtrToPoolPageBitShift;
 
-	// PageSize dependent constants
-	const uint64 MaxHashBuckets;
+		/** Shift required to get required hash table key. */
+		uint64 HashKeyShift;
+
+		/** Used to mask off the bits that have been used to lookup the indirect table */
+		uint64 PoolMask;
+	};
 
-	/** Shift required to get required hash table key. */
-	const uint64 HashKeyShift;
+	const FPtrToPoolMapping PtrToPoolMapping;
 
-	/** Used to mask off the bits that have been used to lookup the indirect table */
-	const uint64 PoolMask;
-	const uint64 BinnedSizeLimit;
-	const uint64 BinnedOSTableIndex;
+	// Pool tables for different pool sizes
+	FPoolTable SmallPoolTables[SMALL_POOL_COUNT];
 
-	// Variables.
-	FPoolTable  PoolTable[POOL_COUNT];
-	FPoolTable  PagePoolTable[EXTENDED_PAGE_POOL_ALLOCATION_COUNT];
-	FPoolTable* MemSizeToPoolTable[MAX_POOLED_ALLOCATION_SIZE + EXTENDED_PAGE_POOL_ALLOCATION_COUNT];
+	// Mapping of indices to both large and small tables
+	FPoolTable* MemSizeToPoolTable[MAX_SMALL_POOLED_ALLOCATION_SIZE];
 
 	PoolHashBucket* HashBuckets;
 	PoolHashBucket* HashBucketFreeList;
 
 	TCachedOSPageAllocator<BINNED2_MAX_CACHED_OS_FREES, BINNED2_MAX_CACHED_OS_FREES_BYTE_LIMIT> CachedOSPageAllocator;
 
-#if STATS
-	struct FStats
-	{
-		BINNED2_STAT		OsCurrent;
-		BINNED2_STAT		OsPeak;
-		BINNED2_STAT		WasteCurrent;
-		BINNED2_STAT		WastePeak;
-		BINNED2_STAT		UsedCurrent;
-		BINNED2_STAT		UsedPeak;
-		BINNED2_STAT		CurrentAllocs;
-		BINNED2_STAT		TotalAllocs;
-		/** OsCurrent - WasteCurrent - UsedCurrent. */
-		BINNED2_STAT		SlackCurrent;
-		double				MemTime;
-
-		FStats()
-			: OsCurrent    (0)
-			, OsPeak       (0)
-			, WasteCurrent (0)
-			, WastePeak    (0)
-			, UsedCurrent  (0)
-			, UsedPeak     (0)
-			, CurrentAllocs(0)
-			, TotalAllocs  (0)
-			, SlackCurrent (0)
-			, MemTime      (0.0)
-		{
-		}
-	};
-
-	FStats Stats;
-#endif
-
 public:
 	// FMalloc interface.
 	// InPageSize - First parameter is page size, all allocs from BinnedAllocFromOS() MUST be aligned to this size
@@ -160,8 +104,6 @@ public:
 	// It is ok to go outside this range, lookups will just be a little slower
 	FMallocBinned2(uint32 InPageSize, uint64 AddressLimit);
 
-	virtual void InitializeStatsMetadata() override;
-
 	virtual ~FMallocBinned2();
 
 	/**
@@ -199,20 +141,5 @@ public:
 	 */
 	virtual bool ValidateHeap() override;
 
-	/** Called once per frame, gathers and sets all memory allocator statistics into the corresponding stats. MUST BE THREAD SAFE. */
-	virtual void UpdateStats() override;
-
-	/** Writes allocator stats from the last update into the specified destination. */
-	virtual void GetAllocatorStats( FGenericMemoryStats& out_Stats ) override;
-
-	/**
-	 * Dumps allocator stats to an output device. Subclasses should override to add additional info
-	 *
-	 * @param Ar	[in] Output device
-	 */
-	virtual void DumpAllocatorStats( class FOutputDevice& Ar ) override;
-
 	virtual const TCHAR* GetDescriptiveName() override;
 };
-
-#undef BINNED2_STAT
diff --git a/Engine/Source/Runtime/Core/Public/Stats/Stats2.h b/Engine/Source/Runtime/Core/Public/Stats/Stats2.h
--- a/Engine/Source/Runtime/Core/Public/Stats/Stats2.h
+++ b/Engine/Source/Runtime/Core/Public/Stats/Stats2.h
@@ -1557,7 +1557,6 @@ public:
 };
 
 /** Manages startup messages, usually to update the metadata. */
-// @TODO yrx 2014-11-18 Eventually add startup messages for memory profiling?
 class FStartupMessages
 {
 	friend class FStatsThread;
diff --git a/Engine/Source/Runtime/Core/Public/Stats/StatsData.h b/Engine/Source/Runtime/Core/Public/Stats/StatsData.h
--- a/Engine/Source/Runtime/Core/Public/Stats/StatsData.h
+++ b/Engine/Source/Runtime/Core/Public/Stats/StatsData.h
@@ -470,7 +470,7 @@ struct IItemFiler
 };
 
 
-// #YRX_STATS 2014-03-21 Move metadata functionality into a separate class
+// #Stats: 2014-03-21 Move metadata functionality into a separate class
 // 
 /**
  * Tracks stat state and history
diff --git a/Engine/Source/Runtime/Core/Public/Stats/StatsFile.h b/Engine/Source/Runtime/Core/Public/Stats/StatsFile.h
--- a/Engine/Source/Runtime/Core/Public/Stats/StatsFile.h
+++ b/Engine/Source/Runtime/Core/Public/Stats/StatsFile.h
@@ -1000,7 +1000,7 @@ protected:
 	virtual void ProcessSpecialMessageMarkerOperation( const FStatMessage& Message, const FStackState& StackState )
 	{}
 
-	// #YRX_Stats: 2015-07-13 Not implemented yet.
+	// #Stats: 2015-07-13 Not implemented yet.
 // 	/** Processes set operation. */
 // 	virtual void ProcessSetOperation( const FStatMessage& Message, const FStackState& StackState )
 // 	{}
diff --git a/Engine/Source/Runtime/Core/Public/Templates/UnrealTemplate.h b/Engine/Source/Runtime/Core/Public/Templates/UnrealTemplate.h
--- a/Engine/Source/Runtime/Core/Public/Templates/UnrealTemplate.h
+++ b/Engine/Source/Runtime/Core/Public/Templates/UnrealTemplate.h
@@ -27,6 +27,19 @@ inline CONSTEXPR T Align( const T Ptr, int32 Alignment )
 	return (T)(((PTRINT)Ptr + Alignment - 1) & ~(Alignment-1));
 }
 
+/**
+ * Aligns a value to the nearest lower multiple of 'Alignment', which must be a power of two.
+ *
+ * @param Ptr			Value to align
+ * @param Alignment		Alignment, must be a power of two
+ * @return				Aligned value
+ */
+template <typename T>
+inline CONSTEXPR T AlignDown( const T Ptr, int32 Alignment )
+{
+	return (T)(((PTRINT)Ptr) & ~(Alignment-1));
+}
+
 /**
  * Checks if a pointer is aligned to the specified alignment.
  *
@@ -367,7 +380,7 @@ FORCEINLINE T&& Forward(typename TRemoveReference<T>::Type&& Obj)
 }
 
 /**
- * Swap two values, using moves if possible
+ * Swap two values.  Assumes the types are trivially relocatable.
  */
 template <typename T>
 inline void Swap(T& A, T& B)
diff --git a/Engine/Source/Runtime/CoreUObject/Private/Serialization/AsyncLoading.cpp b/Engine/Source/Runtime/CoreUObject/Private/Serialization/AsyncLoading.cpp
--- a/Engine/Source/Runtime/CoreUObject/Private/Serialization/AsyncLoading.cpp
+++ b/Engine/Source/Runtime/CoreUObject/Private/Serialization/AsyncLoading.cpp
@@ -341,8 +341,10 @@ void FAsyncLoadingThread::CancelAsyncLoadingInternal()
 #if THREADSAFE_UOBJECTS
 		FScopeLock QueueLock(&QueueCritical);
 #endif
+		const EAsyncLoadingResult::Type Result = EAsyncLoadingResult::Canceled;
 		for (FAsyncPackageDesc* PackageDesc : QueuedPackages)
 		{
+			PackageDesc->PackageLoadedDelegate.ExecuteIfBound(PackageDesc->Name, nullptr, Result);
 			delete PackageDesc;
 		}
 		QueuedPackages.Reset();
@@ -2041,6 +2043,14 @@ EAsyncPackageState::Type FAsyncPackage::FinishObjects()
 	#endif		// WITH_ENGINE
 	}
 
+	if (Linker)
+	{
+		// Flush linker cache now to reduce peak memory usage (5.5-10x)
+		// We shouldn't need it anyway at this point and even if something attempts to read in PostLoad, 
+		// we're just going to re-cache then.
+		Linker->FlushCache();
+	}
+
 	{
 		const bool bInternalCallbacks = true;
 		CallCompletionCallbacks(bInternalCallbacks, LoadingResult);
@@ -2415,7 +2425,7 @@ void FArchiveAsync::FlushCache()
 	PrecacheEndPos[CURRENT]		= 0;
 	PrecacheBufferSize[CURRENT] = 0;
 	PrecacheBufferProtected[CURRENT] = false;
-
+	
 	// Invalidate any precached data and free memory for next buffer.
 	FreeAsyncBuffer(PrecacheBuffer[NEXT], PrecacheBufferSize[NEXT]);
 	PrecacheBuffer[NEXT]		= nullptr;
@@ -2489,15 +2499,15 @@ bool FArchiveAsync::SetCompressionMap( TArray<FCompressedChunk>* InCompressedChu
  */
 void FArchiveAsync::BufferSwitcheroo()
 {
-	check(PrecacheReadStatus[CURRENT].GetValue() == 0);
-	check(PrecacheReadStatus[NEXT].GetValue() == 0);
+	check( PrecacheReadStatus[CURRENT].GetValue() == 0 );
+	check( PrecacheReadStatus[NEXT].GetValue() == 0 );
 
 	// Switcheroo.
 	DEC_DWORD_STAT_BY(STAT_StreamingAllocSize, PrecacheEndPos[CURRENT] - PrecacheStartPos[CURRENT]);
 	FreeAsyncBuffer(PrecacheBuffer[CURRENT], PrecacheBufferSize[CURRENT]);
-	PrecacheBuffer[CURRENT] = PrecacheBuffer[NEXT];
-	PrecacheStartPos[CURRENT] = PrecacheStartPos[NEXT];
-	PrecacheEndPos[CURRENT] = PrecacheEndPos[NEXT];
+	PrecacheBuffer[CURRENT]		= PrecacheBuffer[NEXT];
+	PrecacheStartPos[CURRENT]	= PrecacheStartPos[NEXT];
+	PrecacheEndPos[CURRENT]		= PrecacheEndPos[NEXT];
 	PrecacheBufferSize[CURRENT] = PrecacheBufferSize[NEXT];
 	PrecacheBufferProtected[CURRENT] = PrecacheBufferProtected[NEXT];
 
@@ -2720,7 +2730,7 @@ bool FArchiveAsync::Precache( int64 RequestOffset, int64 RequestSize )
  * @param	Count	Number of bytes to read
  */
 void FArchiveAsync::Serialize(void* Data, int64 Count)
-{	
+{
 #if PLATFORM_DESKTOP
 	// Show a message box indicating, possible, corrupt data (desktop platforms only)
 	if (CurrentPos + Count > TotalSize())
@@ -2738,8 +2748,8 @@ void FArchiveAsync::Serialize(void* Data, int64 Count)
 		FPlatformMisc::MessageBoxExt(EAppMsgType::Ok, *ErrorMessage.ToString(), *ErrorCaption.ToString());
 	}
 #endif
-	// Ensure we aren't reading beyond the end of the file	
-	checkf( CurrentPos + Count <= TotalSize(), TEXT("Seeked past end of file %s (%lld / %lld)"), *FileName, CurrentPos + Count, TotalSize() );	
+	// Ensure we aren't reading beyond the end of the file
+	checkf( CurrentPos + Count <= TotalSize(), TEXT("Seeked past end of file %s (%lld / %lld)"), *FileName, CurrentPos + Count, TotalSize() );
 
 #if LOOKING_FOR_PERF_ISSUES
 	uint32 StartCycles = 0;
diff --git a/Engine/Source/Runtime/CoreUObject/Private/UObject/LazyObjectPtr.cpp b/Engine/Source/Runtime/CoreUObject/Private/UObject/LazyObjectPtr.cpp
--- a/Engine/Source/Runtime/CoreUObject/Private/UObject/LazyObjectPtr.cpp
+++ b/Engine/Source/Runtime/CoreUObject/Private/UObject/LazyObjectPtr.cpp
@@ -133,10 +133,13 @@ void FLazyObjectPtr::PossiblySerializeObjectGuid(UObject *Object, FArchive& Ar)
 					{
 						// IsGame returns true for GIsPlayInEditorWorld
 						UE_CLOG(!((FApp::IsGame() || IsRunningCommandlet()) && Package && Package->ContainsMap()), LogUObjectGlobals, Warning, TEXT("Guid is in use by %s and %s, which should never happen in the editor but could happen at runtime with duplicate level loading or PIE"), *OtherObject->GetFullName(), !!Object ? *Object->GetFullName() : TEXT("NULL"));
-						// This guid is in use, which should never happen in the editor but could happen at runtime with duplicate level loading or PIE. If so give it a new GUID to avoid crashing
-						Guid = FGuid::NewGuid();
+						// This guid is in use, which should never happen in the editor but could happen at runtime with duplicate level loading or PIE. If so give it an invalid GUID and don't add to the annotation map.
+						Guid = FGuid();
+					}
+					else
+					{
+						GuidAnnotation.AddAnnotation(Object, Guid);
 					}
-					GuidAnnotation.AddAnnotation(Object, Guid);
 					FUniqueObjectGuid::InvalidateTag();
 				}
 			}
diff --git a/Engine/Source/Runtime/CoreUObject/Private/UObject/Obj.cpp b/Engine/Source/Runtime/CoreUObject/Private/UObject/Obj.cpp
--- a/Engine/Source/Runtime/CoreUObject/Private/UObject/Obj.cpp
+++ b/Engine/Source/Runtime/CoreUObject/Private/UObject/Obj.cpp
@@ -2659,7 +2659,7 @@ TArray<const TCHAR*> ParsePropertyFlags(uint64 Flags)
 	return Results;
 }
 
-// @TODO yrx 2014-09-15 Move to ObjectCommads.cpp or ObjectExec.cpp
+// #UObject: 2014-09-15 Move to ObjectCommads.cpp or ObjectExec.cpp
 bool StaticExec( UWorld* InWorld, const TCHAR* Cmd, FOutputDevice& Ar )
 {
 	const TCHAR *Str = Cmd;
diff --git a/Engine/Source/Runtime/CoreUObject/Private/UObject/PackageFileSummary.cpp b/Engine/Source/Runtime/CoreUObject/Private/UObject/PackageFileSummary.cpp
--- a/Engine/Source/Runtime/CoreUObject/Private/UObject/PackageFileSummary.cpp
+++ b/Engine/Source/Runtime/CoreUObject/Private/UObject/PackageFileSummary.cpp
@@ -77,13 +77,23 @@ FArchive& operator<<( FArchive& Ar, FPackageFileSummary& Sum )
 		 *		-5 indicates the replacement of writing out the "UE3 version" so older versions of engine can gracefully fail to open newer packages
 		 *		-6 indicates optimizations to how custom versions are being serialized
 		 */
-		int32 LegacyFileVersion = -6;
+		const int32 CurrentLegacyFileVersion = -6;
+		int32 LegacyFileVersion = CurrentLegacyFileVersion;
 		Ar << LegacyFileVersion;
 
 		if (Ar.IsLoading())
 		{
 			if (LegacyFileVersion < 0) // means we have modern version numbers
 			{
+				if (LegacyFileVersion < CurrentLegacyFileVersion)
+				{
+					// we can't safely load more than this because the legacy version code differs in ways we can not predict.
+					// Make sure that the linker will fail to load with it.
+					Sum.FileVersionUE4 = 0;
+					Sum.FileVersionLicenseeUE4 = 0;
+					return Ar;
+				}
+
 				if (LegacyFileVersion != -4)
 				{
 					int32 LegacyUE3Version = 0;
diff --git a/Engine/Source/Runtime/CoreUObject/Private/UObject/PropertyByte.cpp b/Engine/Source/Runtime/CoreUObject/Private/UObject/PropertyByte.cpp
--- a/Engine/Source/Runtime/CoreUObject/Private/UObject/PropertyByte.cpp
+++ b/Engine/Source/Runtime/CoreUObject/Private/UObject/PropertyByte.cpp
@@ -64,8 +64,7 @@ void UByteProperty::SerializeItem( FArchive& Ar, void* Value, void const* Defaul
 }
 bool UByteProperty::NetSerializeItem( FArchive& Ar, UPackageMap* Map, void* Data, TArray<uint8> * MetaData ) const
 {
-	// -1 because the last item in the enum is the autogenerated _MAX item
-	Ar.SerializeBits( Data, Enum ? FMath::CeilLogTwo(Enum->NumEnums() - 1) : 8 );
+	Ar.SerializeBits( Data, Enum ? FMath::CeilLogTwo(Enum->GetMaxEnumValue()) : 8 );
 	return 1;
 }
 void UByteProperty::Serialize( FArchive& Ar )
diff --git a/Engine/Source/Runtime/CoreUObject/Private/UObject/UObjectGlobals.cpp b/Engine/Source/Runtime/CoreUObject/Private/UObject/UObjectGlobals.cpp
--- a/Engine/Source/Runtime/CoreUObject/Private/UObject/UObjectGlobals.cpp
+++ b/Engine/Source/Runtime/CoreUObject/Private/UObject/UObjectGlobals.cpp
@@ -814,7 +814,7 @@ UObject* StaticLoadObjectInternal(UClass* ObjectClass, UObject* InOuter, const T
 				Result = StaticFindObjectFast(ObjectClass, InOuter, *StrName);
 
 				// If the object was not found, check for a redirector and follow it if the class matches
-				if (!Result)
+				if (!Result && !(LoadFlags & LOAD_NoRedirects))
 				{
 					UObjectRedirector* Redirector = FindObjectFast<UObjectRedirector>(InOuter, *StrName);
 					if (Redirector && Redirector->DestinationObject && Redirector->DestinationObject->IsA(ObjectClass))
@@ -3191,7 +3191,7 @@ bool IsReferenced(UObject*& Obj, EObjectFlags KeepFlags, EInternalObjectFlags In
 	bool bIsReferenced = false;
 	if (FoundReferences)
 	{
-		bIsReferenced = FoundReferences->ExternalReferences.Num() > 0 || !Obj->IsUnreachable();
+		bool bReferencedByOuters = false;		
 		// Move some from external to internal before returning
 		for (int32 i = 0; i < FoundReferences->ExternalReferences.Num(); i++)
 		{
@@ -3203,11 +3203,13 @@ bool IsReferenced(UObject*& Obj, EObjectFlags KeepFlags, EInternalObjectFlags In
 			}
 			else if (OldRef->Referencer->IsIn(Obj))
 			{
+				bReferencedByOuters = true;
 				FReferencerInformation *NewRef = new(FoundReferences->InternalReferences) FReferencerInformation(OldRef->Referencer, OldRef->TotalReferences, OldRef->ReferencingProperties);
 				FoundReferences->ExternalReferences.RemoveAt(i);
 				i--;
 			}
 		}
+		bIsReferenced = FoundReferences->ExternalReferences.Num() > 0 || bReferencedByOuters || !Obj->IsUnreachable();
 	}
 	else
 	{
diff --git a/Engine/Source/Runtime/Engine/Private/StatsRender2.cpp b/Engine/Source/Runtime/Engine/Private/StatsRender2.cpp
--- a/Engine/Source/Runtime/Engine/Private/StatsRender2.cpp
+++ b/Engine/Source/Runtime/Engine/Private/StatsRender2.cpp
@@ -309,7 +309,7 @@ static int32 RenderCycle( const FComplexStatMessage& Item, class FCanvas* Canvas
 		const float InMs = FPlatformTime::ToMilliseconds(Item.GetValue_Duration(EComplexStatField::IncAve));
 		// Color will be determined by the average value of history
 		// If show inclusive and and show exclusive is on, then it will choose color based on inclusive average
-		// #YRX_STATS: 2015-06-09 This is slow, fix this
+		// #Stats: 2015-06-09 This is slow, fix this
 		FString CounterName = Item.GetShortName().ToString();
 		CounterName.RemoveFromStart(TEXT("STAT_"), ESearchCase::CaseSensitive);
 		GEngine->GetStatValueColoration(CounterName, InMs, Color);
@@ -333,7 +333,7 @@ static int32 RenderCycle( const FComplexStatMessage& Item, class FCanvas* Canvas
 		}
 	}
 
-	// #TODO: Move to the stats thread to avoid expensive computation on the game thread.
+	// #Stats: Move to the stats thread to avoid expensive computation on the game thread.
 	const FString StatDesc = Item.GetDescription();
 	const FString StatDisplay = StatDesc.Len() == 0 ? Item.GetShortName().GetPlainNameString() : StatDesc;
 
diff --git a/Engine/Source/Runtime/Engine/Private/Texture2D.cpp b/Engine/Source/Runtime/Engine/Private/Texture2D.cpp
--- a/Engine/Source/Runtime/Engine/Private/Texture2D.cpp
+++ b/Engine/Source/Runtime/Engine/Private/Texture2D.cpp
@@ -2115,6 +2115,14 @@ void FTexture2DResource::FinalizeMipCount()
 		// We're done.
 		Owner->PendingMipChangeRequestStatus.Decrement();
 
+		// @todo: make it smarter, close only when we know we won't be streaming anymore or at least for a while
+		if (FPlatformProperties::RequiresCookedData())
+		{
+			const TIndirectArray<FTexture2DMipMap>& OwnerMips = Owner->GetPlatformMips();
+			const FTexture2DMipMap& MipMap = OwnerMips[0];
+			FIOSystem::Get().HintDoneWithFile(MipMap.BulkData.GetFilename());
+		}
+
 		return;
 	}
 
diff --git a/Engine/Source/Runtime/PakFile/Private/IPlatformFilePak.cpp b/Engine/Source/Runtime/PakFile/Private/IPlatformFilePak.cpp
--- a/Engine/Source/Runtime/PakFile/Private/IPlatformFilePak.cpp
+++ b/Engine/Source/Runtime/PakFile/Private/IPlatformFilePak.cpp
@@ -301,15 +301,15 @@ void FPakFile::Initialize(FArchive* Reader)
 {
 	if (Reader->TotalSize() < Info.GetSerializedSize())
 	{
-		UE_LOG(LogPakFile, Fatal, TEXT("Corrupted pak file (too short)."));
+		UE_LOG(LogPakFile, Fatal, TEXT("Corrupted pak file (too short). Verify your installation."));
 	}
 	else
 	{
 		// Serialize trailer and check if everything is as expected.
 		Reader->Seek(Reader->TotalSize() - Info.GetSerializedSize());
 		Info.Serialize(*Reader);
-		check(Info.Magic == FPakInfo::PakFile_Magic);
-		check(Info.Version >= FPakInfo::PakFile_Version_Initial && Info.Version <= FPakInfo::PakFile_Version_Latest);
+		UE_CLOG(Info.Magic != FPakInfo::PakFile_Magic, LogPakFile, Fatal, TEXT("Trailing magic number (%ud) is different than the expected one. Verify your installation."), Info.Magic);
+		UE_CLOG(!(Info.Version >= FPakInfo::PakFile_Version_Initial && Info.Version <= FPakInfo::PakFile_Version_Latest), LogPakFile, Fatal, TEXT("Invalid pak file version (%d). Verify your installation."), Info.Version);
 
 		LoadIndex(Reader);
 		// LoadIndex should crash in case of an error, so just assume everything is ok if we got here.
