UE4 - Fixed task thread stack size not big enough for debug editor.

diff --git a/Engine/Source/Runtime/Core/Private/Async/TaskGraph.cpp b/Engine/Source/Runtime/Core/Private/Async/TaskGraph.cpp
--- a/Engine/Source/Runtime/Core/Private/Async/TaskGraph.cpp
+++ b/Engine/Source/Runtime/Core/Private/Async/TaskGraph.cpp
@@ -1456,7 +1456,7 @@ public:
 				Name = FString::Printf(TEXT("TaskGraphThreadNP %d"), ThreadIndex - (LastExternalThread + 1));
 				ThreadPri = TPri_BelowNormal; // we want normal tasks below normal threads like the game thread
 			}
-			uint32 StackSize = 64 * 1024;
+			uint32 StackSize = 256 * 1024;
 			WorkerThreads[ThreadIndex].RunnableThread = FRunnableThread::Create(&Thread(ThreadIndex), *Name, StackSize, ThreadPri, FPlatformAffinity::GetTaskGraphThreadMask()); // these are below normal threads so that they sleep when the named threads are active
 			WorkerThreads[ThreadIndex].bAttached = true;
 		}
