Copying //UE4/Orion-Staging to //UE4/Main (//Orion/Dev-General @ 2855324)

==========================
MAJOR FEATURES + CHANGES
==========================

Change 2855265 on 2016/02/03 by [REDACTED_PERSON]

	Sequencer: Release track editors when destroying sequencer

Change 2855247 on 2016/02/03 by [REDACTED_PERSON]

	PlacementMode: Null factory check in constructor to fix cooking.

ChangeChange 2855116 on 2016/02/03 by [REDACTED_PERSON]

	PS4: Added a log statement when the gap between SubmitDone calls exceeds 2 seconds and removed a duplicate call to set the LastSubmitDoneTime
	PS4: Fixed a bogus log statement when PS4_GNM_SLOW_FRAME_DEBUGGING=1

	Merging CL# 2854751 and 2852176 by way of 2855100

	#tests Tested on PS4 with PS4_GNM_SLOW_FRAME_DEBUGGING=1

	--------
	Integrated using branch //Orion/Main_to_//Orion/Dev-General of change#2855109 by [REDACTED_PERSON] on 2016/02/03 20:59:51.

Change 2855109 on 2016/02/03 by [REDACTED_PERSON]

	PS4: Added a log statement when the gap between SubmitDone calls exceeds 2 seconds and removed a duplicate call to set the LastSubmitDoneTime
	PS4: Fixed a bogus log statement when PS4_GNM_SLOW_FRAME_DEBUGGING=1

	Merging CL# 2854751 and 2852176 by way of 2855100

	#tests Tested on PS4 with PS4_GNM_SLOW_FRAME_DEBUGGING=1

Change 2855100 on 2016/02/03 by [REDACTED_PERSON]

	PS4: Added a log statement when the gap between SubmitDone calls exceeds 2 seconds and removed a duplicate call to set the LastSubmitDoneTime
	PS4: Fixed a bogus log statement when PS4_GNM_SLOW_FRAME_DEBUGGING=1

	Merging CL# 2854751 and 2852176 using //Orion/Release-Next_to_//Orion/Release-Live

	#tests Tested on PS4 with PS4_GNM_SLOW_FRAME_DEBUGGING=1

ChangeChangeChangeChangeChange 2854825 on 2016/02/03 by [REDACTED_PERSON]

	Harden MaterialParameterCollection from ending up with duplicate parameter names of GUIDs.

	#Tests: Ran debug editor, create materail param collection with >500 elements. Still only ~18ms. Used param in shader.

Change 2854788 on 2016/02/03 by [REDACTED_PERSON]

	#UE4 - JsonObjectConverter changes
	- added the ability for a UStruct to emit json as a string if type traits are setup with ExportTextItem / ImportTextItem
	- allows the UStruct to convert to json as something other than FJsonValueObject
	-- things like FColor, FDateTime but they are already handled differently
	- checked for possible change in existing behavior, no classes currently use type traits for this that aren't handled special already
	- FUniqueNetIdRepl can now convert to/from json as a string

	#tests various online tests connecting to servers, etc

Change 2854751 on 2016/02/03 by [REDACTED_PERSON]

	PS4: Fixed a bogus log statement when PS4_GNM_SLOW_FRAME_DEBUGGING=1

	#tests Tested on PS4 with PS4_GNM_SLOW_FRAME_DEBUGGING=1

ChangeChange 2854712 on 2016/02/03 by [REDACTED_PERSON]

	#UE4 - added some json compatibility features to FUniqueNetIdRepl struct
	- ImportTextItem
	- To/FromJson

	#tests various online features, additional unit tests added to class

Change 2854696 on 2016/02/03 by [REDACTED_PERSON]

	Making Memprofiler usable (by MichaelN).

	#tests Used Memprofiler on a number of captures.

ChangeChange 2854536 on 2016/02/03 by [REDACTED_PERSON]

	Add event groups as users to replay, so we can quickly find replays with certain events types in them

	#tests Replays and events

	Merging using OrionDevGeneral->ReleaseCandidate

Change 2854526 on 2016/02/03 by [REDACTED_PERSON]

	Add support for getting replay id

	#tests Replays

	Merging using OrionDevGeneral->ReleaseCandidate

Change 2854522 on 2016/02/03 by [REDACTED_PERSON]

	Support setting string values in perf counters through the perf  counters helper class.

	#tests Client/Server match

diff --git a/Engine/Config/BaseEngine.ini b/Engine/Config/BaseEngine.ini
--- a/Engine/Config/BaseEngine.ini
+++ b/Engine/Config/BaseEngine.ini
@@ -1887,9 +1887,6 @@ FLightPrimitiveInteractionInitialBlockSize=512
 [SystemSettings]
 ; NOTE THAT ANY ITEMS IN THIS SECTION WILL AFFECT ALL PLATFORMS!!!
 
-// disabled for now to prevent issues http://crashreport/Crashes/Show/2948912, http://crashreport/Crashes/Show/2948842
-r.DetectAndWarnOfBadDrivers = 0
-
 [SystemSettingsEditor]
 ; System settings overrides for the editor.  Ideally the editor should use the same settings as the game.
 
diff --git a/Engine/Config/BaseHardware.ini b/Engine/Config/BaseHardware.ini
new file mode 100644
--- /dev/null
+++ b/Engine/Config/BaseHardware.ini
@@ -0,0 +1,59 @@
+;
+; Base Driver Blacklist
+; =====================
+;
+; Hardware checks are controlled by the console variable "r.WarnOfBadDrivers". The the console variable options for more information.
+;
+; What to do when a certain driver or hardware shows problems:
+;  - If the driver exhibiting the problem is old, add it to the blacklist with "<=". In general it's better to ask users to upgrade.
+;  - If the driver exhibiting the problem is the latest problem, add it to the blacklist with ">=". We should not assume the next driver version fixes the problem.
+;
+; Expectations on QA:
+;  - QA always tests the latest driver version AND the suggested driver version. In most cases these should be the same thing.
+;  - When a new driver is released, QA should test it and if it does not introduce bugs the SuggestedDriverVersion should be updated to the latest driver.
+;  - If the latest driver is blacklisted, it is not removed until QA verifies it is fixed.
+;  - If driver specific bugs are found, enter them in JIRA so that they can be tracked.
+;
+; Expectations on Engineering:
+;  - Bugs on older drivers are added to the blacklist and otherwise ignored.
+;  - Bugs on the latest driver are added to the blacklist but treated as critical bugs for which we must try to find a workaround. We must also notify the IHV to help if we cannot find the problem.
+;
+; INI sections:
+;    [GPU_NVIDIA] [GPU_AMD] [GPU_Intel]
+;
+; Explanation of the INI entries:
+;  - "SuggestedDriverVersion": string used to communicate to the user which driver version is appropriate (not used by any code), empty if unknown
+;  - "Blacklist": array of tests that can be made to mark a driver or a range of them as problematic
+;    "DriverVersion": to compare against the "unified driver" (see log) version e.g. "123.45.2" "==123.45" "<123.45", comparison operators: == != > < >= <=
+;    "Reason": string used to document the reason for blacklisting the driver, could be shown to the user
+;    comment above each blacklist entry: more details about the problem, potential workarounds
+;
+;-----------------------------------------------------------------------------------------------------------------
+
+[GPU_NVIDIA]
+; a newer one is out but has issues (see below)
+SuggestedDriverVersion="359.06"
+
+; Seems this driver was not setting uninitialized global variables
+; (static, non const) to 0 but leaving them uninitialized.
+; Also blacklist future driver versions until QA has verified them to work.
+; NOTE: UE-25096 was fixed in //Orion/Dev-General in CL 2828817, need Paragon QA to verify.
++Blacklist=(DriverVersion=">=361.43", Reason="UE-25096 Viewport flashes black and white when moving in the scene on latest Nvidia drivers")
+
+; To prevent problems with older drivers.
+; We might have to adjust that further.
+;+Blacklist=(DriverVersion="<=347.09", Reason="Crashes with Paragon content using more recent GPU features")
+
+;-----------------------------------------------------------------------------------------------------------------
+
+[GPU_AMD]
+// latest as of 1/27/2016 is "15.30.1025.1001 12/17/2015 (Crimson Edition 15.12)"
+SuggestedDriverVersion="15.30.1025.1001 12/17/2015"
+
+// 14.201.1001.1004 (7/31/14) - 13.152.1.1000 (9/10/13) are the AMD drivers that seem to be causing crashes when we try to compile a compute shader.
++Blacklist=(DriverVersion="<=14.201.1001.1004", Reason="Crashes with Paragon content using more recent GPU features")
+
+;-----------------------------------------------------------------------------------------------------------------
+
+[GPU_Intel]
+SuggestedDriverVersion=""
diff --git a/Engine/Plugins/Runtime/PacketHandlers/CompressionComponents/Oodle/Source/OodleHandlerComponent/Private/OodleHandlerComponent.cpp b/Engine/Plugins/Runtime/PacketHandlers/CompressionComponents/Oodle/Source/OodleHandlerComponent/Private/OodleHandlerComponent.cpp
--- a/Engine/Plugins/Runtime/PacketHandlers/CompressionComponents/Oodle/Source/OodleHandlerComponent/Private/OodleHandlerComponent.cpp
+++ b/Engine/Plugins/Runtime/PacketHandlers/CompressionComponents/Oodle/Source/OodleHandlerComponent/Private/OodleHandlerComponent.cpp
@@ -572,6 +572,8 @@ bool OodleHandlerComponent::GetDictionaryPaths(FString& OutServerDictionary, FSt
 		FPaths::CollapseRelativeDirectories(ServerDictionaryPath);
 		FPaths::CollapseRelativeDirectories(ClientDictionaryPath);
 
+		FPaths::NormalizeDirectoryName(ServerDictionaryPath);
+		FPaths::NormalizeDirectoryName(ClientDictionaryPath);
 
 		// Don't allow directory traversal to escape the game directory
 		if (!ServerDictionaryPath.StartsWith(FPaths::GameDir()) || !ClientDictionaryPath.StartsWith(FPaths::GameDir()))
diff --git a/Engine/Shaders/Fxaa3_11.usf b/Engine/Shaders/Fxaa3_11.usf
--- a/Engine/Shaders/Fxaa3_11.usf
+++ b/Engine/Shaders/Fxaa3_11.usf
@@ -150,7 +150,7 @@ Getting luma correct is required for the algorithm to work correctly.
 ------------------------------------------------------------------------------
 Applying FXAA to a framebuffer with linear RGB color will look worse.
 This is very counter intuitive, but happends to be true in this case.
-The reason is because dithering artifacts will be more visiable 
+The reason is because dithering artifacts will be more visible 
 in a linear colorspace.
 
 
diff --git a/Engine/Source/Editor/FoliageEdit/Private/FoliageEdMode.cpp b/Engine/Source/Editor/FoliageEdit/Private/FoliageEdMode.cpp
--- a/Engine/Source/Editor/FoliageEdit/Private/FoliageEdMode.cpp
+++ b/Engine/Source/Editor/FoliageEdit/Private/FoliageEdMode.cpp
@@ -143,38 +143,21 @@ public:
 //
 // Painting filtering options
 //
-struct FFoliagePaintingGeometryFilter
+bool FFoliagePaintingGeometryFilter::operator() (const UPrimitiveComponent* Component) const
 {
-	bool bAllowLandscape;
-	bool bAllowStaticMesh;
-	bool bAllowBSP;
-	bool bAllowTranslucent;
-
-	FFoliagePaintingGeometryFilter(const FFoliageUISettings& InUISettings)
-		: bAllowLandscape(InUISettings.bFilterLandscape)
-		, bAllowStaticMesh(InUISettings.bFilterStaticMesh)
-		, bAllowBSP(InUISettings.bFilterBSP)
-		, bAllowTranslucent(InUISettings.bFilterTranslucent)
+	if (Component)
 	{
-	}
+		bool bShallNotPass =
+			(!bAllowLandscape	&& Component->IsA(ULandscapeHeightfieldCollisionComponent::StaticClass())) ||
+			(!bAllowStaticMesh	&& Component->IsA(UStaticMeshComponent::StaticClass())) ||
+			(!bAllowBSP			&& Component->IsA(UModelComponent::StaticClass())) ||
+			(!bAllowTranslucent	&& Component->GetMaterial(0) && IsTranslucentBlendMode(Component->GetMaterial(0)->GetBlendMode()));
 
-	bool operator() (const UPrimitiveComponent* Component) const
-	{
-		if (Component)
-		{
-			bool bShallNotPass = 
-				(!bAllowLandscape	&& Component->IsA(ULandscapeHeightfieldCollisionComponent::StaticClass())) ||
-				(!bAllowStaticMesh	&& Component->IsA(UStaticMeshComponent::StaticClass())) ||
-				(!bAllowBSP			&& Component->IsA(UModelComponent::StaticClass())) ||
-				(!bAllowTranslucent	&& Component->GetMaterial(0) && IsTranslucentBlendMode(Component->GetMaterial(0)->GetBlendMode()));
-			
-			return !bShallNotPass;
-		}
-		
-		return false;		
+		return !bShallNotPass;
 	}
-};
 
+	return false;
+}
 
 //
 // FEdModeFoliage
@@ -887,7 +870,7 @@ bool LandscapeLayerCheck(const FHitResult& Hit, const UFoliageType* Settings, FE
 	return true;
 }
 
-void FEdModeFoliage::CalculatePotentialInstances_ThreadSafe(const UWorld* InWorld, const UFoliageType* Settings, const TArray<FDesiredFoliageInstance>* DesiredInstances, TArray<FPotentialInstance> OutPotentialInstances[NUM_INSTANCE_BUCKETS], const FFoliageUISettings* UISettings, const int32 StartIdx, const int32 LastIdx)
+void FEdModeFoliage::CalculatePotentialInstances_ThreadSafe(const UWorld* InWorld, const UFoliageType* Settings, const TArray<FDesiredFoliageInstance>* DesiredInstances, TArray<FPotentialInstance> OutPotentialInstances[NUM_INSTANCE_BUCKETS], const FFoliageUISettings* UISettings, const int32 StartIdx, const int32 LastIdx, const FFoliagePaintingGeometryFilter* OverrideGeometryFilter)
 {
 	LandscapeLayerCacheData LocalCache;
 
@@ -910,6 +893,11 @@ void FEdModeFoliage::CalculatePotentialInstances_ThreadSafe(const UWorld* InWorl
 			// Enable geometry filters when painting foliage manually
 			TraceFilterFunc = FFoliagePaintingGeometryFilter(*UISettings);
 		}
+
+		if(OverrideGeometryFilter)
+		{
+			TraceFilterFunc = *OverrideGeometryFilter;
+		}
 		
 		if (AInstancedFoliageActor::FoliageTrace(InWorld, Hit, DesiredInst, NAME_AddFoliageInstances, true, TraceFilterFunc))
 		{
@@ -927,7 +915,7 @@ void FEdModeFoliage::CalculatePotentialInstances_ThreadSafe(const UWorld* InWorl
 	}
 }
 
-void FEdModeFoliage::CalculatePotentialInstances(const UWorld* InWorld, const UFoliageType* Settings, const TArray<FDesiredFoliageInstance>& DesiredInstances, TArray<FPotentialInstance> OutPotentialInstances[NUM_INSTANCE_BUCKETS], LandscapeLayerCacheData* LandscapeLayerCachesPtr, const FFoliageUISettings* UISettings)
+void FEdModeFoliage::CalculatePotentialInstances(const UWorld* InWorld, const UFoliageType* Settings, const TArray<FDesiredFoliageInstance>& DesiredInstances, TArray<FPotentialInstance> OutPotentialInstances[NUM_INSTANCE_BUCKETS], LandscapeLayerCacheData* LandscapeLayerCachesPtr, const FFoliageUISettings* UISettings, const FFoliagePaintingGeometryFilter* OverrideGeometryFilter)
 {
 	LandscapeLayerCacheData LocalCache;
 	LandscapeLayerCachesPtr = LandscapeLayerCachesPtr ? LandscapeLayerCachesPtr : &LocalCache;
@@ -952,6 +940,11 @@ void FEdModeFoliage::CalculatePotentialInstances(const UWorld* InWorld, const UF
 			// Enable geometry filters when painting foliage manually
 			TraceFilterFunc = FFoliagePaintingGeometryFilter(*UISettings);
 		}
+
+		if(OverrideGeometryFilter)
+		{
+			TraceFilterFunc = *OverrideGeometryFilter;
+		}
 				
 		FHitResult Hit;
 		static FName NAME_AddFoliageInstances = FName(TEXT("AddFoliageInstances"));
@@ -981,7 +974,7 @@ void FEdModeFoliage::CalculatePotentialInstances(const UWorld* InWorld, const UF
 	}
 }
 
-void FEdModeFoliage::AddInstances(UWorld* InWorld, const TArray<FDesiredFoliageInstance>& DesiredInstances)
+void FEdModeFoliage::AddInstances(UWorld* InWorld, const TArray<FDesiredFoliageInstance>& DesiredInstances, const FFoliagePaintingGeometryFilter& OverrideGeometryFilter)
 {
 	TMap<const UFoliageType*, TArray<FDesiredFoliageInstance>> SettingsInstancesMap;
 	for (const FDesiredFoliageInstance& DesiredInst : DesiredInstances)
@@ -995,7 +988,7 @@ void FEdModeFoliage::AddInstances(UWorld* InWorld, const TArray<FDesiredFoliageI
 		const UFoliageType* FoliageType = It.Key();
 
 		const TArray<FDesiredFoliageInstance>& Instances = It.Value();
-		AddInstancesImp(InWorld, FoliageType, Instances);
+		AddInstancesImp(InWorld, FoliageType, Instances, TArray<int32>(), 1.f, nullptr, nullptr, &OverrideGeometryFilter);
 	}
 }
 
@@ -1011,7 +1004,7 @@ static void SpawnFoliageInstance(UWorld* InWorld, const UFoliageType* Settings,
 	MeshInfo->AddInstance(IFA, FoliageSettings, Instance, BaseComponent);
 }
 
-void FEdModeFoliage::AddInstancesImp(UWorld* InWorld, const UFoliageType* Settings, const TArray<FDesiredFoliageInstance>& DesiredInstances, const TArray<int32>& ExistingInstanceBuckets, const float Pressure, LandscapeLayerCacheData* LandscapeLayerCachesPtr, const FFoliageUISettings* UISettings)
+void FEdModeFoliage::AddInstancesImp(UWorld* InWorld, const UFoliageType* Settings, const TArray<FDesiredFoliageInstance>& DesiredInstances, const TArray<int32>& ExistingInstanceBuckets, const float Pressure, LandscapeLayerCacheData* LandscapeLayerCachesPtr, const FFoliageUISettings* UISettings, const FFoliagePaintingGeometryFilter* OverrideGeometryFilter)
 {
 	if (DesiredInstances.Num() == 0)
 	{
@@ -1021,12 +1014,12 @@ void FEdModeFoliage::AddInstancesImp(UWorld* InWorld, const UFoliageType* Settin
 	TArray<FPotentialInstance> PotentialInstanceBuckets[NUM_INSTANCE_BUCKETS];
 	if (DesiredInstances[0].PlacementMode == EFoliagePlacementMode::Manual)
 	{
-		CalculatePotentialInstances(InWorld, Settings, DesiredInstances, PotentialInstanceBuckets, LandscapeLayerCachesPtr, UISettings);
+		CalculatePotentialInstances(InWorld, Settings, DesiredInstances, PotentialInstanceBuckets, LandscapeLayerCachesPtr, UISettings, OverrideGeometryFilter);
 	}
 	else
 	{
 		//@TODO: actual threaded part coming, need parts of this refactor sooner for content team
-		CalculatePotentialInstances_ThreadSafe(InWorld, Settings, &DesiredInstances, PotentialInstanceBuckets, nullptr, 0, DesiredInstances.Num() - 1);
+		CalculatePotentialInstances_ThreadSafe(InWorld, Settings, &DesiredInstances, PotentialInstanceBuckets, nullptr, 0, DesiredInstances.Num() - 1, OverrideGeometryFilter);
 
 		// Existing foliage types in the palette  we want to override any existing mesh settings with the procedural settings.
 		TMap<AInstancedFoliageActor*, TArray<const UFoliageType*>> UpdatedTypesByIFA;
diff --git a/Engine/Source/Editor/FoliageEdit/Private/FoliageEdMode.h b/Engine/Source/Editor/FoliageEdit/Private/FoliageEdMode.h
--- a/Engine/Source/Editor/FoliageEdit/Private/FoliageEdMode.h
+++ b/Engine/Source/Editor/FoliageEdit/Private/FoliageEdMode.h
@@ -184,6 +184,36 @@ public:
 	}
 };
 
+//
+// Painting filtering options
+//
+struct FFoliagePaintingGeometryFilter
+{
+	bool bAllowLandscape;
+	bool bAllowStaticMesh;
+	bool bAllowBSP;
+	bool bAllowTranslucent;
+
+	FFoliagePaintingGeometryFilter(const FFoliageUISettings& InUISettings)
+		: bAllowLandscape(InUISettings.bFilterLandscape)
+		, bAllowStaticMesh(InUISettings.bFilterStaticMesh)
+		, bAllowBSP(InUISettings.bFilterBSP)
+		, bAllowTranslucent(InUISettings.bFilterTranslucent)
+	{
+	}
+
+	FFoliagePaintingGeometryFilter()
+		: bAllowLandscape(false)
+		, bAllowStaticMesh(false)
+		, bAllowBSP(false)
+		, bAllowTranslucent(false)
+	{
+	}
+
+	bool operator() (const UPrimitiveComponent* Component) const;
+};
+
+
 // Number of buckets for layer weight histogram distribution.
 #define NUM_INSTANCE_BUCKETS 10
 
@@ -370,7 +400,7 @@ public:
 	void AdjustBrushRadius(float Adjustment);
 
 	/** Add desired instances. Uses foliage settings to determine location/scale/rotation and whether instances should be ignored */
-	static void AddInstances(UWorld* InWorld, const TArray<FDesiredFoliageInstance>& DesiredInstances);
+	static void AddInstances(UWorld* InWorld, const TArray<FDesiredFoliageInstance>& DesiredInstances, const FFoliagePaintingGeometryFilter& OverrideGeometryFilter);
 
 	typedef TMap<FName, TMap<ULandscapeComponent*, TArray<uint8> > > LandscapeLayerCacheData;
 
@@ -451,13 +481,13 @@ private:
 	void HandleOnFoliageTypeMeshChanged(UFoliageType* FoliageType);
 
 	/** Common code for adding instances to world based on settings */
-	static void AddInstancesImp(UWorld* InWorld, const UFoliageType* Settings, const TArray<FDesiredFoliageInstance>& DesiredInstances, const TArray<int32>& ExistingInstances = TArray<int32>(), const float Pressure = 1.f, LandscapeLayerCacheData* LandscapeLayerCaches = nullptr, const FFoliageUISettings* UISettings = nullptr);
+	static void AddInstancesImp(UWorld* InWorld, const UFoliageType* Settings, const TArray<FDesiredFoliageInstance>& DesiredInstances, const TArray<int32>& ExistingInstances = TArray<int32>(), const float Pressure = 1.f, LandscapeLayerCacheData* LandscapeLayerCaches = nullptr, const FFoliageUISettings* UISettings = nullptr, const FFoliagePaintingGeometryFilter* OverrideGeometryFilter = nullptr);
 
 	/** Logic for determining which instances can be placed in the world*/
-	static void CalculatePotentialInstances(const UWorld* InWorld, const UFoliageType* Settings, const TArray<FDesiredFoliageInstance>& DesiredInstances, TArray<FPotentialInstance> OutPotentialInstances[NUM_INSTANCE_BUCKETS], LandscapeLayerCacheData* LandscaleLayerCachesPtr, const FFoliageUISettings* UISettings);
+	static void CalculatePotentialInstances(const UWorld* InWorld, const UFoliageType* Settings, const TArray<FDesiredFoliageInstance>& DesiredInstances, TArray<FPotentialInstance> OutPotentialInstances[NUM_INSTANCE_BUCKETS], LandscapeLayerCacheData* LandscaleLayerCachesPtr, const FFoliageUISettings* UISettings, const FFoliagePaintingGeometryFilter* OverrideGeometryFilter = nullptr);
 
 	/** Similar to CalculatePotentialInstances, but it doesn't do any overlap checks which are much harder to thread. Meant to be run in parallel for placing lots of instances */
-	static void CalculatePotentialInstances_ThreadSafe(const UWorld* InWorld, const UFoliageType* Settings, const TArray<FDesiredFoliageInstance>* DesiredInstances, TArray<FPotentialInstance> OutPotentialInstances[NUM_INSTANCE_BUCKETS], const FFoliageUISettings* UISettings, const int32 StartIdx, const int32 LastIdx);
+	static void CalculatePotentialInstances_ThreadSafe(const UWorld* InWorld, const UFoliageType* Settings, const TArray<FDesiredFoliageInstance>* DesiredInstances, TArray<FPotentialInstance> OutPotentialInstances[NUM_INSTANCE_BUCKETS], const FFoliageUISettings* UISettings, const int32 StartIdx, const int32 LastIdx, const FFoliagePaintingGeometryFilter* OverrideGeometryFilter = nullptr);
 	
 	/** Lookup the vertex color corresponding to a location traced on a static mesh */
 	static bool GetStaticMeshVertexColorForHit(const UStaticMeshComponent* InStaticMeshComponent, int32 InTriangleIndex, const FVector& InHitLocation, FColor& OutVertexColor);
diff --git a/Engine/Source/Editor/Kismet/Private/BlueprintDetailsCustomization.cpp b/Engine/Source/Editor/Kismet/Private/BlueprintDetailsCustomization.cpp
--- a/Engine/Source/Editor/Kismet/Private/BlueprintDetailsCustomization.cpp
+++ b/Engine/Source/Editor/Kismet/Private/BlueprintDetailsCustomization.cpp
@@ -729,7 +729,7 @@ void FBlueprintVarActionDetails::CustomizeDetails( IDetailLayoutBuilder& DetailL
 					TargetBlueprintDefaultObject = CachedVariableProperty->GetOwnerClass()->GetDefaultObject();
 				}
 
-				if (TargetBlueprintDefaultObject)
+				if (TargetBlueprintDefaultObject != nullptr)
 				{
 					// Things are in order, show the property and allow it to be edited
 					TArray<UObject*> ObjectList;
diff --git a/Engine/Source/Editor/Persona/Private/SAnimCurvePanel.cpp b/Engine/Source/Editor/Persona/Private/SAnimCurvePanel.cpp
--- a/Engine/Source/Editor/Persona/Private/SAnimCurvePanel.cpp
+++ b/Engine/Source/Editor/Persona/Private/SAnimCurvePanel.cpp
@@ -318,6 +318,7 @@ const FSlateBrush* SCurveEdTrack::GetExpandContent() const
 	}
 
 }
+
 void SCurveEdTrack::NewCurveNameEntered( const FText& NewText, ETextCommit::Type CommitInfo )
 {
 	if(CommitInfo == ETextCommit::OnEnter || CommitInfo == ETextCommit::OnUserMovedFocus)
@@ -327,8 +328,38 @@ void SCurveEdTrack::NewCurveNameEntered( const FText& NewText, ETextCommit::Type
 			// Check that the name doesn't already exist
 			FName RequestedName = FName(*NewText.ToString());
 			const FSmartNameMapping* NameMapping = Skeleton->GetSmartNameContainer(USkeleton::AnimCurveMappingName);
-			if(!NameMapping->Exists(RequestedName))
+
+			// If requested name exists, make sure it's not currently in use in this sequence.
+			const FSmartNameMapping::UID* RequestedNameUID = NameMapping->FindUID(RequestedName);
+			if (RequestedNameUID != nullptr)
+			{
+				// Already in use in this sequence, skip
+				if (CurveInterface->BaseSequence->RawCurveData.GetCurveData(*RequestedNameUID) != nullptr)
+				{
+					FFormatNamedArguments Args;
+					Args.Add(TEXT("DestinationName"), FText::FromName(RequestedName));
+					const FText DialogMessage = FText::Format(LOCTEXT("CurveEditor_RenameCurve_AlreadyExists", "ERROR: A curve named '{DestinationName}' already exists in this Sequence."), Args);
+					FMessageDialog::Open(EAppMsgType::Ok, DialogMessage);
+					return;
+				}
+
+				// Not in use in this sequence, switch it to the other curve name.
+				FScopedTransaction Transaction(LOCTEXT("CurveEditor_RenameCurve", "Rename Curve"));
+				CurveInterface->BaseSequence->Modify(true);
+
+				CurveInterface->CurveData->CurveUid = *RequestedNameUID;
+				CurveInterface->CurveData->LastObservedName = RequestedName;
+
+				// 2. refresh panel
+				TSharedPtr<SAnimCurvePanel> SharedPanel = PanelPtr.Pin();
+				if (SharedPanel.IsValid())
+				{
+					SharedPanel.Get()->UpdatePanel();
+				}
+			}
+			else
 			{
+				// If it doesn't exist, rename Curve.
 				FScopedTransaction Transaction(LOCTEXT("CurveEditor_RenameCurve", "Rename Curve"));
 				Skeleton->RenameSmartnameAndModify(USkeleton::AnimCurveMappingName, CurveInterface->CurveData->CurveUid, FName(*NewText.ToString()));
 			}
diff --git a/Engine/Source/Editor/UMGEditor/Private/WidgetBlueprintEditorUtils.cpp b/Engine/Source/Editor/UMGEditor/Private/WidgetBlueprintEditorUtils.cpp
--- a/Engine/Source/Editor/UMGEditor/Private/WidgetBlueprintEditorUtils.cpp
+++ b/Engine/Source/Editor/UMGEditor/Private/WidgetBlueprintEditorUtils.cpp
@@ -674,25 +674,50 @@ void FWidgetBlueprintEditorUtils::CutWidgets(UWidgetBlueprint* BP, TSet<FWidgetR
 
 void FWidgetBlueprintEditorUtils::CopyWidgets(UWidgetBlueprint* BP, TSet<FWidgetReference> Widgets)
 {
-	TArray<UWidget*> CopyableWidgets;
+	TSet<UWidget*> TemplateWidgets;
+
+	// Convert the set of widget references into the list of widget templates we're going to copy.
 	for ( const FWidgetReference& Widget : Widgets )
 	{
-		UWidget* ParentWidget = Widget.GetTemplate();
-		CopyableWidgets.Add(ParentWidget);
-
-		// When copying a widget users expect all sub widgets to be copied as well, so we need to ensure that
-		// we gather all the child widgets and copy them as well.
-		UWidgetTree::GetChildWidgets(ParentWidget, CopyableWidgets);
+		UWidget* TemplateWidget = Widget.GetTemplate();
+		TemplateWidgets.Add(TemplateWidget);
 	}
 
-	TSet<UWidget*> CopyableWidgetsSet(CopyableWidgets);
+	TArray<UWidget*> FinalWidgets;
+
+	// Pair down copied widgets to the legitimate root widgets, if they're parent is not already
+	// in the set we're planning to copy, then keep them in the list, otherwise remove widgets that
+	// will already be handled when their parent copies into the array.
+	for ( UWidget* TemplateWidget : TemplateWidgets )
+	{
+		bool bFoundParent = false;
+
+		// See if the widget already has a parent in the set we're copying.
+		for ( UWidget* PossibleParent : TemplateWidgets )
+		{
+			if ( PossibleParent != TemplateWidget )
+			{
+				if ( TemplateWidget->IsChildOf(PossibleParent) )
+				{
+					bFoundParent = true;
+					break;
+				}
+			}
+		}
+
+		if ( !bFoundParent )
+		{
+			FinalWidgets.Add(TemplateWidget);
+			UWidgetTree::GetChildWidgets(TemplateWidget, FinalWidgets);
+		}
+	}
 
 	FString ExportedText;
-	FWidgetBlueprintEditorUtils::ExportWidgetsToText(CopyableWidgetsSet, /*out*/ ExportedText);
+	FWidgetBlueprintEditorUtils::ExportWidgetsToText(FinalWidgets, /*out*/ ExportedText);
 	FPlatformMisc::ClipboardCopy(*ExportedText);
 }
 
-void FWidgetBlueprintEditorUtils::ExportWidgetsToText(TSet<UWidget*> WidgetsToExport, /*out*/ FString& ExportedText)
+void FWidgetBlueprintEditorUtils::ExportWidgetsToText(TArray<UWidget*> WidgetsToExport, /*out*/ FString& ExportedText)
 {
 	// Clear the mark state for saving.
 	UnMarkAllObjects(EObjectMark(OBJECTMARK_TagExp | OBJECTMARK_TagImp));
diff --git a/Engine/Source/Editor/UMGEditor/Private/WidgetBlueprintEditorUtils.h b/Engine/Source/Editor/UMGEditor/Private/WidgetBlueprintEditorUtils.h
--- a/Engine/Source/Editor/UMGEditor/Private/WidgetBlueprintEditorUtils.h
+++ b/Engine/Source/Editor/UMGEditor/Private/WidgetBlueprintEditorUtils.h
@@ -30,7 +30,7 @@ public:
 	static bool IsUsableWidgetClass(UClass* WidgetClass);
 
 public:
-	static void ExportWidgetsToText(TSet<UWidget*> WidgetsToExport, /*out*/ FString& ExportedText);
+	static void ExportWidgetsToText(TArray<UWidget*> WidgetsToExport, /*out*/ FString& ExportedText);
 
 	static void ImportWidgetsFromText(UWidgetBlueprint* BP, const FString& TextToImport, /*out*/ TSet<UWidget*>& ImportedWidgetSet, /*out*/ TMap<FName, UWidgetSlotPair*>& PastedExtraSlotData);
 
diff --git a/Engine/Source/Editor/UnrealEd/Private/CookOnTheFlyServer.cpp b/Engine/Source/Editor/UnrealEd/Private/CookOnTheFlyServer.cpp
--- a/Engine/Source/Editor/UnrealEd/Private/CookOnTheFlyServer.cpp
+++ b/Engine/Source/Editor/UnrealEd/Private/CookOnTheFlyServer.cpp
@@ -3031,6 +3031,7 @@ bool UCookOnTheFlyServer::GetCurrentIniVersionStrings( const ITargetPlatform* Ta
 	IniFiles.Add(FString(TEXT("Compat")));
 	IniFiles.Add(FString(TEXT("Lightmass")));
 	IniFiles.Add(FString(TEXT("Scalability")));
+	IniFiles.Add(FString(TEXT("Hardware")));
 	IniFiles.Add(FString(TEXT("Input")));
 	IniFiles.Add(FString(TEXT("GameUserSettings")));
 	IniFiles.Add(FString(TEXT("Game")));
diff --git a/Engine/Source/Programs/AutomationTool/Scripts/RebuildLightMapsCommand.Automation.cs b/Engine/Source/Programs/AutomationTool/Scripts/RebuildLightMapsCommand.Automation.cs
--- a/Engine/Source/Programs/AutomationTool/Scripts/RebuildLightMapsCommand.Automation.cs
+++ b/Engine/Source/Programs/AutomationTool/Scripts/RebuildLightMapsCommand.Automation.cs
@@ -76,6 +76,7 @@ namespace AutomationScripts.Automation
 		{
 			Log("Running Step:- RebuildLightMaps::BuildNecessaryTargets");
 			UE4Build.BuildAgenda Agenda = new UE4Build.BuildAgenda();
+            Agenda.AddTarget("UnrealHeaderTool", UnrealBuildTool.UnrealTargetPlatform.Win64, UnrealBuildTool.UnrealTargetConfiguration.Development);
 			Agenda.AddTarget("ShaderCompileWorker", UnrealBuildTool.UnrealTargetPlatform.Win64, UnrealBuildTool.UnrealTargetConfiguration.Development);
 			Agenda.AddTarget("UnrealLightmass", UnrealBuildTool.UnrealTargetPlatform.Win64, UnrealBuildTool.UnrealTargetConfiguration.Development);
 			Agenda.AddTarget(CommandletTargetName, UnrealBuildTool.UnrealTargetPlatform.Win64, UnrealBuildTool.UnrealTargetConfiguration.Development);
diff --git a/Engine/Source/Programs/CrashReporter/CrashReportClient/CrashReportClient.Build.cs b/Engine/Source/Programs/CrashReporter/CrashReportClient/CrashReportClient.Build.cs
--- a/Engine/Source/Programs/CrashReporter/CrashReportClient/CrashReportClient.Build.cs
+++ b/Engine/Source/Programs/CrashReporter/CrashReportClient/CrashReportClient.Build.cs
@@ -27,6 +27,7 @@ public class CrashReportClient : ModuleRules
 				"XmlParser",
 				"Analytics",
 				"AnalyticsET",
+				"QoSReporter",
 			}
 			);
 
diff --git a/Engine/Source/Programs/CrashReporter/CrashReportClient/Private/CrashDescription.cpp b/Engine/Source/Programs/CrashReporter/CrashReportClient/Private/CrashDescription.cpp
--- a/Engine/Source/Programs/CrashReporter/CrashReportClient/Private/CrashDescription.cpp
+++ b/Engine/Source/Programs/CrashReporter/CrashReportClient/Private/CrashDescription.cpp
@@ -10,6 +10,7 @@
 #include "Runtime/Analytics/Analytics/Public/Analytics.h"
 #include "Runtime/Analytics/Analytics/Public/Interfaces/IAnalyticsProvider.h"
 #include "EngineBuildSettings.h"
+#include "QoSReporter.h"
 
 // #CrashReport: 2015-07-23 Move crashes from [REDACTED_USER_HOME]\AppData\Local\Microsoft\Windows\WER\ReportQueue to [REDACTED_USER_HOME]\AppData\Local\CrashReportClient\Saved
 
@@ -213,9 +214,18 @@ void FPrimaryCrashProperties::SendAnalytics()
 	CrashAttributes.Add( FAnalyticsEventAttribute( TEXT( "UserActivityHint" ), UserActivityHint.AsString() ) );
 
 	Analytics.RecordEvent( TEXT( "CrashReportClient.ReportCrash" ), CrashAttributes );
-
 	// Shutdown analytics.
 	FCrashReportAnalytics::Shutdown();
+
+	// duplicate the event to QoS Reporter
+	FQoSReporter::Initialize();
+	FQoSReporter::SetBackendDeploymentName(DeploymentName);
+
+	if (FQoSReporter::IsAvailable())
+	{
+		FQoSReporter::GetProvider().RecordEvent(TEXT("CrashReportClient.ReportCrash"), CrashAttributes);
+	}
+	FQoSReporter::Shutdown();
 }
 
 void FPrimaryCrashProperties::Save()
@@ -265,6 +275,7 @@ FCrashContext::FCrashContext( const FString& CrashContextFilepath )
 		}
 
 		GetCrashProperty( EngineMode, FGenericCrashContext::RuntimePropertiesTag, TEXT( "EngineMode" ) );
+		GetCrashProperty( DeploymentName, FGenericCrashContext::RuntimePropertiesTag, TEXT( "DeploymentName" ) );
 		GetCrashProperty( AppDefaultLocale, FGenericCrashContext::RuntimePropertiesTag, TEXT( "AppDefaultLocale" ) );
 
 		if (CrashDumpMode == ECrashDumpMode::FullDump)
@@ -384,6 +395,8 @@ FCrashWERContext::FCrashWERContext( const FString& WERXMLFilepath )
 			InitializeEngineVersion( BuildVersion, BranchName, BuiltFromCL );
 		}
 
+		GetCrashProperty(DeploymentName, TEXT("DynamicSignatures"), TEXT("DeploymentName"));
+
 		bHasPrimaryData = true;
 	}
 }
diff --git a/Engine/Source/Programs/CrashReporter/CrashReportClient/Private/CrashReportClientApp.cpp b/Engine/Source/Programs/CrashReporter/CrashReportClient/Private/CrashReportClientApp.cpp
--- a/Engine/Source/Programs/CrashReporter/CrashReportClient/Private/CrashReportClientApp.cpp
+++ b/Engine/Source/Programs/CrashReporter/CrashReportClient/Private/CrashReportClientApp.cpp
@@ -171,6 +171,12 @@ void RunCrashReportClient(const TCHAR* CommandLine)
 			FCrashReportClientUnattended CrashReportClient(ErrorReport);
 			ErrorReport.SetUserComment(NSLOCTEXT("CrashReportClient", "UnattendedMode", "Sent in the unattended mode"));
 
+			if (ErrorReport.HasFilesToUpload())
+			{
+				// Send analytics.
+				FPrimaryCrashProperties::Get()->SendAnalytics();
+			}
+
 			// loop until the app is ready to quit
 			while (!GIsRequestingExit)
 			{
diff --git a/Engine/Source/Programs/CrashReporter/CrashReportClient/Private/Linux/CrashReportClientMainLinux.cpp b/Engine/Source/Programs/CrashReporter/CrashReportClient/Private/Linux/CrashReportClientMainLinux.cpp
--- a/Engine/Source/Programs/CrashReporter/CrashReportClient/Private/Linux/CrashReportClientMainLinux.cpp
+++ b/Engine/Source/Programs/CrashReporter/CrashReportClient/Private/Linux/CrashReportClientMainLinux.cpp
@@ -16,7 +16,7 @@ void CrashReporterCrashHandler(const FGenericCrashContext& GenericContext)
 	const FLinuxCrashContext& Context = static_cast< const FLinuxCrashContext& >( GenericContext );
 
 	printf("CrashHandler: Signal=%d\n", Context.Signal);
-	ReportCrash(Context);
+	const_cast< FLinuxCrashContext& >(Context).CaptureStackTrace();
 	if (GLog)
 	{
 		GLog->Flush();
diff --git a/Engine/Source/Programs/MemoryProfiler2/CallStack.cs b/Engine/Source/Programs/MemoryProfiler2/CallStack.cs
--- a/Engine/Source/Programs/MemoryProfiler2/CallStack.cs
+++ b/Engine/Source/Programs/MemoryProfiler2/CallStack.cs
@@ -36,6 +36,9 @@ namespace MemoryProfiler2
 		/// <summary> Current amount of memory that is allocated in this callstack. </summary>
 		public long LatestSize;
 
+		/// <summary> The last processed stream index </summary>
+		private ulong LastStreamIndex = ulong.MaxValue;
+
 		/// <summary> Array of graph points used to draw a timeline graph in the callstack history view. </summary>
 		public List<FSizeGraphPoint> SizeGraphPoints;
 
@@ -504,6 +507,7 @@ namespace MemoryProfiler2
 
 		public void PropagateSizeGraphPoint(FAllocationLifecycle Lifecycle, ulong StreamIndex, int SizeChange)
 		{
+#if NOT_ENABLED
 			if (SizeGraphPoints == null)
 			{
 				return;
@@ -526,27 +530,12 @@ namespace MemoryProfiler2
 				PreviousCallStack = PreviousLifecycle.AllocEvent.PreviousCallStack;
 				PreviousLifecycle = PreviousLifecycle.AllocEvent.PreviousLifecycle;
 			}
-		}
-
-		// Not used, causes stack overflow exception.
-		public void PropagateSizeGraphPointOld(FAllocationLifecycle Lifecycle, ulong StreamIndex, int SizeChange)
-		{
-			if (SizeGraphPoints == null)
-			{
-				return;
-			}
-
-			PropagateSizeGraphPointInner(StreamIndex, SizeChange);
-
-			if (Lifecycle.AllocEvent.PreviousCallStack != null)
-			{
-				Lifecycle.AllocEvent.PreviousCallStack.PropagateSizeGraphPoint(Lifecycle.AllocEvent.PreviousLifecycle, StreamIndex, SizeChange);
-			}
+#endif
 		}
 
 		private void PropagateSizeGraphPointInner(ulong StreamIndex, int SizeChange)
 		{
-			if (SizeGraphPoints.Count > 0 && SizeGraphPoints[SizeGraphPoints.Count - 1].StreamIndex == StreamIndex)
+			if (LastStreamIndex == StreamIndex)
 			{
 				FSizeGraphPoint LastPoint = SizeGraphPoints[SizeGraphPoints.Count - 1];
 
@@ -567,6 +556,7 @@ namespace MemoryProfiler2
 				}
 
 				SizeGraphPoints.Add(new FSizeGraphPoint(StreamIndex, SizeChange, false));
+				LastStreamIndex = StreamIndex;
 			}
 		}
 	};
@@ -678,9 +668,9 @@ namespace MemoryProfiler2
 			if (ReallocsEvents == null)
 			{
 				ReallocsEvents = new List<FReallocationEvent>();
+				ReallocsEvents.Capacity = 1;
 			}
 			ReallocsEvents.Add(new FReallocationEvent(StreamToken, ReallocCallStack));
-			ReallocsEvents.TrimExcess();
 
 			if (ReallocCallStack != InitialCallStack)
 			{
diff --git a/Engine/Source/Programs/MemoryProfiler2/StreamParser.cs b/Engine/Source/Programs/MemoryProfiler2/StreamParser.cs
--- a/Engine/Source/Programs/MemoryProfiler2/StreamParser.cs
+++ b/Engine/Source/Programs/MemoryProfiler2/StreamParser.cs
@@ -12,7 +12,9 @@ namespace MemoryProfiler2
 {	
 	public static class FStreamParser
 	{
-        private static BinaryReader SwitchStreams(int PartIndex, string BaseFilename, bool bIsBigEndian, out FileStream ParserStream)
+		const int AllocationsPerSlice = 5000 * 100;
+
+		private static BinaryReader SwitchStreams(int PartIndex, string BaseFilename, bool bIsBigEndian, out FileStream ParserStream)
         {
             // Determine the effective filename for this part
             string CurrentFilename = BaseFilename;
@@ -39,11 +41,32 @@ namespace MemoryProfiler2
 
 		//-----------------------------------------------------------------------------
 		//@DEBUG
+#if DEBUG_TIMINGS
 		static FGlobalTimer MallocTimer = new FGlobalTimer( "Malloc" );
 		static FGlobalTimer FreeTimer = new FGlobalTimer( "Free" );
 		static FGlobalTimer ReallocTimer = new FGlobalTimer( "Realloc" );
 		static FGlobalTimer OtherTimer = new FGlobalTimer( "Other" );
+#endif
 		static long IncompleteLifeCycles = 0;
+
+		static void WriteTimings()
+		{
+#if DEBUG_TIMINGS
+			string Info = "Timings\n" +
+									MallocTimer.ToString() + "\n" +
+									FreeTimer.ToString() + "\n" +
+									ReallocTimer.ToString() + "\n" +
+									OtherTimer.ToString() + "\n" +
+									"IncompleteLifeCycles " + IncompleteLifeCycles;
+
+			Debug.WriteLine(Info);
+#else
+			string Info = "IncompleteLifeCycles " + IncompleteLifeCycles;
+
+			Debug.WriteLine(Info);
+#endif
+		}
+
 		//-----------------------------------------------------------------------------
 
 		/// <summary> Parses the passed in token stream and returns list of snapshots. </summary>
@@ -275,8 +298,13 @@ namespace MemoryProfiler2
             StartOfMetadata = Math.Min(StartOfMetadata, Header.CallStackTableOffset);
             StartOfMetadata = Math.Min(StartOfMetadata, Header.ModulesOffset);
 
-            long ProgressInterval = ((Int64)StartOfMetadata - TokenStreamOffset) / 100;
-			double ProgressScaleFactor = 100.0f / ( (Int64)StartOfMetadata - TokenStreamOffset );
+			long ProgressInterval = ((Int64)StartOfMetadata - TokenStreamOffset) / 1000;
+			if (ProgressInterval < 1)
+			{
+				ProgressInterval = 1;
+			}
+
+			double ProgressScaleFactor = 100.0f / ((Int64)StartOfMetadata - TokenStreamOffset);
             long NextProgressUpdate = TokenStreamOffset;
 
 			// Parse tokens till we reach the end of the stream.
@@ -305,7 +333,7 @@ namespace MemoryProfiler2
 						NextProgressUpdate += ProgressInterval;
 					}
 
-					if( CustomSnapshots.Count > 0 && Snapshot.AllocationCount >= CustomSnapshots[ 0 ] * 5000 )
+					if ( CustomSnapshots.Count > 0 && (Snapshot.AllocationCount >= CustomSnapshots[0] * AllocationsPerSlice) ) 
 					{
 						// Create an unnamed snapshot.
 						FStreamSnapshot MarkerSnapshot = Snapshot.DeepCopy( PointerToPointerInfoMap );
@@ -335,9 +363,11 @@ namespace MemoryProfiler2
 						// Malloc
 						case EProfilingPayloadType.TYPE_Malloc:
 						{
+#if DEBUG_TIMINGS
 							MallocTimer.Start();
+#endif
 
-							if( Token.Pointer != 0 )
+							if ( Token.Pointer != 0 )
 							{
 								Token.CallStackIndex = GetVirtualCallStackIndex( Token, Observer );
 
@@ -355,16 +385,19 @@ namespace MemoryProfiler2
 								CurrentCallstack.ProcessMalloc( Token, ref NewLifecycle );
 							}
 
+#if DEBUG_TIMINGS
 							MallocTimer.Stop();
+#endif
 						}
 						break;
 
 						// Free
 						case EProfilingPayloadType.TYPE_Free:
 						{
+#if DEBUG_TIMINGS
 							FreeTimer.Start();
-
-							if( bDetectingUseGlobalReallocZeroPtr )
+#endif
+							if ( bDetectingUseGlobalReallocZeroPtr )
 							{
 								if( ReallocZeroCount > 0 && Token.Pointer == ReallocZeroPtr )
 								{
@@ -385,14 +418,18 @@ namespace MemoryProfiler2
 								}
 							}
 
+#if DEBUG_TIMINGS
 							FreeTimer.Stop();
+#endif
 						}
 						break;
 
 						// Realloc
 						case EProfilingPayloadType.TYPE_Realloc:
 						{
+#if DEBUG_TIMINGS
 							ReallocTimer.Start();
+#endif
 
 							Token.CallStackIndex = GetVirtualCallStackIndex( Token, Observer );
 
@@ -440,7 +477,9 @@ namespace MemoryProfiler2
 								if( bUseGlobalReallocZeroPtr )
 								{
 									// break out of case to avoid 'double malloc' warnings
+#if DEBUG_TIMINGS
 									ReallocTimer.Stop();
+#endif
 									break;
 								}
 							}
@@ -468,7 +507,9 @@ namespace MemoryProfiler2
 
 								}
 
+#if DEBUG_TIMINGS
 								ReallocTimer.Stop();
+#endif
 							}
 						}
 						break;
@@ -476,7 +517,9 @@ namespace MemoryProfiler2
 						// Status/ payload.
 						case EProfilingPayloadType.TYPE_Other:
 						{
+#if DEBUG_TIMINGS
 							OtherTimer.Start();
+#endif
 
 							switch( Token.SubType )
 							{
@@ -643,7 +686,9 @@ namespace MemoryProfiler2
 									throw new InvalidDataException();
 								}
 							}
+#if DEBUG_TIMINGS
 							OtherTimer.Stop();
+#endif
 							break;
 						}
 
@@ -666,14 +711,7 @@ namespace MemoryProfiler2
 			}
 			//-----------------------------------------------------------------------------
 			//@DEBUG
-			string Info = "Timings\n" +
-						MallocTimer.ToString() + "\n" +
-						FreeTimer.ToString() + "\n" +
-						ReallocTimer.ToString() + "\n" +
-						OtherTimer.ToString() + "\n" +
-						"IncompleteLifeCycles " + IncompleteLifeCycles;
-
-			Debug.WriteLine( Info );
+			WriteTimings();
 			//-----------------------------------------------------------------------------
 
             if (MainMProfWindow != null && bFoundMultiPoolCallStacks)
@@ -836,7 +874,7 @@ namespace MemoryProfiler2
 			// Maintain timeline view
 			Snapshot.AllocationSize += PointerInfo.Size;
 			Snapshot.AllocationCount++;
-			if( Snapshot.AllocationCount % 5000 == 0 )
+			if (Snapshot.AllocationCount % AllocationsPerSlice == 0)
 			{
 				FMemorySlice Slice = new FMemorySlice( Snapshot );
 				Snapshot.OverallMemorySlice.Add( Slice );
@@ -856,7 +894,7 @@ namespace MemoryProfiler2
 			// Maintain timeline view
             Snapshot.AllocationSize -= FreedAllocInfo.Size;
 			Snapshot.AllocationCount++;
-			if( Snapshot.AllocationCount % 5000 == 0 )
+			if (Snapshot.AllocationCount % AllocationsPerSlice == 0)
 			{
 				FMemorySlice Slice = new FMemorySlice( Snapshot );
 				Snapshot.OverallMemorySlice.Add( Slice );
diff --git a/Engine/Source/Runtime/Analytics/AnalyticsET/Private/AnalyticsET.cpp b/Engine/Source/Runtime/Analytics/AnalyticsET/Private/AnalyticsET.cpp
--- a/Engine/Source/Runtime/Analytics/AnalyticsET/Private/AnalyticsET.cpp
+++ b/Engine/Source/Runtime/Analytics/AnalyticsET/Private/AnalyticsET.cpp
@@ -385,7 +385,7 @@ void FAnalyticsProviderET::FlushEvents()
 				FString EventParams;
 				if (Event.Attributes.Num() > 0)
 				{
-					for (int Ndx = 0; Ndx<FMath::Min(Event.Attributes.Num(), 10); ++Ndx)
+					for (int Ndx = 0; Ndx<FMath::Min(Event.Attributes.Num(), 40); ++Ndx)
 					{
 						EventParams += FString::Printf(TEXT("&AttributeName%d=%s&AttributeValue%d=%s"), 
 							Ndx, 
diff --git a/Engine/Source/Runtime/Analytics/QoSReporter/Private/QoSReporter.cpp b/Engine/Source/Runtime/Analytics/QoSReporter/Private/QoSReporter.cpp
--- a/Engine/Source/Runtime/Analytics/QoSReporter/Private/QoSReporter.cpp
+++ b/Engine/Source/Runtime/Analytics/QoSReporter/Private/QoSReporter.cpp
@@ -9,7 +9,9 @@
 #if WITH_PERFCOUNTERS
 	#include "PerfCountersModule.h"
 #endif // WITH_PERFCOUNTERS
-#include "PerfCountersHelpers.h"
+#if WITH_ENGINE
+	#include "PerfCountersHelpers.h"
+#endif // WITH_ENGINE
 
 #ifndef WITH_QOSREPORTER
 	#error "WITH_QOSREPORTER should be defined in Build.cs file"
@@ -144,6 +146,9 @@ void FQoSReporter::ReportStartupCompleteEvent()
 	ParamArray.Add(FAnalyticsEventAttribute(EQoSEvents::ToString(EQoSEventParam::StartupTime), Difference));
 	Analytics->RecordEvent(EQoSEvents::ToString(EQoSEventParam::StartupTime), ParamArray);
 
+	// log it so we can do basic analysis by log scraping
+	UE_LOG(LogQoSReporter, Log, TEXT("Startup complete, took %f seconds."), Difference);
+
 	bStartupEventReported = true;
 }
 
@@ -178,7 +183,11 @@ void FQoSReporter::EnableCountingHitches(bool bEnable)
 
 double FQoSReporter::HeartbeatInterval = 300;
 double FQoSReporter::LastHeartbeatTimestamp = 0;
+#if WITH_ENGINE
 extern ENGINE_API float GAverageFPS;
+#else
+float GAverageFPS = 0.0f;	// fake
+#endif // WITH_ENGINE
 
 void FQoSReporter::Tick()
 {
diff --git a/Engine/Source/Runtime/Analytics/QoSReporter/Private/QoSReporterModule.cpp b/Engine/Source/Runtime/Analytics/QoSReporter/Private/QoSReporterModule.cpp
--- a/Engine/Source/Runtime/Analytics/QoSReporter/Private/QoSReporterModule.cpp
+++ b/Engine/Source/Runtime/Analytics/QoSReporter/Private/QoSReporterModule.cpp
@@ -14,7 +14,7 @@ DEFINE_LOG_CATEGORY(LogQoSReporter);
 IMPLEMENT_MODULE(FQoSReporterModule, QoSReporter);
 
 // helps to version QoS events (date*10 to allow for 10 revisions per day)
-#define QOS_EVENTS_REVISION					201601200
+#define QOS_EVENTS_REVISION					201601280
 
 FString FQoSReporterModule::Config::GetDefaultAppVersion()
 { 
diff --git a/Engine/Source/Runtime/Analytics/QoSReporter/QoSReporter.Build.cs b/Engine/Source/Runtime/Analytics/QoSReporter/QoSReporter.Build.cs
--- a/Engine/Source/Runtime/Analytics/QoSReporter/QoSReporter.Build.cs
+++ b/Engine/Source/Runtime/Analytics/QoSReporter/QoSReporter.Build.cs
@@ -14,10 +14,14 @@ namespace UnrealBuildTool.Rules
 					"Analytics",
 					"HTTP",
 					"Json",
-					"Engine",	// for GAverageFPS
 				}
 			);
 
+			if (Target.Type != TargetRules.TargetType.Program)
+			{
+				PrivateDependencyModuleNames.Add("Engine"); // for GAverageFPS - unused by programs
+			}
+
 			// servers expose QoS metrics via perfcounters
 			if (Target.Type != TargetRules.TargetType.Client && Target.Type != TargetRules.TargetType.Program)
 			{
diff --git a/Engine/Source/Runtime/Core/Private/GenericPlatform/GenericPlatformCrashContext.cpp b/Engine/Source/Runtime/Core/Private/GenericPlatform/GenericPlatformCrashContext.cpp
--- a/Engine/Source/Runtime/Core/Private/GenericPlatform/GenericPlatformCrashContext.cpp
+++ b/Engine/Source/Runtime/Core/Private/GenericPlatform/GenericPlatformCrashContext.cpp
@@ -139,6 +139,7 @@ void FGenericCrashContext::SerializeContentToBuffer()
 	AddCrashProperty( TEXT( "PlatformName" ), *NCachedCrashContextProperties::PlatformName );
 	AddCrashProperty( TEXT( "PlatformNameIni" ), *NCachedCrashContextProperties::PlatformNameIni );
 	AddCrashProperty( TEXT( "EngineMode" ), FPlatformMisc::GetEngineMode() );
+	AddCrashProperty( TEXT( "DeploymentName"), FApp::GetDeploymentName() );
 	AddCrashProperty( TEXT( "EngineVersion" ), *FEngineVersion::Current().ToString() );
 	AddCrashProperty( TEXT( "CommandLine" ), FCommandLine::IsInitialized() ? FCommandLine::Get() : TEXT( "" ) );
 	AddCrashProperty( TEXT( "LanguageLCID" ), FInternationalization::Get().GetCurrentCulture()->GetLCID() );
diff --git a/Engine/Source/Runtime/Core/Private/GenericPlatform/GenericPlatformMisc.cpp b/Engine/Source/Runtime/Core/Private/GenericPlatform/GenericPlatformMisc.cpp
--- a/Engine/Source/Runtime/Core/Private/GenericPlatform/GenericPlatformMisc.cpp
+++ b/Engine/Source/Runtime/Core/Private/GenericPlatform/GenericPlatformMisc.cpp
@@ -11,6 +11,7 @@
 #include "ExceptionHandling.h"
 #include "Containers/Map.h"
 #include "GenericPlatformCrashContext.h"
+#include "GenericPlatformDriver.h"
 
 #include "UProjectInfo.h"
 
@@ -247,12 +248,9 @@ FString FGenericPlatformMisc::GetPrimaryGPUBrand()
 	return FString( TEXT( "GenericGPUBrand" ) );
 }
 
-void FGenericPlatformMisc::GetGPUDriverInfo(const FString DeviceDescription, FString& InternalDriverVersion, FString& UserDriverVersion, FString& DriverDate)
+FGPUDriverInfo FGenericPlatformMisc::GetGPUDriverInfo(const FString& DeviceDescription)
 {
-	// to prevent bugs
-	InternalDriverVersion.Empty();
-	UserDriverVersion.Empty();
-	DriverDate.Empty();
+	return FGPUDriverInfo();
 }
 
 void FGenericPlatformMisc::GetOSVersions( FString& out_OSVersionLabel, FString& out_OSSubVersionLabel )
diff --git a/Engine/Source/Runtime/Core/Private/Linux/LinuxPlatformCrashContext.cpp b/Engine/Source/Runtime/Core/Private/Linux/LinuxPlatformCrashContext.cpp
--- a/Engine/Source/Runtime/Core/Private/Linux/LinuxPlatformCrashContext.cpp
+++ b/Engine/Source/Runtime/Core/Private/Linux/LinuxPlatformCrashContext.cpp
@@ -251,6 +251,7 @@ void GenerateWindowsErrorReport(const FString & WERPath)
 		WriteLine(ReportFile, TEXT("\t<DynamicSignatures>"));
 		WriteLine(ReportFile, TEXT("\t\t<Parameter1>6.1.7601.2.1.0.256.48</Parameter1>"));
 		WriteLine(ReportFile, TEXT("\t\t<Parameter2>1033</Parameter2>"));
+		WriteLine(ReportFile, *FString::Printf(TEXT("\t\t<DeploymentName>%s</DeploymentName>"), FApp::GetDeploymentName()));
 		WriteLine(ReportFile, TEXT("\t</DynamicSignatures>"));
 
 		WriteLine(ReportFile, TEXT("\t<SystemInformation>"));
@@ -286,34 +287,26 @@ void GenerateMinidump(const FString & Path)
 }
 
 
-int32 DLLEXPORT ReportCrash(const FLinuxCrashContext & Context)
+void FLinuxCrashContext::CaptureStackTrace()
 {
-	static bool GAlreadyCreatedMinidump = false;
-	// Only create a minidump the first time this function is called.
-	// (Can be called the first time from the RenderThread, then a second time from the MainThread.)
-	if ( GAlreadyCreatedMinidump == false )
+	// Only do work the first time this function is called - this is mainly a carry over from Windows where it can be called multiple times, left intact for extra safety.
+	if (!bCapturedBacktrace)
 	{
-		GAlreadyCreatedMinidump = true;
-
 		const SIZE_T StackTraceSize = 65535;
 		ANSICHAR* StackTrace = (ANSICHAR*) FMemory::Malloc( StackTraceSize );
 		StackTrace[0] = 0;
 		// Walk the stack and dump it to the allocated memory (ignore first 2 callstack lines as those are in stack walking code)
-		FPlatformStackWalk::StackWalkAndDump( StackTrace, StackTraceSize, 2, const_cast< FLinuxCrashContext* >( &Context ) );
+		FPlatformStackWalk::StackWalkAndDump( StackTrace, StackTraceSize, 2, this);
 
 		FCString::Strncat( GErrorHist, ANSI_TO_TCHAR(StackTrace), ARRAY_COUNT(GErrorHist) - 1 );
-		CreateExceptionInfoString(Context.Signal, Context.Info);
+		CreateExceptionInfoString(Signal, Info);
 
 		FMemory::Free( StackTrace );
+		bCapturedBacktrace = true;
 	}
-
-	return 0;
 }
 
-/**
- * Generates information for crash reporter
- */
-void DLLEXPORT GenerateCrashInfoAndLaunchReporter(const FLinuxCrashContext & Context)
+void FLinuxCrashContext::GenerateCrashInfoAndLaunchReporter(bool bReportingNonCrash) const
 {
 	// do not report crashes for tools (particularly for crash reporter itself)
 #if !IS_PROGRAM
@@ -325,12 +318,12 @@ void DLLEXPORT GenerateCrashInfoAndLaunchReporter(const FLinuxCrashContext & Con
 		CrashGuid = FGuid::NewGuid().ToString();
 	}
 
-	FString CrashInfoFolder = FString::Printf(TEXT("crashinfo-%s-pid-%d-%s"), FApp::GetGameName(), getpid(), *CrashGuid);
+	FString CrashInfoFolder = FString::Printf(TEXT("%sinfo-%s-pid-%d-%s"), bReportingNonCrash ? TEXT("ensure") : TEXT("crash"), FApp::GetGameName(), getpid(), *CrashGuid);
 	FString CrashInfoAbsolute = FPaths::ConvertRelativePathToFull(CrashInfoFolder);
 	if (IFileManager::Get().MakeDirectory(*CrashInfoFolder))
 	{
 		// generate "minidump"
-		Context.GenerateReport(FPaths::Combine(*CrashInfoFolder, TEXT("Diagnostics.txt")));
+		GenerateReport(FPaths::Combine(*CrashInfoFolder, TEXT("Diagnostics.txt")));
 
 		// generate "WER"
 		GenerateWindowsErrorReport(FPaths::Combine(*CrashInfoFolder, TEXT("wermeta.xml")));
@@ -368,8 +361,6 @@ void DLLEXPORT GenerateCrashInfoAndLaunchReporter(const FLinuxCrashContext & Con
 
 		CrashReportClientArguments += CrashInfoAbsolute + TEXT("/");
 
-		// show on the console
-		printf("Starting %s\n", StringCast<ANSICHAR>(RelativePathToCrashReporter).Get());
 		FProcHandle RunningProc = FPlatformProcess::CreateProc(RelativePathToCrashReporter, *CrashReportClientArguments, true, false, false, NULL, 0, NULL, NULL);
 		if (FPlatformProcess::IsProcRunning(RunningProc))
 		{
@@ -395,16 +386,19 @@ void DLLEXPORT GenerateCrashInfoAndLaunchReporter(const FLinuxCrashContext & Con
 
 #endif
 
-	// remove the handler for this signal and re-raise it (which should generate the proper core dump)
-	UE_LOG(LogLinux, Log, TEXT("Engine crash handling finished; re-raising signal %d for the default handler. Good bye."), Context.Signal);
+	if (!bReportingNonCrash)
+	{
+		// remove the handler for this signal and re-raise it (which should generate the proper core dump)
+		UE_LOG(LogLinux, Log, TEXT("Engine crash handling finished; re-raising signal %d for the default handler. Good bye."), Signal);
 
-	struct sigaction ResetToDefaultAction;
-	FMemory::Memzero(ResetToDefaultAction);
-	ResetToDefaultAction.sa_handler = SIG_DFL;
-	sigemptyset(&ResetToDefaultAction.sa_mask);
-	sigaction(Context.Signal, &ResetToDefaultAction, nullptr);
+		struct sigaction ResetToDefaultAction;
+		FMemory::Memzero(ResetToDefaultAction);
+		ResetToDefaultAction.sa_handler = SIG_DFL;
+		sigemptyset(&ResetToDefaultAction.sa_mask);
+		sigaction(Signal, &ResetToDefaultAction, nullptr);
 
-	raise(Context.Signal);
+		raise(Signal);
+	}
 }
 
 /**
@@ -416,7 +410,7 @@ void DefaultCrashHandler(const FLinuxCrashContext & Context)
 
 	// at this point we should already be using malloc crash handler (see PlatformCrashHandler)
 
-	ReportCrash(Context);
+	const_cast<FLinuxCrashContext&>(Context).CaptureStackTrace();
 	if (GLog)
 	{
 		GLog->Flush();
@@ -431,7 +425,7 @@ void DefaultCrashHandler(const FLinuxCrashContext & Context)
 		GError->HandleError();
 	}
 
-	return GenerateCrashInfoAndLaunchReporter(Context);
+	return Context.GenerateCrashInfoAndLaunchReporter();
 }
 
 /** Global pointer to crash handler */
diff --git a/Engine/Source/Runtime/Core/Private/Linux/LinuxPlatformStackWalk.cpp b/Engine/Source/Runtime/Core/Private/Linux/LinuxPlatformStackWalk.cpp
--- a/Engine/Source/Runtime/Core/Private/Linux/LinuxPlatformStackWalk.cpp
+++ b/Engine/Source/Runtime/Core/Private/Linux/LinuxPlatformStackWalk.cpp
@@ -759,6 +759,32 @@ void FLinuxPlatformStackWalk::CaptureStackBackTrace( uint64* BackTrace, uint32 M
 	}
 }
 
-void NewReportEnsure( const TCHAR* ErrorMessage )
+static FCriticalSection EnsureLock;
+static bool bReentranceGuard = false;
+
+void NewReportEnsure(const TCHAR* ErrorMessage)
 {
+	// Simple re-entrance guard.
+	EnsureLock.Lock();
+
+	if (bReentranceGuard)
+	{
+		EnsureLock.Unlock();
+		return;
+	}
+
+	bReentranceGuard = true;
+
+	siginfo_t Signal;
+	Signal.si_signo = SIGTRAP;
+	Signal.si_code = TRAP_TRACE;
+	Signal.si_addr = __builtin_return_address(0);
+
+	FLinuxCrashContext EnsureContext;
+	EnsureContext.InitFromSignal(SIGTRAP, &Signal, nullptr);
+	EnsureContext.CaptureStackTrace();
+	EnsureContext.GenerateCrashInfoAndLaunchReporter(true);
+
+	bReentranceGuard = false;
+	EnsureLock.Unlock();
 }
diff --git a/Engine/Source/Runtime/Core/Private/Mac/MacPlatformMisc.cpp b/Engine/Source/Runtime/Core/Private/Mac/MacPlatformMisc.cpp
--- a/Engine/Source/Runtime/Core/Private/Mac/MacPlatformMisc.cpp
+++ b/Engine/Source/Runtime/Core/Private/Mac/MacPlatformMisc.cpp
@@ -1564,6 +1564,8 @@ void FMacCrashContext::GenerateWindowsErrorReport(char const* WERPath) const
 		WriteUTF16String(ReportFile, TEXT("\t\t<Parameter2>"));
 		WriteUTF16String(ReportFile, *GMacAppInfo.LCID);
 		WriteLine(ReportFile, TEXT("</Parameter2>"));
+		WriteLine(ReportFile, *FString::Printf(TEXT("\t\t<DeploymentName>%s</DeploymentName>"), FApp::GetDeploymentName()));
+
 		WriteLine(ReportFile, TEXT("\t</DynamicSignatures>"));
 		
 		WriteLine(ReportFile, TEXT("\t<SystemInformation>"));
diff --git a/Engine/Source/Runtime/Core/Private/Misc/App.cpp b/Engine/Source/Runtime/Core/Private/Misc/App.cpp
--- a/Engine/Source/Runtime/Core/Private/Misc/App.cpp
+++ b/Engine/Source/Runtime/Core/Private/Misc/App.cpp
@@ -44,6 +44,21 @@ FString FApp::GetEpicProductIdentifier()
 	return FString(TEXT(EPIC_PRODUCT_IDENTIFIER));
 }
 
+const TCHAR * FApp::GetDeploymentName()
+{
+	static TCHAR StaticDeploymentName[64] = {0};
+	static bool bHaveDeployment = false;
+
+	if (!bHaveDeployment)
+	{
+		// use -epicapp value from the commandline. Default deployment is not captured by this,
+		// but this may not be a problem as that would be the case only during the development
+		FParse::Value(FCommandLine::Get(), TEXT("EPICAPP="), StaticDeploymentName, ARRAY_COUNT(StaticDeploymentName) - 1);
+		bHaveDeployment = true;
+	}
+
+	return StaticDeploymentName;
+}
 
 EBuildConfigurations::Type FApp::GetBuildConfiguration()
 {
diff --git a/Engine/Source/Runtime/Core/Private/Misc/ConfigCacheIni.cpp b/Engine/Source/Runtime/Core/Private/Misc/ConfigCacheIni.cpp
--- a/Engine/Source/Runtime/Core/Private/Misc/ConfigCacheIni.cpp
+++ b/Engine/Source/Runtime/Core/Private/Misc/ConfigCacheIni.cpp
@@ -2823,6 +2823,8 @@ void FConfigCacheIni::InitializeConfigSystem()
 
 	// Load scalability settings.
 	FConfigCacheIni::LoadGlobalIniFile(GScalabilityIni, TEXT("Scalability"));
+	// Load driver blacklist
+	FConfigCacheIni::LoadGlobalIniFile(GHardwareIni, TEXT("Hardware"));
 	
 	// Load user game settings .ini, allowing merging. This also updates the user .ini if necessary.
 	FConfigCacheIni::LoadGlobalIniFile(GGameUserSettingsIni, TEXT("GameUserSettings"));
diff --git a/Engine/Source/Runtime/Core/Private/Misc/Core.cpp b/Engine/Source/Runtime/Core/Private/Misc/Core.cpp
--- a/Engine/Source/Runtime/Core/Private/Misc/Core.cpp
+++ b/Engine/Source/Runtime/Core/Private/Misc/Core.cpp
@@ -145,6 +145,7 @@ FString				GEditorPerProjectIni;										/* Editor User Settings ini filename *
 FString				GCompatIni;
 FString				GLightmassIni;												/* Lightmass settings ini filename */
 FString				GScalabilityIni;											/* Scalability settings ini filename */
+FString				GHardwareIni;												/* Hardware ini filename */
 FString				GInputIni;													/* Input ini filename */
 FString				GGameIni;													/* Game ini filename */
 FString				GGameUserSettingsIni;										/* User Game Settings ini filename */
diff --git a/Engine/Source/Runtime/Core/Private/Windows/WindowsPlatformMisc.cpp b/Engine/Source/Runtime/Core/Private/Windows/WindowsPlatformMisc.cpp
--- a/Engine/Source/Runtime/Core/Private/Windows/WindowsPlatformMisc.cpp
+++ b/Engine/Source/Runtime/Core/Private/Windows/WindowsPlatformMisc.cpp
@@ -8,6 +8,7 @@
 #include "WindowsPlatformCrashContext.h"
 
 #include "GenericPlatformChunkInstall.h"
+#include "GenericPlatformDriver.h"			// FGPUDriverInfo
 
 // Resource includes.
 #include "Runtime/Launch/Resources/Windows/Resource.h"
@@ -46,6 +47,18 @@
 #define SM_CONVERTIBLESLATEMODE			0x2003
 #endif
 
+// this cvar can be removed once we have a single method that works well
+static TAutoConsoleVariable<int32> CVarDriverDetectionMethod(
+	TEXT("r.DriverDetectionMethod"),
+	4,
+	TEXT("Defined which implementation is used to detect the GPU driver (to check for old drivers and for logs and statistics)\n")
+	TEXT("  0: Iterate available drivers in registry and choose the one with the same name, if in question use next method (happens)\n")
+	TEXT("  1: Get the driver of the primary adpater (might not be correct when dealing with multiple adapters)\n")
+	TEXT("  2: Use DirectX LUID (would be the best, not yet implemented)\n")
+	TEXT("  3: Use Windows functions, use the primary device (might be wrong when API is using another adapter)\n")
+	TEXT("  4: Use Windows functions, use the one names like the DirectX Device (newest, most promising)"),
+	ECVF_RenderThreadSafe);
+
 namespace
 {
 	/**
@@ -2421,83 +2434,189 @@ static void GetVideoDriverDetails(const FString& Key, FGPUDriverInfo& Out)
 		}
 	}
 
+	// some key/value pairs explained: http://www.helpdoc-online.com/SCDMS01EN1A330P306~Windows-NT-Workstation-3.51-Resource-Kit-Help-en~Video-Device-Driver-Entries.htm
+
 	FWindowsPlatformMisc::QueryRegKey(HKEY_LOCAL_MACHINE, *Key, TEXT("ProviderName"), Out.ProviderName);
 
-	// AMD (not the Catalyst one) and NVIDIA
+	if (!Out.ProviderName.IsEmpty())
+	{
+		if (Out.ProviderName.Find(TEXT("NVIDIA")) != INDEX_NONE)
+		{
+			Out.SetNVIDIA();
+		}
+		else if (Out.ProviderName.Find(TEXT("Advanced Micro Devices")) != INDEX_NONE)
+		{
+			Out.SetAMD();
+		}
+		else if (Out.ProviderName.Find(TEXT("Intel")) != INDEX_NONE)	// usually TEXT("Intel Corporation")
+		{
+			Out.SetIntel();
+		}
+	}
+
+	// technical driver version, AMD and NVIDIA
 	FWindowsPlatformMisc::QueryRegKey(HKEY_LOCAL_MACHINE, *Key, TEXT("DriverVersion"), Out.InternalDriverVersion);
 
 	Out.UserDriverVersion = Out.InternalDriverVersion;
 
-	if(Out.ProviderName == TEXT("NVIDIA"))
+	if(Out.IsNVIDIA())
 	{
-		// https://forums.geforce.com/default/topic/378546/confusion-over-driver-version-numbers/
-		// we don't care about the windows version so we don't look at the front part of the driver version
-		// "9.18.13.4788" -> "347.88"
-		// "10.18.13.4788" -> "347.88"
-		// the following code works with the current numbering scheme, if needed we have to update that
-		{
-			// e.g. 3.4788
-			FString RightPart = Out.UserDriverVersion.Right(6);
-
-			if(RightPart.Len() == 6 && RightPart[1] == (TCHAR)'.')
-			{
-				// e.g. 347.88
-				Out.UserDriverVersion = RightPart.Left(1) + RightPart[2] + RightPart[3] + TEXT(".") + RightPart[4] + RightPart[5];
-			}
-		}
+		Out.UserDriverVersion = Out.GetUnifiedDriverVersion();
 	}
-	else 
+	else if(Out.IsAMD())
 	{
-		// we assume Out.ProviderName == TEXT("Advanced Micro Devices, Inc.")
-
-		// e.g. 15.7.1
 		if(FWindowsPlatformMisc::QueryRegKey(HKEY_LOCAL_MACHINE, *Key, TEXT("Catalyst_Version"), Out.UserDriverVersion))
 		{
 			Out.UserDriverVersion = FString(TEXT("Catalyst ")) + Out.UserDriverVersion;
 		}
+
+		FString Edition;
+		if(FWindowsPlatformMisc::QueryRegKey(HKEY_LOCAL_MACHINE, *Key, TEXT("RadeonSoftwareEdition"), Edition))
+		{
+			FString Version;
+			if(FWindowsPlatformMisc::QueryRegKey(HKEY_LOCAL_MACHINE, *Key, TEXT("RadeonSoftwareVersion"), Version))
+			{
+				// e.g. TEXT("Crimson 15.12") or TEXT("Catalyst 14.1")
+				Out.UserDriverVersion = Edition + TEXT(" ") + Version;
+			}
+		}
 	}
 
 	// AMD and NVIDIA
 	FWindowsPlatformMisc::QueryRegKey(HKEY_LOCAL_MACHINE, *Key, TEXT("DriverDate"), Out.DriverDate);
 }
 
-void FWindowsPlatformMisc::GetGPUDriverInfo(const FString DeviceDescription, FString& InternalDriverVersion, FString& UserDriverVersion, FString& DriverDate)
+FGPUDriverInfo FWindowsPlatformMisc::GetGPUDriverInfo(const FString& DeviceDescription)
 {
 	// to distinguish failed GetGPUDriverInfo() from call to GetGPUDriverInfo()
-	InternalDriverVersion = TEXT("Unknown");
-	UserDriverVersion = TEXT("Unknown");
-	DriverDate = TEXT("Unknown");
+	FGPUDriverInfo Ret;
+
+	Ret.InternalDriverVersion = TEXT("Unknown");
+	Ret.UserDriverVersion = TEXT("Unknown");
+	Ret.DriverDate = TEXT("Unknown");
 
 	// for debugging, useful even in shipping to see what went wrong
-	FString ErrorString;
+	FString DebugString;
 
 	uint32 FoundDriverCount = 0;
 
-	for(uint32 i = 0;; ++i)
+	int32 Method = CVarDriverDetectionMethod.GetValueOnGameThread();
+
+	if(Method == 3 || Method == 4)
 	{
-		// Iterate all installed display adapters
-		const FString Key = FString::Printf(TEXT("SYSTEM\\CurrentControlSet\\Control\\Class\\{4D36E968-E325-11CE-BFC1-08002BE10318}\\%04d"), i);
+		UE_LOG(LogWindows, Log, TEXT("EnumDisplayDevices:"));
+
+		for(uint32 i = 0; i < 256; ++i)
+		{
+			DISPLAY_DEVICE Device;
+			
+			ZeroMemory(&Device, sizeof(Device));
+			Device.cb = sizeof(Device);
+			
+			// see https://msdn.microsoft.com/en-us/library/windows/desktop/dd162609(v=vs.85).aspx
+			if(EnumDisplayDevices(0, i, &Device, EDD_GET_DEVICE_INTERFACE_NAME) == 0)
+			{
+				// last device or error
+				break;
+			}
+
+			UE_LOG(LogWindows, Log, TEXT("   %d. '%s' (P:%d D:%d)"),
+				i,
+				Device.DeviceString,
+				(Device.StateFlags & DISPLAY_DEVICE_PRIMARY_DEVICE) != 0,
+				(Device.StateFlags & DISPLAY_DEVICE_ATTACHED_TO_DESKTOP) != 0
+				);
+
+			if(Method == 3)
+			{
+				if (!(Device.StateFlags & DISPLAY_DEVICE_PRIMARY_DEVICE))
+				{
+					// see http://www.vistaheads.com/forums/microsoft-public-windows-vista-hardware-devices/286017-find-out-active-graphics-device-programmatically-registry-key.html
+					DebugString += TEXT("JumpOverNonPrimary ");
+					// we want the primary device
+					continue;
+				}
+			}
+
+			FString DriverLocation = Device.DeviceKey;
+
+			if(DriverLocation.Left(18) == TEXT("\\Registry\\Machine\\"))		// not case sensitive
+			{
+				DriverLocation = FString(TEXT("\\HKEY_LOCAL_MACHINE\\")) + DriverLocation.RightChop(18);
+			}
+			if(DriverLocation.Left(20) == TEXT("\\HKEY_LOCAL_MACHINE\\"))		// not case sensitive
+			{
+				FString Key = DriverLocation.RightChop(20);
+				
+				FGPUDriverInfo Local;
+				GetVideoDriverDetails(Key, Local);
+
+				if(!Local.IsValid())
+				{
+					DebugString += TEXT("GetVideoDriverDetailsInvalid ");
+				}
+
+				if((Method == 3) || (Local.DeviceDescription == DeviceDescription))
+				{
+					if(!FoundDriverCount)
+					{
+						Ret = Local;
+					}
+					++FoundDriverCount;
+				}
+				else
+				{
+					DebugString += TEXT("PrimaryIsNotTheChoosenAdapter ");
+				}
+			}
+			else
+			{
+				DebugString += TEXT("PrimaryDriverLocationFailed ");
+			}
+		}
 		
-		FGPUDriverInfo Local;
-		GetVideoDriverDetails(Key, Local);
+		if(FoundDriverCount != 1)
+		{
+			// We assume if multiple entries are found they are all the same driver. If that is correct - this is no error.
+			DebugString += FString::Printf(TEXT("FoundDriverCount:%d "), FoundDriverCount);
+		}
 
-		if(!Local.IsValid())
+		if(!DebugString.IsEmpty())
 		{
-			ErrorString += TEXT("GetVideoDriverDetailsInvalid ");
+			UE_LOG(LogWindows, Log, TEXT("DebugString: %s"), *DebugString);
 		}
 
-		if(Local.DeviceDescription == DeviceDescription)
+		return Ret;
+	}
+
+	const bool bIterateAvailableAndChoose = Method == 0;
+
+	if(bIterateAvailableAndChoose)
+	{
+		for(uint32 i = 0;; ++i)
 		{
-			// found the one we are searching for
-			InternalDriverVersion = Local.InternalDriverVersion;
-			UserDriverVersion = Local.UserDriverVersion;
-			DriverDate = Local.DriverDate;
-			++FoundDriverCount;
-			break;
+			// Iterate all installed display adapters
+			const FString Key = FString::Printf(TEXT("SYSTEM\\CurrentControlSet\\Control\\Class\\{4D36E968-E325-11CE-BFC1-08002BE10318}\\%04d"), i);
+		
+			FGPUDriverInfo Local;
+			GetVideoDriverDetails(Key, Local);
+
+			if(!Local.IsValid())
+			{
+				DebugString += TEXT("GetVideoDriverDetailsInvalid ");
+			}
+
+			if(Local.DeviceDescription == DeviceDescription)
+			{
+				// found the one we are searching for
+				Ret = Local;
+				++FoundDriverCount;
+				break;
+			}
 		}
 	}
 
-	// should be 1 after the code but we've seen it >1 and the code needs to be improved (seems it was from upgrading a machine to a new windows version)
+	// FoundDriverCount can be != 1, we take the primary adapater (can be from upgrading a machine to a new OS or old drivers) which also might be wrong
 	// see: http://win7settings.blogspot.com/2014/10/how-to-extract-installed-drivers-from.html
 	// https://support.microsoft.com/en-us/kb/200435
 	// http://www.experts-exchange.com/questions/10198207/Windows-NT-Display-adapter-information.html
@@ -2505,16 +2624,20 @@ void FWindowsPlatformMisc::GetGPUDriverInfo(const FString DeviceDescription, FSt
 	if(FoundDriverCount != 1)
 	{
 		// we start again, this time we only look at the primary adapter
-		InternalDriverVersion = TEXT("Unknown");
-		UserDriverVersion = TEXT("Unknown");
-		DriverDate = TEXT("Unknown");
+		Ret.InternalDriverVersion = TEXT("Unknown");
+		Ret.UserDriverVersion = TEXT("Unknown");
+		Ret.DriverDate = TEXT("Unknown");
 
-		ErrorString += FString::Printf(TEXT("FoundDriverCount:%d FallbackToPrimary "), FoundDriverCount);
+		if(bIterateAvailableAndChoose)
+		{
+			DebugString += FString::Printf(TEXT("FoundDriverCount:%d FallbackToPrimary "), FoundDriverCount);
+		}
 	
 		// Iterate all installed display adapters
 		const FString Key = TEXT("HARDWARE\\DEVICEMAP\\VIDEO");
 
 		FString DriverLocation; // e.g. HKEY_LOCAL_MACHINE\System\CurrentControlSet\Services\<videodriver>\Device0
+		// Video0 is the first logical one, not neccesarily the primary, would have to iterate multiple to get the right one (see https://support.microsoft.com/en-us/kb/102992)
 		bool bOk = FWindowsPlatformMisc::QueryRegKey(HKEY_LOCAL_MACHINE, *Key, TEXT("\\Device\\Video0"), /*out*/ DriverLocation);
 
 		if(bOk)
@@ -2532,36 +2655,36 @@ void FWindowsPlatformMisc::GetGPUDriverInfo(const FString DeviceDescription, FSt
 
 				if(!Local.IsValid())
 				{
-					ErrorString += TEXT("GetVideoDriverDetailsInvalid ");
+					DebugString += TEXT("GetVideoDriverDetailsInvalid ");
 				}
 
 				if(Local.DeviceDescription == DeviceDescription)
 				{
-					InternalDriverVersion = Local.InternalDriverVersion;
-					UserDriverVersion = Local.UserDriverVersion;
-					DriverDate = Local.DriverDate;
+					Ret = Local;
 					FoundDriverCount = 1;
 				}
 				else
 				{
-					ErrorString += TEXT("PrimaryIsNotTheChoosenAdapter ");
+					DebugString += TEXT("PrimaryIsNotTheChoosenAdapter ");
 				}
 			}
 			else
 			{
-				ErrorString += TEXT("PrimaryDriverLocationFailed ");
+				DebugString += TEXT("PrimaryDriverLocationFailed ");
 			}
 		}
 		else
 		{
-			ErrorString += TEXT("QueryForPrimaryFailed ");
+			DebugString += TEXT("QueryForPrimaryFailed ");
 		}
 	}
 
-	if(!ErrorString.IsEmpty())
+	if(!DebugString.IsEmpty())
 	{
-		UE_LOG(LogWindows, Error, TEXT("Error: %s"), *ErrorString);
+		UE_LOG(LogWindows, Log, TEXT("DebugString: %s"), *DebugString);
 	}
+
+	return Ret;
 }
 
 #include "HideWindowsPlatformTypes.h"
diff --git a/Engine/Source/Runtime/Core/Private/Windows/WindowsWindow.cpp b/Engine/Source/Runtime/Core/Private/Windows/WindowsWindow.cpp
--- a/Engine/Source/Runtime/Core/Private/Windows/WindowsWindow.cpp
+++ b/Engine/Source/Runtime/Core/Private/Windows/WindowsWindow.cpp
@@ -479,8 +479,11 @@ void FWindowsWindow::Destroy()
 	if (OLEReferenceCount > 0 && IsWindow(HWnd))
 	{
 		HRESULT Result = RevokeDragDrop(HWnd);
-		checkf(Result == S_OK, TEXT("Failed in RevokeDragDrop, error-code: %i"), Result);
-		checkf(OLEReferenceCount == 0, TEXT("Not all references to window are released, %i left"), OLEReferenceCount);
+		// If we decremented OLEReferenceCount check it for being null (shutdown)
+		if (Result == S_OK)
+		{
+			checkf(OLEReferenceCount == 0, TEXT("Not all references to window are released, %i left"), OLEReferenceCount);
+		}		
 	}
 
 	::DestroyWindow( HWnd );
diff --git a/Engine/Source/Runtime/Core/Public/CoreGlobals.h b/Engine/Source/Runtime/Core/Public/CoreGlobals.h
--- a/Engine/Source/Runtime/Core/Public/CoreGlobals.h
+++ b/Engine/Source/Runtime/Core/Public/CoreGlobals.h
@@ -176,6 +176,7 @@ extern CORE_API FString GEditorPerProjectIni;
 extern CORE_API FString GCompatIni;
 extern CORE_API FString GLightmassIni;
 extern CORE_API FString GScalabilityIni;
+extern CORE_API FString GHardwareIni;
 extern CORE_API FString GInputIni;
 extern CORE_API FString GGameIni;
 extern CORE_API FString GGameUserSettingsIni;
diff --git a/Engine/Source/Runtime/Core/Public/GenericPlatform/GenericPlatformDriver.h b/Engine/Source/Runtime/Core/Public/GenericPlatform/GenericPlatformDriver.h
new file mode 100644
--- /dev/null
+++ b/Engine/Source/Runtime/Core/Public/GenericPlatform/GenericPlatformDriver.h
@@ -0,0 +1,509 @@
+// Copyright 1998-2016 Epic Games, Inc. All Rights Reserved.
+
+#pragma once
+
+// TSize integer values separated by '.' e.g. "120.210.11.22"
+// made to compare different GPU driver versions
+// handles whitespace
+// if there are to too many numbers we take the left most ones
+// comparison operators (for driver comparison)
+template <int32 TSize>
+class FMultiInt
+{
+public:
+	// constructor
+	FMultiInt()
+	{
+		for(uint32 i = 0; i < Size; ++i)
+		{
+			Value[i] = 0;
+		}
+	}
+
+	FMultiInt(const TCHAR* In)
+	{
+		Parse(In);
+	}
+
+	// like Parse but it doesn't alter the pointer
+	void GetValue(const TCHAR* In)
+	{
+		const TCHAR* Cpy = In;
+
+		Parse(Cpy);
+	}
+	// @param In  "123.456" "132", pointer is advanced to be after the data
+	// if there isn't enough input values, the rightmost value gets set
+	void Parse(const TCHAR*& In)
+	{
+		check(In);
+		// clear
+		*this = FMultiInt();
+
+		// find out how many numbers we have
+		uint32 NumberCount = 1;
+		{
+			const TCHAR* p = In;
+
+			while(!IsSeparator(*p))
+			{
+				// '.' separates numbers
+				if(*p == TCHAR('.'))
+				{
+					++NumberCount;
+				}
+				++p;
+			}
+		}
+
+		NumberCount = FMath::Min(NumberCount, Size);
+
+		// parse the data
+		{
+			const TCHAR* p = In;
+			for(uint32 i = 0; i < NumberCount; ++i)
+			{
+				Value[Size - NumberCount + i] = FCString::Atoi(p);
+
+				// go to next '.', operator or end
+				while(!IsSeparator(*p) && *p != TCHAR('.'))
+				{
+					++p;
+				}
+
+				// jump over '.'
+				if(*p == TCHAR('.'))
+				{
+					++p;
+				}
+			}
+
+			In = p;
+		}
+	}
+
+	// the base comparison operator, we derive the others from it
+	bool operator>(const FMultiInt<TSize>& rhs) const
+	{
+		for(uint32 i = 0 ; i < Size; ++i)
+		{
+			if(Value[i] > rhs.Value[i])
+			{
+				return true; 
+			}
+			if(Value[i] < rhs.Value[i])
+			{
+				return false; 
+			}
+		}
+
+		return false;
+	}
+	// could be optimized but doesn't have to be fast
+	bool operator==(const FMultiInt<TSize>& rhs) const
+	{
+		// map to existing operators
+		return !(*this > rhs) && !(rhs > *this);
+	}
+	// could be optimized but doesn't have to be fast
+	bool operator<(const FMultiInt<TSize>& rhs) const
+	{
+		// map to existing operators
+		return rhs > *this;
+	}
+	// could be optimized but doesn't have to be fast
+	bool operator!=(const FMultiInt<TSize>& rhs) const
+	{
+		// map to existing operators
+		return !(*this == rhs);
+	}
+	// could be optimized but doesn't have to be fast
+	bool operator>=(const FMultiInt<TSize>& rhs) const
+	{
+		// map to existing operators
+		return (*this == rhs) || (*this > rhs);
+	}
+	// could be optimized but doesn't have to be fast
+	bool operator<=(const FMultiInt<TSize>& rhs) const
+	{
+		// map to existing operators
+		return (*this == rhs) || (*this < rhs);
+	}
+
+	static const uint32 Size = TSize;
+
+	// [0]:left .. [Size-1]:right
+	uint32 Value[Size];
+
+private:
+	bool IsSeparator(const TCHAR c)
+	{
+		// comparison operators are considered separators (e.g. "1.21 < 0.121")
+		return c == 0
+			|| c == TCHAR('=')
+			|| c == TCHAR('!')
+			|| c == TCHAR('<')
+			|| c == TCHAR('&')
+			|| c == TCHAR('|')
+			|| c == TCHAR('>');
+	}
+};
+
+// used to compare driver versions in Hardware.ini
+enum EComparisonOp
+{
+	ECO_Unknown, ECO_Equal, ECO_NotEqual, ECO_Larger, ECO_LargerThan, ECO_Smaller, ECO_SmallerThan
+};
+
+// Find the comparison enum based on input text
+// @param In pointer is advanced to be after the data
+inline EComparisonOp ParseComparisonOp(const TCHAR*& In)
+{
+	if(In[0] == '=' && In[1] == '=')
+	{
+		In += 2;
+		return ECO_Equal;
+	}
+	if(In[0] == '!' && In[1] == '=')
+	{
+		In += 2;
+		return ECO_NotEqual;
+	}
+	if(*In == '>')
+	{
+		++In;
+		if(*In == '=')
+		{
+			++In;
+			return ECO_LargerThan;
+		}
+		return ECO_Larger;
+	}
+	if(*In == '<')
+	{
+		++In;
+		if(*In == '=')
+		{
+			++In;
+			return ECO_SmallerThan;
+		}
+		return ECO_Smaller;
+	}
+
+	return ECO_Unknown;
+}
+
+// general comparison with the comparison operator given as enum
+template <class T> bool Compare(const T& A, EComparisonOp Op, const T& B)
+{
+	switch(Op)
+	{
+		case ECO_Equal:			return A == B;
+		case ECO_NotEqual:		return A != B;
+		case ECO_Larger:		return A >  B;
+		case ECO_LargerThan:	return A >= B;
+		case ECO_Smaller:		return A <  B;
+		case ECO_SmallerThan:	return A <= B;
+	}
+	check(0);
+	return false;
+}
+
+// useful to express very simple math, later we might extend to express a range e.g. ">10 && <12.121"
+// @param InOpWithVersion e.g. "<=220.2"
+// @param CurrentVersion e.g. "219.1"
+inline bool CompareStringOp(const TCHAR* InOpWithMultiInt, const TCHAR* CurrentMultiInt)
+{
+	const TCHAR* p = InOpWithMultiInt;
+
+	EComparisonOp Op = ParseComparisonOp(p);
+
+	// it makes sense to compare for equal if there is no operator
+	if(Op == ECO_Unknown)
+	{
+		Op = ECO_Equal;
+	}
+
+	FMultiInt<6> A, B;
+		
+	A.GetValue(CurrentMultiInt);
+	B.Parse(p);
+
+	return Compare(A, Op, B);
+}
+
+
+// video driver details
+struct FGPUDriverInfo
+{
+	FGPUDriverInfo()
+		: VendorId(0)
+	{
+	}
+
+	// DirectX VendorId, 0 if not set, use functions below to set/get
+	uint32 VendorId;
+	// e.g. "NVIDIA GeForce GTX 680" or "AMD Radeon R9 200 / HD 7900 Series"
+	FString DeviceDescription;
+	// e.g. "NVIDIA" or "Advanced Micro Devices, Inc."
+	FString ProviderName;
+	// e.g. "15.200.1062.1004"(AMD)
+	// e.g. "9.18.13.4788"(NVIDIA) first number is Windows version (e.g. 7:Vista, 6:XP, 4:Me, 9:Win8(1), 10:Win7), last 5 have the UserDriver version encoded
+	// also called technical version number (https://wiki.mozilla.org/Blocklisting/Blocked_Graphics_Drivers)
+	// TEXT("Unknown") if driver detection failed
+	FString InternalDriverVersion;	
+	// e.g. "Catalyst 15.7.1"(AMD) or "Crimson 15.7.1"(AMD) or "347.88"(NVIDIA)
+	// also called commercial version number (https://wiki.mozilla.org/Blocklisting/Blocked_Graphics_Drivers)
+	FString UserDriverVersion;
+	// e.g. 3-13-2015
+	FString DriverDate;
+
+	bool IsValid() const
+	{
+		return !DeviceDescription.IsEmpty()
+			&& VendorId
+			&& (InternalDriverVersion != TEXT("Unknown"))		// if driver detection code fails
+			&& (InternalDriverVersion != TEXT(""));				// if running on non Windows platform we don't fill in the driver version, later we need to check for the OS as well.
+	}
+	
+	// set VendorId
+	void SetAMD() { VendorId = 0x1002; }
+	// set VendorId
+	void SetIntel() { VendorId = 0x8086; }
+	// set VendorId
+	void SetNVIDIA() { VendorId = 0x10DE; }
+	// get VendorId
+	bool IsAMD() const { return VendorId == 0x1002; }
+	// get VendorId
+	bool IsIntel() const { return VendorId == 0x8086; }
+	// get VendorId
+	bool IsNVIDIA() const { return VendorId == 0x10DE; }
+
+	FString GetUnifiedDriverVersion() const
+	{
+		// we use the internal version, not the user version to avoid problem where the name was altered 
+		const FString& FullVersion = InternalDriverVersion;
+
+		if(IsNVIDIA())
+		{
+			// on the internal driver number: https://forums.geforce.com/default/topic/378546/confusion-over-driver-version-numbers/
+			//   The first 7 shows u that is a Vista driver, 6 that is an XP and 4 that is Me
+			// we don't care about the windows version so we don't look at the front part of the driver version
+			// "9.18.13.4788" -> "347.88"
+			// "10.18.13.4788" -> "347.88"
+			// the following code works with the current numbering scheme, if needed we have to update that
+
+			// we don't care about the windows version so we don't look at the front part of the driver version
+			// e.g. 36.143
+			FString RightPart = FullVersion.Right(6);
+
+			// move the dot
+			RightPart = RightPart.Replace(TEXT("."), TEXT(""));
+			RightPart.InsertAt(3, TEXT("."));
+			return RightPart;
+		}
+		else if(IsAMD())
+		{
+			// examples for AMD: "13.12" "15.101.1007" "13.351"
+		}
+		else if(IsIntel())
+		{
+		}
+		return FullVersion;
+	}
+};
+
+// one entry in the Hardware.ini file
+struct FBlackListEntry
+{
+	// required, e.g. "<=223.112.21.1", might includes comparison operators, later even things multiple ">12.22 <=12.44"
+	FString DriverVersionString;
+	// required
+	FString Reason;
+
+	// @param e.g. "DriverVersion=\"361.43\", Reason=\"UE-25096 Viewport flashes black and white when moving in the scene on latest Nvidia drivers\""
+	void LoadFromINIString(const TCHAR* In)
+	{
+		FParse::Value(In, TEXT("DriverVersion="), DriverVersionString);
+		ensure(!DriverVersionString.IsEmpty());
+
+		// later:
+//		FParse::Value(In, TEXT("DeviceId="), DeviceId);
+//		FParse::Value(In, TEXT("OS="), OS);
+//		FParse::Value(In, TEXT("API="), API);
+//		ensure(API == TEXT("DX11"));
+
+		FParse::Value(In, TEXT("Reason="), Reason);
+		ensure(!Reason.IsEmpty());
+	}
+
+	// test if the given driver version is mentioned in the this blacklist entry
+	// @return true:yes, inform used, false otherwise
+	bool Test(const FGPUDriverInfo& Info) const
+	{
+		return CompareStringOp(*DriverVersionString, *Info.GetUnifiedDriverVersion());
+	}
+
+	bool IsValid() const
+	{
+		return !DriverVersionString.IsEmpty();
+	}
+};
+
+struct FGPUHardware
+{
+	// const as this is set in the constructor to be conveniently available
+	const FGPUDriverInfo DriverInfo;
+
+	// constructor
+	FGPUHardware(const FGPUDriverInfo InDriverInfo)
+		: DriverInfo(InDriverInfo)
+	{
+		// tests (should be very fast)
+#if !(UE_BUILD_SHIPPING || UE_BUILD_TEST)
+		{
+			FMultiInt<2> A;
+			check(A.Value[0] == 0 && A.Value[1] == 0 && A.Size == 2);
+			A.GetValue(TEXT("18.98"));
+			check(A.Value[0] == 18 && A.Value[1] == 98);
+			A.GetValue(TEXT(""));
+			check(A.Value[0] == 0 && A.Value[1] == 0);
+			A.GetValue(TEXT("98"));
+			check(A.Value[0] == 0 && A.Value[1] == 98);
+			A.GetValue(TEXT("98.34.56"));
+			check(A.Value[0] == 98 && A.Value[1] == 34);
+			A.GetValue(TEXT(" 98 . 034 "));
+			check(A.Value[0] == 98 && A.Value[1] == 34);
+			A.GetValue(TEXT("\t 98\t.\t34\t"));
+			check(A.Value[0] == 98 && A.Value[1] == 34);
+		
+			check(FMultiInt<2>(TEXT("3.07")) == FMultiInt<2>(TEXT("3.07")));
+			check(FMultiInt<2>(TEXT("3.05")) < FMultiInt<2>(TEXT("3.07")));
+			check(FMultiInt<2>(TEXT("3.05")) <= FMultiInt<2>(TEXT("3.07")));
+			check(FMultiInt<2>(TEXT("3.07")) <= FMultiInt<2>(TEXT("3.07")));
+			check(FMultiInt<2>(TEXT("3.08")) > FMultiInt<2>(TEXT("3.07")));
+			check(FMultiInt<2>(TEXT("3.08")) >= FMultiInt<2>(TEXT("3.07")));
+			check(FMultiInt<2>(TEXT("3.07")) >= FMultiInt<2>(TEXT("3.07")));
+			check(FMultiInt<2>(TEXT("3.05")) != FMultiInt<2>(TEXT("3.07")));
+			check(FMultiInt<2>(TEXT("4.05")) > FMultiInt<2>(TEXT("3.07")));
+			check(FMultiInt<2>(TEXT("4.05")) >= FMultiInt<2>(TEXT("3.07")));
+			check(FMultiInt<2>(TEXT("2.05")) < FMultiInt<2>(TEXT("3.07")));
+			check(FMultiInt<2>(TEXT("2.05")) <= FMultiInt<2>(TEXT("3.07")));
+			
+			check(Compare(10, ECO_Equal, 10));
+			check(Compare(10, ECO_NotEqual, 20));
+			check(Compare(20, ECO_Larger, 10));
+			check(Compare(20, ECO_LargerThan, 10));
+			check(Compare(10, ECO_LargerThan, 10));
+			check(Compare(10, ECO_Smaller, 20));
+			check(Compare(10, ECO_SmallerThan, 10));
+
+			check(CompareStringOp(TEXT("<20.10"), TEXT("19.12")));
+			check(CompareStringOp(TEXT("<=20.10"), TEXT("19.12")));
+			check(CompareStringOp(TEXT("<=19.12"), TEXT("19.12")));
+			check(CompareStringOp(TEXT("==19.12"), TEXT("19.12")));
+			check(CompareStringOp(TEXT(">=19.12"), TEXT("19.12")));
+			check(CompareStringOp(TEXT(">=10.12"), TEXT("19.12")));
+			check(CompareStringOp(TEXT("!=20.12"), TEXT("19.12")));
+			check(CompareStringOp(TEXT(">10.12"), TEXT("19.12")));
+
+			{
+				FGPUDriverInfo Version;
+				Version.SetNVIDIA();
+				check(Version.IsNVIDIA());
+				check(!Version.IsAMD());
+				check(!Version.IsIntel());
+				Version.InternalDriverVersion = TEXT("10.18.13.4788");
+				check(Version.GetUnifiedDriverVersion() == TEXT("347.88"));
+			}
+			{
+				FGPUDriverInfo Version;
+				Version.SetAMD();
+				check(Version.IsAMD());
+				check(!Version.IsNVIDIA());
+				check(!Version.IsIntel());
+				Version.InternalDriverVersion = TEXT("15.200.1062.1004");
+				check(Version.GetUnifiedDriverVersion() == TEXT("15.200.1062.1004"));
+			}
+			{
+				FGPUDriverInfo Version;
+				Version.SetIntel();
+				check(Version.IsIntel());
+				check(!Version.IsAMD());
+				check(!Version.IsNVIDIA());
+				Version.DeviceDescription = TEXT("Intel(R) HD Graphics 4600");
+				Version.InternalDriverVersion = TEXT("9.18.10.3310");
+				Version.DriverDate = TEXT("9-17-2013");
+				check(Version.GetUnifiedDriverVersion() == TEXT("9.18.10.3310"));
+			}
+		}
+#endif// !(UE_BUILD_SHIPPING || UE_BUILD_TEST)
+	}	
+	
+	// @return a driver version intended to be shown to the user e.g. "15.30.1025.1001 12/17/2015 (Crimson Edition 15.12)"
+	FString GetSuggestedDriverVersion() const
+	{
+		const TCHAR* Section = GetVendorSectionName();
+
+		FString Ret;
+
+		if(Section)
+		{
+			GConfig->GetString(Section, TEXT("SuggestedDriverVersion"), Ret, GHardwareIni);
+		}
+
+		return Ret;
+	}
+
+	// @return 0 if there is none
+	FBlackListEntry FindDriverBlacklistEntry() const
+	{
+		const TCHAR* Section = GetVendorSectionName();
+
+		if(Section)
+		{
+			TArray<FString> BlacklistStrings;
+			GConfig->GetArray(GetVendorSectionName(), TEXT("Blacklist"), BlacklistStrings, GHardwareIni);
+
+			for(int32 i = 0; i < BlacklistStrings.Num(); ++i)
+			{
+				FBlackListEntry Entry;
+
+				const TCHAR* Line = *BlacklistStrings[i];
+			
+				ensure(Line[0] == TCHAR('('));
+
+				Entry.LoadFromINIString(&Line[1]);
+
+				if(Entry.Test(DriverInfo))
+				{
+					return Entry;
+				}
+			}
+		}
+
+		return FBlackListEntry();
+	}
+
+	// to get a section name in the Hardware.ini file
+	// @return 0 if not found
+	const TCHAR* GetVendorSectionName() const
+	{
+		if(DriverInfo.IsNVIDIA())
+		{
+			return TEXT("GPU_NVIDIA");
+		}
+		if(DriverInfo.IsAMD())
+		{
+			return TEXT("GPU_AMD");
+		}
+		else if(DriverInfo.IsIntel())
+		{
+			return TEXT("GPU_0x8086");
+		}
+		// more GPU vendors can be added on demand
+		return 0;
+	}
+};
+
+
diff --git a/Engine/Source/Runtime/Core/Public/GenericPlatform/GenericPlatformMisc.h b/Engine/Source/Runtime/Core/Public/GenericPlatform/GenericPlatformMisc.h
--- a/Engine/Source/Runtime/Core/Public/GenericPlatform/GenericPlatformMisc.h
+++ b/Engine/Source/Runtime/Core/Public/GenericPlatform/GenericPlatformMisc.h
@@ -319,10 +319,7 @@ struct CORE_API FGenericPlatformMisc
 	 */
 	static FString GetPrimaryGPUBrand();
 
-	// @param InternalDriverVersion e.g. "15.200.1062.1004"(AMD) "9.18.13.4788" (NVIDIA) for comparison
-	// @param UserDriverVersion e.g. "15.7.1"(Catalyst AMD) "9.18.13.4788" (NVIDIA) for user feedback
-	// @param DriverDate e.g. 3-13-2015
-	static void GetGPUDriverInfo(const FString DeviceDescription, FString& InternalDriverVersion, FString& UserDriverVersion, FString& DriverDate);
+	static struct FGPUDriverInfo GetGPUDriverInfo(const FString& DeviceDescription);
 
 	/**
 	 * Gets the OS Version and OS Subversion.
@@ -942,3 +939,5 @@ protected:
 	static bool bPromptForRemoteDebugOnEnsure;
 #endif	//#if !UE_BUILD_SHIPPING
 };
+
+
diff --git a/Engine/Source/Runtime/Core/Public/GenericPlatform/GenericWindowDefinition.h b/Engine/Source/Runtime/Core/Public/GenericPlatform/GenericWindowDefinition.h
--- a/Engine/Source/Runtime/Core/Public/GenericPlatform/GenericWindowDefinition.h
+++ b/Engine/Source/Runtime/Core/Public/GenericPlatform/GenericWindowDefinition.h
@@ -5,6 +5,7 @@
 #include "Containers/UnrealString.h"
 #include "HAL/Platform.h"
 #include "GenericApplicationMessageHandler.h"
+#include "StaticArray.h"
 
 
 #ifndef ALPHA_BLENDED_WINDOWS
@@ -103,37 +104,3 @@ struct CORE_API FGenericWindowDefinition
 
 	FWindowSizeLimits SizeLimits;
 };
-
-// video driver details
-struct FGPUDriverInfo
-{
-	FGPUDriverInfo()
-		: VendorId(0)
-	{
-	}
-
-	// DirectX vendor Id, use functions below
-	uint32 VendorId;
-	// e.g. "NVIDIA GeForce GTX 680" or "AMD Radeon R9 200 / HD 7900 Series"
-	FString DeviceDescription;
-	// e.g. "NVIDIA" or "Advanced Micro Devices, Inc."
-	FString ProviderName;
-	// e.g. "15.200.1062.1004"(AMD)
-	// e.g. "9.18.13.4788"(NVIDIA) first number is Windows version (e.g. 7:Vista, 6:XP, 4:Me, 9:Win8(1), 10:Win7), last 5 have the UserDriver version encoded
-	FString InternalDriverVersion;	
-	// e.g. "Catalyst 15.7.1"(AMD) or "347.88"(NVIDIA)
-	FString UserDriverVersion;
-	// e.g. 3-13-2015
-	FString DriverDate;
-
-	bool IsValid() const
-	{
-		return !DeviceDescription.IsEmpty();
-	}
-	
-	bool IsAMD() const { return VendorId == 0x1002; }
-
-	bool IsIntel() const { return VendorId == 0x8086; }
-
-	bool IsNVIDIA() const { return VendorId == 0x10DE; }
-};
\ No newline at end of file
diff --git a/Engine/Source/Runtime/Core/Public/HAL/ExceptionHandling.h b/Engine/Source/Runtime/Core/Public/HAL/ExceptionHandling.h
--- a/Engine/Source/Runtime/Core/Public/HAL/ExceptionHandling.h
+++ b/Engine/Source/Runtime/Core/Public/HAL/ExceptionHandling.h
@@ -27,6 +27,5 @@ extern CORE_API void NewReportEnsure( const TCHAR* ErrorMessage );
 extern CORE_API int32 ReportCrash( ucontext_t *Context, int32 Signal, struct __siginfo* Info );
 extern CORE_API void NewReportEnsure( const TCHAR* ErrorMessage );
 #elif PLATFORM_LINUX
-extern CORE_API int32 ReportCrash( ucontext_t *Context, int32 Signal, siginfo_t* Info );
 extern CORE_API void NewReportEnsure( const TCHAR* ErrorMessage );
 #endif
diff --git a/Engine/Source/Runtime/Core/Public/Linux/LinuxPlatformCrashContext.h b/Engine/Source/Runtime/Core/Public/Linux/LinuxPlatformCrashContext.h
--- a/Engine/Source/Runtime/Core/Public/Linux/LinuxPlatformCrashContext.h
+++ b/Engine/Source/Runtime/Core/Public/Linux/LinuxPlatformCrashContext.h
@@ -18,6 +18,9 @@ struct CORE_API FLinuxCrashContext : public FGenericCrashContext
 	/** Thread context */
 	ucontext_t*	Context;
 
+	/** Whether backtrace was already captured */
+	bool bCapturedBacktrace;
+
 	/** Symbols received via backtrace_symbols(), if any (note that we will need to clean it up) */
 	char ** BacktraceSymbols;
 
@@ -29,9 +32,10 @@ struct CORE_API FLinuxCrashContext : public FGenericCrashContext
 
 	FLinuxCrashContext()
 		:	Signal(0)
-		,	Info(NULL)
-		,	Context(NULL)
-		,	BacktraceSymbols(NULL)
+		,	Info(nullptr)
+		,	Context(nullptr)
+		,	bCapturedBacktrace(false)
+		,	BacktraceSymbols(nullptr)
 	{
 		SignalDescription[ 0 ] = 0;
 		MinidumpCallstackInfo[ 0 ] = 0;
@@ -48,6 +52,22 @@ struct CORE_API FLinuxCrashContext : public FGenericCrashContext
 	 */
 	void InitFromSignal(int32 InSignal, siginfo_t* InInfo, void* InContext);
 
+	/**
+	 * Populates crash context stack trace and a few related fields
+	 *
+	 */
+	void CaptureStackTrace();
+
+	/**
+	 * Generates a new crash report containing information needed for the crash reporter and launches it; may not return.
+	 *
+	 * @param bReportingNonCrash if true, we are not reporting a crash (but e.g. ensure()), so don't re-raise the signal.
+	 *
+	 * @return If bReportingNonCrash is false, the function will not return
+	 */
+	void GenerateCrashInfoAndLaunchReporter(bool bReportingNonCrash = false) const;
+
+protected:
 	/**
 	 * Dumps all the data from crash context to the "minidump" report.
 	 *
diff --git a/Engine/Source/Runtime/Core/Public/Math/UnrealMathUtility.h b/Engine/Source/Runtime/Core/Public/Math/UnrealMathUtility.h
--- a/Engine/Source/Runtime/Core/Public/Math/UnrealMathUtility.h
+++ b/Engine/Source/Runtime/Core/Public/Math/UnrealMathUtility.h
@@ -93,8 +93,9 @@ struct FMath : public FPlatformMath
 	/** Helper function for rand implementations. Returns a random number in [0..A) */
 	static FORCEINLINE int32 RandHelper(int32 A)
 	{
-		// RAND_MAX+1 give interval [0..A) with even distribution.
-		return A>0 ? TruncToInt(Rand()/(float)((uint32)RAND_MAX+1) * A) : 0;
+		// Note that on some platforms RAND_MAX is a large number so we cannot do ((rand()/(RAND_MAX+1)) * A)
+		// or else we may include the upper bound results, which should be excluded.
+		return A > 0 ? Min(TruncToInt(FRand() * A), A - 1) : 0;
 	}
 
 	/** Helper function for rand implementations. Returns a random number >= Min and <= Max */
diff --git a/Engine/Source/Runtime/Core/Public/Misc/App.h b/Engine/Source/Runtime/Core/Public/Misc/App.h
--- a/Engine/Source/Runtime/Core/Public/Misc/App.h
+++ b/Engine/Source/Runtime/Core/Public/Misc/App.h
@@ -28,6 +28,15 @@ public:
 	 */
 	static EBuildConfigurations::Type GetBuildConfiguration();
 
+	/**
+	 * Gets the deployment name (also known as "EpicApp" in the launcher), e.g. DevPlaytest, PublicTest, Live etc.
+	 *
+	 * Does not return FString because it can be used during crash handling, so it should avoid memory allocation.
+	 *
+	 * @return The branch name (guaranteed not to be nullptr).
+	 */
+	static const TCHAR * GetDeploymentName();
+
 	/**
 	 * Gets the date at which this application was built.
 	 *
diff --git a/Engine/Source/Runtime/Core/Public/Windows/WindowsPlatformMisc.h b/Engine/Source/Runtime/Core/Public/Windows/WindowsPlatformMisc.h
--- a/Engine/Source/Runtime/Core/Public/Windows/WindowsPlatformMisc.h
+++ b/Engine/Source/Runtime/Core/Public/Windows/WindowsPlatformMisc.h
@@ -196,7 +196,7 @@ struct CORE_API FWindowsPlatformMisc
 	static FString GetCPUVendor();
 	static FString GetCPUBrand();
 	static FString GetPrimaryGPUBrand();
-	static void GetGPUDriverInfo(const FString DeviceDescription, FString& InternalDriverVersion, FString& UserDriverVersion, FString& DriverDate);
+	static struct FGPUDriverInfo GetGPUDriverInfo(const FString& DeviceDescription);
 	static void GetOSVersions( FString& out_OSVersionLabel, FString& out_OSSubVersionLabel );
 	static bool GetDiskTotalAndFreeSpace( const FString& InPath, uint64& TotalNumberOfBytes, uint64& NumberOfFreeBytes );
 
diff --git a/Engine/Source/Runtime/Engine/Classes/Components/ReflectionCaptureComponent.h b/Engine/Source/Runtime/Engine/Classes/Components/ReflectionCaptureComponent.h
--- a/Engine/Source/Runtime/Engine/Classes/Components/ReflectionCaptureComponent.h
+++ b/Engine/Source/Runtime/Engine/Classes/Components/ReflectionCaptureComponent.h
@@ -129,6 +129,10 @@ class UReflectionCaptureComponent : public USceneComponent
 	virtual void BeginDestroy() override;
 	virtual bool IsReadyForFinishDestroy() override;
 	virtual void FinishDestroy() override;
+	virtual bool NeedsLoadForServer() const override
+	{
+		return false;
+	}
 	//~ End UObject Interface
 
 private:
diff --git a/Engine/Source/Runtime/Engine/Classes/Engine/Canvas.h b/Engine/Source/Runtime/Engine/Classes/Engine/Canvas.h
--- a/Engine/Source/Runtime/Engine/Classes/Engine/Canvas.h
+++ b/Engine/Source/Runtime/Engine/Classes/Engine/Canvas.h
@@ -98,9 +98,11 @@ public:
 
 	void AddColumnIfNeeded();
 
+	float GetTextScale() const;
+
 	float GetYStep() const
 	{
-		return MaxCharHeight * 1.15f;
+		return MaxCharHeight * 1.15f * GetTextScale();
 	}
 
 	float GetXPos() const
diff --git a/Engine/Source/Runtime/Engine/Classes/Engine/DemoNetDriver.h b/Engine/Source/Runtime/Engine/Classes/Engine/DemoNetDriver.h
--- a/Engine/Source/Runtime/Engine/Classes/Engine/DemoNetDriver.h
+++ b/Engine/Source/Runtime/Engine/Classes/Engine/DemoNetDriver.h
@@ -105,6 +105,12 @@ private:
 	/** Set via GotoTimeInSeconds, only fired once (at most). Called for successful or failed scrub. */
 	FOnGotoTimeDelegate OnGotoTimeDelegate_Transient;
 	
+	/** Saved server time after loading a checkpoint, so that we can set the server time as accurately as possible after the fast-forward */
+	float SavedReplicatedWorldTimeSeconds;
+
+	/** Saved fast-forward time, used for correcting world time after the fast-forward is complete */
+	float SavedSecondsToSkip;
+
 public:
 
 	// UNetDriver interface.
diff --git a/Engine/Source/Runtime/Engine/Classes/Engine/NetDriver.h b/Engine/Source/Runtime/Engine/Classes/Engine/NetDriver.h
--- a/Engine/Source/Runtime/Engine/Classes/Engine/NetDriver.h
+++ b/Engine/Source/Runtime/Engine/Classes/Engine/NetDriver.h
@@ -266,7 +266,8 @@ public:
 	double						StatUpdateTime;
 	/** Interval between gathering stats */
 	float						StatPeriod;
-
+	/** Time of last netdriver cleanup pass */
+	double						LastCleanupTime;
 	/** Used to determine if checking for standby cheats should occur */
 	bool						bIsStandbyCheckingEnabled;
 	/** Used to determine whether we've already caught a cheat or not */
diff --git a/Engine/Source/Runtime/Engine/Classes/GameFramework/CharacterMovementComponent.h b/Engine/Source/Runtime/Engine/Classes/GameFramework/CharacterMovementComponent.h
--- a/Engine/Source/Runtime/Engine/Classes/GameFramework/CharacterMovementComponent.h
+++ b/Engine/Source/Runtime/Engine/Classes/GameFramework/CharacterMovementComponent.h
@@ -1870,7 +1870,7 @@ public:
 		It will also handle TimeStamp resets if it detects a gap larger than MinTimeBetweenTimeStampResets / 2.f
 		!! ServerData.CurrentClientTimeStamp can be reset !!
 		@returns true if TimeStamp is valid, or false if it has expired. */
-	bool VerifyClientTimeStamp(float TimeStamp, FNetworkPredictionData_Server_Character & ServerData);
+	virtual bool VerifyClientTimeStamp(float TimeStamp, FNetworkPredictionData_Server_Character & ServerData);
 protected:
 	/** Internal const check for client timestamp validity without side-effects. 
 	  * @see VerifyClientTimeStamp */
diff --git a/Engine/Source/Runtime/Engine/Classes/GameFramework/GameState.h b/Engine/Source/Runtime/Engine/Classes/GameFramework/GameState.h
--- a/Engine/Source/Runtime/Engine/Classes/GameFramework/GameState.h
+++ b/Engine/Source/Runtime/Engine/Classes/GameFramework/GameState.h
@@ -152,6 +152,8 @@ protected:
 private:
 	// Hidden functions that don't make sense to use on this class.
 	HIDE_ACTOR_TRANSFORM_FUNCTIONS();
+
+	friend class UDemoNetDriver;
 };
 
 
diff --git a/Engine/Source/Runtime/Engine/Classes/GameFramework/OnlineReplStructs.h b/Engine/Source/Runtime/Engine/Classes/GameFramework/OnlineReplStructs.h
--- a/Engine/Source/Runtime/Engine/Classes/GameFramework/OnlineReplStructs.h
+++ b/Engine/Source/Runtime/Engine/Classes/GameFramework/OnlineReplStructs.h
@@ -9,6 +9,7 @@
 #include "Runtime/Online/OnlineSubsystem/Public/UniqueNetIdWrapper.h"
 #include "OnlineReplStructs.generated.h"
 
+class FJsonValue;
 class FUniqueNetId;
 
 /**
@@ -72,6 +73,9 @@ struct FUniqueNetIdRepl
     /** Export contents of this struct as a string */
 	bool ExportTextItem(FString& ValueStr, FUniqueNetIdRepl const& DefaultValue, UObject* Parent, int32 PortFlags, UObject* ExportRootScope) const;
 
+	/** Import string contexts and try to map them into a unique id */
+	bool ImportTextItem(const TCHAR*& Buffer, int32 PortFlags, UObject* Parent, FOutputDevice* ErrorText);
+
 	/** Network serialization */
 	ENGINE_API bool NetSerialize(FArchive& Ar, class UPackageMap* Map, bool& bOutSuccess);
 
@@ -87,6 +91,11 @@ struct FUniqueNetIdRepl
 		return IsValid() ? UniqueNetId->ToString() : TEXT("INVALID");
 	}
 
+	/** Convert this unique id to a json value */
+	TSharedRef<FJsonValue> ToJson() const;
+	/** Create a unique id from a json string */
+	void FromJson(const FString& InValue);
+
 	/** Is the FUniqueNetId wrapped in this object valid */
 	bool IsValid() const
 	{
@@ -145,6 +154,9 @@ struct FUniqueNetIdRepl
 
 protected:
 	TSharedPtr<const FUniqueNetId> UniqueNetId;
+
+	/** Helper to create an FUniqueNetId from a string */
+	void UniqueIdFromString(const FString& Contents);
 };
 
 /** Specify net delta serializer support for the active skill cooldown array */
@@ -163,6 +175,8 @@ struct TStructOpsTypeTraits<FUniqueNetIdRepl> : public TStructOpsTypeTraitsBase
 		WithIdenticalViaEquality = true,
 		// Export contents of this struct as a string (displayall, obj dump, etc)
 		WithExportTextItem = true,
+		// Import string contents as a unique id
+		WithImportTextItem = true
 	};
 };
 
diff --git a/Engine/Source/Runtime/Engine/Classes/GameFramework/PlayerState.h b/Engine/Source/Runtime/Engine/Classes/GameFramework/PlayerState.h
--- a/Engine/Source/Runtime/Engine/Classes/GameFramework/PlayerState.h
+++ b/Engine/Source/Runtime/Engine/Classes/GameFramework/PlayerState.h
@@ -90,6 +90,7 @@ class ENGINE_API APlayerState : public AInfo
 	float ExactPing;
 
 	/** Used to match up InactivePlayerState with rejoining playercontroller. */
+	UPROPERTY()
 	FString SavedNetworkAddress;
 
 	/** The id used by the network to uniquely identify a player.
diff --git a/Engine/Source/Runtime/Engine/Classes/GameFramework/RootMotionSource.h b/Engine/Source/Runtime/Engine/Classes/GameFramework/RootMotionSource.h
--- a/Engine/Source/Runtime/Engine/Classes/GameFramework/RootMotionSource.h
+++ b/Engine/Source/Runtime/Engine/Classes/GameFramework/RootMotionSource.h
@@ -44,7 +44,12 @@ enum class ERootMotionSourceSettingsFlags : uint8
 	// Source will switch character to Falling mode with any "Z up" velocity added.
 	// Use this for jump-like root motion. If not enabled, uses default jump impulse
 	// detection (which keeps you stuck on ground in Walking fairly strongly)
-	UseSensitiveLiftoffCheck	= 0x01
+	UseSensitiveLiftoffCheck	= 0x01,
+	// If Duration of Source would end partway through the last tick it is active,
+	// do not reduce SimulationTime. Disabling this is useful for sources that
+	// are more about providing velocity (like jumps), vs. sources that need
+	// the precision of partial ticks for say ending up at an exact location (MoveTo)
+	DisablePartialEndTick		= 0x02
 };
 
 enum class ERootMotionSourceID : uint16 { Invalid = 0 };
diff --git a/Engine/Source/Runtime/Engine/Classes/Materials/MaterialExpression.h b/Engine/Source/Runtime/Engine/Classes/Materials/MaterialExpression.h
--- a/Engine/Source/Runtime/Engine/Classes/Materials/MaterialExpression.h
+++ b/Engine/Source/Runtime/Engine/Classes/Materials/MaterialExpression.h
@@ -187,6 +187,11 @@ class ENGINE_API UMaterialExpression : public UObject
 
 	virtual bool Modify( bool bAlwaysMarkDirty=true ) override;
 	virtual void Serialize( FArchive& Ar ) override;
+	virtual bool NeedsLoadForServer() const override
+	{
+		// [RCL] FIXME 2016-02-01: should not be loaded for servers, but currently causes problems with GC clusters (see OR-14209).
+		return Super::NeedsLoadForServer();
+	}
 	//~ End UObject Interface.
 
 	/**
