Copying //UE4/Dev-Core to //UE4/Main

==========================
MAJOR FEATURES + CHANGES
==========================

Change 2855309 on 2016/02/04 by [REDACTED_PERSON][REDACTED_HANDLE]

	GC Cluster improvements:

	- Moved cluster creation from PostLoad to after the entire package has been loaded to make sure all referenced objects are accounted for
	- Added better support for objects that can't be added to a cluster for some reason (lazy loaded referenced etc)
	- Added -gcnoclusters option to 'obj list' to list all GC objects excluding objects in clusters
	- Reference chain search will now print info on clusters

Change 2854937 on 2016/02/03 by [REDACTED_PERSON][REDACTED_HANDLE]

	Pool headers now contain the small block size and lookups are done using that.
	Size->table array is now one-based indexed to allow allocations of the maximum block size (32768) to be binned as expected.
	Misc refactors.

Change 2851056 on 2016/02/01 by [REDACTED_PERSON][REDACTED_HANDLE]

	Disregard for GC improvements: added the ability to re-open disregard for GC set at any time and keep adding objects to it.

	+ Maybe fixed startup packages loading code: load startup packages AFTER initial moduels have been loaded, don't use the seek free flag.

Change 2848788 on 2016/01/29 by [REDACTED_PERSON][REDACTED_HANDLE]

	ALIGNOF now supports abstract types.

Change 2852323 on 2016/02/02 by [REDACTED_PERSON][REDACTED_HANDLE]

	GC Clusters: Making sure VerifyClusterAssumptions catches all cases:

	- Invalid objects in GC
	- Objects from other clusters that don't have their parent cluster referenced by the source cluster

Change 2846860 on 2016/01/28 by [REDACTED_PERSON][REDACTED_HANDLE]

	Added extra debug info to error messages sent back in the response from the AddCrash method
	[OPP-4995] - Crash reports are not appearing for Portal Main branch

Change 2846864 on 2016/01/28 by [REDACTED_PERSON][REDACTED_HANDLE]

	Added retry to AddCrash web requests due to unexplained missing payloads in some crashes
	[OPP-4995] - Crash reports are not appearing for Portal Main branch

Change 2846875 on 2016/01/28 by [REDACTED_PERSON][REDACTED_HANDLE]

	Ensure that streams are properly closed when sending and receiving data in SimpleWebRequest

	This is to resolve an intermittent issue whereby CrashReporterWebsite sometimes receives no payload on an HTTP POST request from CrashReportProcessor. This could be consistent with the request stream not being closed and flushed before the HTTP request is sent.  MSDN documentation for GetRequestStream states "You must call the Stream.Close method to close the stream and release the connection for reuse. Failure to close the stream causes your application to run out of connections."  Note that Close() is automatically called on a Stream during disposal, and Dispose() is automatically called when exiting the scope of a using statement.
	https://msdn.microsoft.com/en-us/library/d4cek6cc(v=vs.110).aspx

Change 2847012 on 2016/01/28 by [REDACTED_PERSON][REDACTED_HANDLE]

	Allocations with alignment > OSPageSize are now illegal.
	Removal of trailing pool support.

Change 2847106 on 2016/01/28 by [REDACTED_PERSON][REDACTED_HANDLE]

	AllocateBlockFromPool inlined and redundancy removed.
	Misc tidy-ups.

Change 2849028 on 2016/01/29 by [REDACTED_PERSON][REDACTED_HANDLE]

	TIsClass fix for union types.

Change 2849039 on 2016/01/29 by [REDACTED_PERSON][REDACTED_HANDLE]

	Basic custom deleter support for TUniquePtr.

Change 2850853 on 2016/02/01 by [REDACTED_PERSON][REDACTED_HANDLE]

	Added per branch filtering for Reports output page

Change 2851035 on 2016/02/01 by [REDACTED_PERSON][REDACTED_HANDLE]

	Generating Platform and Platform Variant strings in the Crash Report Process to pass to the MDD
	[OPP-4792] - Crash reports not appearing for launcher on Windows 32-bit OS

Change 2853124 on 2016/02/02 by [REDACTED_PERSON][REDACTED_HANDLE]

	Changed crash reporter to respect command line parameter filtering when that is enabled via the project define (UECORE-318)

Change 2853960 on 2016/02/03 by [REDACTED_PERSON][REDACTED_HANDLE]

	Limiting the size of UObject Array for programs.

Change 2855582 on 2016/02/04 by [REDACTED_PERSON][REDACTED_HANDLE]

	Revert to MallocBinned for non-Windows and non-PS4 platforms, as it requires a 64K page size.

Change 2856837 on 2016/02/05 by [REDACTED_PERSON][REDACTED_HANDLE]

	Changing FUObjectItem::ThisThreadAtomicallyClearedFlag to better support changes to multiple different flags from different threads. Made sure NoStrongReference is also cleared atomically.

Change 2857113 on 2016/02/05 by [REDACTED_PERSON][REDACTED_HANDLE]

	Exhausted check moved out of Allocate function.

Change 2858600 on 2016/02/08 by [REDACTED_PERSON][REDACTED_HANDLE]

	Preventing async loading suspended log spam in the editor.

Change 2858700 on 2016/02/08 by [REDACTED_PERSON][REDACTED_HANDLE]

	Code for short block support.

Change 2859112 on 2016/02/08 by [REDACTED_PERSON][REDACTED_HANDLE]

	Array of pool table pointers changed to array of indices.

Change 2859138 on 2016/02/08 by [REDACTED_PERSON][REDACTED_HANDLE]

	Index map moved into a global.

diff --git a/Engine/Source/Programs/CrashReporter/CrashReportProcess/ReportProcessor.cs b/Engine/Source/Programs/CrashReporter/CrashReportProcess/ReportProcessor.cs
--- a/Engine/Source/Programs/CrashReporter/CrashReportProcess/ReportProcessor.cs
+++ b/Engine/Source/Programs/CrashReporter/CrashReportProcess/ReportProcessor.cs
@@ -379,19 +379,28 @@ namespace Tools.CrashReporter.CrashReportProcess
 						RequestString = "http://localhost:80/Crashes/AddCrash/-1"; 
 					}
 
-					string ResponseString = SimpleWebRequest.GetWebServiceResponse( RequestString, Payload );
-					if( ResponseString.Length > 0 )
+					string ErrorMessage = string.Empty;
+
+					for (int Retry = 0; Retry < 3; ++Retry)
 					{
-						// Convert response into a string
-						CrashReporterResult Result = XmlHandler.FromXmlString<CrashReporterResult>( ResponseString );
-						if( Result.ID > 0 )
-						{
-							NewID = Result.ID;
-						}
-						else
+						string ResponseString = SimpleWebRequest.GetWebServiceResponse(RequestString, Payload);
+						if (ResponseString.Length > 0)
 						{
-							CrashReporterProcessServicer.WriteFailure("UploadCrash: " + Result.Message);
+							// Convert response into a string
+							CrashReporterResult Result = XmlHandler.FromXmlString<CrashReporterResult>(ResponseString);
+							if (Result.ID > 0)
+							{
+								NewID = Result.ID;
+								break;
+							}
+							ErrorMessage = Result.Message;
 						}
+						Thread.Sleep(200);
+					}
+
+					if (NewID == -1)
+					{
+						CrashReporterProcessServicer.WriteFailure("UploadCrash: " + ErrorMessage);
 					}
 				}
 			}
@@ -706,7 +715,23 @@ namespace Tools.CrashReporter.CrashReportProcess
 			// Purge logs every 2048 processed crashes
 			string PurgeLogsDays = ProcessedReports % 2048 == 0 ? "2" : "-1";
 
-			FEngineVersion EngineVersion = new FEngineVersion( NewContext.PrimaryCrashProperties.EngineVersion ); 
+			FEngineVersion EngineVersion = new FEngineVersion( NewContext.PrimaryCrashProperties.EngineVersion );
+
+			// Pass Windows variants (Win32/64) to MinidumpDiagnostics
+			string PlatformVariant = NewContext.PrimaryCrashProperties.PlatformName;
+			if (PlatformVariant != null && NewContext.PrimaryCrashProperties.PlatformFullName != null && PlatformVariant.ToUpper().Contains("WINDOWS"))
+			{
+				if (NewContext.PrimaryCrashProperties.PlatformFullName.Contains("Win32") ||
+				    NewContext.PrimaryCrashProperties.PlatformFullName.Contains("32b"))
+				{
+					PlatformVariant = "Win32";
+				}
+				else if (NewContext.PrimaryCrashProperties.PlatformFullName.Contains("Win64") ||
+						NewContext.PrimaryCrashProperties.PlatformFullName.Contains("64b"))
+				{
+					PlatformVariant = "Win64";
+				}
+			}
 
 			List<string> MinidumpDiagnosticsParams = new List<string>
 			(
@@ -717,6 +742,8 @@ namespace Tools.CrashReporter.CrashReportProcess
 					"-BuiltFromCL=" + EngineVersion.Changelist,		// Backward compatibility
 					"-GameName=" + NewContext.PrimaryCrashProperties.GameName,
 					"-EngineVersion=" + NewContext.PrimaryCrashProperties.EngineVersion,
+					"-PlatformName=" + NewContext.PrimaryCrashProperties.PlatformName,
+					"-PlatformVariantName=" + PlatformVariant,
 					"-bUsePDBCache=true",
 					"-Annotate",
 					"-SyncSymbols",
diff --git a/Engine/Source/Programs/CrashReporter/CrashReportProcess/ReportQueue.cs b/Engine/Source/Programs/CrashReporter/CrashReportProcess/ReportQueue.cs
--- a/Engine/Source/Programs/CrashReporter/CrashReportProcess/ReportQueue.cs
+++ b/Engine/Source/Programs/CrashReporter/CrashReportProcess/ReportQueue.cs
@@ -112,30 +112,33 @@ namespace Tools.CrashReporter.CrashReportProcess
 		}
 
 		/// <summary> Converts WER metadata xml file to the crash context. </summary>
-		private FGenericCrashContext ConvertMetadataToCrashContext( WERReportMetadata Metadata, string NewReportPath )
+		private void ConvertMetadataToCrashContext(WERReportMetadata Metadata, string NewReportPath, ref FGenericCrashContext OutCrashContext)
 		{
-			FGenericCrashContext CrashContext = new FGenericCrashContext();
+			if (OutCrashContext == null)
+			{
+				OutCrashContext = new FGenericCrashContext();
+			}
 
 			FReportData ReportData = new FReportData( Metadata, NewReportPath );
 
-			CrashContext.PrimaryCrashProperties.CrashVersion = (int)ECrashDescVersions.VER_1_NewCrashFormat;
-
-			CrashContext.PrimaryCrashProperties.ProcessId = 0;
-			CrashContext.PrimaryCrashProperties.CrashGUID = new DirectoryInfo( NewReportPath ).Name;  
-// 			CrashContext.PrimaryCrashProperties.IsInternalBuild
-// 			CrashContext.PrimaryCrashProperties.IsPerforceBuild
-// 			CrashContext.PrimaryCrashProperties.IsSourceDistribution
-// 			CrashContext.PrimaryCrashProperties.SecondsSinceStart
-			CrashContext.PrimaryCrashProperties.GameName = ReportData.GameName;
-// 			CrashContext.PrimaryCrashProperties.ExecutableName
-// 			CrashContext.PrimaryCrashProperties.BuildConfiguration
-// 			CrashContext.PrimaryCrashProperties.PlatformName
-// 			CrashContext.PrimaryCrashProperties.PlatformNameIni
-			CrashContext.PrimaryCrashProperties.PlatformFullName = ReportData.Platform;
-			CrashContext.PrimaryCrashProperties.EngineMode = ReportData.EngineMode;
-			CrashContext.PrimaryCrashProperties.EngineVersion = ReportData.GetEngineVersion();
-			CrashContext.PrimaryCrashProperties.CommandLine = ReportData.CommandLine;
-// 			CrashContext.PrimaryCrashProperties.LanguageLCID
+			OutCrashContext.PrimaryCrashProperties.CrashVersion = (int)ECrashDescVersions.VER_1_NewCrashFormat;
+
+			//OutCrashContext.PrimaryCrashProperties.ProcessId = 0; don't overwrite valid ids, zero is default anyway
+			OutCrashContext.PrimaryCrashProperties.CrashGUID = new DirectoryInfo( NewReportPath ).Name;  
+// 			OutCrashContext.PrimaryCrashProperties.IsInternalBuild
+// 			OutCrashContext.PrimaryCrashProperties.IsPerforceBuild
+// 			OutCrashContext.PrimaryCrashProperties.IsSourceDistribution
+// 			OutCrashContext.PrimaryCrashProperties.SecondsSinceStart
+			OutCrashContext.PrimaryCrashProperties.GameName = ReportData.GameName;
+// 			OutCrashContext.PrimaryCrashProperties.ExecutableName
+// 			OutCrashContext.PrimaryCrashProperties.BuildConfiguration
+// 			OutCrashContext.PrimaryCrashProperties.PlatformName
+// 			OutCrashContext.PrimaryCrashProperties.PlatformNameIni
+			OutCrashContext.PrimaryCrashProperties.PlatformFullName = ReportData.Platform;
+			OutCrashContext.PrimaryCrashProperties.EngineMode = ReportData.EngineMode;
+			OutCrashContext.PrimaryCrashProperties.EngineVersion = ReportData.GetEngineVersion();
+			OutCrashContext.PrimaryCrashProperties.CommandLine = ReportData.CommandLine;
+// 			OutCrashContext.PrimaryCrashProperties.LanguageLCID
 
 			// Create a locate and get the language.
 			int LanguageCode = 0;
@@ -145,54 +148,66 @@ namespace Tools.CrashReporter.CrashReportProcess
 				if (LanguageCode > 0)
 				{
 					CultureInfo LanguageCI = new CultureInfo( LanguageCode );
-					CrashContext.PrimaryCrashProperties.AppDefaultLocale = LanguageCI.Name;
+					OutCrashContext.PrimaryCrashProperties.AppDefaultLocale = LanguageCI.Name;
 				}
 			}
-			catch (System.Exception)
+			catch (Exception)
+			{
+				OutCrashContext.PrimaryCrashProperties.AppDefaultLocale = null;
+			}
+
+			if (string.IsNullOrEmpty(OutCrashContext.PrimaryCrashProperties.AppDefaultLocale))
 			{
 				// Default to en-US
-				CrashContext.PrimaryCrashProperties.AppDefaultLocale = "en-US";
+				OutCrashContext.PrimaryCrashProperties.AppDefaultLocale = "en-US";				
+			}
+
+// 			OutCrashContext.PrimaryCrashProperties.IsUE4Release
+			string WERUserName = ReportData.UserName;
+			if (!string.IsNullOrEmpty(WERUserName) || string.IsNullOrEmpty(OutCrashContext.PrimaryCrashProperties.UserName))
+			{
+				OutCrashContext.PrimaryCrashProperties.UserName = WERUserName;
 			}
 
+			OutCrashContext.PrimaryCrashProperties.BaseDir = ReportData.BaseDir;
+// 			OutCrashContext.PrimaryCrashProperties.RootDir
+			OutCrashContext.PrimaryCrashProperties.MachineId = ReportData.MachineId;
+			OutCrashContext.PrimaryCrashProperties.EpicAccountId = ReportData.EpicAccountId;
+// 			OutCrashContext.PrimaryCrashProperties.CallStack
+// 			OutCrashContext.PrimaryCrashProperties.SourceContext
+			OutCrashContext.PrimaryCrashProperties.UserDescription = string.Join( "\n", ReportData.UserDescription );
+
+			if (string.IsNullOrEmpty(OutCrashContext.PrimaryCrashProperties.ErrorMessage))
+			{
+				OutCrashContext.PrimaryCrashProperties.ErrorMessage = string.Join("\n", ReportData.ErrorMessage);
+			}
 
-// 			CrashContext.PrimaryCrashProperties.IsUE4Release
-			CrashContext.PrimaryCrashProperties.UserName = ReportData.UserName;
-			CrashContext.PrimaryCrashProperties.BaseDir = ReportData.BaseDir;
-// 			CrashContext.PrimaryCrashProperties.RootDir
-			CrashContext.PrimaryCrashProperties.MachineId = ReportData.MachineId;
-			CrashContext.PrimaryCrashProperties.EpicAccountId = ReportData.EpicAccountId;
-// 			CrashContext.PrimaryCrashProperties.CallStack
-// 			CrashContext.PrimaryCrashProperties.SourceContext
-			CrashContext.PrimaryCrashProperties.UserDescription = string.Join( "\n", ReportData.UserDescription );
-			CrashContext.PrimaryCrashProperties.ErrorMessage = string.Join( "\n", ReportData.ErrorMessage );
-// 			CrashContext.PrimaryCrashProperties.CrashDumpMode
-// 			CrashContext.PrimaryCrashProperties.Misc.NumberOfCores
-// 			CrashContext.PrimaryCrashProperties.Misc.NumberOfCoresIncludingHyperthreads
-// 			CrashContext.PrimaryCrashProperties.Misc.Is64bitOperatingSystem
-// 			CrashContext.PrimaryCrashProperties.Misc.CPUVendor
-// 			CrashContext.PrimaryCrashProperties.Misc.CPUBrand
-// 			CrashContext.PrimaryCrashProperties.Misc.PrimaryGPUBrand
-// 			CrashContext.PrimaryCrashProperties.Misc.OSVersionMajor
-// 			CrashContext.PrimaryCrashProperties.Misc.OSVersionMinor
-// 			CrashContext.PrimaryCrashProperties.Misc.AppDiskTotalNumberOfBytes
-// 			CrashContext.PrimaryCrashProperties.Misc.AppDiskNumberOfFreeBytes
-// 			CrashContext.PrimaryCrashProperties.MemoryStats.TotalPhysical
-// 			CrashContext.PrimaryCrashProperties.MemoryStats.TotalVirtual
-// 			CrashContext.PrimaryCrashProperties.MemoryStats.PageSize
-// 			CrashContext.PrimaryCrashProperties.MemoryStats.TotalPhysicalGB
-// 			CrashContext.PrimaryCrashProperties.MemoryStats.AvailablePhysical
-// 			CrashContext.PrimaryCrashProperties.MemoryStats.AvailableVirtual
-// 			CrashContext.PrimaryCrashProperties.MemoryStats.UsedPhysical
-// 			CrashContext.PrimaryCrashProperties.MemoryStats.PeakUsedPhysical
-// 			CrashContext.PrimaryCrashProperties.MemoryStats.UsedVirtual
-// 			CrashContext.PrimaryCrashProperties.MemoryStats.PeakUsedVirtual
-// 			CrashContext.PrimaryCrashProperties.MemoryStats.bIsOOM
-// 			CrashContext.PrimaryCrashProperties.MemoryStats.OOMAllocationSize
-// 			CrashContext.PrimaryCrashProperties.MemoryStats.OOMAllocationAlignment
-			CrashContext.PrimaryCrashProperties.TimeofCrash = new DateTime( ReportData.Ticks );
-			CrashContext.PrimaryCrashProperties.bAllowToBeContacted = ReportData.AllowToBeContacted;
-
-			return CrashContext;
+			// 			OutCrashContext.PrimaryCrashProperties.CrashDumpMode
+// 			OutCrashContext.PrimaryCrashProperties.Misc.NumberOfCores
+// 			OutCrashContext.PrimaryCrashProperties.Misc.NumberOfCoresIncludingHyperthreads
+// 			OutCrashContext.PrimaryCrashProperties.Misc.Is64bitOperatingSystem
+// 			OutCrashContext.PrimaryCrashProperties.Misc.CPUVendor
+// 			OutCrashContext.PrimaryCrashProperties.Misc.CPUBrand
+// 			OutCrashContext.PrimaryCrashProperties.Misc.PrimaryGPUBrand
+// 			OutCrashContext.PrimaryCrashProperties.Misc.OSVersionMajor
+// 			OutCrashContext.PrimaryCrashProperties.Misc.OSVersionMinor
+// 			OutCrashContext.PrimaryCrashProperties.Misc.AppDiskTotalNumberOfBytes
+// 			OutCrashContext.PrimaryCrashProperties.Misc.AppDiskNumberOfFreeBytes
+// 			OutCrashContext.PrimaryCrashProperties.MemoryStats.TotalPhysical
+// 			OutCrashContext.PrimaryCrashProperties.MemoryStats.TotalVirtual
+// 			OutCrashContext.PrimaryCrashProperties.MemoryStats.PageSize
+// 			OutCrashContext.PrimaryCrashProperties.MemoryStats.TotalPhysicalGB
+// 			OutCrashContext.PrimaryCrashProperties.MemoryStats.AvailablePhysical
+// 			OutCrashContext.PrimaryCrashProperties.MemoryStats.AvailableVirtual
+// 			OutCrashContext.PrimaryCrashProperties.MemoryStats.UsedPhysical
+// 			OutCrashContext.PrimaryCrashProperties.MemoryStats.PeakUsedPhysical
+// 			OutCrashContext.PrimaryCrashProperties.MemoryStats.UsedVirtual
+// 			OutCrashContext.PrimaryCrashProperties.MemoryStats.PeakUsedVirtual
+// 			OutCrashContext.PrimaryCrashProperties.MemoryStats.bIsOOM
+// 			OutCrashContext.PrimaryCrashProperties.MemoryStats.OOMAllocationSize
+// 			OutCrashContext.PrimaryCrashProperties.MemoryStats.OOMAllocationAlignment
+			OutCrashContext.PrimaryCrashProperties.TimeofCrash = new DateTime( ReportData.Ticks );
+			OutCrashContext.PrimaryCrashProperties.bAllowToBeContacted = ReportData.AllowToBeContacted;
 		}
 
 		/// <summary> Looks for the WER metadata xml file, if found, will return a new instance of the WERReportMetadata. </summary>
@@ -367,7 +382,7 @@ namespace Tools.CrashReporter.CrashReportProcess
 				if (MetaData != null)
 				{
 					FReportData ReportData = new FReportData( MetaData, NewReportPath );
-					Context = ConvertMetadataToCrashContext( MetaData, NewReportPath );
+					ConvertMetadataToCrashContext(MetaData, NewReportPath, ref Context);
 					bFromWER = true;
 				}
 			}
@@ -379,17 +394,11 @@ namespace Tools.CrashReporter.CrashReportProcess
 			}
 			else
 			{
-				if (GenericContext != null && GenericContext.PrimaryCrashProperties.ErrorMessage.Length > 0)
-				{
-					// Get error message from the crash context and fix value in the metadata.
-					Context.PrimaryCrashProperties.ErrorMessage = GenericContext.PrimaryCrashProperties.ErrorMessage;
-				}
-
 				Context.CrashDirectory = NewReportPath;
 				Context.PrimaryCrashProperties.SetPlatformFullName();
 
-				// If based on WER, save to the file.
-				if (bFromWER && GenericContext == null)
+				// Added data from WER, save to the crash context file.
+				if (bFromWER)
 				{
 					Context.ToFile();
 				}
diff --git a/Engine/Source/Programs/CrashReporter/CrashReportWebSite/Controllers/CrashesController.cs b/Engine/Source/Programs/CrashReporter/CrashReportWebSite/Controllers/CrashesController.cs
--- a/Engine/Source/Programs/CrashReporter/CrashReportWebSite/Controllers/CrashesController.cs
+++ b/Engine/Source/Programs/CrashReporter/CrashReportWebSite/Controllers/CrashesController.cs
@@ -3,6 +3,7 @@
 using System;
 using System.Diagnostics;
 using System.IO;
+using System.Text;
 using System.Web.Mvc;
 
 using Tools.DotNETCommon.XmlHandler;
@@ -161,12 +162,19 @@ namespace Tools.CrashReporter.CrashReportWebSite.Controllers
 
 				CrashReporterResult NewCrashResult = new CrashReporterResult();
 				NewCrashResult.ID = -1;
+				string PayloadString = string.Empty;
 
 				for (int Index = 0; Index < 3; Index++)
 				{
 					try
 					{
-						UnsafeAddCrash( NewCrashResult, Crashes );
+						using (StreamReader Reader = new StreamReader( Request.InputStream, Request.ContentEncoding ))
+						{
+							PayloadString = Reader.ReadToEnd();
+							CrashDescription NewCrash = XmlHandler.FromXmlString<CrashDescription>(PayloadString);
+							NewCrashResult.ID = Crashes.AddNewCrash( NewCrash );
+							NewCrashResult.bSuccess = true;
+						}
 						break;
 					}
 					catch (SqlException SqlExc)
@@ -184,7 +192,13 @@ namespace Tools.CrashReporter.CrashReportWebSite.Controllers
 					}
 					catch (Exception Ex)
 					{
-						NewCrashResult.Message = Ex.ToString();
+						StringBuilder MessageBuilder = new StringBuilder();
+						MessageBuilder.AppendLine("Exception was:");
+						MessageBuilder.AppendLine(Ex.ToString());
+						MessageBuilder.AppendLine("Received payload was:");
+						MessageBuilder.AppendLine(PayloadString);
+
+						NewCrashResult.Message = MessageBuilder.ToString();
 						NewCrashResult.bSuccess = false;
 						break;
 					}
@@ -195,16 +209,5 @@ namespace Tools.CrashReporter.CrashReportWebSite.Controllers
 				return Content( ReturnResult, "text/xml" );
 			}
 		}
-
-		private void UnsafeAddCrash( CrashReporterResult NewCrashResult, CrashRepository Crashes )
-		{
-			using (StreamReader Reader = new StreamReader( Request.InputStream, Request.ContentEncoding ))
-			{
-				string Result = Reader.ReadToEnd();
-				CrashDescription NewCrash = XmlHandler.FromXmlString<CrashDescription>( Result );
-				NewCrashResult.ID = Crashes.AddNewCrash( NewCrash );
-				NewCrashResult.bSuccess = true;
-			}
-		}
 	}
 }
diff --git a/Engine/Source/Programs/CrashReporter/CrashReportWebSite/Controllers/ReportsController.cs b/Engine/Source/Programs/CrashReporter/CrashReportWebSite/Controllers/ReportsController.cs
--- a/Engine/Source/Programs/CrashReporter/CrashReportWebSite/Controllers/ReportsController.cs
+++ b/Engine/Source/Programs/CrashReporter/CrashReportWebSite/Controllers/ReportsController.cs
@@ -32,14 +32,87 @@ namespace Tools.CrashReporter.CrashReportWebSite.Controllers
 	{
 		/// <summary>Fake id for all user groups</summary>
 		public static readonly int AllUserGroupId = -1;
-
+        
 		/// <summary>
+		/// 
+		/// </summary>
+		/// <param name="ReportsForm"></param>
+		/// <returns></returns>
+		public ActionResult Index( FormCollection ReportsForm )
+		{
+			using( FAutoScopedLogTimer LogTimer = new FAutoScopedLogTimer( this.GetType().ToString(), bCreateNewLog: true ) )
+			{
+				FormHelper FormData = new FormHelper( Request, ReportsForm, "JustReport" );
+
+				// Handle 'CopyToJira' button
+				int BuggIDToBeAddedToJira = 0;
+				foreach( var Entry in ReportsForm )
+				{
+					if( Entry.ToString().Contains( Bugg.JiraSubmitName ) )
+					{
+						int.TryParse( Entry.ToString().Substring( Bugg.JiraSubmitName.Length ), out BuggIDToBeAddedToJira );
+						break;
+					}
+				}
+
+				ReportsViewModel Results = GetResults( FormData, BuggIDToBeAddedToJira );
+				Results.GenerationTime = LogTimer.GetElapsedSeconds().ToString( "F2" );
+				return View( "Index", Results );
+			}
+		}
+
+	    /// <summary>
+	    /// 
+	    /// </summary>
+	    /// <param name="ReportsForm"></param>
+	    /// <returns></returns>
+	    public ActionResult BranchReport(FormCollection ReportsForm)
+	    {
+            using (var logTimer = new FAutoScopedLogTimer(this.GetType().ToString(), bCreateNewLog: true))
+            {
+                var formData = new FormHelper(Request, ReportsForm, "JustReport");
+                
+                var buggIdToBeAddedToJira = 0;
+                // Handle 'CopyToJira' button
+                foreach (var Entry in ReportsForm)
+                {
+                    if (Entry.ToString().Contains(Bugg.JiraSubmitName))
+                    {
+                        
+                        int.TryParse(Entry.ToString().Substring(Bugg.JiraSubmitName.Length), out buggIdToBeAddedToJira);
+                        break;
+                    }
+                }
+
+                var results = GetResults(formData, buggIdToBeAddedToJira);
+                results.GenerationTime = logTimer.GetElapsedSeconds().ToString("F2");
+                return View("Index", results);
+            }
+        }
+
+        #region Private Methods
+
+        private void AddBuggJiraMapping(Bugg NewBugg, ref HashSet<string> FoundJiras, ref Dictionary<string, List<Bugg>> JiraIDtoBugg)
+        {
+            string JiraID = NewBugg.Jira;
+            FoundJiras.Add("key = " + JiraID);
+
+            bool bAdd = !JiraIDtoBugg.ContainsKey(JiraID);
+            if (bAdd)
+            {
+                JiraIDtoBugg.Add(JiraID, new List<Bugg>());
+            }
+
+            JiraIDtoBugg[JiraID].Add(NewBugg);
+        }
+
+        /// <summary>
 		/// Retrieve all Buggs matching the search criteria.
 		/// </summary>
 		/// <param name="FormData">The incoming form of search criteria from the client.</param>
 		/// <param name="BuggIDToBeAddedToJira">ID of the bugg that will be added to JIRA</param>
 		/// <returns>A view to display the filtered Buggs.</returns>
-		public ReportsViewModel GetResults( FormHelper FormData, int BuggIDToBeAddedToJira )
+		private ReportsViewModel GetResults( FormHelper FormData, int BuggIDToBeAddedToJira )
 		{
 			BuggRepository BuggsRepo = new BuggRepository();
 			CrashRepository CrashRepo = new CrashRepository();
@@ -57,16 +130,28 @@ namespace Tools.CrashReporter.CrashReportWebSite.Controllers
 				string AnonumousGroup = "Anonymous";
 				//List<String> Users = FRepository.Get().GetUserNamesFromGroupName( AnonumousGroup );
 				int AnonymousGroupID = FRepository.Get( BuggsRepo ).FindOrAddGroup( AnonumousGroup );
-				HashSet<int> AnonumousIDs = FRepository.Get( BuggsRepo ).GetUserIdsFromUserGroup( AnonumousGroup );
+				List<int> AnonumousIDs = FRepository.Get( BuggsRepo ).GetUserIdsFromUserGroup( AnonumousGroup ).ToList();
 				int AnonymousID = AnonumousIDs.First();
 
-				var Crashes = CrashRepo
-					.FilterByDate( CrashRepo.ListAll(), FormData.DateFrom, FormData.DateTo )
+				var CrashesQuery = CrashRepo
+					.QueryByDate( CrashRepo.ListAll(), FormData.DateFrom, FormData.DateTo )
 					// Only crashes and asserts
 					.Where( Crash => Crash.CrashType == 1 || Crash.CrashType == 2 )
 					// Only anonymous user
-					.Where( Crash => Crash.UserNameId.Value == AnonymousID )
-					.Select( Crash => new
+                    .Where(Crash => Crash.UserNameId.Value == AnonymousID);
+
+                // Filter by BranchName
+                if (!string.IsNullOrEmpty(FormData.BranchName))
+                {
+                    CrashesQuery =
+                    (
+                        from CrashDetail in CrashesQuery
+                        where CrashDetail.Branch.Equals(FormData.BranchName)
+                        select CrashDetail
+                    );
+                }
+
+                var Crashes = CrashesQuery.Select(Crash => new
 					{
 						ID = Crash.Id,
 						TimeOfCrash = Crash.TimeOfCrash.Value,
@@ -241,17 +326,20 @@ namespace Tools.CrashReporter.CrashReportWebSite.Controllers
 						var JiraResults = new Dictionary<string, Dictionary<string, object>>();
 						try
 						{
-							JiraResults = JC.SearchJiraTickets(
-							JiraSearchQuery,
-							new string[] 
-							{ 
-								"key",				// string
-								"summary",			// string
-								"components",		// System.Collections.ArrayList, Dictionary<string,object>, name
-								"resolution",		// System.Collections.Generic.Dictionary`2[System.String,System.Object], name
-								"fixVersions",		// System.Collections.ArrayList, Dictionary<string,object>, name
-								"customfield_11200" // string
-							} );
+						    if (!string.IsNullOrWhiteSpace(JiraSearchQuery))
+						    {
+                                JiraResults = JC.SearchJiraTickets(
+                                    JiraSearchQuery,
+                                    new string[] 
+                                    { 
+                                        "key",				// string
+								        "summary",			// string
+								        "components",		// System.Collections.ArrayList, Dictionary<string,object>, name
+								        "resolution",		// System.Collections.Generic.Dictionary`2[System.String,System.Object], name
+								        "fixVersions",		// System.Collections.ArrayList, Dictionary<string,object>, name
+								        "customfield_11200" // string
+							        });
+                            }
 						}
 						catch (System.Exception)
 						{
@@ -368,46 +456,6 @@ namespace Tools.CrashReporter.CrashReportWebSite.Controllers
 			}
 		}
 
-		private void AddBuggJiraMapping( Bugg NewBugg, ref HashSet<string> FoundJiras, ref Dictionary<string, List<Bugg>> JiraIDtoBugg )
-		{
-			string JiraID = NewBugg.Jira;
-			FoundJiras.Add( "key = " + JiraID );
-
-			bool bAdd = !JiraIDtoBugg.ContainsKey( JiraID );
-			if( bAdd )
-			{
-				JiraIDtoBugg.Add( JiraID, new List<Bugg>() );
-			}
-
-			JiraIDtoBugg[JiraID].Add( NewBugg );
-		}
-
-		/// <summary>
-		/// 
-		/// </summary>
-		/// <param name="ReportsForm"></param>
-		/// <returns></returns>
-		public ActionResult Index( FormCollection ReportsForm )
-		{
-			using( FAutoScopedLogTimer LogTimer = new FAutoScopedLogTimer( this.GetType().ToString(), bCreateNewLog: true ) )
-			{
-				FormHelper FormData = new FormHelper( Request, ReportsForm, "JustReport" );
-
-				// Handle 'CopyToJira' button
-				int BuggIDToBeAddedToJira = 0;
-				foreach( var Entry in ReportsForm )
-				{
-					if( Entry.ToString().Contains( Bugg.JiraSubmitName ) )
-					{
-						int.TryParse( Entry.ToString().Substring( Bugg.JiraSubmitName.Length ), out BuggIDToBeAddedToJira );
-						break;
-					}
-				}
-
-				ReportsViewModel Results = GetResults( FormData, BuggIDToBeAddedToJira );
-				Results.GenerationTime = LogTimer.GetElapsedSeconds().ToString( "F2" );
-				return View( "Index", Results );
-			}
-		}
-	}
+        #endregion
+    }
 }
\ No newline at end of file
diff --git a/Engine/Source/Programs/CrashReporter/CrashReportWebSite/Models/CrashRepository.cs b/Engine/Source/Programs/CrashReporter/CrashReportWebSite/Models/CrashRepository.cs
--- a/Engine/Source/Programs/CrashReporter/CrashReportWebSite/Models/CrashRepository.cs
+++ b/Engine/Source/Programs/CrashReporter/CrashReportWebSite/Models/CrashRepository.cs
@@ -74,6 +74,7 @@ namespace Tools.CrashReporter.CrashReportWebSite.Models
 
 		private static DateTime LastBranchesDate = DateTime.UtcNow.AddDays(-1);
 		private static List<SelectListItem> BranchesAsSelectList = null;
+
 		/// <summary>
 		/// Retrieves a list of distinct UE4 Branches from the CrashRepository
 		/// </summary>
@@ -87,17 +88,17 @@ namespace Tools.CrashReporter.CrashReportWebSite.Models
 
 				if (Now - LastBranchesDate > TimeSpan.FromHours( 1 ))
 				{
-					var BranchList = Context.Crashes
-					.Where( n => n.TimeOfCrash > DateTime.Now.AddDays( -14 ) )
+                    var BranchList = Context.Crashes
+                    .Where( n => n.TimeOfCrash > DateTime.Now.AddDays( -14 ) )
                     .Where(c => c.CrashType != 3) // Ignore ensures
-					// Depot - //depot/UE4* || Stream //UE4, //Something etc.
-					.Where( n => n.Branch.StartsWith( "UE4" ) || n.Branch.StartsWith( "//" ) )
-					.Select( n => n.Branch )
-					.Distinct()
-					.ToList();
-					BranchesAsSelectList = BranchList
-						.Select( listitem => new SelectListItem { Selected = false, Text = listitem, Value = listitem } )
-						.ToList();
+                    // Depot - //depot/UE4* || Stream //UE4, //Something etc.
+                    .Where( n => n.Branch.StartsWith( "UE4" ) || n.Branch.StartsWith( "//" ) )
+                    .Select( n => n.Branch )
+                    .Distinct()
+                    .ToList();
+                    BranchesAsSelectList = BranchList
+                        .Select( listitem => new SelectListItem { Selected = false, Text = listitem, Value = listitem } )
+                        .ToList();
 					BranchesAsSelectList.Insert( 0, new SelectListItem { Selected = true, Text = "", Value = "" } );
 
 					LastBranchesDate = Now;
@@ -142,29 +143,29 @@ namespace Tools.CrashReporter.CrashReportWebSite.Models
 
 				if (Now - LastVersionDate > TimeSpan.FromHours( 1 ))
 				{
-					var BuildVersions = Context.Crashes
-					.Where( c => c.TimeOfCrash > DateTime.Now.AddDays( -14 ) )
-                    .Where(c => c.CrashType != 3 ) // Ignore ensures
-					.Select( c => c.BuildVersion )
-					.Distinct()
-					.ToList();
+                    var BuildVersions = Context.Crashes
+                    .Where(c => c.TimeOfCrash > DateTime.Now.AddDays(-14))
+                    .Where(c => c.CrashType != 3) // Ignore ensures
+                    .Select(c => c.BuildVersion)
+                    .Distinct()
+                    .ToList();
 
 					DistinctBuildVersions = new HashSet<string>();
 
-					foreach (var BuildVersion in BuildVersions)
-					{
-						var BVParts = BuildVersion.Split( new char[] { '.' }, StringSplitOptions.RemoveEmptyEntries );
-						if (BVParts.Length > 2 && BVParts[0] != "0")
-						{
-							string CleanBV = string.Format( "{0}.{1}.{2}", BVParts[0], BVParts[1], BVParts[2] );
-							DistinctBuildVersions.Add( CleanBV );
-						}
-					}
+                    foreach (var BuildVersion in BuildVersions)
+                    {
+                        var BVParts = BuildVersion.Split(new char[] { '.' }, StringSplitOptions.RemoveEmptyEntries);
+                        if (BVParts.Length > 2 && BVParts[0] != "0")
+                        {
+                            string CleanBV = string.Format("{0}.{1}.{2}", BVParts[0], BVParts[1], BVParts[2]);
+                            DistinctBuildVersions.Add(CleanBV);
+                        }
+                    }
 
 
-					VersionsAsSelectList = DistinctBuildVersions
-						.Select( listitem => new SelectListItem { Selected = false, Text = listitem, Value = listitem } )
-						.ToList();
+                    VersionsAsSelectList = DistinctBuildVersions
+                        .Select( listitem => new SelectListItem { Selected = false, Text = listitem, Value = listitem } )
+                        .ToList();
 					VersionsAsSelectList.Insert( 0, new SelectListItem { Selected = true, Text = "", Value = "" } );
 
 					LastVersionDate = Now;
@@ -492,6 +493,24 @@ namespace Tools.CrashReporter.CrashReportWebSite.Models
 			}
 		}
 
+        /// <summary>
+        /// Filter a set of crashes to a date range but do not enumerate the data.
+        /// This allows further filtering on the query expression to take place before the data is retrieved from the data source.
+        /// </summary>
+        /// <param name="Results">The unfiltered set of crashes.</param>
+        /// <param name="DateFrom">The earliest date to filter by.</param>
+        /// <param name="DateTo">The latest date to filter by.</param>
+        /// <returns>A query statement describing the set of crashes between the earliest and latest date.</returns>
+        public IQueryable<Crash> QueryByDate(IQueryable<Crash> Results, DateTime DateFrom, DateTime DateTo)
+        {
+            using (FAutoScopedLogTimer LogTimer = new FAutoScopedLogTimer(this.GetType().ToString() + " SQL"))
+            {
+                IQueryable<Crash> CrashesInTimeFrame = Results
+                    .Where(MyCrash => MyCrash.TimeOfCrash >= DateFrom && MyCrash.TimeOfCrash <= DateTo.AddDays(1));
+                return CrashesInTimeFrame;
+            }
+        }
+
 		/// <summary>
 		/// 
 		/// </summary>
diff --git a/Engine/Source/Programs/CrashReporter/CrashReportWebSite/Models/ReportsViewModel.cs b/Engine/Source/Programs/CrashReporter/CrashReportWebSite/Models/ReportsViewModel.cs
--- a/Engine/Source/Programs/CrashReporter/CrashReportWebSite/Models/ReportsViewModel.cs
+++ b/Engine/Source/Programs/CrashReporter/CrashReportWebSite/Models/ReportsViewModel.cs
@@ -7,36 +7,50 @@ using System.Web.Mvc;
 
 namespace Tools.CrashReporter.CrashReportWebSite.Models
 {
-	/// <summary>
-	/// The view model for the crash summary page.
-	/// </summary>
-	public class ReportsViewModel
-	{
-		/// <summary>A container of sorted Buggs.</summary>
-		public List<Bugg> Buggs { get; set; }
+    /// <summary>
+    /// The view model for the crash summary page.
+    /// </summary>
+    public class ReportsViewModel
+    {
+        /// <summary>A container of sorted Buggs.</summary>
+        public List<Bugg> Buggs { get; set; }
 
-		/// <summary>A container of sorted Crashes.</summary>
-		public List<Crash> Crashes { get; set; }
+        /// <summary>A container of sorted Crashes.</summary>
+        public List<Crash> Crashes { get; set; }
 
-		/// <summary>The date of the earliest crash in a Bugg.</summary>
-		public long DateFrom { get; set; }
+        /// <summary>The date of the earliest crash in a Bugg.</summary>
+        public long DateFrom { get; set; }
 
-		/// <summary>The date of the most recent crash in a Bugg.</summary>
-		public long DateTo { get; set; }
+        /// <summary>The date of the most recent crash in a Bugg.</summary>
+        public long DateTo { get; set; }
 
-		/// <summary></summary>
-		public int TotalAnonymousCrashes { get; set; }
+        /// <summary></summary>
+        public int TotalAnonymousCrashes { get; set; }
 
-		/// <summary></summary>
-		public int TotalUniqueAnonymousCrashes { get; set; }
+        /// <summary></summary>
+        public int TotalUniqueAnonymousCrashes { get; set; }
 
-		/// <summary></summary>
-		public int TotalAffectedUsers { get; set; }
+        /// <summary></summary>
+        public int TotalAffectedUsers { get; set; }
 
-		/// <summary>Time spent in generating this site, formatted as a string.</summary>
-		public string GenerationTime { get; set; }
+        /// <summary>Time spent in generating this site, formatted as a string.</summary>
+        public string GenerationTime { get; set; }
 
-		/// <summary>User input from the client.</summary>
-		public FormCollection FormCollection { get; set; }
-	}
+        /// <summary>User input from the client.</summary>
+        public FormCollection FormCollection { get; set; }
+
+        /// <summary>The name of the branch to filter by.</summary>
+        public string BranchName { get; set; }
+
+        /// <summary>A collection of Branch names used in the drop down on the main search form</summary>
+        public List<SelectListItem> BranchNames { get; set; }
+
+        /// <summary>
+        /// Default constructor, used by ReportsController
+        /// </summary>
+        public ReportsViewModel()
+        {
+            BranchNames = CrashRepository.GetBranchesAsListItems();
+        }
+    }
 }
\ No newline at end of file
diff --git a/Engine/Source/Programs/CrashReporter/CrashReportWebSite/Views/Crashes/Show.aspx b/Engine/Source/Programs/CrashReporter/CrashReportWebSite/Views/Crashes/Show.aspx
--- a/Engine/Source/Programs/CrashReporter/CrashReportWebSite/Views/Crashes/Show.aspx
+++ b/Engine/Source/Programs/CrashReporter/CrashReportWebSite/Views/Crashes/Show.aspx
@@ -225,7 +225,7 @@
 				<dd ><%=Html.DisplayFor(m => Model.Crash.Module) %></dd>
 			
 			<dt>User Activity</dt>
-				<dd ><%=Html.DisplayFor(m => Model.Crash.UserActivityHint) %></dd>
+				<%--<dd ><%=Html.DisplayFor(m => Model.Crash.UserActivityHint) %></dd>--%>
 		</dl>
 	</div>
 
diff --git a/Engine/Source/Programs/CrashReporter/CrashReportWebSite/Views/Reports/Index.aspx b/Engine/Source/Programs/CrashReporter/CrashReportWebSite/Views/Reports/Index.aspx
--- a/Engine/Source/Programs/CrashReporter/CrashReportWebSite/Views/Reports/Index.aspx
+++ b/Engine/Source/Programs/CrashReporter/CrashReportWebSite/Views/Reports/Index.aspx
@@ -54,13 +54,16 @@
 		<span class="SearchTextTitle">To:</span>
 		<input id="dateToVisible" type="text" class="date" autocomplete="OFF" style="width:80px" />
 		<input id="dateTo" name="dateTo" type="hidden" value="<%=Model.DateTo %>"autocomplete="OFF"  />
+        
 	</p>
+    <p class="SearchTextTitle">Filter by Branch</p>
+    <p id="BranchFilter"><%=Html.DropDownListFor( m=>m.BranchName, Model.BranchNames )%></p>
 
-	 <p>
-		 <input type="submit" value="Search" class='SearchButton' />
-	 </p>
+	<p>
+		<input type="submit" value="Search" class='SearchButton' />
+	</p>
 	
-	<%} %>
+<%} %>
 
 </div>
 </asp:Content>
diff --git a/Engine/Source/Programs/CrashReporter/CrashReportWebSite/Views/Web.config b/Engine/Source/Programs/CrashReporter/CrashReportWebSite/Views/Web.config
--- a/Engine/Source/Programs/CrashReporter/CrashReportWebSite/Views/Web.config
+++ b/Engine/Source/Programs/CrashReporter/CrashReportWebSite/Views/Web.config
@@ -9,7 +9,7 @@
   </configSections>
 
   <system.web.webPages.razor>
-    <host factoryType="System.Web.Mvc.MvcWebRazorHostFactory, System.Web.Mvc, Version=3.0.0.0, Culture=neutral, PublicKeyToken=31BF3856AD364E35" />
+    <host factoryType="System.Web.Mvc.MvcWebRazorHostFactory, System.Web.Mvc, Version=3.0.0.1, Culture=neutral, PublicKeyToken=31BF3856AD364E35" />
     <pages pageBaseType="System.Web.Mvc.WebViewPage">
       <namespaces>
         <add namespace="System.Web.Mvc" />
@@ -46,10 +46,8 @@
       </controls>
     </pages>
   </system.web>
-
   <system.webServer>
     <validation validateIntegratedModeConfiguration="false" />
-
     <handlers>
       <remove name="BlockViewHandler"/>
       <add name="BlockViewHandler" path="*" verb="*" preCondition="integratedMode" type="System.Web.HttpNotFoundHandler" />
diff --git a/Engine/Source/Programs/CrashReporter/CrashReportWebSite/Web.config b/Engine/Source/Programs/CrashReporter/CrashReportWebSite/Web.config
--- a/Engine/Source/Programs/CrashReporter/CrashReportWebSite/Web.config
+++ b/Engine/Source/Programs/CrashReporter/CrashReportWebSite/Web.config
@@ -9,13 +9,17 @@
 -->
 <configuration>
   <configSections>
-    <sectionGroup name="applicationSettings" type="System.Configuration.ApplicationSettingsGroup, System, Version=4.0.0.0, Culture=neutral, PublicKeyToken=b77a5c561934e089">
-      <section name="Tools.CrashReporter.CrashReportWebSite.Properties.Settings" type="System.Configuration.ClientSettingsSection, System, Version=4.0.0.0, Culture=neutral, PublicKeyToken=b77a5c561934e089" requirePermission="false"/>
+    <sectionGroup name="applicationSettings"
+      type="System.Configuration.ApplicationSettingsGroup, System, Version=4.0.0.0, Culture=neutral, PublicKeyToken=b77a5c561934e089">
+      <section name="Tools.CrashReporter.CrashReportWebSite.Properties.Settings"
+        type="System.Configuration.ClientSettingsSection, System, Version=4.0.0.0, Culture=neutral, PublicKeyToken=b77a5c561934e089" requirePermission="false"/>
     </sectionGroup>
   </configSections>
   <appSettings/>
   <connectionStrings>
-    <add name="CrashReportConnectionString" connectionString="Data Source=db-09;Initial Catalog=CrashReport;Server=db-09;Database=CrashReport;Integrated Security=true" providerName="System.Data.SqlClient"/>
+    <add name="CrashReportConnectionString"
+      connectionString="Data Source=db-09;Initial Catalog=CrashReport;Server=db-09;Database=CrashReport;Integrated Security=true"
+      providerName="System.Data.SqlClient"/>
   </connectionStrings>
   <!--
     For a description of web.config changes see http://go.microsoft.com/fwlink/?LinkId=235367.
@@ -34,7 +38,7 @@
     -->
     <compilation debug="true" targetFramework="4.5">
       <assemblies>
-        <add assembly="System.Web.Mvc, Version=3.0.0.0, Culture=neutral, PublicKeyToken=31BF3856AD364E35"/>
+        <add assembly="System.Web.Mvc, Version=3.0.0.1, Culture=neutral, PublicKeyToken=31BF3856AD364E35"/>
         <add assembly="System.Web.Abstractions, Version=4.0.0.0, Culture=neutral, PublicKeyToken=31BF3856AD364E35"/>
         <add assembly="System.Web.Routing, Version=4.0.0.0, Culture=neutral, PublicKeyToken=31BF3856AD364E35"/>
         <add assembly="System.Data.Linq, Version=4.0.0.0, Culture=neutral, PublicKeyToken=B77A5C561934E089"/>
@@ -101,4 +105,12 @@
       </setting>
     </Tools.CrashReporter.CrashReportWebSite.Properties.Settings>
   </applicationSettings>
+  <runtime>
+    <assemblyBinding xmlns="urn:schemas-microsoft-com:asm.v1">
+      <dependentAssembly>
+        <assemblyIdentity name="System.Web.Mvc" publicKeyToken="31bf3856ad364e35"/>
+        <bindingRedirect oldVersion="3.0.0.0-3.0.0.1" newVersion="3.0.0.1"/>
+      </dependentAssembly>
+    </assemblyBinding>
+  </runtime>
 </configuration>
\ No newline at end of file
diff --git a/Engine/Source/Programs/DotNETCommon/DotNETUtilities/SimpleWebRequest.cs b/Engine/Source/Programs/DotNETCommon/DotNETUtilities/SimpleWebRequest.cs
--- a/Engine/Source/Programs/DotNETCommon/DotNETUtilities/SimpleWebRequest.cs
+++ b/Engine/Source/Programs/DotNETCommon/DotNETUtilities/SimpleWebRequest.cs
@@ -42,7 +42,10 @@ namespace Tools.DotNETCommon.SimpleWebRequest
 					byte[] Buffer = Encoding.UTF8.GetBytes( Payload );
 					Request.ContentType = "text/xml";
 					Request.ContentLength = Buffer.Length;
-					Request.GetRequestStream().Write( Buffer, 0, Buffer.Length );
+					using (Stream RequestStream = Request.GetRequestStream())
+					{
+						RequestStream.Write(Buffer, 0, Buffer.Length);
+					}
 				}
 				else
 				{
@@ -71,17 +74,21 @@ namespace Tools.DotNETCommon.SimpleWebRequest
 
 			try
 			{
-				HttpWebResponse WebResponse = ( HttpWebResponse )Request.GetResponse();
-				// Simple size check to prevent exploits
-				if( WebResponse.StatusCode == HttpStatusCode.OK && WebResponse.ContentLength < 4096 )
+				using (HttpWebResponse WebResponse = (HttpWebResponse)Request.GetResponse())
 				{
-					// Process the response
-					Stream ResponseStream = WebResponse.GetResponseStream();
-					byte[] RawResponse = new byte[WebResponse.ContentLength];
-					ResponseStream.Read( RawResponse, 0, ( int )WebResponse.ContentLength );
-					ResponseStream.Close();
+					// Simple size check to prevent exploits
+					if (WebResponse.StatusCode == HttpStatusCode.OK && WebResponse.ContentLength < 4096)
+					{
+						// Process the response
+						using (Stream ResponseStream = WebResponse.GetResponseStream())
+						{
+							byte[] RawResponse = new byte[WebResponse.ContentLength];
+							ResponseStream.Read(RawResponse, 0, (int)WebResponse.ContentLength);
+							ResponseStream.Close();
 
-					ResponseString = Encoding.UTF8.GetString( RawResponse );
+							ResponseString = Encoding.UTF8.GetString(RawResponse);
+						}
+					}
 				}
 			}
 			catch( Exception Ex )
diff --git a/Engine/Source/Runtime/Core/Private/Android/AndroidMemory.cpp b/Engine/Source/Runtime/Core/Private/Android/AndroidMemory.cpp
--- a/Engine/Source/Runtime/Core/Private/Android/AndroidMemory.cpp
+++ b/Engine/Source/Runtime/Core/Private/Android/AndroidMemory.cpp
@@ -1,7 +1,7 @@
 // Copyright 1998-2016 Epic Games, Inc. All Rights Reserved.
 
 #include "CorePrivatePCH.h"
-#include "MallocBinned2.h"
+#include "MallocBinned.h"
 #include "MallocAnsi.h"
 #include "unistd.h"
 #include <jni.h>
@@ -92,7 +92,7 @@ FMalloc* FAndroidPlatformMemory::BaseAllocator()
 	uint64 MemoryLimit = FMath::Min<uint64>( uint64(1) << FMath::CeilLogTwo(MemoryConstants.TotalPhysical), 0x100000000);
 
 	//return new FMallocAnsi();
-	return new FMallocBinned2(MemoryConstants.PageSize, MemoryLimit);
+	return new FMallocBinned(MemoryConstants.PageSize, MemoryLimit);
 }
 
 void* FAndroidPlatformMemory::BinnedAllocFromOS( SIZE_T Size )
diff --git a/Engine/Source/Runtime/Core/Private/GenericPlatform/GenericPlatformCrashContext.cpp b/Engine/Source/Runtime/Core/Private/GenericPlatform/GenericPlatformCrashContext.cpp
--- a/Engine/Source/Runtime/Core/Private/GenericPlatform/GenericPlatformCrashContext.cpp
+++ b/Engine/Source/Runtime/Core/Private/GenericPlatform/GenericPlatformCrashContext.cpp
@@ -49,6 +49,7 @@ namespace NCachedCrashContextProperties
 	static int32 SecondsSinceStart;
 	static FString CrashGUID;
 	static FString UserActivityHint;
+	static FString CommandLine;
 }
 
 void FGenericCrashContext::Initialize()
@@ -75,6 +76,7 @@ void FGenericCrashContext::Initialize()
 	NCachedCrashContextProperties::PrimaryGPUBrand = FPlatformMisc::GetPrimaryGPUBrand();
 	NCachedCrashContextProperties::UserName = FPlatformProcess::UserName();
 	NCachedCrashContextProperties::DefaultLocale = FPlatformMisc::GetDefaultLocale();
+	NCachedCrashContextProperties::CommandLine = FCommandLine::IsInitialized() ? FCommandLine::GetForLogging() : TEXT("");
 
 	// Using the -fullcrashdump parameter will cause full memory minidumps to be created for crashes
 	NCachedCrashContextProperties::CrashDumpMode = (int32)ECrashDumpMode::Default;
@@ -141,7 +143,7 @@ void FGenericCrashContext::SerializeContentToBuffer()
 	AddCrashProperty( TEXT( "EngineMode" ), FPlatformMisc::GetEngineMode() );
 	AddCrashProperty( TEXT( "DeploymentName"), FApp::GetDeploymentName() );
 	AddCrashProperty( TEXT( "EngineVersion" ), *FEngineVersion::Current().ToString() );
-	AddCrashProperty( TEXT( "CommandLine" ), FCommandLine::IsInitialized() ? FCommandLine::Get() : TEXT( "" ) );
+	AddCrashProperty( TEXT("CommandLine"), *NCachedCrashContextProperties::CommandLine );
 	AddCrashProperty( TEXT( "LanguageLCID" ), FInternationalization::Get().GetCurrentCulture()->GetLCID() );
 	AddCrashProperty( TEXT( "AppDefaultLocale" ), *NCachedCrashContextProperties::DefaultLocale );
 
diff --git a/Engine/Source/Runtime/Core/Private/HAL/MallocBinned2.cpp b/Engine/Source/Runtime/Core/Private/HAL/MallocBinned2.cpp
--- a/Engine/Source/Runtime/Core/Private/HAL/MallocBinned2.cpp
+++ b/Engine/Source/Runtime/Core/Private/HAL/MallocBinned2.cpp
@@ -10,6 +10,10 @@
 #include "MemoryMisc.h"
 #include "HAL/PlatformAtomics.h"
 
+// Alignment of OS-allocated pointer - pool-allocated pointers will have a non-aligned pointer
+#define BINNED2_LARGE_ALLOC 65536
+
+
 // Block sizes are based around getting the maximum amount of allocations per pool, with as little alignment waste as possible.
 // Block sizes should be close to even divisors of the system page size, and well distributed.
 // They must be 16-byte aligned as well.
@@ -24,97 +28,90 @@ const uint32 GMallocBinned2SmallBlockSizes[] =
 };
 
 /** Information about a piece of free memory. */
-struct FMallocBinned2::FFreeMem
+struct FMallocBinned2::FFreeBlock
 {
-	/** Next or MemLastPool[], always in order by pool. */
-	FFreeMem*	NextFreePool;
-	/** Number of consecutive free blocks here, at least 1. */
-	uint32		NumFreeBlocks;
+	void*  NextFreeBlock; // Next free block in another pool
+	uint32 NumFreeBlocks; // Number of consecutive free blocks here, at least 1.
+	uint32 BlockSize;     // Size of the blocks that this list points to
 };
 
-// Memory pool info. 32 bytes.
+FMallocBinned2::FPoolList::FPoolList()
+	: Front(nullptr)
+{
+}
+
+FMallocBinned2::FPoolTable::FPoolTable()
+	: BlockSize(0)
+{
+}
+
 struct FMallocBinned2::FPoolInfo
 {
-	/** Number of allocated elements in this pool, when counts down to zero can free the entire pool. */
-	uint16			Taken;		// 2
-	/** Index of pool. Index into MemSizeToPoolTable[]. Valid when < MAX_SMALL_POOLED_ALLOCATION_SIZE, -1 is OsTable.
-		When AllocSize is 0, this is the number of pages to step back to find the base address of an allocation. See FindPoolInfo() */
-	uint16			TableIndex; // 4
-	/** Number of bytes allocated, or 0, this is a trailing pool (i.e. a pool used to find its way back to the base pool) */
-	uint32			AllocSize;	// 8
-	/** Pointer to first free memory in this pool or the OS Allocation Size in bytes if this allocation is not binned*/
-	FFreeMem*		FirstMem;   // 12/16
-	FPoolInfo*		Next;		// 16/24
-	FPoolInfo**		PrevLink;	// 20/32
-#if PLATFORM_32BITS
-	/** Explicit padding for 32 bit builds */
-	uint8 Padding[12]; // 32
-#endif
-	uint32 GetTrailingPoolOffset() const
-	{
-		checkSlow(IsTrailingPool());
-		return TableIndex;
-	}
+ public:	uint16      Taken;          // Number of allocated elements in this pool, when counts down to zero can free the entire pool.
+private:	uint32      AllocSize;      // Number of bytes allocated
+ public:	FFreeBlock* FirstFreeBlock; // Pointer to first free memory in this pool or the OS Allocation Size in bytes if this allocation is not binned
+ public:	FPoolInfo*  Next;           // Pointer to next pool
+ public:	FPoolInfo** PtrToPrevNext;  // Pointer to whichever pointer points to this pool
 
-	void SetTrailingPoolOffset(uint32 InOffset)
+public:
+	bool IsExhausted() const
 	{
-		TableIndex = InOffset;
-		AllocSize  = 0;
+		return !FirstFreeBlock || FirstFreeBlock->NumFreeBlocks == 0;
 	}
 
-	void SetPoolAllocationSizes(uint32 InBytes, uint32 InTableIndex)
+	void* AllocateBlock(uint32 BlockSize)
 	{
-		// Shouldn't be pooling zero byte allocations
-		checkSlow(InBytes != 0);
-
-		TableIndex = InTableIndex;
-		AllocSize  = InBytes;
-	}
+		check(!IsExhausted());
+		++Taken;
+		--FirstFreeBlock->NumFreeBlocks;
+		void* Result = (uint8*)FirstFreeBlock + (FirstFreeBlock->NumFreeBlocks + IsAligned(FirstFreeBlock, BINNED2_LARGE_ALLOC)) * BlockSize;
 
-	void SetOSAllocationSizes(uint32 InBytes, UPTRINT InOsBytes)
-	{
-		// Shouldn't be pooling zero byte allocations
-		checkSlow(InBytes != 0);
+		if (FirstFreeBlock->NumFreeBlocks == 0)
+		{
+			FirstFreeBlock = (FFreeBlock*)FirstFreeBlock->NextFreeBlock;
+		}
+		checkSlow(!FirstFreeBlock || FirstFreeBlock->NumFreeBlocks != 0);
 
-		TableIndex = (uint16)-1;
-		AllocSize  = InBytes;
-		FirstMem   = (FFreeMem*)InOsBytes;
+		return Result;
 	}
 
-	UPTRINT GetOsBytes() const
+	uint32 GetOSRequestedBytes() const
 	{
-		check(IsOSAllocation());
-		return (UPTRINT)FirstMem;
+		return AllocSize;
 	}
 
-	FORCEINLINE bool IsTrailingPool() const
+	UPTRINT GetOsAllocatedBytes() const
 	{
-		return AllocSize == 0;
+		return (UPTRINT)FirstFreeBlock;
 	}
 
-	FORCEINLINE bool IsOSAllocation() const
+	void SetOSAllocationSizes(uint32 InRequestedBytes, UPTRINT InAllocatedBytes)
 	{
-		return TableIndex == (uint16)-1;
+		checkSlow(InRequestedBytes != 0);                // Shouldn't be pooling zero byte allocations
+		checkSlow(InAllocatedBytes >= InRequestedBytes); // We must be allocating at least as much as we requested
+
+		AllocSize      = InRequestedBytes;
+		FirstFreeBlock = (FFreeBlock*)InAllocatedBytes;
 	}
 
-	void Link( FPoolInfo*& Before )
+	void Link(FPoolInfo*& PrevNext)
 	{
-		if( Before )
+		if (PrevNext)
 		{
-			Before->PrevLink = &Next;
+			PrevNext->PtrToPrevNext = &Next;
 		}
-		Next     = Before;
-		PrevLink = &Before;
-		Before   = this;
+		Next          = PrevNext;
+		PtrToPrevNext = &PrevNext;
+		PrevNext      = this;
 	}
 
 	void Unlink()
 	{
-		if( Next )
+		if (Next)
 		{
-			Next->PrevLink = PrevLink;
+			Next->PtrToPrevNext = PtrToPrevNext;
 		}
-		*PrevLink = Next;
+		*PtrToPrevNext = Next;
 	}
 };
 
@@ -159,11 +156,18 @@ struct FMallocBinned2::PoolHashBucket
 
 struct FMallocBinned2::Private
 {
-	static_assert(ARRAY_COUNT(GMallocBinned2SmallBlockSizes) == SMALL_POOL_COUNT, "Small block size array size must match SMALL_POOL_COUNT");
-	static_assert(sizeof(FMallocBinned2::FPoolInfo)          == 32,               "sizeof(FPoolInfo) is expected to be 32 bytes");
+	enum
+	{
+		DEFAULT_BINNED_ALLOCATOR_ALIGNMENT = 16, // Default alignment for binned allocator
+	};
 
-	// Default alignment for binned allocator
-	enum { DEFAULT_BINNED_ALLOCATOR_ALIGNMENT = 16 };
+	FORCEINLINE static bool IsOSAllocation(const void* Ptr)
+	{
+		return IsAligned(Ptr, BINNED2_LARGE_ALLOC);
+	}
+
+	static_assert(ARRAY_COUNT(GMallocBinned2SmallBlockSizes) == SMALL_POOL_COUNT, "Small block size array size must match SMALL_POOL_COUNT");
+	static_assert(sizeof(FFreeBlock) <= DEFAULT_BINNED_ALLOCATOR_ALIGNMENT,       "Free block struct must be small enough to fit into a block.");
 
 	// Implementation. 
 	static CA_NO_RETURN void OutOfMemory(uint64 Size, uint32 Alignment=0)
@@ -172,28 +176,28 @@ struct FMallocBinned2::Private
 		FPlatformMemory::OnOutOfMemory(Size, Alignment);
 	}
 
-	/** 
-	 * Creates an array of FPoolInfo structures for tracking allocations.
-	 */
-	static FPoolInfo* CreatePoolArray(uint64 NumPools)
-	{
-		uint64 PoolArraySize = NumPools * sizeof(FPoolInfo);
-
-		void* Result = FPlatformMemory::BinnedAllocFromOS(PoolArraySize);
-		if (!Result)
-		{
-			OutOfMemory(PoolArraySize);
-		}
-
-		DefaultConstructItems<FPoolInfo>(Result, NumPools);
-		return (FPoolInfo*)Result;
-	}
-
 	/**
 	 * Gets the FPoolInfo for a memory address. If no valid info exists one is created.
 	 */
 	static FPoolInfo* GetOrCreatePoolInfo(FMallocBinned2& Allocator, void* InPtr)
 	{
+		/** 
+		 * Creates an array of FPoolInfo structures for tracking allocations.
+		 */
+		auto CreatePoolArray = [](uint64 NumPools)
+		{
+			uint64 PoolArraySize = NumPools * sizeof(FPoolInfo);
+
+			void* Result = FPlatformMemory::BinnedAllocFromOS(PoolArraySize);
+			if (!Result)
+			{
+				OutOfMemory(PoolArraySize);
+			}
+
+			DefaultConstructItems<FPoolInfo>(Result, NumPools);
+			return (FPoolInfo*)Result;
+		};
+
 		UPTRINT BucketIndex;
 		uint32  PoolIndex;
 		Allocator.PtrToPoolMapping.GetHashBucketAndPoolIndices(InPtr, BucketIndex, PoolIndex);
@@ -239,7 +243,6 @@ struct FMallocBinned2::Private
 
 		if (NextFree == NewBucket)
 		{
-			// Are we leaking memory here?
 			NextFree = nullptr;
 		}
 		Allocator.HashBucketFreeList = NextFree;
@@ -256,76 +259,83 @@ struct FMallocBinned2::Private
 		return &NewBucket->FirstPool[PoolIndex];
 	}
 
-	static FPoolInfo* FindPoolInfo(FMallocBinned2& Allocator, void* InPtr, void*& AllocationBase)
+	static FPoolInfo* FindPoolInfo(FMallocBinned2& Allocator, void* InPtr)
 	{
 		const UPTRINT PageSize = (UPTRINT)Allocator.PageSize;
 
-		uint8* Ptr = AlignDown((uint8*)InPtr, PageSize);
-		for (uint32 Repeat = 2; Repeat; --Repeat)
-		{
-			UPTRINT BucketIndex;
-			uint32  PoolIndex;
-			Allocator.PtrToPoolMapping.GetHashBucketAndPoolIndices(Ptr, BucketIndex, PoolIndex);
+		UPTRINT BucketIndex;
+		uint32  PoolIndex;
+		Allocator.PtrToPoolMapping.GetHashBucketAndPoolIndices(InPtr, BucketIndex, PoolIndex);
 
-			PoolHashBucket* FirstBucket = &Allocator.HashBuckets[BucketIndex];
-			PoolHashBucket* Collision   = FirstBucket;
-			for (;;)
+		PoolHashBucket* FirstBucket = &Allocator.HashBuckets[BucketIndex];
+		PoolHashBucket* Collision   = FirstBucket;
+		do
+		{
+			if (Collision->BucketIndex == BucketIndex)
 			{
-				if (Collision->BucketIndex == BucketIndex)
-				{
-					if (!Collision->FirstPool[PoolIndex].IsTrailingPool())
-					{
-						AllocationBase = Ptr;
-						return &Collision->FirstPool[PoolIndex];
-					}
-
-					Ptr -= PageSize * (Collision->FirstPool[PoolIndex].GetTrailingPoolOffset());
-					break;
-				}
-
-				Collision = Collision->Next;
-				if (Collision == FirstBucket)
-				{
-					Ptr -= PageSize;
-					break;
-				}
+				return &Collision->FirstPool[PoolIndex];
 			}
+
+			Collision = Collision->Next;
 		}
+		while (Collision != FirstBucket);
 
 		return nullptr;
 	}
 
-	static FORCEINLINE FPoolInfo* FindPoolInfo(FMallocBinned2& Allocator, void* Ptr)
+	static FFreeBlock* GetPoolHeaderFromPointer(void* Ptr)
 	{
-		void* BasePtr;
-		return FindPoolInfo(Allocator, Ptr, BasePtr);
+		return (FFreeBlock*)AlignDown(Ptr, BINNED2_LARGE_ALLOC);
 	}
+};
 
-	static void* AllocateBlockFromPool(FPoolTable& Table, FPoolInfo* Pool)
-	{
-		checkSlow(!Pool->IsOSAllocation());
-		checkSlow(Pool->FirstMem);
-		checkSlow(Pool->FirstMem->NumFreeBlocks > 0);
+FORCEINLINE bool FMallocBinned2::FPoolList::IsEmpty() const
+{
+	return Front == nullptr;
+}
 
-		// Pick first available block and unlink it.
-		++Pool->Taken;
-		--Pool->FirstMem->NumFreeBlocks;
+FORCEINLINE FMallocBinned2::FPoolInfo& FMallocBinned2::FPoolList::GetFrontPool()
+{
+	checkSlow(!IsEmpty());
+	return *Front;
+}
 
-		void* Free = (uint8*)Pool->FirstMem + Pool->FirstMem->NumFreeBlocks * Table.BlockSize;
-		if (Pool->FirstMem->NumFreeBlocks == 0)
-		{
-			Pool->FirstMem = Pool->FirstMem->NextFreePool;
-			if( !Pool->FirstMem )
-			{
-				// Move to exhausted list.
-				Pool->Unlink();
-				Pool->Link( Table.ExhaustedPool );
-			}
-		}
+FORCEINLINE const FMallocBinned2::FPoolInfo& FMallocBinned2::FPoolList::GetFrontPool() const
+{
+	checkSlow(!IsEmpty());
+	return *Front;
+}
+
+void FMallocBinned2::FPoolList::LinkToFront(FPoolInfo* Pool)
+{
+	Pool->Unlink();
+	Pool->Link(Front);
+}
 
-		return Free;
+FMallocBinned2::FPoolInfo& FMallocBinned2::FPoolList::PushNewPoolToFront(FMallocBinned2& Allocator, uint32 InBlockSize)
+{
+	const uint32 LocalPageSize = Allocator.PageSize;
+
+	// Allocate memory.
+	FFreeBlock* Free = (FFreeBlock*)Allocator.CachedOSPageAllocator.Allocate(LocalPageSize);
+	if (!Free)
+	{
+		Private::OutOfMemory(LocalPageSize);
 	}
-};
+
+	// Create pool
+	FPoolInfo* Result = Private::GetOrCreatePoolInfo(Allocator, Free);
+	Result->Link(Front);
+	Result->Taken          = 0;
+	Result->FirstFreeBlock = Free;
+
+	// Create first free item.
+	Free->NumFreeBlocks = LocalPageSize / InBlockSize - 1;
+	Free->NextFreeBlock = nullptr;
+	Free->BlockSize     = InBlockSize;
+
+	return *Result;
+}
 
 FMallocBinned2::FMallocBinned2(uint32 InPageSize, uint64 AddressLimit)
 	: PageSize          (InPageSize)
@@ -333,10 +343,12 @@ FMallocBinned2::FMallocBinned2(uint32 InPageSize, uint64 AddressLimit)
 	, PtrToPoolMapping  (PageSize, NumPoolsPerPage)
 	, HashBucketFreeList(nullptr)
 {
-	check(FMath::IsPowerOfTwo(PageSize));
-	check(FMath::IsPowerOfTwo(AddressLimit));
-	check(AddressLimit > PageSize); // Check to catch 32 bit overflow in AddressLimit
-	check(GMallocBinned2SmallBlockSizes[SMALL_POOL_COUNT - 1] == MAX_SMALL_POOLED_ALLOCATION_SIZE);
+	checkf(FMath::IsPowerOfTwo(PageSize),                                                           TEXT("OS page size must be a power of two"));
+	checkf(FMath::IsPowerOfTwo(AddressLimit),                                                       TEXT("OS address limit must be a power of two"));
+	checkf(AddressLimit > PageSize,                                                                 TEXT("OS address limit must be greater than the page size")); // Check to catch 32 bit overflow in AddressLimit
+	checkf(GMallocBinned2SmallBlockSizes[SMALL_POOL_COUNT - 1] == MAX_SMALL_POOLED_ALLOCATION_SIZE, TEXT("MAX_SMALL_POOLED_ALLOCATION_SIZE must equal the smallest block size"));
+	checkf(PageSize % BINNED2_LARGE_ALLOC == 0,                                                     TEXT("OS page size must be a multiple of BINNED2_LARGE_ALLOC"));
+	checkf(sizeof(FMallocBinned2::FFreeBlock) <= GMallocBinned2SmallBlockSizes[0],                  TEXT("Pool header must be able to fit into the smallest block"));
 
 	// Init pool tables.
 	for (uint32 Index = 0; Index != SMALL_POOL_COUNT; ++Index)
@@ -353,7 +365,7 @@ FMallocBinned2::FMallocBinned2(uint32 InPageSize, uint64 AddressLimit)
 	FPoolTable*  Pool       = SmallPoolTables;
 	for (uint32 Index = 0; Index != MAX_SMALL_POOLED_ALLOCATION_SIZE; ++Index)
 	{
-		while (Pool->BlockSize < Index)
+		while (Pool->BlockSize - 1 < Index)
 		{
 			++Pool;
 			check(Pool != SmallPoolTables + SMALL_POOL_COUNT);
@@ -379,75 +391,50 @@ bool FMallocBinned2::IsInternallyThreadSafe() const
 void* FMallocBinned2::Malloc(SIZE_T Size, uint32 Alignment)
 {
 	Alignment = FMath::Max<uint32>(Alignment, Private::DEFAULT_BINNED_ALLOCATOR_ALIGNMENT);
-	Size      = Align(Size, Alignment);
+	Size      = Align(FMath::Max((SIZE_T)1, Size), Alignment);
 
 	checkf(FMath::IsPowerOfTwo(Alignment), TEXT("Alignment must be a power of two"));
 	checkf(Alignment <= PageSize, TEXT("Alignment cannot exceed the system page size"));
 
 	// Only allocate from the small pools if the size is small enough and the alignment isn't crazy large.
 	// With large alignments, we'll waste a lot of memory allocating an entire page, but such alignments are highly unlikely in practice.
-	if (Size < MAX_SMALL_POOLED_ALLOCATION_SIZE && Alignment <= Private::DEFAULT_BINNED_ALLOCATOR_ALIGNMENT)
+	if (Size <= MAX_SMALL_POOLED_ALLOCATION_SIZE && Alignment <= Private::DEFAULT_BINNED_ALLOCATOR_ALIGNMENT)
 	{
 		// Allocate from small object pool.
-		FPoolTable* Table = MemSizeToPoolTable[Size];
-		FPoolInfo*  Pool  = Table->FirstPool;
-		if (!Pool)
-		{
-			// Allocate memory.
-			FFreeMem* Free = (FFreeMem*)CachedOSPageAllocator.Allocate(PageSize);
-			if (!Free)
-			{
-				Private::OutOfMemory(PageSize);
-			}
-
-			// Create pool in the pool page.
-			Pool = Private::GetOrCreatePoolInfo(*this, Free);
+		FPoolTable& Table = *MemSizeToPoolTable[Size - 1];
 
-			// Init pool.
-			Pool->Link(Table->FirstPool);
-			Pool->SetPoolAllocationSizes(PageSize, Size);
-			Pool->Taken    = 0;
-			Pool->FirstMem = Free;
+		FPoolInfo* Pool;
+		if (!Table.ActivePools.IsEmpty())
+		{
+			Pool = &Table.ActivePools.GetFrontPool();
+		}
+		else
+		{
+			Pool = &Table.ActivePools.PushNewPoolToFront(*this, Table.BlockSize);
+		}
 
-			// Create first free item.
-			Free->NumFreeBlocks = PageSize / Table->BlockSize;
-			Free->NextFreePool  = nullptr;
+		void* Result = Pool->AllocateBlock(Table.BlockSize);
+		if (Pool->IsExhausted())
+		{
+			Table.ExhaustedPools.LinkToFront(Pool);
 		}
 
-		void* Result = Private::AllocateBlockFromPool(*Table, Pool);
-		return Align(Result, Alignment);
+		return Result;
 	}
 
 	// Use OS for non-pooled allocations.
 	UPTRINT AlignedSize = Align(Size, PageSize);
-	FFreeMem* Result = (FFreeMem*)CachedOSPageAllocator.Allocate(AlignedSize);
+	void* Result = CachedOSPageAllocator.Allocate(AlignedSize);
 	if (!Result)
 	{
 		Private::OutOfMemory(AlignedSize);
 	}
 
-	// OS allocations should always be page-aligned.
-	checkSlow(IsAligned(Result, PageSize));
-
-	void* AlignedResult = Align(Result, Alignment);
-
 	// Create pool.
 	FPoolInfo* Pool = Private::GetOrCreatePoolInfo(*this, Result);
-	if (Result != AlignedResult)
-	{
-		// Mark the FPoolInfo for AlignedResult to jump back to the FPoolInfo for ptr.
-		for (UPTRINT i = (UPTRINT)PageSize, Offset = 1; i < AlignedSize; i += PageSize, ++Offset)
-		{
-			FPoolInfo* TrailingPool = Private::GetOrCreatePoolInfo(*this, (uint8*)Result + i);
-
-			// Set trailing pools to point back to first pool
-			TrailingPool->SetTrailingPoolOffset(Offset);
-		}
-	}
-
 	Pool->SetOSAllocationSizes(Size, AlignedSize);
 
-	return AlignedResult;
+	return Result;
 }
 
 void* FMallocBinned2::Realloc(void* Ptr, SIZE_T NewSize, uint32 Alignment)
@@ -470,30 +457,31 @@ void* FMallocBinned2::Realloc(void* Ptr, SIZE_T NewSize, uint32 Alignment)
 	checkf(FMath::IsPowerOfTwo(Alignment), TEXT("Alignment must be a power of two"));
 	checkf(Alignment <= PageSize, TEXT("Alignment cannot exceed the system page size"));
 
-	FPoolInfo* Pool = Private::FindPoolInfo(*this, Ptr);
-	if (!Pool->IsOSAllocation())
+	if (!Private::IsOSAllocation(Ptr))
 	{
 		// Reallocate to a smaller/bigger pool if necessary
-		FPoolTable* Table = MemSizeToPoolTable[Pool->TableIndex];
-		if (NewSize > Table->BlockSize || (Pool->TableIndex != 0 && NewSize <= Table[-1].BlockSize) || !IsAligned(Ptr, Alignment))
+		FFreeBlock* Free  = Private::GetPoolHeaderFromPointer(Ptr);
+		FPoolTable& Table = *MemSizeToPoolTable[Free->BlockSize - 1];
+		if (NewSize <= Table.BlockSize && (&Table == &SmallPoolTables[0] || NewSize > (&Table)[-1].BlockSize) && IsAligned(Ptr, Alignment))
 		{
-			// Reallocate and copy the data across
-			void* Result = FMallocBinned2::Malloc(NewSize, Alignment);
-			FMemory::Memcpy(Result, Ptr, FMath::Min<SIZE_T>(NewSize, Table->BlockSize));
-			FMallocBinned2::Free(Ptr);
-			return Result;
+			return Ptr;
 		}
 
-		return Ptr;
+		// Reallocate and copy the data across
+		void* Result = FMallocBinned2::Malloc(NewSize, Alignment);
+		FMemory::Memcpy(Result, Ptr, FMath::Min<SIZE_T>(NewSize, Table.BlockSize));
+		FMallocBinned2::Free(Ptr);
+		return Result;
 	}
 
 	// Allocated from OS.
-	UPTRINT PoolOsBytes = Pool->GetOsBytes();
-	if( NewSize > PoolOsBytes || NewSize * 3 < PoolOsBytes * 2 )
+	FPoolInfo* Pool        = Private::FindPoolInfo(*this, Ptr);
+	UPTRINT    PoolOsBytes = Pool->GetOsAllocatedBytes();
+	if (NewSize > PoolOsBytes || NewSize * 3 < PoolOsBytes * 2)
 	{
 		// Grow or shrink.
 		void* Result = FMallocBinned2::Malloc(NewSize, Alignment);
-		FMemory::Memcpy(Result, Ptr, FMath::Min<SIZE_T>(NewSize, Pool->AllocSize));
+		FMemory::Memcpy(Result, Ptr, FMath::Min<SIZE_T>(NewSize, Pool->GetOSRequestedBytes()));
 		FMallocBinned2::Free(Ptr);
 		return Result;
 	}
@@ -503,15 +491,14 @@ void* FMallocBinned2::Realloc(void* Ptr, SIZE_T NewSize, uint32 Alignment)
 	return Ptr;
 }
 
-void FMallocBinned2::Free( void* Ptr )
+void FMallocBinned2::Free(void* Ptr)
 {
 	if (!Ptr)
 	{
 		return;
 	}
 
-	void* BasePtr;
-	FPoolInfo* Pool = Private::FindPoolInfo(*this, Ptr, BasePtr);
+	FPoolInfo* Pool = Private::FindPoolInfo(*this, Ptr);
 #if PLATFORM_IOS || PLATFORM_MAC
 	if (Pool == NULL)
 	{
@@ -519,33 +506,29 @@ void FMallocBinned2::Free( void* Ptr )
 		return;
 	}
 #endif
+
 	checkSlow(Pool);
-	checkSlow(!Pool->IsTrailingPool());
-	if (!Pool->IsOSAllocation())
+	if (!Private::IsOSAllocation(Ptr))
 	{
-		FPoolTable* Table = MemSizeToPoolTable[Pool->TableIndex];
+		FFreeBlock* BasePtr = Private::GetPoolHeaderFromPointer(Ptr);
+		FPoolTable& Table   = *MemSizeToPoolTable[BasePtr->BlockSize - 1];
 
 		// If this pool was exhausted, move to available list.
-		if (!Pool->FirstMem)
+		if (!Pool->FirstFreeBlock)
 		{
-			Pool->Unlink();
-			Pool->Link(Table->FirstPool);
+			Table.ActivePools.LinkToFront(Pool);
 		}
 
-		uint32 AlignOffset = ((UPTRINT)Ptr - (UPTRINT)BasePtr) % Table->BlockSize;
-
-		// Patch pointer to include previously applied alignment.
-		Ptr = (uint8*)Ptr - AlignOffset;
-
 		// Free a pooled allocation.
-		FFreeMem* Free      = (FFreeMem*)Ptr;
-		Free->NumFreeBlocks = 1;
-		Free->NextFreePool  = Pool->FirstMem;
-		Pool->FirstMem      = Free;
+		FFreeBlock* Free     = (FFreeBlock*)Ptr;
+		Free->NumFreeBlocks  = 1;
+		Free->NextFreeBlock  = Pool->FirstFreeBlock;
+		Free->BlockSize      = Table.BlockSize;
+		Pool->FirstFreeBlock = Free;
 
 		// Free this pool.
 		checkSlow(Pool->Taken >= 1);
-		if( --Pool->Taken == 0 )
+		if (--Pool->Taken == 0)
 		{
 			// Free the OS memory.
 			Pool->Unlink();
@@ -555,45 +538,60 @@ void FMallocBinned2::Free( void* Ptr )
 	else
 	{
 		// Free an OS allocation.
-		checkSlow(IsAligned(Ptr, PageSize));
-		SIZE_T OsBytes = Pool->GetOsBytes();
-
-		CachedOSPageAllocator.Free(BasePtr, OsBytes);
+		CachedOSPageAllocator.Free(Ptr, Pool->GetOsAllocatedBytes());
 	}
 }
 
-bool FMallocBinned2::GetAllocationSize(void* Original, SIZE_T& SizeOut)
+bool FMallocBinned2::GetAllocationSize(void* Ptr, SIZE_T& SizeOut)
 {
-	if (!Original)
+	if (!Ptr)
 	{
 		return false;
 	}
 
-	FPoolInfo* Pool = Private::FindPoolInfo(*this, Original);
-	SizeOut = Pool->IsOSAllocation() ? Pool->AllocSize : MemSizeToPoolTable[Pool->TableIndex]->BlockSize;
+	if (!Private::IsOSAllocation(Ptr))
+	{
+		FFreeBlock* Free = Private::GetPoolHeaderFromPointer(Ptr);
+		SizeOut = MemSizeToPoolTable[Free->BlockSize - 1]->BlockSize;
+		return true;
+	}
+
+	FPoolInfo* Pool = Private::FindPoolInfo(*this, Ptr);
+	SizeOut = Pool->GetOSRequestedBytes();
 	return true;
 }
 
+void FMallocBinned2::FPoolList::Validate(TFunctionRef<void(FMallocBinned2::FPoolInfo&)> PoolValidate)
+{
+	for (FPoolInfo** PoolPtr = &Front; *PoolPtr; PoolPtr = &(*PoolPtr)->Next)
+	{
+		FPoolInfo* Pool = *PoolPtr;
+		check(Pool->PtrToPrevNext == PoolPtr);
+		PoolValidate(*Pool);
+	}
+}
+
 bool FMallocBinned2::ValidateHeap()
 {
 	for (FPoolTable& Table : SmallPoolTables)
 	{
-		for( FPoolInfo** PoolPtr = &Table.FirstPool; *PoolPtr; PoolPtr = &(*PoolPtr)->Next )
-		{
-			FPoolInfo* Pool = *PoolPtr;
-			check(Pool->PrevLink == PoolPtr);
-			check(Pool->FirstMem);
-			for( FFreeMem* Free = Pool->FirstMem; Free; Free = Free->NextFreePool )
+		Table.ActivePools.Validate(
+			[](FMallocBinned2::FPoolInfo& Pool)
 			{
-				check(Free->NumFreeBlocks > 0);
+				check(Pool.FirstFreeBlock);
+				for (FFreeBlock* Free = Pool.FirstFreeBlock; Free; Free = (FFreeBlock*)Free->NextFreeBlock)
+				{
+					check(Free->NumFreeBlocks > 0);
+				}
 			}
-		}
-		for( FPoolInfo** PoolPtr = &Table.ExhaustedPool; *PoolPtr; PoolPtr = &(*PoolPtr)->Next )
-		{
-			FPoolInfo* Pool = *PoolPtr;
-			check(Pool->PrevLink == PoolPtr);
-			check(!Pool->FirstMem);
-		}
+		);
+
+		Table.ExhaustedPools.Validate(
+			[](FMallocBinned2::FPoolInfo& Pool)
+			{
+				check(!Pool.FirstFreeBlock);
+			}
+		);
 	}
 
 	return true;
diff --git a/Engine/Source/Runtime/Core/Private/HTML5/HTML5PlatformMemory.cpp b/Engine/Source/Runtime/Core/Private/HTML5/HTML5PlatformMemory.cpp
--- a/Engine/Source/Runtime/Core/Private/HTML5/HTML5PlatformMemory.cpp
+++ b/Engine/Source/Runtime/Core/Private/HTML5/HTML5PlatformMemory.cpp
@@ -5,7 +5,7 @@
 =============================================================================*/
 
 #include "CorePrivatePCH.h"
-#include "MallocBinned2.h"
+#include "MallocBinned.h"
 #include "MallocAnsi.h"
 
 void FHTML5PlatformMemory::Init()
@@ -51,7 +51,7 @@ FMalloc* FHTML5PlatformMemory::BaseAllocator()
 #if !PLATFORM_HTML5_WIN32 
 	return new FMallocAnsi();
 #else 
-	return new FMallocBinned2(32 * 1024, 1 << 30 );
+	return new FMallocBinned(32 * 1024, 1 << 30 );
 #endif 
 }
 
diff --git a/Engine/Source/Runtime/Core/Private/IOS/IOSPlatformMemory.cpp b/Engine/Source/Runtime/Core/Private/IOS/IOSPlatformMemory.cpp
--- a/Engine/Source/Runtime/Core/Private/IOS/IOSPlatformMemory.cpp
+++ b/Engine/Source/Runtime/Core/Private/IOS/IOSPlatformMemory.cpp
@@ -5,7 +5,7 @@
 =============================================================================*/
 
 #include "CorePrivatePCH.h"
-#include "MallocBinned2.h"
+#include "MallocBinned.h"
 #include "MallocAnsi.h"
 
 
@@ -99,7 +99,7 @@ FMalloc* FIOSPlatformMemory::BaseAllocator()
 	uint64 MemoryLimit = FMath::Min<uint64>( uint64(1) << FMath::CeilLogTwo(Stats.free_count * PageSize), 0x100000000);
 
 	//return new FMallocAnsi();
-	return new FMallocBinned2(PageSize, MemoryLimit);
+	return new FMallocBinned(PageSize, MemoryLimit);
 }
 
 void* FIOSPlatformMemory::BinnedAllocFromOS( SIZE_T Size )
diff --git a/Engine/Source/Runtime/Core/Private/Linux/LinuxPlatformMemory.cpp b/Engine/Source/Runtime/Core/Private/Linux/LinuxPlatformMemory.cpp
--- a/Engine/Source/Runtime/Core/Private/Linux/LinuxPlatformMemory.cpp
+++ b/Engine/Source/Runtime/Core/Private/Linux/LinuxPlatformMemory.cpp
@@ -7,7 +7,7 @@
 #include "CorePrivatePCH.h"
 #include "MallocAnsi.h"
 #include "MallocJemalloc.h"
-#include "MallocBinned2.h"
+#include "MallocBinned.h"
 #include <sys/sysinfo.h>
 #include <sys/file.h>
 #include <sys/mman.h>
@@ -105,7 +105,7 @@ class FMalloc* FLinuxPlatformMemory::BaseAllocator()
 
 		default:	// intentional fall-through
 		case Binned:
-			Allocator = new FMallocBinned2(FPlatformMemory::GetConstants().PageSize & MAX_uint32, 0x100000000);
+			Allocator = new FMallocBinned(FPlatformMemory::GetConstants().PageSize & MAX_uint32, 0x100000000);
 			break;
 	}
 
diff --git a/Engine/Source/Runtime/Core/Private/Mac/MacPlatformMemory.cpp b/Engine/Source/Runtime/Core/Private/Mac/MacPlatformMemory.cpp
--- a/Engine/Source/Runtime/Core/Private/Mac/MacPlatformMemory.cpp
+++ b/Engine/Source/Runtime/Core/Private/Mac/MacPlatformMemory.cpp
@@ -7,7 +7,7 @@
 #include "CorePrivatePCH.h"
 #include "MallocTBB.h"
 #include "MallocAnsi.h"
-#include "MallocBinned2.h"
+#include "MallocBinned.h"
 #include "MallocStomp.h"
 
 #include <sys/param.h>
@@ -53,7 +53,7 @@ FMalloc* FMacPlatformMemory::BaseAllocator()
 #elif (WITH_EDITORONLY_DATA || IS_PROGRAM) && TBB_ALLOCATOR_ALLOWED
 		return new FMallocTBB();
 #else
-		return new FMallocBinned2((uint32)(GetConstants().PageSize&MAX_uint32), 0x100000000);
+		return new FMallocBinned((uint32)(GetConstants().PageSize&MAX_uint32), 0x100000000);
 #endif
 	}
 }
diff --git a/Engine/Source/Runtime/Core/Public/Containers/ContainerAllocationPolicies.h b/Engine/Source/Runtime/Core/Public/Containers/ContainerAllocationPolicies.h
--- a/Engine/Source/Runtime/Core/Public/Containers/ContainerAllocationPolicies.h
+++ b/Engine/Source/Runtime/Core/Public/Containers/ContainerAllocationPolicies.h
@@ -8,27 +8,57 @@
 
 
 /** Used to determine the alignment of an element type. */
-template<typename ElementType>
+template<typename ElementType, bool IsClass = TIsClass<ElementType>::Value>
 class TElementAlignmentCalculator
 {
-private:
 	/**
-	 * This is a dummy type that's used to calculate the padding added between the byte and the element
+	 * We use a dummy FAlignedElement type that's used to calculate the padding added between the byte and the element
 	 * to fulfill the type's required alignment.
+	 *
+	 * Its default constructor and destructor are declared but never implemented to avoid the need for a ElementType default constructor.
+	 */
+
+private:
+	/**
+	 * In the case of class ElementTypes, we inherit it to allow abstract types to work.
+	 */
+	struct FAlignedElements : ElementType
+	{
+		uint8 MisalignmentPadding;
+
+		FAlignedElements();
+		~FAlignedElements();
+	};
+
+	// We calculate the alignment here and then handle the zero case in the result by forwarding it to the non-class variant.
+	// This is necessary because the compiler can perform empty-base-optimization to eliminate a redundant ElementType state.
+	// Forwarding it to the non-class implementation should always work because an abstract type should never be empty.
+	enum { CalculatedAlignment = sizeof(FAlignedElements) - sizeof(ElementType) };
+
+public:
+	enum { Value = TChooseClass<CalculatedAlignment != 0, TIntegralConstant<SIZE_T, CalculatedAlignment>, TElementAlignmentCalculator<ElementType, false>>::Result::Value };
+};
+
+template<typename ElementType>
+class TElementAlignmentCalculator<ElementType, false>
+{
+private:
+	/**
+	 * In the case of non-class ElementTypes, we contain it because non-class types cannot be inherited.
 	 */
 	struct FAlignedElements
 	{
 		uint8 MisalignmentPadding;
 		ElementType Element;
-		/** FAlignedElement's default constructor and destructor are declared but never implemented to avoid the need for a ElementType default constructor. */
+
 		FAlignedElements();
 		~FAlignedElements();
 	};
 public:
-	enum { Alignment = sizeof(FAlignedElements) - sizeof(ElementType) };
+	enum { Value = sizeof(FAlignedElements) - sizeof(ElementType) };
 };
 
-#define ALIGNOF(T) (TElementAlignmentCalculator<T>::Alignment)
+#define ALIGNOF(T) (TElementAlignmentCalculator<T>::Value)
 
 
 /** branchless pointer selection
diff --git a/Engine/Source/Runtime/Core/Public/HAL/MallocBinned2.h b/Engine/Source/Runtime/Core/Public/HAL/MallocBinned2.h
--- a/Engine/Source/Runtime/Core/Public/HAL/MallocBinned2.h
+++ b/Engine/Source/Runtime/Core/Public/HAL/MallocBinned2.h
@@ -3,6 +3,7 @@
 #pragma once
 
 #include "Allocators/CachedOSPageAllocator.h"
+#include "Function.h"
 
 #define BINNED2_MAX_CACHED_OS_FREES (64)
 #if PLATFORM_64BITS
@@ -27,24 +28,38 @@ private:
 	enum { MAX_LARGE_POOLED_PAGE_FRACTIONS  = 2 };
 
 	// Forward declares.
-	struct FFreeMem;
-	struct FPoolTable;
+	struct FFreeBlock;
 	struct FPoolInfo;
+	struct FPoolTable;
 	struct PoolHashBucket;
 
+	struct FPoolList
+	{
+		FPoolList();
+
+		bool IsEmpty() const;
+
+		      FPoolInfo& GetFrontPool();
+		const FPoolInfo& GetFrontPool() const;
+
+		void LinkToFront(FPoolInfo* Pool);
+
+		FPoolInfo& PushNewPoolToFront(FMallocBinned2& Allocator, uint32 InBytes);
+
+		void Validate(TFunctionRef<void(FPoolInfo&)>);
+
+	private:
+		FPoolInfo* Front;
+	};
+
 	/** Pool table. */
 	struct FPoolTable
 	{
-		FPoolInfo* FirstPool;
-		FPoolInfo* ExhaustedPool;
-		uint32     BlockSize;
-
-		FPoolTable()
-			: FirstPool(nullptr)
-			, ExhaustedPool(nullptr)
-			, BlockSize(0)
-		{
-		}
+		FPoolList ActivePools;
+		FPoolList ExhaustedPools;
+		uint32    BlockSize;
+
+		FPoolTable();
 	};
 
 	const uint32 PageSize;
@@ -115,7 +130,7 @@ public:
 	/** 
 	 * Malloc
 	 */
-	virtual void* Malloc( SIZE_T Size, uint32 Alignment ) override;
+	virtual void* Malloc(SIZE_T Size, uint32 Alignment) override;
 
 	/** 
 	 * Realloc
diff --git a/Engine/Source/Runtime/Core/Public/Templates/UniquePtr.h b/Engine/Source/Runtime/Core/Public/Templates/UniquePtr.h
--- a/Engine/Source/Runtime/Core/Public/Templates/UniquePtr.h
+++ b/Engine/Source/Runtime/Core/Public/Templates/UniquePtr.h
@@ -3,6 +3,7 @@
 #pragma once
 
 #include "UnrealTemplate.h"
+#include "Templates/EnableIf.h"
 
 // Single-ownership smart pointer in the vein of std::unique_ptr.
 // Use this when you need an object's lifetime to be strictly bound to the lifetime of a single smart pointer.
@@ -14,9 +15,37 @@
 // TUniquePtr<MyClass> Ptr3(MoveTemp(Ptr1)); // Ptr3 now owns the MyClass object - Ptr1 is now nullptr.
 
 template <typename T>
-class TUniquePtr
+struct TDefaultDelete
 {
-	template <typename OtherT>
+	TDefaultDelete()
+	{
+	}
+
+	template <typename U, typename = typename TEnableIf<TPointerIsConvertibleFromTo<U, T>::Value>::Type>
+	TDefaultDelete(TDefaultDelete<U>&&)
+	{
+	}
+
+	template <typename U, typename = typename TEnableIf<TPointerIsConvertibleFromTo<U, T>::Value>::Type>
+	TDefaultDelete& operator=(TDefaultDelete<U>&&)
+	{
+		return *this;
+	}
+
+	~TDefaultDelete()
+	{
+	}
+
+	void operator()(T* Ptr) const
+	{
+		delete Ptr;
+	}
+};
+
+template <typename T, typename Deleter = TDefaultDelete<T>>
+class TUniquePtr : private Deleter
+{
+	template <typename OtherT, typename OtherDeleter>
 	friend class TUniquePtr;
 
 public:
@@ -50,7 +79,8 @@ public:
 	 * Move constructor
 	 */
 	FORCEINLINE TUniquePtr(TUniquePtr&& Other)
-		: Ptr(Other.Ptr)
+		: Deleter(MoveTemp(Other.GetDeleter()))
+		, Ptr    (Other.Ptr)
 	{
 		Other.Ptr = nullptr;
 	}
@@ -58,9 +88,10 @@ public:
 	/**
 	 * Constructor from rvalues of other (usually derived) types
 	 */
-	template <typename OtherT>
-	FORCEINLINE TUniquePtr(TUniquePtr<OtherT>&& Other)
-		: Ptr(Other.Ptr)
+	template <typename OtherT, typename OtherDeleter>
+	FORCEINLINE TUniquePtr(TUniquePtr<OtherT, OtherDeleter>&& Other)
+		: Deleter(MoveTemp(Other.GetDeleter()))
+		, Ptr    (Other.Ptr)
 	{
 		Other.Ptr = nullptr;
 	}
@@ -76,23 +107,27 @@ public:
 			T* OldPtr = Ptr;
 			Ptr = Other.Ptr;
 			Other.Ptr = nullptr;
-			delete OldPtr;
+			GetDeleter()(OldPtr);
 		}
 
+		GetDeleter() = MoveTemp(Other.GetDeleter());
+
 		return *this;
 	}
 
 	/**
 	 * Assignment operator for rvalues of other (usually derived) types
 	 */
-	template <typename OtherT>
+	template <typename OtherT, typename OtherDeleter>
 	FORCEINLINE TUniquePtr& operator=(TUniquePtr<OtherT>&& Other)
 	{
 		// We delete last, because we don't want odd side effects if the destructor of T relies on the state of this or Other
 		T* OldPtr = Ptr;
 		Ptr = Other.Ptr;
 		Other.Ptr = nullptr;
-		delete OldPtr;
+		GetDeleter()(OldPtr);
+
+		GetDeleter() = MoveTemp(Other.GetDeleter());
 
 		return *this;
 	}
@@ -105,7 +140,7 @@ public:
 		// We delete last, because we don't want odd side effects if the destructor of T relies on the state of this
 		T* OldPtr = Ptr;
 		Ptr = nullptr;
-		delete OldPtr;
+		GetDeleter()(OldPtr);
 
 		return *this;
 	}
@@ -115,7 +150,7 @@ public:
 	 */
 	FORCEINLINE ~TUniquePtr()
 	{
-		delete Ptr;
+		GetDeleter()(Ptr);
 	}
 
 	/**
@@ -200,7 +235,27 @@ public:
 		// We delete last, because we don't want odd side effects if the destructor of T relies on the state of this
 		T* OldPtr = Ptr;
 		Ptr = InPtr;
-		delete OldPtr;
+		GetDeleter()(OldPtr);
+	}
+
+	/**
+	 * Returns a reference to the deleter subobject.
+	 *
+	 * @return A reference to the deleter.
+	 */
+	FORCEINLINE Deleter& GetDeleter()
+	{
+		return static_cast<Deleter&>(*this);
+	}
+
+	/**
+	 * Returns a reference to the deleter subobject.
+	 *
+	 * @return A reference to the deleter.
+	 */
+	FORCEINLINE const Deleter& GetDeleter() const
+	{
+		return static_cast<const Deleter&>(*this);
 	}
 
 private:
diff --git a/Engine/Source/Runtime/Core/Public/Templates/UnrealTemplate.h b/Engine/Source/Runtime/Core/Public/Templates/UnrealTemplate.h
--- a/Engine/Source/Runtime/Core/Public/Templates/UnrealTemplate.h
+++ b/Engine/Source/Runtime/Core/Public/Templates/UnrealTemplate.h
@@ -524,6 +524,15 @@ struct TIdentity
 	typedef T Type;
 };
 
+/**
+ * Defines a value metafunction of the given Value.
+ */
+template <typename T, T Val>
+struct TIntegralConstant
+{
+	static const T Value = Val;
+};
+
 /**
  * Does LHS::Value && RHS::Value, but short-circuits if LHS::Value == false.
  */
@@ -580,3 +589,18 @@ struct TBoolConstant
  */
 template <typename T>
 T&& DeclVal();
+
+
+/**
+ * Determines if T is a struct/class type
+ */
+template <typename T>
+struct TIsClass
+{
+private:
+	template <typename U> static uint16 Func(int U::*);
+	template <typename U> static uint8  Func(...);
+
+public:
+	enum { Value = !__is_union(T) && sizeof(Func<T>(0)) - 1 };
+};
diff --git a/Engine/Source/Runtime/CoreUObject/Private/Misc/StartupPackages.cpp b/Engine/Source/Runtime/CoreUObject/Private/Misc/StartupPackages.cpp
--- a/Engine/Source/Runtime/CoreUObject/Private/Misc/StartupPackages.cpp
+++ b/Engine/Source/Runtime/CoreUObject/Private/Misc/StartupPackages.cpp
@@ -11,13 +11,7 @@
 void FStartupPackages::GetStartupPackageNames(TArray<FString>& PackageNames, const FString& EngineConfigFilename, bool bIsCreatingHashes)
 {
 	// if we aren't cooking, we actually just want to use the cooked startup package as the only startup package
-	if (bIsCreatingHashes || (!IsRunningCommandlet() && !GIsEditor && GUseSeekFreeLoading))
-	{
-		// make sure the startup package exists
-		PackageNames.Add(FString(TEXT("Startup_LOC")));
-		PackageNames.Add(FString(TEXT("Startup")));
-	}
-	else
+	if (FPlatformProperties::RequiresCookedData())
 	{
 		// look for any packages that we want to force preload at startup
 		FConfigSection* PackagesToPreload = GConfig->GetSectionPrivate(TEXT("Engine.StartupPackages"), 0, 1, EngineConfigFilename);
diff --git a/Engine/Source/Runtime/CoreUObject/Private/Serialization/AsyncLoading.cpp b/Engine/Source/Runtime/CoreUObject/Private/Serialization/AsyncLoading.cpp
--- a/Engine/Source/Runtime/CoreUObject/Private/Serialization/AsyncLoading.cpp
+++ b/Engine/Source/Runtime/CoreUObject/Private/Serialization/AsyncLoading.cpp
@@ -949,7 +949,9 @@ void FAsyncLoadingThread::SuspendLoading()
 {
 	UE_CLOG(!IsInGameThread(), LogStreaming, Fatal, TEXT("Async loading can only be suspended from the main thread"));
 	const int32 SuspendCount = IsLoadingSuspended.Increment();
+#if !WITH_EDITORONLY_DATA
 	UE_LOG(LogStreaming, Display, TEXT("Suspending async loading (%d)"), SuspendCount);
+#endif
 	if (IsMultithreaded() && SuspendCount == 1)
 	{
 		ThreadSuspendedEvent->Wait();
@@ -1907,6 +1909,8 @@ EAsyncPackageState::Type FAsyncPackage::PostLoadObjects()
 	return Result;
 }
 
+void CreateClustersFromPackage(FLinkerLoad* PackageLinker);
+
 EAsyncPackageState::Type FAsyncPackage::PostLoadDeferredObjects(double InTickStartTime, bool bInUseTimeLimit, float& InOutTimeLimit)
 {
 	SCOPE_CYCLE_COUNTER(STAT_FAsyncPackage_PostLoadObjectsGameThread);
@@ -1982,11 +1986,13 @@ EAsyncPackageState::Type FAsyncPackage::PostLoadDeferredObjects(double InTickSta
 		if (LinkerRoot && !bLoadHasFailed)
 		{
 			LinkerRoot->AtomicallyClearInternalFlags(EInternalObjectFlags::AsyncLoading);
-			LinkerRoot->MarkAsFullyLoaded();
+			LinkerRoot->MarkAsFullyLoaded();			
 			LinkerRoot->SetLoadTime(FPlatformTime::Seconds() - LoadStartTime);
 
 			if (Linker)
 			{
+				CreateClustersFromPackage(Linker);
+
 				// give a hint to the IO system that we are done with this file for now
 				FIOSystem::Get().HintDoneWithFile(Linker->Filename);
 			}
diff --git a/Engine/Source/Runtime/CoreUObject/Private/UObject/FastReferenceCollector.h b/Engine/Source/Runtime/CoreUObject/Private/UObject/FastReferenceCollector.h
--- a/Engine/Source/Runtime/CoreUObject/Private/UObject/FastReferenceCollector.h
+++ b/Engine/Source/Runtime/CoreUObject/Private/UObject/FastReferenceCollector.h
@@ -41,7 +41,7 @@
 		 void ReturnToPool(TArray<UObject*>* Array);
 	 };
  */
-template <typename ReferenceProcessorType, typename CollectorType, typename ArrayPoolType>
+template <typename ReferenceProcessorType, typename CollectorType, typename ArrayPoolType, bool bAutoGenerateTokenStream = false>
 class TFastReferenceCollector
 {
 private:
@@ -210,6 +210,14 @@ private:
 				//@todo rtgc: we need to handle object references in struct defaults
 
 				// Make sure that token stream has been assembled at this point as the below code relies on it.
+				if (bAutoGenerateTokenStream && !ReferenceProcessor.IsRunningMultithreaded())
+				{
+					UClass* ObjectClass = CurrentObject->GetClass();
+					if (!ObjectClass->HasAnyClassFlags(CLASS_TokenStreamAssembled))
+					{
+						ObjectClass->AssembleReferenceTokenStream();
+					}
+				}
 				checkSlow(CurrentObject->GetClass()->HasAnyClassFlags(CLASS_TokenStreamAssembled));
 
 				// Get pointer to token stream and jump to the start.
diff --git a/Engine/Source/Runtime/CoreUObject/Private/UObject/GarbageCollection.cpp b/Engine/Source/Runtime/CoreUObject/Private/UObject/GarbageCollection.cpp
--- a/Engine/Source/Runtime/CoreUObject/Private/UObject/GarbageCollection.cpp
+++ b/Engine/Source/Runtime/CoreUObject/Private/UObject/GarbageCollection.cpp
@@ -14,7 +14,7 @@
    Garbage collection.
 -----------------------------------------------------------------------------*/
 
-#define PERF_DETAILED_PER_CLASS_GC_STATS				(LOOKING_FOR_PERF_ISSUES)
+#define PERF_DETAILED_PER_CLASS_GC_STATS				(LOOKING_FOR_PERF_ISSUES || 0) 
 
 // FastReferenceCollector uses PERF_DETAILED_PER_CLASS_GC_STATS
 #include "UObject/FastReferenceCollector.h"
@@ -338,10 +338,35 @@ public:
 #endif
 	}
 
+	/** Marks all objects that can't be directly in a cluster but are referenced by it as reachable */
 	template <bool bParallel>
-	FORCEINLINE void MarkReferencedClustersAsReachable(int32 ClusterRootIndex)
+	FORCEINLINE void MarkClusterMutableObjectsAsReachable(FUObjectCluster* Cluster, TArray<UObject*>& ObjectsToSerialize)
+	{
+		for (int32 ReferencedMutableObjectIndex : Cluster->MutableObjects)
+		{
+			FUObjectItem* ReferencedMutableObjectItem = GUObjectArray.IndexToObjectUnsafeForGC(ReferencedMutableObjectIndex);
+			if (bParallel)
+			{
+				if (ReferencedMutableObjectItem->IsUnreachable() && ReferencedMutableObjectItem->ThisThreadAtomicallyClearedRFUnreachable())
+				{
+					ObjectsToSerialize.Add(static_cast<UObject*>(ReferencedMutableObjectItem->Object));
+				}
+			}
+			else if (ReferencedMutableObjectItem->IsUnreachable())
+			{
+				ReferencedMutableObjectItem->ClearUnreachable();
+				ObjectsToSerialize.Add(static_cast<UObject*>(ReferencedMutableObjectItem->Object));
+			}
+		}
+	}
+
+	/** Marks all clusters referenced by another cluster as rechable */
+	template <bool bParallel>
+	FORCEINLINE void MarkReferencedClustersAsReachable(int32 ClusterRootIndex, TArray<UObject*>& ObjectsToSerialize)
 	{
 		FUObjectCluster* Cluster = GUObjectClusters.FindChecked(ClusterRootIndex);
+		// Also mark all referenced objects from outside of the cluster as reachable
+		MarkClusterMutableObjectsAsReachable<bParallel>(Cluster, ObjectsToSerialize);
 		for (int32 ReferncedClusterIndex : Cluster->ReferencedClusters)
 		{
 			FUObjectItem* ReferencedClusterRootObjectItem = GUObjectArray.IndexToObjectUnsafeForGC(ReferncedClusterIndex);
@@ -354,7 +379,9 @@ public:
 			{
 				ReferencedClusterRootObjectItem->ClearUnreachable();
 			}
-		}
+			FUObjectCluster* ReferencedCluster = GUObjectClusters.FindChecked(ReferncedClusterIndex);
+			MarkClusterMutableObjectsAsReachable<bParallel>(ReferencedCluster, ObjectsToSerialize);
+		}		
 	}
 
 	/**
@@ -409,7 +436,7 @@ public:
 					{
 						// This is a cluster root reference so mark all referenced clusters as reachable
 						const int32 ObjectIndex = GUObjectArray.ObjectToIndex(Object);
-						MarkReferencedClustersAsReachable<true>(ObjectIndex);
+						MarkReferencedClustersAsReachable<true>(ObjectIndex, ObjectsToSerialize);
 					}
 				}
 			}
@@ -440,7 +467,7 @@ public:
 				{
 					// This is a cluster root reference so mark all referenced clusters as reachable
 					const int32 ObjectIndex = GUObjectArray.ObjectToIndex(Object);
-					MarkReferencedClustersAsReachable<false>(ObjectIndex);
+					MarkReferencedClustersAsReachable<false>(ObjectIndex, ObjectsToSerialize);
 				}
 			}
 		}
@@ -456,20 +483,21 @@ public:
 				if (RootObjectItem->ThisThreadAtomicallyClearedRFUnreachable())
 				{
 					// Make sure all referenced clusters are marked as reachable too
-					MarkReferencedClustersAsReachable<true>(OwnerIndex);
+					MarkReferencedClustersAsReachable<true>(OwnerIndex, ObjectsToSerialize);
 				}
 			}
 			else if (RootObjectItem->IsUnreachable())
 			{
 				RootObjectItem->ClearUnreachable();
 				// Make sure all referenced clusters are marked as reachable too
-				MarkReferencedClustersAsReachable<false>(OwnerIndex);
+				MarkReferencedClustersAsReachable<false>(OwnerIndex, ObjectsToSerialize);
 			}
 		}
 
-		if (bStrongReference)
+		// The second condition seems to improve perf for multithreaded GC
+		if (bStrongReference && ObjectItem->HasAnyFlags(EInternalObjectFlags::NoStrongReference))
 		{
-			ObjectItem->ClearNoStrongReference();
+			ObjectItem->ThisThreadAtomicallyClearedFlag(EInternalObjectFlags::NoStrongReference);
 		}
 #if PERF_DETAILED_PER_CLASS_GC_STATS
 		GCurrentObjectRegularObjectRefs++;
@@ -660,7 +688,7 @@ bool IsTokenStreamDirty();
  * interface doesn't make sense to implement for.
  */
 class FRealtimeGC
-	{
+{
 public:
 	/** Default constructor, initializing all members. */
 	FRealtimeGC()
@@ -733,16 +761,16 @@ public:
 			{
 				// Compile will strip this out when we don't need to update the token stream since this is a template.
 				// Otherwise the GC will suffer a big performance hit.
-			if (UClass* Class = dynamic_cast<UClass*>(Object))
-			{
-				if (!Class->HasAnyClassFlags(CLASS_TokenStreamAssembled))
+				if (UClass* Class = dynamic_cast<UClass*>(Object))
 				{
-					Class->AssembleReferenceTokenStream();
-					check(Class->HasAnyClassFlags(CLASS_TokenStreamAssembled));
+					if (!Class->HasAnyClassFlags(CLASS_TokenStreamAssembled))
+					{
+						Class->AssembleReferenceTokenStream();
+						check(Class->HasAnyClassFlags(CLASS_TokenStreamAssembled));
+					}
 				}
 			}
 		}
-		}
 	}
 
 	/**
@@ -752,7 +780,7 @@ public:
 	 */
 	void PerformReachabilityAnalysis(EObjectFlags KeepFlags, bool bForceSingleThreaded = false)
 	{		
-		DECLARE_SCOPE_CYCLE_COUNTER( TEXT( "FRealtimeGC::PerformReachabilityAnalysis" ), STAT_FArchiveRealtimeGC_PerformReachabilityAnalysis, STATGROUP_GC );
+		DECLARE_SCOPE_CYCLE_COUNTER(TEXT("FRealtimeGC::PerformReachabilityAnalysis"), STAT_FArchiveRealtimeGC_PerformReachabilityAnalysis, STATGROUP_GC);
 
 		/** Growing array of objects that require serialization */
 		TArray<UObject*>& ObjectsToSerialize = *FGCArrayPool::Get().GetArrayFromPool();
@@ -1173,7 +1201,7 @@ void CollectGarbageInternal(EObjectFlags KeepFlags, bool bPerformFullPurge)
 					{
 						UE_LOG(LogGarbage, Warning, TEXT("Disregard for GC object %s referencing %s which is not part of root set"),
 							*Object->GetFullName(),
-							*ReferencedObject->GetFullName());
+							*ReferencedObject->GetFullName());						
 						bShouldAssert = true;
 					}
 				}
diff --git a/Engine/Source/Runtime/CoreUObject/Private/UObject/LinkerLoad.cpp b/Engine/Source/Runtime/CoreUObject/Private/UObject/LinkerLoad.cpp
--- a/Engine/Source/Runtime/CoreUObject/Private/UObject/LinkerLoad.cpp
+++ b/Engine/Source/Runtime/CoreUObject/Private/UObject/LinkerLoad.cpp
@@ -2906,7 +2906,7 @@ void FLinkerLoad::LoadAllObjects( bool bForcePreload )
 	// Mark package as having been fully loaded.
 	if( LinkerRoot )
 	{
-		LinkerRoot->MarkAsFullyLoaded();
+		LinkerRoot->MarkAsFullyLoaded();		
 	}
 }
 
@@ -3828,10 +3828,14 @@ UObject* FLinkerLoad::CreateExport( int32 Index )
 			EInternalObjectFlags::None,
 			Template
 		);
-		if (FPlatformProperties::RequiresCookedData() && GIsInitialLoad)
+		if (FPlatformProperties::RequiresCookedData())
 		{
-			Export.Object->AddToRoot();
+			if (GIsInitialLoad || GUObjectArray.IsOpenForDisregardForGC())
+			{
+				Export.Object->AddToRoot();
+			}
 		}
+		
 		LoadClass = Export.Object->GetClass(); // this may have changed if we are overwriting a CDO component
 
 		if (NewName != Export.ObjectName)
diff --git a/Engine/Source/Runtime/CoreUObject/Private/UObject/Obj.cpp b/Engine/Source/Runtime/CoreUObject/Private/UObject/Obj.cpp
--- a/Engine/Source/Runtime/CoreUObject/Private/UObject/Obj.cpp
+++ b/Engine/Source/Runtime/CoreUObject/Private/UObject/Obj.cpp
@@ -727,12 +727,6 @@ void UObject::ConditionalPostLoad()
 			PostLoad();
 		}
 
-		extern int32 GCreateGCClusters;
-		if (GCreateGCClusters && FPlatformProperties::RequiresCookedData() && !GIsInitialLoad && CanBeClusterRoot())
-		{
-			CreateCluster();
-		}
-
 #if !(UE_BUILD_SHIPPING || UE_BUILD_TEST)
 		if (ThreadContext.DebugPostLoad.Contains(this))
 		{
@@ -3773,50 +3767,6 @@ void StaticExit()
 	UE_LOG(LogExit, Log, TEXT("Object subsystem successfully closed.") );
 }
 
-void MarkObjectsToDisregardForGC()
-{
-	// Iterate over all class objects and force the default objects to be created. Additionally also
-	// assembles the token reference stream at this point. This is required for class objects that are
-	// not taken into account for garbage collection but have instances that are.
-	
-	// Workaround for Visual Studio 2013 analyzer bug. Using a temporary directly in the range-for
-	// errors if the analyzer is enabled.
-	TObjectRange<UClass> Range;
-	for( auto* Class : Range )
-	{
-		// Force the default object to be created.
-		Class->GetDefaultObject(); // Force the default object to be constructed if it isn't already
-		// Assemble reference token stream for garbage collection/ RTGC.
-		Class->AssembleReferenceTokenStream();
-	}
-
-	// Iterate over all objects and mark them to be part of root set.
-	int32 NumAlwaysLoadedObjects = 0;
-	int32 NumRootObjects = 0;
-	for( FObjectIterator It; It; ++It )
-	{
-		UObject* Object = *It;
-		if (Object->IsSafeForRootSet())
-		{
-			NumRootObjects++;
-			Object->AddToRoot();
-		}
-		else if (Object->IsRooted())
-		{
-			Object->RemoveFromRoot();
-		}
-		NumAlwaysLoadedObjects++;
-	}
-
-	UE_LOG(LogObj, Log, TEXT("%i objects as part of root set at end of initial load."), NumAlwaysLoadedObjects);
-	if (GUObjectArray.DisregardForGCEnabled())
-	{
-		UE_LOG(LogObj, Log, TEXT("%i objects are not in the root set, but can never be destroyed because they are in the DisregardForGC set."), NumAlwaysLoadedObjects - NumRootObjects);
-	}
-	GUObjectAllocator.BootMessage();
-}					
-
-
 
 /*-----------------------------------------------------------------------------
 	Misc.
diff --git a/Engine/Source/Runtime/CoreUObject/Private/UObject/ReferenceChainSearch.cpp b/Engine/Source/Runtime/CoreUObject/Private/UObject/ReferenceChainSearch.cpp
--- a/Engine/Source/Runtime/CoreUObject/Private/UObject/ReferenceChainSearch.cpp
+++ b/Engine/Source/Runtime/CoreUObject/Private/UObject/ReferenceChainSearch.cpp
@@ -178,6 +178,18 @@ void FReferenceChainSearch::PrintReferencers( FReferenceChain& Referencer )
 			ObjectReachability += TEXT("(NeverGCed) ");
 		}
 
+		FUObjectItem* ReferencedByObjectItem = GUObjectArray.ObjectToObjectItem(RefInfo.ReferencedBy);
+		bool bClusterRoot = false;
+		if (ReferencedByObjectItem->HasAnyFlags(EInternalObjectFlags::ClusterRoot))
+		{
+			ObjectReachability += TEXT("(ClusterRoot) ");
+			bClusterRoot = true;
+		}
+		if (ReferencedByObjectItem->GetOwnerIndex())
+		{
+			ObjectReachability += TEXT("(Clustered) ");
+		}
+
 		FString Indent; Indent = Indent.LeftPad(RefLevel*2);
 		UE_LOG(LogReferenceChain, Log, TEXT("%s%s%s->%s"), *Indent, *ObjectReachability, *RefInfo.ReferencedBy->GetFullName(), *ReferencedThrough);
 
diff --git a/Engine/Source/Runtime/CoreUObject/Private/UObject/UObjectArray.cpp b/Engine/Source/Runtime/CoreUObject/Private/UObject/UObjectArray.cpp
--- a/Engine/Source/Runtime/CoreUObject/Private/UObject/UObjectArray.cpp
+++ b/Engine/Source/Runtime/CoreUObject/Private/UObject/UObjectArray.cpp
@@ -10,26 +10,96 @@ DEFINE_LOG_CATEGORY_STATIC(LogUObjectArray, Log, All);
 
 TMap<int32, FUObjectCluster*> GUObjectClusters;
 
-void FUObjectArray::AllocateObjectPool(int32 MaxUObjects, int32 MaxObjectsNotConsideredByGC)
+void FUObjectArray::AllocateObjectPool(int32 InMaxUObjects, int32 InMaxObjectsNotConsideredByGC)
 {
 	check(IsInGameThread());
 
+	MaxObjectsNotConsideredByGC = InMaxObjectsNotConsideredByGC;
+
 	// GObjFirstGCIndex is the index at which the garbage collector will start for the mark phase.
-	ObjFirstGCIndex = MaxObjectsNotConsideredByGC;
+	// If disregard for GC is enabled this will be set to an invalid value so that later we
+	// know if disregard for GC pool has already been closed (at least once)
+	ObjFirstGCIndex = DisregardForGCEnabled() ? -1 : 0;
 
 	// Pre-size array.
 	check(ObjObjects.Num() == 0);
-	ObjObjects.PreAllocate(MaxUObjects);
+	ObjObjects.PreAllocate(InMaxUObjects);
+
+	if (MaxObjectsNotConsideredByGC > 0)
+	{
+		ObjObjects.AddRange(MaxObjectsNotConsideredByGC);
+	}
+}
+
+void FUObjectArray::OpenDisregardForGC()
+{
+	check(IsInGameThread());
+	check(!OpenForDisregardForGC);
+	OpenForDisregardForGC = true;
+	UE_LOG(LogUObjectArray, Log, TEXT("OpenDisregardForGC: %d/%d objects in disregard for GC pool"), ObjLastNonGCIndex + 1, MaxObjectsNotConsideredByGC);
 }
 
 void FUObjectArray::CloseDisregardForGC()
 {
 	check(IsInGameThread());
+	check(OpenForDisregardForGC);
+
+	// Iterate over all class objects and force the default objects to be created. Additionally also
+	// assembles the token reference stream at this point. This is required for class objects that are
+	// not taken into account for garbage collection but have instances that are.
+
+	// Workaround for Visual Studio 2013 analyzer bug. Using a temporary directly in the range-for
+	// errors if the analyzer is enabled.
+	TObjectRange<UClass> Range;
+	for (UClass* Class : Range)
+	{
+		// Force the default object to be created.
+		Class->GetDefaultObject(); // Force the default object to be constructed if it isn't already
+		// Assemble reference token stream for garbage collection/ RTGC.
+		if (!Class->HasAnyClassFlags(CLASS_TokenStreamAssembled))
+		{
+			Class->AssembleReferenceTokenStream();
+		}
+	}
+
+	if (GIsInitialLoad)
+	{
+		// Iterate over all objects and mark them to be part of root set.
+		int32 NumAlwaysLoadedObjects = 0;
+		int32 NumRootObjects = 0;
+		for (FObjectIterator It; It; ++It)
+		{
+			UObject* Object = *It;
+			if (Object->IsSafeForRootSet())
+			{
+				NumRootObjects++;
+				Object->AddToRoot();
+			}
+			else if (Object->IsRooted())
+			{
+				Object->RemoveFromRoot();
+			}
+			NumAlwaysLoadedObjects++;
+		}
+
+		UE_LOG(LogUObjectArray, Log, TEXT("%i objects as part of root set at end of initial load."), NumAlwaysLoadedObjects);
+		if (GUObjectArray.DisregardForGCEnabled())
+		{
+			UE_LOG(LogUObjectArray, Log, TEXT("%i objects are not in the root set, but can never be destroyed because they are in the DisregardForGC set."), NumAlwaysLoadedObjects - NumRootObjects);
+		}
+
+		// When disregard for GC pool is closed for the first time, make sure the first GC index is set after the last non-GC index.
+		// We do allow here for some slack if MaxObjectsNotConsideredByGC > (ObjLastNonGCIndex + 1) so that disregard for GC pool
+		// can be re-opened later.
+		ObjFirstGCIndex = FMath::Max(ObjFirstGCIndex, ObjLastNonGCIndex + 1);
+
+		GUObjectAllocator.BootMessage();
+	}
+
+	UE_LOG(LogUObjectArray, Log, TEXT("CloseDisregardForGC: %d/%d objects in disregard for GC pool"), ObjLastNonGCIndex + 1, MaxObjectsNotConsideredByGC);	
 
 	OpenForDisregardForGC = false;
 	GIsInitialLoad = false;
-	// Make sure the first GC index matches the last non-GC index after DisregardForGC is closed
-	ObjFirstGCIndex = ObjLastNonGCIndex + 1;
 }
 
 void FUObjectArray::AllocateUObjectIndex(UObjectBase* Object, bool bMergingThreads /*= false*/)
@@ -42,9 +112,16 @@ void FUObjectArray::AllocateUObjectIndex(UObjectBase* Object, bool bMergingThrea
 	{
 		// Disregard from GC pool is only available from the game thread, at least for now
 		check(IsInGameThread());
-		Index = ObjObjects.AddSingle();
-		ObjLastNonGCIndex = Index;
-		ObjFirstGCIndex = FMath::Max(ObjFirstGCIndex, Index + 1);
+		Index = ++ObjLastNonGCIndex;
+		// Check if we're not out of bounds, unless there hasn't been any gc objects yet
+		UE_CLOG(ObjLastNonGCIndex >= MaxObjectsNotConsideredByGC && ObjFirstGCIndex >= 0, LogUObjectArray, Fatal, TEXT("Unable to add more objects to disregard for GC pool (Max: %d)"), MaxObjectsNotConsideredByGC);
+		// If we haven't added any GC objects yet, it's fine to keep growing the disregard pool past its initial size.
+		if (ObjLastNonGCIndex >= MaxObjectsNotConsideredByGC)
+		{
+			Index = ObjObjects.AddSingle();
+			check(Index == ObjLastNonGCIndex);
+		}
+		MaxObjectsNotConsideredByGC = FMath::Max(MaxObjectsNotConsideredByGC, ObjLastNonGCIndex + 1);
 	}
 	// Regular pool/ range.
 	else
@@ -61,14 +138,16 @@ void FUObjectArray::AllocateUObjectIndex(UObjectBase* Object, bool bMergingThrea
 		}
 		else
 		{
+			// Make sure ObjFirstGCIndex is valid, otherwise we didn't close the disregard for GC set
+			check(ObjFirstGCIndex >= 0);
 #if THREADSAFE_UOBJECTS
 			FScopeLock ObjObjectsLock(&ObjObjectsCritical);
 #else
 			check(IsInGameThread());
 #endif
-			Index = ObjObjects.AddSingle();
+			Index = ObjObjects.AddSingle();			
 		}
-		check(Index >= ObjFirstGCIndex);
+		check(Index >= ObjFirstGCIndex && Index > ObjLastNonGCIndex);
 	}
 	// Add to global table.
 	if (FPlatformAtomics::InterlockedCompareExchangePointer((void**)&ObjObjects[Index].Object, Object, NULL) != NULL) // we use an atomic operation to check for unexpected concurrency, verify alignment, etc
diff --git a/Engine/Source/Runtime/CoreUObject/Private/UObject/UObjectBase.cpp b/Engine/Source/Runtime/CoreUObject/Private/UObject/UObjectBase.cpp
--- a/Engine/Source/Runtime/CoreUObject/Private/UObject/UObjectBase.cpp
+++ b/Engine/Source/Runtime/CoreUObject/Private/UObject/UObjectBase.cpp
@@ -880,8 +880,14 @@ void UObjectBaseInit()
 	}
 	else
 	{
+#if IS_PROGRAM
+		// Maximum number of UObjects for programs can be low
+		MaxUObjects = 100000; // Default to 100K for programs
+		GConfig->GetInt(TEXT("/Script/Engine.GarbageCollectionSettings"), TEXT("gc.MaxObjectsInProgram"), MaxUObjects, GEngineIni);
+#else
 		// Maximum number of UObjects in the editor
 		GConfig->GetInt(TEXT("/Script/Engine.GarbageCollectionSettings"), TEXT("gc.MaxObjectsInEditor"), MaxUObjects, GEngineIni);
+#endif
 	}
 
 	// Log what we're doing to track down what really happens as log in LaunchEngineLoop doesn't report those settings in pristine form.
diff --git a/Engine/Source/Runtime/CoreUObject/Private/UObject/UObjectClusters.cpp b/Engine/Source/Runtime/CoreUObject/Private/UObject/UObjectClusters.cpp
--- a/Engine/Source/Runtime/CoreUObject/Private/UObject/UObjectClusters.cpp
+++ b/Engine/Source/Runtime/CoreUObject/Private/UObject/UObjectClusters.cpp
@@ -8,6 +8,7 @@
 #include "Interface.h"
 #include "ModuleManager.h"
 #include "FastReferenceCollector.h"
+#include "GarbageCollection.h"
 
 int32 GCreateGCClusters = 1;
 static FAutoConsoleVariableRef CCreateGCClusters(
@@ -226,9 +227,10 @@ public:
 	 */
 	void AddObjectToCluster(int32 ObjectIndex, FUObjectItem* ObjectItem, UObject* Obj, TArray<UObject*>& ObjectsToSerialize, bool bOuterAndClass)
 	{
-		if (ObjectItem->GetOwnerIndex() == 0 && !ObjectItem->IsRootSet() && !GUObjectArray.IsDisregardForGC(Obj) && Obj->CanBeInCluster())
+		if (ObjectIndex != ClusterRootIndex && ObjectItem->GetOwnerIndex() == 0 && !ObjectItem->IsRootSet() && !GUObjectArray.IsDisregardForGC(Obj) && Obj->CanBeInCluster())
 		{
 			ObjectsToSerialize.Add(Obj);
+			check(!ObjectItem->HasAnyFlags(EInternalObjectFlags::ClusterRoot));
 			ObjectItem->SetOwnerIndex(ClusterRootIndex);
 			Cluster.Objects.Add(ObjectIndex);
 
@@ -242,11 +244,10 @@ public:
 				}
 				if (!Obj->GetClass()->HasAllClassFlags(CLASS_Native))
 				{
-					int32 ClassIndex = GUObjectArray.ObjectToIndex(Obj->GetClass());
-					AddObjectToCluster(ClassIndex, GUObjectArray.IndexToObjectUnsafeForGC(ClassIndex), Obj->GetClass(), ObjectsToSerialize, false);
-
-					int32 ClassOuterIndex = GUObjectArray.ObjectToIndex(Obj->GetClass()->GetOuter());
-					AddObjectToCluster(ClassOuterIndex, GUObjectArray.IndexToObjectUnsafeForGC(ClassOuterIndex), Obj->GetClass()->GetOuter(), ObjectsToSerialize, false);
+					UObject* ObjectClass = Obj->GetClass();
+					HandleTokenStreamObjectReference(ObjectsToSerialize, Obj, ObjectClass, INDEX_NONE, true);
+					UObject* ObjectClassOuter = Obj->GetClass()->GetOuter();
+					HandleTokenStreamObjectReference(ObjectsToSerialize, Obj, ObjectClassOuter, INDEX_NONE, true);
 				}
 			}
 		}		
@@ -329,12 +330,24 @@ public:
 							check(OtherClusterReferencedCluster != ClusterRootIndex);
 							Cluster.ReferencedClusters.AddUnique(OtherClusterReferencedCluster);
 						}
+						for (int32 OtherClusterReferencedMutableObjectIndex : OtherCluster->MutableObjects)
+						{
+							Cluster.MutableObjects.AddUnique(OtherClusterReferencedMutableObjectIndex);
+						}
 					}
 				}
-				else if (ObjectItem->GetOwnerIndex() == 0 && !ObjectItem->IsRootSet() && !GUObjectArray.IsDisregardForGC(Object) && Object->CanBeInCluster())
+				else if (ObjectItem->GetOwnerIndex() == 0 && !ObjectItem->IsRootSet() && !GUObjectArray.IsDisregardForGC(Object) &&
+					!(Object->CanBeClusterRoot() && Object->HasAnyFlags(RF_NeedLoad|RF_NeedPostLoad))) // Objects that can create clusters themselves and haven't been postloaded yet should be excluded
 				{
 					// New object, add it to the cluster.
-					AddObjectToCluster(GUObjectArray.ObjectToIndex(Object), ObjectItem, Object, ObjectsToSerialize, true);
+					if (Object->CanBeInCluster())
+					{
+						AddObjectToCluster(GUObjectArray.ObjectToIndex(Object), ObjectItem, Object, ObjectsToSerialize, true);
+					}
+					else
+					{
+						Cluster.MutableObjects.AddUnique(GUObjectArray.ObjectToIndex(Object));
+					}
 				}
 			}
 		}
@@ -378,21 +391,78 @@ public:
 	}
 };
 
+/** Looks through objects loaded with a package and creates clusters from them */
+void CreateClustersFromPackage(FLinkerLoad* PackageLinker)
+{	
+	if (FPlatformProperties::RequiresCookedData() && !GIsInitialLoad && GCreateGCClusters && !GUObjectArray.IsOpenForDisregardForGC())
+	{
+		check(PackageLinker);
+
+		for (FObjectExport& Export : PackageLinker->ExportMap)
+		{
+			if (Export.Object && Export.Object->CanBeClusterRoot())
+			{
+				Export.Object->CreateCluster();
+			}
+		}
+	}
+}
+
+void UObjectBaseUtility::AddToCluster(UObjectBaseUtility* ClusterRootOrObjectFromCluster)
+{
+	check(ClusterRootOrObjectFromCluster);
+
+	const int32 OuterIndex = GUObjectArray.ObjectToIndex(ClusterRootOrObjectFromCluster);
+	FUObjectItem* OuterItem = GUObjectArray.IndexToObjectUnsafeForGC(OuterIndex);
+	int32 ClusterRootIndex = 0;
+	if (OuterItem->HasAnyFlags(EInternalObjectFlags::ClusterRoot))
+	{
+		ClusterRootIndex = OuterIndex;
+	}
+	else
+	{
+		ClusterRootIndex = OuterItem->GetOwnerIndex();
+	}
+	if (ClusterRootIndex != 0)
+	{
+		FUObjectCluster* Cluster = GUObjectClusters.FindChecked(ClusterRootIndex);
+		FClusterReferenceProcessor Processor(ClusterRootIndex, *Cluster);			
+		TFastReferenceCollector<FClusterReferenceProcessor, TClusterCollector<FClusterReferenceProcessor>, FClusterArrayPool, true> ReferenceCollector(Processor, FClusterArrayPool::Get());
+		TArray<UObject*> ObjectsToProcess;
+		UObject* ThisObject = static_cast<UObject*>(this);
+		Processor.HandleTokenStreamObjectReference(ObjectsToProcess, static_cast<UObject*>(ClusterRootOrObjectFromCluster), ThisObject, INDEX_NONE, true);
+		if (ObjectsToProcess.Num())
+		{
+			ReferenceCollector.CollectReferences(ObjectsToProcess, true);
+		}
+	}
+}
+
+bool UObjectBaseUtility::CanBeInCluster() const
+{
+	return true;
+}
+
 void UObjectBaseUtility::CreateCluster()
 {
+	DECLARE_SCOPE_CYCLE_COUNTER(TEXT("UObjectBaseUtility::CreateCluster"), STAT_FArchiveRealtimeGC_CreateCluster, STATGROUP_GC);
+
 	FUObjectItem* RootItem = GUObjectArray.IndexToObject(InternalIndex);
-	if (RootItem->GetOwnerIndex() != 0)
+	if (RootItem->GetOwnerIndex() != 0 || RootItem->HasAnyFlags(EInternalObjectFlags::ClusterRoot))
 	{
 		return;
 	}
 
+	// If we haven't finished loading, we can't be sure we know all the references
+	check(!HasAnyFlags(RF_NeedLoad | RF_NeedPostLoad));
+
 	// Create a new cluster, reserve an arbitrary amount of memory for it.
 	FUObjectCluster* Cluster = new FUObjectCluster;
 	Cluster->Objects.Reserve(64);
 
 	// Collect all objects referenced by cluster root and by all objects it's referencing
 	FClusterReferenceProcessor Processor(InternalIndex, *Cluster);
-	TFastReferenceCollector<FClusterReferenceProcessor, TClusterCollector<FClusterReferenceProcessor>, FClusterArrayPool> ReferenceCollector(Processor, FClusterArrayPool::Get());
+	TFastReferenceCollector<FClusterReferenceProcessor, TClusterCollector<FClusterReferenceProcessor>, FClusterArrayPool, true> ReferenceCollector(Processor, FClusterArrayPool::Get());
 	TArray<UObject*> ObjectsToProcess;
 	ObjectsToProcess.Add(static_cast<UObject*>(this));
 	ReferenceCollector.CollectReferences(ObjectsToProcess, true);
@@ -404,6 +474,7 @@ void UObjectBaseUtility::CreateCluster()
 	{
 		// Add new cluster to the global cluster map.
 		GUObjectClusters.Add(InternalIndex, Cluster);
+		check(RootItem->GetOwnerIndex() == 0);
 		RootItem->SetFlags(EInternalObjectFlags::ClusterRoot);
 	}
 	else
@@ -420,6 +491,7 @@ class FClusterVerifyReferenceProcessor
 {	
 	const UObject* const ClusterRootObject;
 	const int32 ClusterRootIndex;
+	const FUObjectCluster& Cluster;
 	volatile bool bIsRunningMultithreaded;
 	bool bFailed;
 	TSet<UObject*> ProcessedObjects;
@@ -429,6 +501,7 @@ public:
 	FClusterVerifyReferenceProcessor(UObject* InClusterRootObject)		
 		: ClusterRootObject(InClusterRootObject)
 		, ClusterRootIndex(GUObjectArray.ObjectToIndex(InClusterRootObject))
+		, Cluster(*GUObjectClusters.FindChecked(ClusterRootIndex)) // This can't fail otherwise there's something wrong with cluster creation code
 		, bIsRunningMultithreaded(false)
 		, bFailed(false)
 	{
@@ -475,19 +548,57 @@ public:
 	*/
 	FORCEINLINE void HandleTokenStreamObjectReference(TArray<UObject*>& ObjectsToSerialize, UObject* ReferencingObject, UObject*& Object, const int32 TokenIndex, bool bAllowReferenceElimination)
 	{
+#if !(UE_BUILD_TEST || UE_BUILD_SHIPPING)
+		if (Object && !Object->IsValidLowLevelFast())
+		{
+			FString TokenDebugInfo;
+			if (UClass *Class = (ReferencingObject ? ReferencingObject->GetClass() : nullptr))
+			{
+				const FTokenInfo& TokenInfo = Class->DebugTokenMap.GetTokenInfo(TokenIndex);
+				TokenDebugInfo = FString::Printf(TEXT("ReferencingObjectClass: %s, Property Name: %s, Offset: %d"),
+					*Class->GetFullName(), *TokenInfo.Name.GetPlainNameString(), TokenInfo.Offset);
+			}
+			else
+			{
+				// This means this objects is most likely being referenced by AddReferencedObjects
+				TokenDebugInfo = TEXT("Native Reference");
+			}
+
+			UE_LOG(LogObj, Fatal, TEXT("Invalid object while verifying cluster assumptions: 0x%016llx, ReferencingObject: %s, %s, TokenIndex: %d"),
+				(int64)(PTRINT)Object,
+				ReferencingObject ? *ReferencingObject->GetFullName() : TEXT("NULL"),
+				*TokenDebugInfo, TokenIndex);
+		}
+#endif
 		if (Object && !ProcessedObjects.Contains(Object))
 		{
 			ProcessedObjects.Add(Object);
 			FUObjectItem* ObjectItem = GUObjectArray.ObjectToObjectItem(Object);
 			if (ObjectItem->GetOwnerIndex() == 0)
 			{
-				if (!ObjectItem->HasAnyFlags(EInternalObjectFlags::ClusterRoot|EInternalObjectFlags::RootSet) && !GUObjectArray.IsDisregardForGC(Object))
+				// We are allowed to reference other clusters, root set objects and objects from diregard for GC pool
+				if (!ObjectItem->HasAnyFlags(EInternalObjectFlags::ClusterRoot|EInternalObjectFlags::RootSet) && !GUObjectArray.IsDisregardForGC(Object) && Object->CanBeInCluster())
 				{
-					UE_LOG(LogObj, Warning, TEXT("Object %s from cluster %s is referencing %s which is not part of root set or cluster."),
+					UE_LOG(LogObj, Warning, TEXT("Object %s from cluster %s is referencing 0x%016llx %s which is not part of root set or cluster."),
 						*ReferencingObject->GetFullName(),
 						*ClusterRootObject->GetFullName(),
+						(int64)(PTRINT)Object,
 						*Object->GetFullName());
 					bFailed = true;
+#if UE_BUILD_DEBUG
+					FReferenceChainSearch RefChainSearch(Object, FReferenceChainSearch::ESearchMode::Shortest | FReferenceChainSearch::ESearchMode::PrintResults);
+#endif
+				}
+				else if (ObjectItem->HasAnyFlags(EInternalObjectFlags::ClusterRoot))
+				{
+					// However, clusters need to be referenced by the current cluster otherwise they can also get GC'd too early.
+					const int32 OtherClusterRootIndex = GUObjectArray.ObjectToIndex(Object);
+					UE_CLOG(OtherClusterRootIndex != ClusterRootIndex && !Cluster.ReferencedClusters.Contains(OtherClusterRootIndex), LogObj, Fatal,
+						TEXT("Object %s from source cluster %s is referencing cluster root object 0x%016llx %s which is not referenced by the source cluster."),
+						*ReferencingObject->GetFullName(),
+						*ClusterRootObject->GetFullName(),
+						(int64)(PTRINT)Object,
+						*Object->GetFullName());
 				}
 			}
 			else if (ObjectItem->GetOwnerIndex() == ClusterRootIndex)
@@ -495,6 +606,21 @@ public:
 				// If this object belongs to the current cluster, keep processing its references. Otherwise ignore it as it will be processed by its cluster
 				ObjectsToSerialize.Add(Object);
 			}
+			else
+			{
+				// If we're referencing an object from another cluster, make sure the other cluster is actually referenced by this cluster
+				const int32 OtherClusterRootIndex = ObjectItem->GetOwnerIndex();
+				const FUObjectItem* OtherClusterRootItem = GUObjectArray.IndexToObjectUnsafeForGC(OtherClusterRootIndex);				
+				check(OtherClusterRootItem && OtherClusterRootItem->Object);
+				UObject* OtherClusterRootObject = static_cast<UObject*>(OtherClusterRootItem->Object);
+				UE_CLOG(OtherClusterRootIndex != ClusterRootIndex && !Cluster.ReferencedClusters.Contains(OtherClusterRootIndex), LogObj, Fatal,
+					TEXT("Object %s from source cluster %s is referencing cluster %d object 0x%016llx %s which is not referenced by the source cluster."),
+					*ReferencingObject->GetFullName(),
+					*ClusterRootObject->GetFullName(),
+					*OtherClusterRootObject->GetFullName(),
+					(int64)(PTRINT)Object,
+					*Object->GetFullName());
+			}
 		}
 	}
 };
diff --git a/Engine/Source/Runtime/CoreUObject/Private/UObject/UObjectGlobals.cpp b/Engine/Source/Runtime/CoreUObject/Private/UObject/UObjectGlobals.cpp
--- a/Engine/Source/Runtime/CoreUObject/Private/UObject/UObjectGlobals.cpp
+++ b/Engine/Source/Runtime/CoreUObject/Private/UObject/UObjectGlobals.cpp
@@ -1460,6 +1460,20 @@ void EndLoad()
 				}
 			}
 
+			// Create clusters after all objects have been loaded
+			extern int32 GCreateGCClusters;
+			if (FPlatformProperties::RequiresCookedData() && !GIsInitialLoad && GCreateGCClusters && !GUObjectArray.IsOpenForDisregardForGC())
+			{
+				for (UObject* Obj : ObjLoaded)
+				{
+					check(Obj);
+					if (Obj->CanBeClusterRoot())
+					{
+						Obj->CreateCluster();
+					}
+				}
+			}
+
 #if WITH_EDITOR
 			// Send global notification for each object that was loaded.
 			// Useful for updating UI such as ContentBrowser's loaded status.
@@ -2833,7 +2847,8 @@ UObject* StaticConstructObject_Internal
 
 	// Subobjects are always created in the constructor, no need to re-create them unless their archetype != CDO or they're blueprint generated.
 	// If the existing subobject is to be re-used it can't have BeginDestroy called on it so we need to pass this information to StaticAllocateObject.	
-	const bool bIsNativeFromCDO = InClass->HasAnyClassFlags(CLASS_Native | CLASS_Intrinsic) && 
+	const bool bIsNativeClass = InClass->HasAnyClassFlags(CLASS_Native | CLASS_Intrinsic);
+	const bool bIsNativeFromCDO = bIsNativeClass &&
 		(
 			!InTemplate || 
 			(InName != NAME_None && InTemplate == UObject::GetArchetypeFromRequiredInfo(InClass, InOuter, InName, InFlags))
diff --git a/Engine/Source/Runtime/CoreUObject/Public/UObject/UObjectArray.h b/Engine/Source/Runtime/CoreUObject/Public/UObject/UObjectArray.h
--- a/Engine/Source/Runtime/CoreUObject/Public/UObject/UObjectArray.h
+++ b/Engine/Source/Runtime/CoreUObject/Public/UObject/UObjectArray.h
@@ -59,6 +59,11 @@ struct FUObjectItem
 		ClusterAndFlags &= ~int32(FlagsToClear);
 	}
 
+	/**
+	 * Uses atomics to clear the specified flag(s).
+	 * @param FlagsToClear
+	 * @return True if this call cleared the flag, false if it has been cleared by another thread.
+	 */
 	FORCEINLINE bool ThisThreadAtomicallyClearedFlag(EInternalObjectFlags FlagToClear)
 	{
 		static_assert(sizeof(int32) == sizeof(ClusterAndFlags), "Flags must be 32-bit for atomics.");
@@ -71,14 +76,20 @@ struct FUObjectItem
 				break;
 			}
 			int32 OldValue = (int32)FPlatformAtomics::InterlockedCompareExchange((int32*)&ClusterAndFlags, StartValue & ~int32(FlagToClear), StartValue);
-			if (OldValue == StartValue)
+			// We know the flag was set when we entered this iteration,
+			// so if the old value returned by atomics had the flag set, we must have cleared it.
+			// (there is always a chance that another thread cleared some other flag and the above function did nothing)
+			// But we only care about the flags we want to clear
+			if (!(ClusterAndFlags & int32(FlagToClear)) && (OldValue & int32(FlagToClear)) == (StartValue & int32(FlagToClear)))
 			{
+				// if (the flag has actually been cleared) && (the previous value had the flag set) we must have cleared it
 				bIChangedIt = true;
 				break;
 			}
-			// Remove later.
-			checkSlow(OldValue == (StartValue & ~int32(FlagToClear)));
+			// We didn't clear the flag, probably because some other thread changed flags in the meantime (either the one we want to clear or some other). Try again.
 		}
+		// Make sure the flag was actually cleared
+		checkSlow((ClusterAndFlags & int32(FlagToClear)) == 0);
 		return bIChangedIt;
 	}
 
@@ -198,6 +209,17 @@ public:
 		return Result;
 	}
 
+	int32 AddRange(int32 Count)
+	{
+		int32 Result = NumElements + Count - 1;
+		checkf(NumElements + Count <= MaxElements, TEXT("Maximum number of Objects exceeded, make sure you update MaxObjectsInGame/MaxObjectsInEditor in project settings."));
+		check(Result == (NumElements + Count - 1));
+		NumElements += Count;
+		FPlatformMisc::MemoryBarrier();
+		check(Objects[Result].Object == nullptr);
+		return Result;
+	}
+
 	FORCEINLINE FUObjectItem const* GetObjectPtr(int32 Index) const
 	{
 		check(Index >= 0 && Index < NumElements);
@@ -327,20 +349,30 @@ public:
 		ObjFirstGCIndex = 0;
 	}
 
+	/**
+	* If there's enough slack in the disregard pool, we can re-open it and keep adding objects to it
+	*/
+	void OpenDisregardForGC();
+
 	/**
 	 * After the initial load, this closes the disregard pool so that new object are GC-able
-	 *
 	 */
 	void CloseDisregardForGC();
 
+	/** Returns true if the disregard for GC pool is open */
+	bool IsOpenForDisregardForGC() const
+	{
+		return OpenForDisregardForGC;
+	}
+
 	/**
 	 * indicates if the disregard for GC optimization is active
 	 *
-	 * @return true if the first GC index is greater than zero; this indicates that the disregard for GC optimization is enabled
+	 * @return true if MaxObjectsNotConsideredByGC is greater than zero; this indicates that the disregard for GC optimization is enabled
 	 */
 	bool DisregardForGCEnabled() const 
 	{ 
-		return ObjFirstGCIndex > 0; 
+		return MaxObjectsNotConsideredByGC > 0;
 	}
 
 	/**
@@ -690,8 +722,11 @@ private:
 	int32 ObjFirstGCIndex;
 	/** Index pointing to last object created in range disregarded for GC.					*/
 	int32 ObjLastNonGCIndex;
+	/** Maximum number of objects in the disregard for GC Pool */
+	int32 MaxObjectsNotConsideredByGC;
+
 	/** If true this is the intial load and we should load objects int the disregarded for GC range.	*/
-	int32 OpenForDisregardForGC;
+	bool OpenForDisregardForGC;
 	/** Array of all live objects.											*/
 	TUObjectArray ObjObjects;
 	/** Synchronization object for all live objects.											*/
@@ -718,10 +753,15 @@ private:
 	FThreadSafeCounter	MasterSerialNumber;
 };
 
+/** UObject cluster. Groups UObjects into a single unit for GC. */
 struct FUObjectCluster
 {
+	/** Objects that belong to this cluster */
 	TArray<int32> Objects;
+	/** Other clusters referenced by this cluster */
 	TArray<int32> ReferencedClusters;
+	/** Objects that could not be added to the cluster but still need to be referenced by it */
+	TArray<int32> MutableObjects;
 };
 
 /** Global UObject allocator							*/
diff --git a/Engine/Source/Runtime/CoreUObject/Public/UObject/UObjectBaseUtility.h b/Engine/Source/Runtime/CoreUObject/Public/UObject/UObjectBaseUtility.h
--- a/Engine/Source/Runtime/CoreUObject/Public/UObject/UObjectBaseUtility.h
+++ b/Engine/Source/Runtime/CoreUObject/Public/UObject/UObjectBaseUtility.h
@@ -286,6 +286,8 @@ public:
 	 */
 	FString GetPathName( const UObject* StopOuter=NULL ) const;
 
+public:
+
 	/**
 	* Called after load to determine if the object can be a cluster root
 	*
@@ -301,23 +303,29 @@ public:
 	*
 	* @return	true if this object can be inside of a cluster
 	*/
-	virtual bool CanBeInCluster() const
-	{
-		return true;
-	}
+	virtual bool CanBeInCluster() const;
+
+	/**
+	* Called during PostLoad to create UObject cluster
+	*/
+	virtual void CreateCluster();
+
+	/** 
+	 * Adds this objects to a GC cluster that already exists 
+	 * @param ClusterRootOrObjectFromCluster Object that belongs to the cluster we want to add this object to.
+	 */
+	virtual void AddToCluster(UObjectBaseUtility* ClusterRootOrObjectFromCluster);
 
 protected:
 
+	/** Helper function to create a cluster from UObject */
+	static void CreateClusterFromObject(UObjectBaseUtility* ClusterRootObject, UObjectBaseUtility* ReferencingObject);
+
 	/**
 	 * Internal version of GetPathName() that eliminates lots of copies.
 	 */
 	void GetPathName( const UObject* StopOuter, FString& ResultString ) const;
 
-	/**
-	* Called during PostLoad to create UObject cluster
-	*/
-	virtual void CreateCluster();
-
 public:
 	/**
 	 * Walks up the chain of packages until it reaches the top level, which it ignores.
diff --git a/Engine/Source/Runtime/CoreUObject/Public/UObject/UObjectGlobals.h b/Engine/Source/Runtime/CoreUObject/Public/UObject/UObjectGlobals.h
--- a/Engine/Source/Runtime/CoreUObject/Public/UObject/UObjectGlobals.h
+++ b/Engine/Source/Runtime/CoreUObject/Public/UObject/UObjectGlobals.h
@@ -277,7 +277,6 @@ COREUOBJECT_API void PreInitUObject();
 /** 
  *   Iterate over all objects considered part of the root to setup GC optimizations
  */
-COREUOBJECT_API void MarkObjectsToDisregardForGC();
 COREUOBJECT_API bool StaticExec( UWorld* InWorld, const TCHAR* Cmd, FOutputDevice& Ar=*GLog );
 COREUOBJECT_API void StaticTick( float DeltaTime, bool bUseFullTimeLimit = true, float AsyncLoadingTime = 0.005f );
 
diff --git a/Engine/Source/Runtime/Engine/Private/UnrealEngine.cpp b/Engine/Source/Runtime/Engine/Private/UnrealEngine.cpp
--- a/Engine/Source/Runtime/Engine/Private/UnrealEngine.cpp
+++ b/Engine/Source/Runtime/Engine/Private/UnrealEngine.cpp
@@ -5419,12 +5419,13 @@ bool UEngine::HandleObjCommand( const TCHAR* Cmd, FOutputDevice& Ar )
 			&&	(InsidePackage	||	!FCString::Strfind(Cmd,TEXT("PACKAGE="))) 
 			&&	(InsideObject	||	!FCString::Strfind(Cmd,TEXT("INSIDE=")))))
 		{
-			const bool bTrackDetailedObjectInfo		= bAll || (CheckType != NULL && CheckType != UObject::StaticClass()) || CheckOuter != NULL || InsideObject != NULL || InsidePackage != NULL || !ObjectName.IsEmpty();
-			const bool bOnlyListGCObjects				= FParse::Param( Cmd, TEXT("GCONLY") );
-			const bool bOnlyListRootObjects				= FParse::Param( Cmd, TEXT("ROOTONLY") );
-			const bool bShouldIncludeDefaultObjects	= FParse::Param( Cmd, TEXT("INCLUDEDEFAULTS") );
-			const bool bOnlyListDefaultObjects			= FParse::Param( Cmd, TEXT("DEFAULTSONLY") );
-			const bool bShowDetailedObjectInfo			= FParse::Param( Cmd, TEXT("NODETAILEDINFO") ) == false && bTrackDetailedObjectInfo;
+			const bool bTrackDetailedObjectInfo = bAll || (CheckType != NULL && CheckType != UObject::StaticClass()) || CheckOuter != NULL || InsideObject != NULL || InsidePackage != NULL || !ObjectName.IsEmpty();
+			const bool bOnlyListGCObjects = FParse::Param(Cmd, TEXT("GCONLY"));
+			const bool bOnlyListGCObjectsNoClusters = FParse::Param(Cmd, TEXT("GCNOCLUSTERS"));
+			const bool bOnlyListRootObjects = FParse::Param(Cmd, TEXT("ROOTONLY"));
+			const bool bShouldIncludeDefaultObjects = FParse::Param(Cmd, TEXT("INCLUDEDEFAULTS"));
+			const bool bOnlyListDefaultObjects = FParse::Param(Cmd, TEXT("DEFAULTSONLY"));
+			const bool bShowDetailedObjectInfo = FParse::Param(Cmd, TEXT("NODETAILEDINFO")) == false && bTrackDetailedObjectInfo;
 
 			for( FObjectIterator It; It; ++It )
 			{
@@ -5445,6 +5446,19 @@ bool UEngine::HandleObjCommand( const TCHAR* Cmd, FOutputDevice& Ar )
 					continue;
 				}
 
+				if (bOnlyListGCObjectsNoClusters)
+				{
+					if (GUObjectArray.IsDisregardForGC(*It))
+					{
+						continue;
+					}
+					FUObjectItem* ObjectItem = GUObjectArray.ObjectToObjectItem(*It);
+					if (ObjectItem->GetOwnerIndex())
+					{
+						continue;
+					}
+				}
+
 				if ( bOnlyListRootObjects && !It->IsRooted() )
 				{
 					continue;
diff --git a/Engine/Source/Runtime/Launch/Private/LaunchEngineLoop.cpp b/Engine/Source/Runtime/Launch/Private/LaunchEngineLoop.cpp
--- a/Engine/Source/Runtime/Launch/Private/LaunchEngineLoop.cpp
+++ b/Engine/Source/Runtime/Launch/Private/LaunchEngineLoop.cpp
@@ -1465,13 +1465,6 @@ int32 FEngineLoop::PreInit( const TCHAR* CmdLine )
 	// Tell the module manager is may now process newly-loaded UObjects when new C++ modules are loaded
 	FModuleManager::Get().StartProcessingNewlyLoadedObjects();
 
-	// load up the seek-free startup packages
-	if ( !FStartupPackages::LoadAll() )
-	{
-		// At least one startup package failed to load, return 1 to indicate an error
-		return 1;
-	}
-
 	// Setup GC optimizations
 	if (bIsSeekFreeDedicatedServer || bHasEditorToken)
 	{
@@ -1566,7 +1559,13 @@ int32 FEngineLoop::PreInit( const TCHAR* CmdLine )
 		return 1;
 	}
 
-	MarkObjectsToDisregardForGC(); 
+	// load up the seek-free startup packages
+	if ( !FStartupPackages::LoadAll() )
+	{
+		// At least one startup package failed to load, return 1 to indicate an error
+		return 1;
+	}
+
 	GUObjectArray.CloseDisregardForGC();
 
 	if (IOnlineSubsystem::IsLoaded())
