Copying //UE4/Orion-Staging to //UE4/Main (Origin //Orion/Dev-General @ 2861092)

diff --git a/Engine/Source/Editor/UnrealEd/Private/SoundCueGraphSchema.cpp b/Engine/Source/Editor/UnrealEd/Private/SoundCueGraphSchema.cpp
--- a/Engine/Source/Editor/UnrealEd/Private/SoundCueGraphSchema.cpp
+++ b/Engine/Source/Editor/UnrealEd/Private/SoundCueGraphSchema.cpp
@@ -15,6 +15,8 @@
 #include "GraphEditor.h"
 #include "Sound/SoundCue.h"
 #include "Engine/Selection.h"
+#include "Sound/SoundNodeDialoguePlayer.h"
+#include "Sound/DialogueWave.h"
 
 #define LOCTEXT_NAMESPACE "SoundCueSchema"
 
@@ -117,9 +119,14 @@ UEdGraphNode* FSoundCueGraphSchemaAction_NewFromSelected::PerformAction(class UE
 	}
 
 	TArray<USoundWave*> SelectedWaves;
+	TArray<UDialogueWave*> SelectedDialogues;
 	TArray<USoundNode*> CreatedPlayers;
+
 	GEditor->GetSelectedObjects()->GetSelectedObjects<USoundWave>(SelectedWaves);
+	GEditor->GetSelectedObjects()->GetSelectedObjects<UDialogueWave>(SelectedDialogues);
+
 	FSoundCueEditorUtilities::CreateWaveContainers(SelectedWaves, SoundCue, CreatedPlayers, WaveStartLocation);
+	FSoundCueEditorUtilities::CreateDialogueContainers(SelectedDialogues, SoundCue, CreatedPlayers, WaveStartLocation);
 
 	if (SoundNodeClass)
 	{
@@ -377,6 +384,8 @@ void USoundCueGraphSchema::BreakPinLinks(UEdGraphPin& TargetPin, bool bSendsNode
 
 void USoundCueGraphSchema::DroppedAssetsOnGraph(const TArray<FAssetData>& Assets, const FVector2D& GraphPosition, UEdGraph* Graph) const
 {
+	//////////////////////////////////////////////////////////////////////////
+	// Handle dropped USoundWaves
 	TArray<USoundWave*> Waves;
 	for (int32 AssetIdx = 0; AssetIdx < Assets.Num(); ++AssetIdx)
 	{
@@ -399,6 +408,139 @@ void USoundCueGraphSchema::DroppedAssetsOnGraph(const TArray<FAssetData>& Assets
 		TArray<USoundNode*> CreatedPlayers;
 		FSoundCueEditorUtilities::CreateWaveContainers(Waves, SoundCue, CreatedPlayers, GraphPosition);
 	}
+
+	//////////////////////////////////////////////////////////////////////////
+	// Handle dropped UDialogueWaves
+	TArray<UDialogueWave*> Dialogues;
+	for (int32 AssetIdx = 0; AssetIdx < Assets.Num(); ++AssetIdx)
+	{
+		UDialogueWave* DialogueWave = Cast<UDialogueWave>(Assets[AssetIdx].GetAsset());
+		if (DialogueWave)
+		{
+			Dialogues.Add(DialogueWave);
+		}
+	}
+
+	if (Dialogues.Num() > 0)
+	{
+		const FScopedTransaction Transaction(LOCTEXT("SoundCueEditorDropDialogue", "Sound Cue Editor: Drag and Drop Dialogue Wave"));
+
+		USoundCueGraph* SoundCueGraph = CastChecked<USoundCueGraph>(Graph);
+		USoundCue* SoundCue = SoundCueGraph->GetSoundCue();
+
+		SoundCueGraph->Modify();
+
+		TArray<USoundNode*> CreatedPlayers;
+		FSoundCueEditorUtilities::CreateDialogueContainers(Dialogues, SoundCue, CreatedPlayers, GraphPosition);
+	}
+}
+
+void USoundCueGraphSchema::DroppedAssetsOnNode(const TArray<FAssetData>& Assets, const FVector2D& GraphPosition, UEdGraphNode* Node) const
+{
+	USoundCueGraphNode* SoundCueGraphNode = CastChecked<USoundCueGraphNode>(Node);
+	USoundCueGraph* SoundCueGraph = CastChecked<USoundCueGraph>(Node->GetGraph());
+	USoundCue* SoundCue = SoundCueGraph->GetSoundCue();
+
+	TArray<USoundWave*> Waves;
+	TArray<UDialogueWave*> Dialogues;
+	for (int32 AssetIdx = 0; AssetIdx < Assets.Num(); ++AssetIdx)
+	{
+		USoundWave* SoundWav = Cast<USoundWave>(Assets[AssetIdx].GetAsset());
+		if (SoundWav)
+		{
+			Waves.Add(SoundWav);
+		}
+		else
+		{
+			UDialogueWave* Dialogue = Cast<UDialogueWave>(Assets[AssetIdx].GetAsset());
+			if (Dialogue)
+			{
+				Dialogues.Add(Dialogue);
+			}
+		}
+	}
+
+	USoundNodeWavePlayer* SoundNodeWavePlayer = Cast<USoundNodeWavePlayer>(SoundCueGraphNode->SoundNode);
+	if (SoundNodeWavePlayer != nullptr)
+	{
+		if (Waves.Num() > 0)
+		{
+			if (Waves.Num() >= 1)
+			{
+				SoundCueGraph->Modify();
+				SoundNodeWavePlayer->SetSoundWave(Waves[0]);
+			}
+
+			for (int32 Index = 1; Index < Waves.Num(); Index++)
+			{
+				TArray<USoundNode*> CreatedPlayers;
+				FSoundCueEditorUtilities::CreateWaveContainers(Waves, SoundCue, CreatedPlayers, GraphPosition);
+			}
+		}
+		else if (Dialogues.Num() > 0)
+		{
+			TArray<USoundNode*> CreatedPlayers;
+			FSoundCueEditorUtilities::CreateDialogueContainers(Dialogues, SoundCue, CreatedPlayers, GraphPosition);
+
+			if (CreatedPlayers.Num() > 0)
+			{
+				USoundNode* OldNode = SoundCueGraphNode->SoundNode;
+				SoundCueGraphNode->SetSoundNode(CreatedPlayers[0]);
+
+				// Make sure SoundCue is updated to match graph
+				SoundCue->CompileSoundNodesFromGraphNodes();
+
+				// Remove this node from the SoundCue's list of all SoundNodes
+				SoundCue->AllNodes.Remove(OldNode);
+				SoundCue->MarkPackageDirty();
+			}
+		}
+	}
+
+	USoundNodeDialoguePlayer* SoundNodeDialoguePlayer = Cast<USoundNodeDialoguePlayer>(SoundCueGraphNode->SoundNode);
+	if (SoundNodeDialoguePlayer != nullptr)
+	{
+		if (Dialogues.Num() > 0)
+		{
+			if (Dialogues.Num() >= 1)
+			{
+				SoundCueGraph->Modify();
+				SoundNodeDialoguePlayer->SetDialogueWave(Dialogues[0]);
+
+				if (Dialogues[0]->ContextMappings.Num() == 1)
+				{
+					SoundNodeDialoguePlayer->DialogueWaveParameter.Context.Speaker = Dialogues[0]->ContextMappings[0].Context.Speaker;
+					SoundNodeDialoguePlayer->DialogueWaveParameter.Context.Targets = Dialogues[0]->ContextMappings[0].Context.Targets;
+				}
+			}
+
+			for (int32 Index = 1; Index < Waves.Num(); Index++)
+			{
+				TArray<USoundNode*> CreatedPlayers;
+				FSoundCueEditorUtilities::CreateDialogueContainers(Dialogues, SoundCue, CreatedPlayers, GraphPosition);
+			}
+		}
+		else if (Waves.Num() > 0)
+		{
+			TArray<USoundNode*> CreatedPlayers;
+			FSoundCueEditorUtilities::CreateWaveContainers(Waves, SoundCue, CreatedPlayers, GraphPosition);
+
+			if (CreatedPlayers.Num() > 0)
+			{
+				USoundNode* OldNode = SoundCueGraphNode->SoundNode;
+				SoundCueGraphNode->SetSoundNode(CreatedPlayers[0]);
+
+				// Make sure SoundCue is updated to match graph
+				SoundCue->CompileSoundNodesFromGraphNodes();
+
+				// Remove this node from the SoundCue's list of all SoundNodes
+				SoundCue->AllNodes.Remove(OldNode);
+				SoundCue->MarkPackageDirty();
+			}
+		}
+	}
+
+	SoundCueGraph->NotifyGraphChanged();
 }
 
 void USoundCueGraphSchema::GetAllSoundNodeActions(FGraphActionMenuBuilder& ActionMenuBuilder, bool bShowSelectedActions) const
@@ -406,6 +548,8 @@ void USoundCueGraphSchema::GetAllSoundNodeActions(FGraphActionMenuBuilder& Actio
 	InitSoundNodeClasses();
 
 	FText SelectedItemText;
+	bool IsSoundWaveSelected = false;
+	bool IsDialogueWaveSelected = false;
 
 	if (bShowSelectedActions)
 	{
@@ -415,10 +559,25 @@ void USoundCueGraphSchema::GetAllSoundNodeActions(FGraphActionMenuBuilder& Actio
 		if (ActionMenuBuilder.FromPin == NULL)
 		{
 			TArray<USoundWave*> SelectedWavs;
+			TArray<UDialogueWave*> SelectedDialogues;
 			GEditor->GetSelectedObjects()->GetSelectedObjects<USoundWave>(SelectedWavs);
-			if (SelectedWavs.Num())
+			GEditor->GetSelectedObjects()->GetSelectedObjects<UDialogueWave>(SelectedDialogues);
+
+			int32 TotalWavs = SelectedWavs.Num() + SelectedDialogues.Num() ;
+
+			if (TotalWavs > 1)
+			{
+				SelectedItemText = LOCTEXT("MultipleWAVsSelected", "Multiple WAVs");
+			}
+			else if (SelectedWavs.Num() == 1)
+			{
+				SelectedItemText = FText::FromString(SelectedWavs[0]->GetName());
+				IsSoundWaveSelected = true;
+			}
+			else if (SelectedDialogues.Num() == 1)
 			{
-				SelectedItemText = FText::FromString(FString::Printf(TEXT("%s"), (SelectedWavs.Num() > 1) ? *(LOCTEXT("MultipleWAVsSelected", "Multiple WAVs").ToString()) : *SelectedWavs[0]->GetName()));
+				SelectedItemText = FText::FromString(SelectedDialogues[0]->GetName());
+				IsDialogueWaveSelected = true;
 			}
 		}
 		else
@@ -427,6 +586,16 @@ void USoundCueGraphSchema::GetAllSoundNodeActions(FGraphActionMenuBuilder& Actio
 			if (SelectedWave && ActionMenuBuilder.FromPin->Direction == EGPD_Input)
 			{
 				SelectedItemText = FText::FromString(SelectedWave->GetName());
+				IsSoundWaveSelected = true;
+			}
+			else
+			{
+				UDialogueWave* SelectedDialogue = GEditor->GetSelectedObjects()->GetTop<UDialogueWave>();
+				if (SelectedDialogue && ActionMenuBuilder.FromPin->Direction == EGPD_Input)
+				{
+					SelectedItemText = FText::FromString(SelectedDialogue->GetName());
+					IsDialogueWaveSelected = true;
+				}
 			}
 		}
 
@@ -451,7 +620,10 @@ void USoundCueGraphSchema::GetAllSoundNodeActions(FGraphActionMenuBuilder& Actio
 				NewNodeAction->SoundNodeClass = SoundNodeClass;
 			}
 
-			if (bShowSelectedActions && (SoundNode->GetMaxChildNodes() == USoundNode::MAX_ALLOWED_CHILD_NODES || SoundNodeClass == USoundNodeWavePlayer::StaticClass()))
+			if (bShowSelectedActions && 
+				(SoundNode->GetMaxChildNodes() == USoundNode::MAX_ALLOWED_CHILD_NODES || 
+				   ((SoundNodeClass == USoundNodeWavePlayer::StaticClass() && IsSoundWaveSelected) || 
+				    (SoundNodeClass == USoundNodeDialoguePlayer::StaticClass() && IsDialogueWaveSelected))))
 			{
 				FFormatNamedArguments Arguments;
 				Arguments.Add(TEXT("Name"), Name);
@@ -462,7 +634,7 @@ void USoundCueGraphSchema::GetAllSoundNodeActions(FGraphActionMenuBuilder& Actio
 					MenuDesc,
 					ToolTip.ToString(), 0));
 				ActionMenuBuilder.AddAction(NewNodeAction);
-				NewNodeAction->SoundNodeClass = (SoundNodeClass == USoundNodeWavePlayer::StaticClass() ? NULL : SoundNodeClass);
+				NewNodeAction->SoundNodeClass = (SoundNodeClass == USoundNodeWavePlayer::StaticClass() || SoundNodeClass == USoundNodeDialoguePlayer::StaticClass() ? NULL : SoundNodeClass);
 			}
 		}
 	}
diff --git a/Engine/Source/Editor/UnrealEd/UnrealEd.Build.cs b/Engine/Source/Editor/UnrealEd/UnrealEd.Build.cs
--- a/Engine/Source/Editor/UnrealEd/UnrealEd.Build.cs
+++ b/Engine/Source/Editor/UnrealEd/UnrealEd.Build.cs
@@ -25,7 +25,6 @@ public class UnrealEd : ModuleRules
 			{
 				"Analytics",
 				"AssetRegistry",
-                "CookingStats",
 				"AssetTools",
                 "BehaviorTreeEditor",
 				"ClassViewer",
diff --git a/Engine/Source/Programs/AutomationTool/AutomationUtils/Platform.cs b/Engine/Source/Programs/AutomationTool/AutomationUtils/Platform.cs
--- a/Engine/Source/Programs/AutomationTool/AutomationUtils/Platform.cs
+++ b/Engine/Source/Programs/AutomationTool/AutomationUtils/Platform.cs
@@ -324,7 +324,24 @@ namespace AutomationTool
             return new List<string>();
         }
 
-
+        /// <summary>
+        /// Get a release pak file path, if we are currently building a patch then get the previous release pak file path, if we are creating a new release this will be the output path
+        /// </summary>
+        /// <param name="SC"></param>
+        /// <param name="Params"></param>
+        /// <param name="PakName"></param>
+        /// <returns></returns>
+        public virtual string GetReleasePakFilePath(DeploymentContext SC, ProjectParams Params, string PakName)
+        {
+            if (Params.IsGeneratingPatch)
+            {
+                return CombinePaths(Params.GetBasedOnReleaseVersionPath(SC), PakName);
+            }
+            else
+            {
+                return CombinePaths(Params.GetCreateReleaseVersionPath(SC), PakName);
+            }        
+        }
 
         /// <summary>
         /// Gets editor cook platform name for this platform. Cooking the editor is not useful, but this is used to fill the derived data cache
diff --git a/Engine/Source/Programs/AutomationTool/AutomationUtils/ProjectParams.cs b/Engine/Source/Programs/AutomationTool/AutomationUtils/ProjectParams.cs
--- a/Engine/Source/Programs/AutomationTool/AutomationUtils/ProjectParams.cs
+++ b/Engine/Source/Programs/AutomationTool/AutomationUtils/ProjectParams.cs
@@ -1931,12 +1931,18 @@ namespace AutomationTool
 			String Platform = SC.StageTargetPlatform.GetCookPlatform(SC.DedicatedServer, false, CookFlavor);
 			if (String.IsNullOrEmpty(BasePath))
 			{
-				BasePath = CommandUtils.CombinePaths(SC.ProjectRoot, "Releases", BasedOnReleaseVersion, Platform);
+                BasePath = CommandUtils.CombinePaths(SC.ProjectRoot, "Releases", BasedOnReleaseVersion, Platform);
 			}
 			else
 			{
 				BasePath = CommandUtils.CombinePaths(BasePath, BasedOnReleaseVersion, Platform);
 			}
+
+            /*if ( TitleID != null && TitleID.Count == 1 )
+            {
+                BasePath = CommandUtils.CombinePaths( BasePath, TitleID[0]);
+            }*/
+
 			return BasePath;
 		}
 
@@ -1957,6 +1963,12 @@ namespace AutomationTool
 			{
 				BasePath = CommandUtils.CombinePaths(BasePath, CreateReleaseVersion, Platform);
 			}
+
+            /*if (TitleID != null && TitleID.Count == 1)
+            {
+                BasePath = CommandUtils.CombinePaths(BasePath, TitleID[0]);
+            }*/
+
 			return BasePath;
 		}
 
diff --git a/Engine/Source/Programs/AutomationTool/AutomationUtils/ProjectUtils.cs b/Engine/Source/Programs/AutomationTool/AutomationUtils/ProjectUtils.cs
--- a/Engine/Source/Programs/AutomationTool/AutomationUtils/ProjectUtils.cs
+++ b/Engine/Source/Programs/AutomationTool/AutomationUtils/ProjectUtils.cs
@@ -611,7 +611,7 @@ namespace AutomationTool
     public class BranchInfo
     {
 
-        public static TargetRules.TargetType[] MonolithicKinds = new TargetRules.TargetType[]
+        public static List<TargetRules.TargetType> MonolithicKinds = new List<TargetRules.TargetType>
         {
             TargetRules.TargetType.Game,
             TargetRules.TargetType.Client,
@@ -691,7 +691,7 @@ namespace AutomationTool
 						CommandUtils.LogVerbose("              bUsesSteam  : " + (ThisTarget.Value.Rules.bUsesSteam ? "YES" : "NO"));
 						CommandUtils.LogVerbose("              bUsesCEF3   : " + (ThisTarget.Value.Rules.bUsesCEF3 ? "YES" : "NO"));
 						CommandUtils.LogVerbose("              bUsesSlate  : " + (ThisTarget.Value.Rules.bUsesSlate ? "YES" : "NO"));
-                        if (Array.IndexOf(MonolithicKinds, ThisTarget.Key) >= 0)
+                        if (MonolithicKinds.Contains(ThisTarget.Key))
                         {
                             var Platforms = ThisTarget.Value.Rules.GUBP_GetPlatforms_MonolithicOnly(HostPlatform);
                             var AdditionalPlatforms = ThisTarget.Value.Rules.GUBP_GetBuildOnlyPlatforms_MonolithicOnly(HostPlatform);
diff --git a/Engine/Source/Programs/AutomationTool/GUBP/GUBP.cs b/Engine/Source/Programs/AutomationTool/GUBP/GUBP.cs
--- a/Engine/Source/Programs/AutomationTool/GUBP/GUBP.cs
+++ b/Engine/Source/Programs/AutomationTool/GUBP/GUBP.cs
@@ -100,6 +100,11 @@ public partial class GUBP : BuildCommand
         string BranchName = P4Enabled ? P4Env.BuildRootP4 : ParseParamValue("BranchName", "");
         GUBPBranchHacker.BranchOptions BranchOptions = GetBranchOptions(BranchName);
 
+		foreach (var Kind in BranchOptions.MonolithicsToRemove)
+		{
+			BranchInfo.MonolithicKinds.Remove(Kind);
+		}
+
 		string BranchNameForTempStorage = P4Enabled ? P4Env.BuildRootEscaped : "NoP4";
 		string RootNameForTempStorage = BranchOptions.RootNameForTempStorage ?? "UE4";
 		string PreflightSuffix = GetPreflightSuffix();
@@ -589,7 +594,10 @@ public partial class GUBP : BuildCommand
                 }
 				foreach (BuildNode Dep in NodeToDo.OrderDependencies.OrderBy(x => x.Name))
 				{
-					Log("           pdep> {0}", Dep.Name);
+					if (!NodeToDo.InputDependencies.Contains(Dep))
+					{
+						Log("           pdep> {0}", Dep.Name);
+					}
 				}
             }
         }
diff --git a/Engine/Source/Programs/AutomationTool/GUBP/LegacyBranchSetup.cs b/Engine/Source/Programs/AutomationTool/GUBP/LegacyBranchSetup.cs
--- a/Engine/Source/Programs/AutomationTool/GUBP/LegacyBranchSetup.cs
+++ b/Engine/Source/Programs/AutomationTool/GUBP/LegacyBranchSetup.cs
@@ -142,6 +142,7 @@ partial class GUBP
         public class BranchOptions
         {            
             public List<UnrealTargetPlatform> PlatformsToRemove = new List<UnrealTargetPlatform>();
+			public List<TargetRules.TargetType> MonolithicsToRemove = new List<TargetRules.TargetType>();
 			public List<string> ExcludeNodes = new List<string>();
 			public List<UnrealTargetPlatform> ExcludePlatformsForEditor = new List<UnrealTargetPlatform>();
             public List<UnrealTargetPlatform> RemovePlatformFromPromotable = new List<UnrealTargetPlatform>();
@@ -154,7 +155,7 @@ partial class GUBP
 			/// </summary>
 			public bool bGameBuildsBehindTriggers = false;
             /// <summary>
-            /// This turns off ALL automatedtesting in ALL branches by default.
+            /// This turns off ALL automated testing in ALL branches by default.
             /// A branch hacker must set this value to false explicitly to use automated testing.
             /// </summary>
 			public bool bNoAutomatedTesting = true;
@@ -188,7 +189,7 @@ partial class GUBP
 			/// </summary>
 			public bool bBuildEngineLocalization = false;
 			public string EngineLocalizationBranchSuffix = "";
-			/// <summary>
+            public string AdditionalCookArgs = "";
 			/// Whether to build all target platforms in parallel, or serialize them to avoid saturating the farm.
 			/// </summary>
 			public bool bTargetPlatformsInParallel = true;
diff --git a/Engine/Source/Programs/AutomationTool/GUBP/LegacyNodes.cs b/Engine/Source/Programs/AutomationTool/GUBP/LegacyNodes.cs
--- a/Engine/Source/Programs/AutomationTool/GUBP/LegacyNodes.cs
+++ b/Engine/Source/Programs/AutomationTool/GUBP/LegacyNodes.cs
@@ -1685,6 +1685,7 @@ partial class GUBP
 			{
 				ExtraArgsForCook += " -WarningsAsErrors";
 			}
+            ExtraArgsForCook += " " + BranchConfig.BranchOptions.AdditionalCookArgs;
 
 			if(GamePlatformMonolithicsNode.IsSample(BranchConfig, GameProj))
 			{
diff --git a/Engine/Source/Programs/AutomationTool/GUBP/TempStorage.cs b/Engine/Source/Programs/AutomationTool/GUBP/TempStorage.cs
--- a/Engine/Source/Programs/AutomationTool/GUBP/TempStorage.cs
+++ b/Engine/Source/Programs/AutomationTool/GUBP/TempStorage.cs
@@ -283,7 +283,7 @@ namespace AutomationTool
 			{
 				if (IsEmptyManifest(Directories))
 				{
-					throw new AutomationException("Attempt to save empty manifest.");
+					throw new AutomationException("Attempt to save empty manifest file {0}.", Filename);
 				}
 
 				new XElement(RootElementName,
diff --git a/Engine/Source/Programs/AutomationTool/Scripts/CopyBuildToStagingDirectory.Automation.cs b/Engine/Source/Programs/AutomationTool/Scripts/CopyBuildToStagingDirectory.Automation.cs
--- a/Engine/Source/Programs/AutomationTool/Scripts/CopyBuildToStagingDirectory.Automation.cs
+++ b/Engine/Source/Programs/AutomationTool/Scripts/CopyBuildToStagingDirectory.Automation.cs
@@ -828,19 +828,6 @@ public partial class Project : CommandUtils
 	}
 
 
-    private static string GetReleasePakFilePath(DeploymentContext SC, ProjectParams Params, string PakName)
-    {
-		if (Params.IsGeneratingPatch)
-		{
-			return CombinePaths(Params.GetBasedOnReleaseVersionPath(SC), PakName);
-		}
-		else
-		{
-			return CombinePaths(Params.GetCreateReleaseVersionPath(SC), PakName);
-		}        
-    }
-
-
 	/// <summary>
 	/// Creates a pak file using response file.
 	/// </summary>
@@ -919,7 +906,7 @@ public partial class Project : CommandUtils
         {
             // don't include the post fix in this filename because we are looking for the source pak path
             string PakFilename = PakName + "-" + SC.FinalCookPlatform + ".pak";
-            PatchSourceContentPath = GetReleasePakFilePath(SC, Params, PakFilename);
+            PatchSourceContentPath = SC.StageTargetPlatform.GetReleasePakFilePath(SC, Params, PakFilename);
         }
 
 		if (!bCopiedExistingPak)
@@ -943,7 +930,7 @@ public partial class Project : CommandUtils
         {
             // copy the created pak to the release version directory we might need this later if we want to generate patches
             //string ReleaseVersionPath = CombinePaths( SC.ProjectRoot, "Releases", Params.CreateReleaseVersion, SC.StageTargetPlatform.GetCookPlatform(Params.DedicatedServer, false, Params.CookFlavor), Path.GetFileName(OutputLocation) );
-            string ReleaseVersionPath = GetReleasePakFilePath(SC, Params, Path.GetFileName(OutputLocation));
+            string ReleaseVersionPath = SC.StageTargetPlatform.GetReleasePakFilePath(SC, Params, Path.GetFileName(OutputLocation));
 
 			InternalUtils.SafeCreateDirectory(Path.GetDirectoryName(ReleaseVersionPath));
 			InternalUtils.SafeCopyFile(OutputLocation, ReleaseVersionPath);
diff --git a/Engine/Source/Programs/AutomationTool/Scripts/Localisation.Automation.cs b/Engine/Source/Programs/AutomationTool/Scripts/Localisation.Automation.cs
--- a/Engine/Source/Programs/AutomationTool/Scripts/Localisation.Automation.cs
+++ b/Engine/Source/Programs/AutomationTool/Scripts/Localisation.Automation.cs
@@ -86,7 +86,7 @@ class Localise : BuildCommand
 			ProjectInfos.Add(GenerateProjectInfo(RootWorkingDirectory, ProjectName));
 		}
 
-		// Export all text from OneSky
+		// Export all text from our localization provider
 		foreach (var ProjectInfo in ProjectInfos)
 		{
 			LocProvider.DownloadProjectFromLocalizationProvider(ProjectInfo.ProjectName, ProjectInfo.ImportInfo);
@@ -109,6 +109,7 @@ class Localise : BuildCommand
 		if (!AllowSubmit) { CommandletSCCArguments += (String.IsNullOrEmpty(CommandletSCCArguments) ? "" : " ") + "-DisableSCCSubmit"; }
 
 		// Execute commandlet for each config in each project.
+		bool bLocCommandletFailed = false;
 		foreach (var ProjectInfo in ProjectInfos)
 		{
 			foreach (var LocalizationConfigFile in ProjectInfo.LocalizationConfigFiles)
@@ -120,24 +121,33 @@ class Localise : BuildCommand
 					CommandletArguments += " " + AdditionalCommandletArguments;
 				}
 
-				Log("Localization for {0} {1}", EditorArguments, CommandletArguments);
-
-				Log("Running UE4Editor to generate localization data");
-
 				string Arguments = String.Format("{0} -run=GatherText {1} {2}", UEProjectName, EditorArguments, CommandletArguments);
-				var RunResult = Run(EditorExe, Arguments);
+				Log("Running localization commandlet: {0}", Arguments);
+				var StartTime = DateTime.UtcNow;
+				var RunResult = Run(EditorExe, Arguments, null, ERunOptions.Default | ERunOptions.NoLoggingOfRunCommand); // Disable logging of the run command as it will print the exit code which GUBP can pick up as an error (we do that ourselves below)
+				var RunDuration = (DateTime.UtcNow - StartTime).TotalMilliseconds;
+				Log("Localization commandlet finished in {0}s", RunDuration / 1000);
 
 				if (RunResult.ExitCode != 0)
 				{
-					Console.WriteLine("[ERROR] Error while executing localization commandlet '{0}'", Arguments);
+					LogWarning("The localization commandlet exited with code {0} which likely indicates a crash. It ran with the following arguments: '{1}'", RunResult.ExitCode, Arguments);
+					bLocCommandletFailed = true;
+					break; // We failed a step, so don't process any other steps in this config chain
 				}
 			}
 		}
 
-		// Upload all text to OneSky
-		foreach (var ProjectInfo in ProjectInfos)
+		if (bLocCommandletFailed)
+		{
+			LogWarning("Skipping upload to the localization provider due to an earlier commandlet failure.");
+		}
+		else
 		{
-			LocProvider.UploadProjectToLocalizationProvider(ProjectInfo.ProjectName, ProjectInfo.ExportInfo);
+			// Upload all text to our localization provider
+			foreach (var ProjectInfo in ProjectInfos)
+			{
+				LocProvider.UploadProjectToLocalizationProvider(ProjectInfo.ProjectName, ProjectInfo.ExportInfo);
+			}
 		}
 	}
 
diff --git a/Engine/Source/Programs/CrashReporter/CrashReportClient/Private/CrashDescription.cpp b/Engine/Source/Programs/CrashReporter/CrashReportClient/Private/CrashDescription.cpp
--- a/Engine/Source/Programs/CrashReporter/CrashReportClient/Private/CrashDescription.cpp
+++ b/Engine/Source/Programs/CrashReporter/CrashReportClient/Private/CrashDescription.cpp
@@ -106,6 +106,7 @@ FPrimaryCrashProperties::FPrimaryCrashProperties()
 	, UserName( FGenericCrashContext::RuntimePropertiesTag, TEXT("UserName"), this )
 	, MachineId( FGenericCrashContext::RuntimePropertiesTag, TEXT( "MachineId" ), this )
 	, EpicAccountId( FGenericCrashContext::RuntimePropertiesTag, TEXT( "EpicAccountId" ), this )
+	, GameSessionID( FGenericCrashContext::RuntimePropertiesTag, TEXT( "GameSessionID" ), this )
 	// Multiline properties
 	, CallStack( FGenericCrashContext::RuntimePropertiesTag, TEXT( "CallStack" ), this )
 	, SourceContext( FGenericCrashContext::RuntimePropertiesTag, TEXT( "SourceContext" ), this )
@@ -212,6 +213,8 @@ void FPrimaryCrashProperties::SendAnalytics()
 	CrashAttributes.Add( FAnalyticsEventAttribute( TEXT( "AppDefaultLocale" ), AppDefaultLocale ) );
 
 	CrashAttributes.Add( FAnalyticsEventAttribute( TEXT( "UserActivityHint" ), UserActivityHint.AsString() ) );
+	CrashAttributes.Add( FAnalyticsEventAttribute( TEXT( "GameSessionID" ), GameSessionID.AsString() ) );
+	CrashAttributes.Add( FAnalyticsEventAttribute( TEXT( "DeploymentName" ), DeploymentName ) );
 
 	Analytics.RecordEvent( TEXT( "CrashReportClient.ReportCrash" ), CrashAttributes );
 	// Shutdown analytics.
diff --git a/Engine/Source/Runtime/AIModule/Classes/AIController.h b/Engine/Source/Runtime/AIModule/Classes/AIController.h
--- a/Engine/Source/Runtime/AIModule/Classes/AIController.h
+++ b/Engine/Source/Runtime/AIModule/Classes/AIController.h
@@ -107,6 +107,10 @@ public:
 	UPROPERTY()
 	uint32 bWantsPlayerState:1;
 
+	/** Copy Pawn rotation to ControlRotation, if there is no focus point. */
+	UPROPERTY(EditAnywhere, BlueprintReadWrite, Category = AI)
+	uint32 bSetControlRotationFromPawnOrientation:1;
+
 private_subobject:
 
 	/** Component used for moving along a path. */
diff --git a/Engine/Source/Runtime/AIModule/Private/AIController.cpp b/Engine/Source/Runtime/AIModule/Private/AIController.cpp
--- a/Engine/Source/Runtime/AIModule/Private/AIController.cpp
+++ b/Engine/Source/Runtime/AIModule/Private/AIController.cpp
@@ -38,6 +38,7 @@ DEFINE_LOG_CATEGORY(LogAINavigation);
 AAIController::AAIController(const FObjectInitializer& ObjectInitializer)
 	: Super(ObjectInitializer)
 {
+	bSetControlRotationFromPawnOrientation = true;
 	PathFollowingComponent = CreateDefaultSubobject<UPathFollowingComponent>(TEXT("PathFollowingComponent"));
 	PathFollowingComponent->OnMoveFinished.AddUObject(this, &AAIController::OnMoveCompleted);
 
@@ -414,28 +415,36 @@ DEFINE_LOG_CATEGORY_STATIC(LogTestAI, All, All);
 
 void AAIController::UpdateControlRotation(float DeltaTime, bool bUpdatePawn)
 {
-	// Look toward focus
-	FVector FocalPoint = GetFocalPoint();
-	APawn* const Pawn = GetPawn();
-
-	if (Pawn)
+	APawn* const MyPawn = GetPawn();
+	if (MyPawn)
 	{
-		FVector Direction = FAISystem::IsValidLocation(FocalPoint) ? (FocalPoint - Pawn->GetPawnViewLocation()) : Pawn->GetActorForwardVector();
-		FRotator NewControlRotation = Direction.Rotation();
+		const FRotator InitialControlRotation = GetControlRotation();		
+		FRotator NewControlRotation = InitialControlRotation;
+
+		// Look toward focus
+		const FVector FocalPoint = GetFocalPoint();
+		if (FAISystem::IsValidLocation(FocalPoint))
+		{
+			NewControlRotation = (FocalPoint - MyPawn->GetPawnViewLocation()).Rotation();
+		}
+		else if (bSetControlRotationFromPawnOrientation)
+		{
+			NewControlRotation = MyPawn->GetActorRotation();
+		}
 
 		// Don't pitch view unless looking at another pawn
-		if (Cast<APawn>(GetFocusActor()) == nullptr)
+		if (NewControlRotation.Pitch != 0 && Cast<APawn>(GetFocusActor()) == nullptr)
 		{
 			NewControlRotation.Pitch = 0.f;
 		}
 
-		if (GetControlRotation().Equals(NewControlRotation, 1e-3f) == false)
+		if (InitialControlRotation.Equals(NewControlRotation, 1e-3f) == false)
 		{
 			SetControlRotation(NewControlRotation);
 
 			if (bUpdatePawn)
 			{
-				Pawn->FaceRotation(NewControlRotation, DeltaTime);
+				MyPawn->FaceRotation(NewControlRotation, DeltaTime);
 			}
 		}
 	}
diff --git a/Engine/Source/Runtime/Analytics/QoSReporter/Private/QoSReporter.cpp b/Engine/Source/Runtime/Analytics/QoSReporter/Private/QoSReporter.cpp
--- a/Engine/Source/Runtime/Analytics/QoSReporter/Private/QoSReporter.cpp
+++ b/Engine/Source/Runtime/Analytics/QoSReporter/Private/QoSReporter.cpp
@@ -178,11 +178,13 @@ void FQoSReporter::SetBackendDeploymentName(const FString & InDeploymentName)
 void FQoSReporter::EnableCountingHitches(bool bEnable)
 {
 	bCountHitches = bEnable;
-	UE_LOG(LogQoSReporter, Verbose, TEXT("Counting hitches in QoSReporter has been %s."), bCountHitches ? TEXT("enabled") : TEXT("disabled"));
+	PreviousTickTime = FPlatformTime::Seconds();	// reset the timer
+	UE_LOG(LogQoSReporter, Log, TEXT("Counting hitches in QoSReporter has been %s."), bCountHitches ? TEXT("enabled") : TEXT("disabled"));
 }
 
 double FQoSReporter::HeartbeatInterval = 300;
 double FQoSReporter::LastHeartbeatTimestamp = 0;
+double FQoSReporter::PreviousTickTime = FPlatformTime::Seconds();
 #if WITH_ENGINE
 extern ENGINE_API float GAverageFPS;
 #else
@@ -196,7 +198,6 @@ void FQoSReporter::Tick()
 		return;
 	}
 
-	static double PreviousTickTime = FPlatformTime::Seconds();
 	double CurrentTime = FPlatformTime::Seconds();
 
 	if (HeartbeatInterval > 0 && CurrentTime - LastHeartbeatTimestamp > HeartbeatInterval)
@@ -219,10 +220,20 @@ void FQoSReporter::Tick()
 			if (DeltaBetweenTicks > 0.25)
 			{
 				PerfCountersIncrement(TEXT("HitchesAbove250msec"), 0, IPerfCounters::Flags::Transient);
+
+				if (DeltaBetweenTicks > 0.5)
+				{
+					PerfCountersIncrement(TEXT("HitchesAbove500msec"), 0, IPerfCounters::Flags::Transient);
+
+					if (DeltaBetweenTicks > 1)
+					{
+						PerfCountersIncrement(TEXT("HitchesAbove1000msec"), 0, IPerfCounters::Flags::Transient);
+					}
+				}
 			}
 #endif // USE_SERVER_PERF_COUNTERS
 
-			UE_LOG(LogQoSReporter, Warning, TEXT("QoS reporter could not tick for %f sec, average FPS is %f."),
+			UE_LOG(LogQoSReporter, Log, TEXT("QoS reporter could not tick for %f sec, average FPS is %f."),
 				DeltaBetweenTicks, GAverageFPS);
 		}
 	}
diff --git a/Engine/Source/Runtime/Analytics/QoSReporter/Private/QoSReporterModule.cpp b/Engine/Source/Runtime/Analytics/QoSReporter/Private/QoSReporterModule.cpp
--- a/Engine/Source/Runtime/Analytics/QoSReporter/Private/QoSReporterModule.cpp
+++ b/Engine/Source/Runtime/Analytics/QoSReporter/Private/QoSReporterModule.cpp
@@ -9,12 +9,16 @@
 #include "EngineVersion.h"
 #include "QoSReporter.h"
 
+#if USE_SERVER_PERF_COUNTERS
+	#include "PerfCountersHelpers.h"
+#endif
+
 DEFINE_LOG_CATEGORY(LogQoSReporter);
 
 IMPLEMENT_MODULE(FQoSReporterModule, QoSReporter);
 
 // helps to version QoS events (date*10 to allow for 10 revisions per day)
-#define QOS_EVENTS_REVISION					201601280
+#define QOS_EVENTS_REVISION					201602080
 
 FString FQoSReporterModule::Config::GetDefaultAppVersion()
 { 
@@ -51,7 +55,7 @@ public:
 	/** We're (ab)using this API to set DeploymentName */
 	virtual void SetLocation(const FString& InLocation) override { DeploymentName = InLocation; };
 
-	virtual void RecordEvent(const FString& EventName, const TArray<FAnalyticsEventAttribute>& Attributes) override;
+	virtual void RecordEvent(const FString& InEventName, const TArray<FAnalyticsEventAttribute>& Attributes) override;
 	virtual ~FAnalyticsProviderQoSReporter();
 
 	FString GetAPIKey() const { return APIKey; }
@@ -170,19 +174,25 @@ FAnalyticsProviderQoSReporter::~FAnalyticsProviderQoSReporter()
 	EndSession();
 }
 
-void FAnalyticsProviderQoSReporter::RecordEvent(const FString& EventName, const TArray<FAnalyticsEventAttribute>& InAttributes)
+void FAnalyticsProviderQoSReporter::RecordEvent(const FString& InEventName, const TArray<FAnalyticsEventAttribute>& InAttributes)
 {
 	if (APIKey.IsEmpty() || APIServer.IsEmpty())
 	{
 		return;
 	}
 
+	// for data router, it is actually preferable to have different events than attach more attributes, so
+	// append application role to the event name instead
+	FString EventName = InEventName;
+	EventName += TEXT(".");
+	EventName += GetApplicationRole();
+
 	// add attributes common to each QoS event first
 	TArray<FAnalyticsEventAttribute> Attributes;
 	Attributes.Add(FAnalyticsEventAttribute(TEXT("QoSRevision"), QOS_EVENTS_REVISION));
-	Attributes.Add(FAnalyticsEventAttribute(TEXT("Role"), GetApplicationRole()));
 	Attributes.Add(FAnalyticsEventAttribute(TEXT("SystemId"), FPlatformMisc::GetOperatingSystemId()));
 	Attributes.Add(FAnalyticsEventAttribute(TEXT("InstanceId"), InstanceId.ToString()));
+	Attributes.Add(FAnalyticsEventAttribute(TEXT("Platform"), FString(FPlatformProperties::PlatformName())));
 	if (LIKELY(DeploymentName.Len() > 0))
 	{
 		Attributes.Add(FAnalyticsEventAttribute(TEXT("Deployment"), DeploymentName));
@@ -258,6 +268,8 @@ void FAnalyticsProviderQoSReporter::RecordEvent(const FString& EventName, const
 
 void FAnalyticsProviderQoSReporter::EventRequestComplete(FHttpRequestPtr HttpRequest, FHttpResponsePtr HttpResponse, bool bSucceeded)
 {
+	bool bCountTowardsFailedAttempts = true;
+
 	if (bSucceeded && HttpResponse.IsValid())
 	{
 		// normal operation is silent, but any problems are reported as warnings
@@ -265,6 +277,8 @@ void FAnalyticsProviderQoSReporter::EventRequestComplete(FHttpRequestPtr HttpReq
 		{
 			UE_LOG(LogQoSReporter, VeryVerbose, TEXT("QoS response for [%s]. Code: %d. Payload: %s"),
 				*HttpRequest->GetURL(), HttpResponse->GetResponseCode(), *HttpResponse->GetContentAsString());
+
+			bCountTowardsFailedAttempts = false;
 		}
 		else
 		{
@@ -277,6 +291,14 @@ void FAnalyticsProviderQoSReporter::EventRequestComplete(FHttpRequestPtr HttpReq
 		// if we cannot report QoS metrics this is pretty bad; report at least a warning
 		UE_LOG(LogQoSReporter, Warning, TEXT("QoS response for [%s]. No response"), *HttpRequest->GetURL());
 	}
+
+	if (bCountTowardsFailedAttempts)
+	{
+		// FIXME: should use retrial with exponential backoff here
+#if USE_SERVER_PERF_COUNTERS
+		PerfCountersIncrement(TEXT("FailedQoSRequests"), 0, IPerfCounters::Flags::Transient);
+#endif // USE_SERVER_PERF_COUNTERS
+	}
 }
 
 /**
diff --git a/Engine/Source/Runtime/Analytics/QoSReporter/Public/QoSReporter.h b/Engine/Source/Runtime/Analytics/QoSReporter/Public/QoSReporter.h
--- a/Engine/Source/Runtime/Analytics/QoSReporter/Public/QoSReporter.h
+++ b/Engine/Source/Runtime/Analytics/QoSReporter/Public/QoSReporter.h
@@ -205,6 +205,9 @@ private:
 	/** Timestamp when we sent the last heartbeat */
 	static double LastHeartbeatTimestamp;
 
+	/** Time when Tick() was called previously */
+	static double PreviousTickTime;
+
 	/**
 	 * Event-specific variables
 	 */
diff --git a/Engine/Source/Runtime/Core/Core.Build.cs b/Engine/Source/Runtime/Core/Core.Build.cs
--- a/Engine/Source/Runtime/Core/Core.Build.cs
+++ b/Engine/Source/Runtime/Core/Core.Build.cs
@@ -51,8 +51,7 @@ public class Core : ModuleRules
 			new string[] {
 				"TargetPlatform",
 				"DerivedDataCache",
-                "CookingStats",
-				"InputDevice",
+                "InputDevice",
                 "Analytics"
 			}
 			);
diff --git a/Engine/Source/Runtime/Core/Private/Android/AndroidApplication.cpp b/Engine/Source/Runtime/Core/Private/Android/AndroidApplication.cpp
--- a/Engine/Source/Runtime/Core/Private/Android/AndroidApplication.cpp
+++ b/Engine/Source/Runtime/Core/Private/Android/AndroidApplication.cpp
@@ -104,6 +104,9 @@ void FDisplayMetrics::GetDisplayMetrics( FDisplayMetrics& OutDisplayMetrics )
 	// Total screen size of the primary monitor
 	OutDisplayMetrics.PrimaryDisplayWidth = OutDisplayMetrics.PrimaryDisplayWorkAreaRect.Right - OutDisplayMetrics.PrimaryDisplayWorkAreaRect.Left;
 	OutDisplayMetrics.PrimaryDisplayHeight = OutDisplayMetrics.PrimaryDisplayWorkAreaRect.Bottom - OutDisplayMetrics.PrimaryDisplayWorkAreaRect.Top;
+
+	// Apply the debug safe zones
+	OutDisplayMetrics.ApplyDefaultSafeZones();
 }
 
 TSharedRef< FGenericWindow > FAndroidApplication::MakeWindow()
diff --git a/Engine/Source/Runtime/Core/Private/GenericPlatform/GenericApplication.cpp b/Engine/Source/Runtime/Core/Private/GenericPlatform/GenericApplication.cpp
--- a/Engine/Source/Runtime/Core/Private/GenericPlatform/GenericApplication.cpp
+++ b/Engine/Source/Runtime/Core/Private/GenericPlatform/GenericApplication.cpp
@@ -101,6 +101,48 @@ const FGamepadKeyNames::Type FGamepadKeyNames::MotionController_Right_TriggerAxi
 const FGamepadKeyNames::Type FGamepadKeyNames::MotionController_Right_Grip1Axis( "MotionController_Right_Grip1Axis" );
 const FGamepadKeyNames::Type FGamepadKeyNames::MotionController_Right_Grip2Axis( "MotionController_Right_Grip2Axis" );
 
+float GDebugSafeZoneRatio = 1.0f;
+float GDebugActionZoneRatio = 1.0f;
+
+FAutoConsoleVariableRef GDebugSafeZoneRatioCVar(
+	TEXT("r.DebugSafeZone.TitleRatio"),
+	GDebugSafeZoneRatio,
+	TEXT("The safe zone ratio that will be returned by FDisplayMetrics::GetDisplayMetrics on platforms that don't have a defined safe zone (0..1)\n")
+	TEXT(" default: 1.0"));
+
+FAutoConsoleVariableRef GDebugActionZoneRatioCVar(
+	TEXT("r.DebugActionZone.ActionRatio"),
+	GDebugActionZoneRatio,
+	TEXT("The action zone ratio that will be returned by FDisplayMetrics::GetDisplayMetrics on platforms that don't have a defined safe zone (0..1)\n")
+	TEXT(" default: 1.0"));
+
+float FDisplayMetrics::GetDebugTitleSafeZoneRatio()
+{
+	return GDebugSafeZoneRatio;
+}
+
+float FDisplayMetrics::GetDebugActionSafeZoneRatio()
+{
+	return GDebugActionZoneRatio;
+}
+
+void FDisplayMetrics::ApplyDefaultSafeZones()
+{
+	const float SafeZoneRatio = GetDebugTitleSafeZoneRatio();
+	if (SafeZoneRatio < 1.0f)
+	{
+		const float HalfUnsafeRatio = (1.0f - SafeZoneRatio) * 0.5f;
+		TitleSafePaddingSize = FVector2D(PrimaryDisplayWidth * HalfUnsafeRatio, PrimaryDisplayHeight * HalfUnsafeRatio);
+	}
+
+	const float ActionSafeZoneRatio = GetDebugActionSafeZoneRatio();
+	if (ActionSafeZoneRatio < 1.0f)
+	{
+		const float HalfUnsafeRatio = (1.0f - ActionSafeZoneRatio) * 0.5f;
+		ActionSafePaddingSize = FVector2D(PrimaryDisplayWidth * HalfUnsafeRatio, PrimaryDisplayHeight * HalfUnsafeRatio);
+	}
+}
+
 void FDisplayMetrics::PrintToLog() const
 {
 	UE_LOG(LogInit, Log, TEXT("Display metrics:"));
diff --git a/Engine/Source/Runtime/Core/Private/GenericPlatform/GenericPlatformCrashContext.cpp b/Engine/Source/Runtime/Core/Private/GenericPlatform/GenericPlatformCrashContext.cpp
--- a/Engine/Source/Runtime/Core/Private/GenericPlatform/GenericPlatformCrashContext.cpp
+++ b/Engine/Source/Runtime/Core/Private/GenericPlatform/GenericPlatformCrashContext.cpp
@@ -49,6 +49,7 @@ namespace NCachedCrashContextProperties
 	static int32 SecondsSinceStart;
 	static FString CrashGUID;
 	static FString UserActivityHint;
+	static FString GameSessionID;
 	static FString CommandLine;
 }
 
@@ -109,6 +110,11 @@ void FGenericCrashContext::Initialize()
 		NCachedCrashContextProperties::UserActivityHint = InUserActivity;
 	});
 
+	FCoreDelegates::GameSessionIDChanged.AddLambda([](const FString& InGameSessionID)
+	{
+		NCachedCrashContextProperties::GameSessionID = InGameSessionID;
+	});
+
 	bIsInitialized = true;
 }
 
@@ -137,6 +143,7 @@ void FGenericCrashContext::SerializeContentToBuffer()
 	AddCrashProperty( TEXT( "GameName" ), *NCachedCrashContextProperties::GameName );
 	AddCrashProperty( TEXT( "ExecutableName" ), *NCachedCrashContextProperties::ExecutableName );
 	AddCrashProperty( TEXT( "BuildConfiguration" ), EBuildConfigurations::ToString( FApp::GetBuildConfiguration() ) );
+	AddCrashProperty( TEXT( "GameSessionID" ), *NCachedCrashContextProperties::GameSessionID );
 
 	AddCrashProperty( TEXT( "PlatformName" ), *NCachedCrashContextProperties::PlatformName );
 	AddCrashProperty( TEXT( "PlatformNameIni" ), *NCachedCrashContextProperties::PlatformNameIni );
diff --git a/Engine/Source/Runtime/Core/Private/GenericPlatform/GenericPlatformStackWalk.cpp b/Engine/Source/Runtime/Core/Private/GenericPlatform/GenericPlatformStackWalk.cpp
--- a/Engine/Source/Runtime/Core/Private/GenericPlatform/GenericPlatformStackWalk.cpp
+++ b/Engine/Source/Runtime/Core/Private/GenericPlatform/GenericPlatformStackWalk.cpp
@@ -151,3 +151,8 @@ void FGenericPlatformStackWalk::StackWalkAndDump( ANSICHAR* HumanReadableString,
 	}
 }
 
+void FGenericPlatformStackWalk::StackWalkAndDumpEx(ANSICHAR* HumanReadableString, SIZE_T HumanReadableStringSize, int32 IgnoreCount, uint32 Flags, void* Context)
+{
+	// generic implementation ignores extra flags
+	return StackWalkAndDump(HumanReadableString, HumanReadableStringSize, IgnoreCount, Context);
+}
diff --git a/Engine/Source/Runtime/Core/Private/HAL/MallocLeakDetection.cpp b/Engine/Source/Runtime/Core/Private/HAL/MallocLeakDetection.cpp
new file mode 100644
--- /dev/null
+++ b/Engine/Source/Runtime/Core/Private/HAL/MallocLeakDetection.cpp
@@ -0,0 +1,219 @@
+// Copyright 1998-2015 Epic Games, Inc. All Rights Reserved.
+
+/*=============================================================================
+	MallocLeakDetection.cpp: Helper class to track memory allocations
+=============================================================================*/
+
+#include "CorePrivatePCH.h"
+#include "MemoryMisc.h"
+#include "MallocLeakDetection.h"
+
+#if MALLOC_LEAKDETECTION
+
+FAutoConsoleCommand InvalidateCachedShaders(
+	TEXT("MallocLeak"),
+	TEXT("Usage:\n")
+	TEXT("'MallocLeak Start'  Begins accumulating unique callstacks\n")
+	TEXT("'MallocLeak Stop'  Dumps outstanding unique callstacks and stops accumulation.  Also clears data.\n")
+	TEXT("'MallocLeak Clear'  Clears all accumulated data.  Does not change start/stop state.\n")
+	TEXT("'MallocLeak Dump Size' Dumps oustanding unique callstacks with optional size filter in bytes.\n"),
+	FConsoleCommandWithArgsDelegate::CreateStatic(FMallocLeakDetection::HandleMallocLeakCommand),
+	ECVF_Default	
+	);
+
+FMallocLeakDetection::FMallocLeakDetection()
+	: bCaptureAllocs(false)
+	, bDumpOustandingAllocs(false)
+{
+}
+
+FMallocLeakDetection& FMallocLeakDetection::Get()
+{
+	static FMallocLeakDetection Singleton;
+	return Singleton;
+}
+
+FMallocLeakDetection::~FMallocLeakDetection()
+{	
+}
+
+void FMallocLeakDetection::HandleMallocLeakCommandInternal(const TArray< FString >& Args)
+{
+	if (Args.Num() >= 1)
+	{
+		if (Args[0].Compare(TEXT("Start"), ESearchCase::IgnoreCase) == 0)
+		{		
+			UE_LOG(LogConsoleResponse, Display, TEXT("Starting allocation tracking."));
+			SetAllocationCollection(true);
+		}
+		else if (Args[0].Compare(TEXT("Stop"), ESearchCase::IgnoreCase) == 0)
+		{
+			UE_LOG(LogConsoleResponse, Display, TEXT("Stopping allocation tracking and clearing data."));
+			SetAllocationCollection(false);
+			DumpOpenCallstacks();
+			ClearData();
+		}
+		else if (Args[0].Compare(TEXT("Clear"), ESearchCase::IgnoreCase) == 0)
+		{
+			UE_LOG(LogConsoleResponse, Display, TEXT("Clearing tracking data."));
+			ClearData();
+		}
+		else if (Args[0].Compare(TEXT("Dump"), ESearchCase::IgnoreCase) == 0)
+		{
+			uint32 FilterSize = 0;
+			if (Args.Num() >= 2)
+			{
+				FilterSize = FCString::Atoi(*Args[1]);
+			}
+			UE_LOG(LogConsoleResponse, Display, TEXT("Dumping unique calltacks with %i bytes or more oustanding."), FilterSize);
+			DumpOpenCallstacks(FilterSize);
+		}
+	}
+}
+
+void FMallocLeakDetection::HandleMallocLeakCommand(const TArray< FString >& Args)
+{
+	Get().HandleMallocLeakCommandInternal(Args);
+}
+
+
+void FMallocLeakDetection::AddCallstack(FCallstackTrack& Callstack)
+{
+	FScopeLock Lock(&AllocatedPointersCritical);
+	uint32 CallstackHash = FCrc::MemCrc32(Callstack.CallStack, sizeof(Callstack.CallStack), 0);
+	FCallstackTrack& UniqueCallstack = UniqueCallstacks.FindOrAdd(CallstackHash);
+	if (UniqueCallstack.Count == 0)
+	{
+		UniqueCallstack = Callstack;
+	}
+	UniqueCallstack.Count++;
+	UniqueCallstack.Size += Callstack.Size;
+}
+
+void FMallocLeakDetection::RemoveCallstack(FCallstackTrack& Callstack)
+{
+	FScopeLock Lock(&AllocatedPointersCritical);
+	uint32 CallstackHash = FCrc::MemCrc32(Callstack.CallStack, sizeof(Callstack.CallStack), 0);
+	FCallstackTrack* UniqueCallstack = UniqueCallstacks.Find(CallstackHash);
+	if (UniqueCallstack)
+	{
+		UniqueCallstack->Count--;
+		UniqueCallstack->Size -= Callstack.Size;
+		if (UniqueCallstack->Count == 0)
+		{
+			UniqueCallstack = nullptr;
+			UniqueCallstacks.Remove(CallstackHash);
+		}
+	}
+}
+
+void FMallocLeakDetection::SetAllocationCollection(bool bEnabled)
+{
+	FScopeLock Lock(&AllocatedPointersCritical);
+	bCaptureAllocs = bEnabled;
+}
+
+void FMallocLeakDetection::DumpOpenCallstacks(uint32 FilterSize)
+{
+	//could be called when OOM so avoid UE_LOG functions.
+	FScopeLock Lock(&AllocatedPointersCritical);
+	FPlatformMisc::LowLevelOutputDebugStringf(TEXT("Dumping out of %i possible open callstacks filtered with more than %u bytes on frame: %i\n"), UniqueCallstacks.Num(), FilterSize, (int32)GFrameCounter);
+	const int32 MaxCallstackLineChars = 2048;
+	char CallstackString[MaxCallstackLineChars];
+	TCHAR CallstackStringWide[MaxCallstackLineChars];
+	FMemory::Memzero(CallstackString);
+	for (const auto& Pair : UniqueCallstacks)
+	{
+		const FCallstackTrack& Callstack = Pair.Value;
+		if (Callstack.Size >= FilterSize)
+		{
+			FPlatformMisc::LowLevelOutputDebugStringf(TEXT("AllocSize: %i, Num: %i, FirstFrameEverAllocated: %i\n"), Callstack.Size, Callstack.Count, Callstack.FrameNumber);
+			for (int32 i = 0; i < FCallstackTrack::Depth; ++i)
+			{
+				FPlatformStackWalk::ProgramCounterToHumanReadableString(i, Callstack.CallStack[i], CallstackString, MaxCallstackLineChars);
+				//convert ansi -> tchar without mallocs in case we are in OOM handler.
+				for (int32 CurrChar = 0; CurrChar < MaxCallstackLineChars; ++CurrChar)
+				{
+					CallstackStringWide[CurrChar] = CallstackString[CurrChar];
+				}
+				FPlatformMisc::LowLevelOutputDebugStringf(TEXT("%s\n"), CallstackStringWide);
+				FMemory::Memzero(CallstackString);
+			}
+		}
+	}
+}
+
+void FMallocLeakDetection::ClearData()
+{
+	FScopeLock Lock(&AllocatedPointersCritical);
+	OpenPointers.Empty();
+	UniqueCallstacks.Empty();
+}
+
+bool FMallocLeakDetection::Exec(UWorld* InWorld, const TCHAR* Cmd, FOutputDevice& Ar)
+{		
+	return false;
+}
+
+void FMallocLeakDetection::Malloc(void* Ptr, SIZE_T Size)
+{
+	if (Ptr)
+	{
+		if (bCaptureAllocs)
+		{
+			FScopeLock Lock(&AllocatedPointersCritical);
+			if (!bRecursive)
+			{
+				bRecursive = true;
+				FCallstackTrack Callstack;
+				FPlatformStackWalk::CaptureStackBackTrace(Callstack.CallStack, FCallstackTrack::Depth);
+				Callstack.FrameNumber = GFrameCounter;
+				Callstack.Size = Size;
+				AddCallstack(Callstack);
+				OpenPointers.Add(Ptr, Callstack);
+				bRecursive = false;
+			}
+		}		
+	}
+}
+
+void FMallocLeakDetection::Realloc(void* OldPtr, void* NewPtr, SIZE_T Size)
+{	
+	if (bCaptureAllocs)
+	{
+		FScopeLock Lock(&AllocatedPointersCritical);		
+		Free(OldPtr);
+		Malloc(NewPtr, Size);
+	}
+}
+
+void FMallocLeakDetection::Free(void* Ptr)
+{
+	if (Ptr)
+	{
+		if (bCaptureAllocs)
+		{
+			FScopeLock Lock(&AllocatedPointersCritical);			
+			if (!bRecursive)
+			{
+				bRecursive = true;
+				FCallstackTrack* Callstack = OpenPointers.Find(Ptr);
+				if (Callstack)
+				{
+					RemoveCallstack(*Callstack);
+				}
+				OpenPointers.Remove(Ptr);
+				bRecursive = false;
+			}
+		}
+	}
+}
+
+FMallocLeakDetectionProxy::FMallocLeakDetectionProxy(FMalloc* InMalloc)
+	: UsedMalloc(InMalloc)
+	, Verify(FMallocLeakDetection::Get())
+{
+	
+}
+
+#endif // MALLOC_LEAKDETECTION
\ No newline at end of file
diff --git a/Engine/Source/Runtime/Core/Private/HAL/MallocLeakDetection.h b/Engine/Source/Runtime/Core/Private/HAL/MallocLeakDetection.h
new file mode 100644
--- /dev/null
+++ b/Engine/Source/Runtime/Core/Private/HAL/MallocLeakDetection.h
@@ -0,0 +1,150 @@
+// Copyright 1998-2015 Epic Games, Inc. All Rights Reserved.
+
+/*=============================================================================
+	MallocLeakDetection.h: Helper class to track memory allocations
+=============================================================================*/
+
+#pragma once
+
+#include "Allocators/AnsiAllocator.h"
+
+#define MALLOC_LEAKDETECTION 0
+
+#if MALLOC_LEAKDETECTION
+
+/**
+ * Maintains a list of all pointers to currently allocated memory.
+ */
+class FMallocLeakDetection
+{
+	struct FCallstackTrack
+	{
+		FCallstackTrack()
+		{
+			FMemory::Memzero(this, sizeof(FCallstackTrack));
+		}
+		static const int32 Depth = 32;
+		uint64 CallStack[Depth];
+		uint32 FrameNumber;
+		uint32 Size;
+		uint32 Count;
+	};
+
+	FMallocLeakDetection();
+	~FMallocLeakDetection();
+
+	void AddCallstack(FCallstackTrack& Callstack);
+	void RemoveCallstack(FCallstackTrack& Callstack);
+	void HandleMallocLeakCommandInternal(const TArray< FString >& Args);
+
+	/** List of all currently allocated pointers */
+	TMap<void*, FCallstackTrack> OpenPointers;
+
+	/** List of all unique callstacks with allocated memory */
+	TMap<uint32, FCallstackTrack> UniqueCallstacks;
+
+	bool bRecursive;
+	bool bCaptureAllocs;
+	bool bDumpOustandingAllocs;	
+		
+	FCriticalSection AllocatedPointersCritical;	
+
+public:	
+
+	static FMallocLeakDetection& Get();
+	static void HandleMallocLeakCommand(const TArray< FString >& Args);
+
+	void SetAllocationCollection(bool bEnabled);
+	void DumpOpenCallstacks(uint32 FilterSize = 0);
+	void ClearData();
+
+	bool Exec(UWorld* InWorld, const TCHAR* Cmd, FOutputDevice& Ar);
+
+	/** Handles new allocated pointer */
+	void Malloc(void* Ptr, SIZE_T Size);
+
+	/** Handles reallocation */
+	void Realloc(void* OldPtr, void* NewPtr, SIZE_T Size);
+
+	/** Removes allocated pointer from list */
+	void Free(void* Ptr);	
+};
+
+/**
+ * A verifying proxy malloc that takes a malloc to be used and tracks unique callstacks with outstanding allocations
+ * to help identify leaks.
+ */
+class FMallocLeakDetectionProxy : public FMalloc
+{
+private:
+	/** Malloc we're based on, aka using under the hood */
+	FMalloc* UsedMalloc;
+
+	/* Verifier object */
+	FMallocLeakDetection& Verify;
+
+public:
+	explicit FMallocLeakDetectionProxy(FMalloc* InMalloc);	
+
+	virtual void* Malloc(SIZE_T Size, uint32 Alignment) override
+	{
+		void* Result = UsedMalloc->Malloc(Size, Alignment);
+		Verify.Malloc(Result, Size);
+		return Result;
+	}
+
+	virtual void* Realloc(void* Ptr, SIZE_T NewSize, uint32 Alignment) override
+	{
+		void* Result = UsedMalloc->Realloc(Ptr, NewSize, Alignment);
+		Verify.Realloc(Ptr, Result, NewSize);
+		return Result;
+	}
+
+	virtual void Free(void* Ptr) override
+	{
+		if (Ptr)
+		{
+			Verify.Free(Ptr);
+			UsedMalloc->Free(Ptr);
+		}
+	}
+
+	virtual void InitializeStatsMetadata() override
+	{
+		UsedMalloc->InitializeStatsMetadata();
+	}
+
+	virtual void GetAllocatorStats(FGenericMemoryStats& OutStats) override
+	{
+		UsedMalloc->GetAllocatorStats(OutStats);
+	}
+
+	virtual void DumpAllocatorStats(FOutputDevice& Ar) override
+	{
+		Verify.DumpOpenCallstacks();
+		UsedMalloc->DumpAllocatorStats(Ar);
+	}
+
+	virtual bool ValidateHeap() override
+	{
+		return UsedMalloc->ValidateHeap();
+	}
+
+	virtual bool Exec(UWorld* InWorld, const TCHAR* Cmd, FOutputDevice& Ar) override
+	{
+		Verify.Exec(InWorld, Cmd, Ar);
+		return UsedMalloc->Exec(InWorld, Cmd, Ar);
+	}
+
+	virtual bool GetAllocationSize(void* Original, SIZE_T& OutSize) override
+	{
+		return UsedMalloc->GetAllocationSize(Original, OutSize);
+	}
+
+	virtual const TCHAR* GetDescriptiveName() override
+	{ 
+		return UsedMalloc->GetDescriptiveName();
+	}
+};
+
+#endif // MALLOC_LEAKDETECTION
diff --git a/Engine/Source/Runtime/Core/Private/HAL/UnrealMemory.cpp b/Engine/Source/Runtime/Core/Private/HAL/UnrealMemory.cpp
--- a/Engine/Source/Runtime/Core/Private/HAL/UnrealMemory.cpp
+++ b/Engine/Source/Runtime/Core/Private/HAL/UnrealMemory.cpp
@@ -14,6 +14,7 @@
 #include "MallocProfiler.h"
 #include "MallocThreadSafeProxy.h"
 #include "MallocVerify.h"
+#include "MallocLeakDetection.h"
 #include "PlatformMallocCrash.h"
 
 
@@ -244,6 +245,10 @@ void GCreateMalloc()
 	// Add the verifier
 	GMalloc = new FMallocVerifyProxy( GMalloc );
 #endif
+
+#if MALLOC_LEAKDETECTION
+	GMalloc = new FMallocLeakDetectionProxy(GMalloc);
+#endif
 }
 
 
diff --git a/Engine/Source/Runtime/Core/Private/HTML5/HTML5Application.cpp b/Engine/Source/Runtime/Core/Private/HTML5/HTML5Application.cpp
--- a/Engine/Source/Runtime/Core/Private/HTML5/HTML5Application.cpp
+++ b/Engine/Source/Runtime/Core/Private/HTML5/HTML5Application.cpp
@@ -215,6 +215,9 @@ void FDisplayMetrics::GetDisplayMetrics(FDisplayMetrics& OutDisplayMetrics)
 	OutDisplayMetrics.PrimaryDisplayWidth   =	OutDisplayMetrics.PrimaryDisplayWorkAreaRect.Right;
 	OutDisplayMetrics.PrimaryDisplayHeight  =	OutDisplayMetrics.PrimaryDisplayWorkAreaRect.Bottom; 
 	UE_LOG(LogHTML5Application, Verbose, TEXT("GetDisplayMetrics Width:%d, Height:%d"), OutDisplayMetrics.PrimaryDisplayWorkAreaRect.Right, OutDisplayMetrics.PrimaryDisplayWorkAreaRect.Bottom);
+
+	// Apply the debug safe zones
+	OutDisplayMetrics.ApplyDefaultSafeZones();
 }
 
 TSharedRef< FGenericWindow > FHTML5Application::MakeWindow()
diff --git a/Engine/Source/Runtime/Core/Private/Linux/LinuxPlatformCrashContext.cpp b/Engine/Source/Runtime/Core/Private/Linux/LinuxPlatformCrashContext.cpp
--- a/Engine/Source/Runtime/Core/Private/Linux/LinuxPlatformCrashContext.cpp
+++ b/Engine/Source/Runtime/Core/Private/Linux/LinuxPlatformCrashContext.cpp
@@ -306,6 +306,61 @@ void FLinuxCrashContext::CaptureStackTrace()
 	}
 }
 
+namespace LinuxCrashReporterTracker
+{
+	FProcHandle CurrentlyRunningCrashReporter;
+	FDelegateHandle CurrentTicker;
+
+	bool Tick(float DeltaTime)
+	{
+		if (!FPlatformProcess::IsProcRunning(CurrentlyRunningCrashReporter))
+		{
+			FPlatformProcess::CloseProc(CurrentlyRunningCrashReporter);
+			CurrentlyRunningCrashReporter = FProcHandle();
+
+			FTicker::GetCoreTicker().RemoveTicker(CurrentTicker);
+			CurrentTicker.Reset();
+
+			UE_LOG(LogLinux, Log, TEXT("Done sending crash report for ensure()."));
+			return false;
+		}
+
+		// tick again
+		return true;
+	}
+
+	/**
+	 * Waits for the proc with timeout (busy loop, workaround for platform abstraction layer not exposing this)
+	 *
+	 * @param Proc proc handle to wait for
+	 * @param TimeoutInSec timeout in seconds
+	 * @param SleepIntervalInSec sleep interval (the smaller the more CPU we will eat, but the faster we will detect the program exiting)
+	 *
+	 * @return true if exited cleanly, false if timeout has expired
+	 */
+	bool WaitForProcWithTimeout(FProcHandle Proc, const double TimeoutInSec, const double SleepIntervalInSec)
+	{
+		double StartSeconds = FPlatformTime::Seconds();
+		for (;;)
+		{
+			if (!FPlatformProcess::IsProcRunning(Proc))
+			{
+				break;
+			}
+
+			if (FPlatformTime::Seconds() - StartSeconds > TimeoutInSec)
+			{
+				return false;
+			}
+
+			FPlatformProcess::Sleep(SleepIntervalInSec);
+		};
+
+		return true;
+	}
+}
+
+
 void FLinuxCrashContext::GenerateCrashInfoAndLaunchReporter(bool bReportingNonCrash) const
 {
 	// do not report crashes for tools (particularly for crash reporter itself)
@@ -361,26 +416,39 @@ void FLinuxCrashContext::GenerateCrashInfoAndLaunchReporter(bool bReportingNonCr
 
 		CrashReportClientArguments += CrashInfoAbsolute + TEXT("/");
 
-		FProcHandle RunningProc = FPlatformProcess::CreateProc(RelativePathToCrashReporter, *CrashReportClientArguments, true, false, false, NULL, 0, NULL, NULL);
-		if (FPlatformProcess::IsProcRunning(RunningProc))
+		if (bReportingNonCrash)
 		{
-			// do not wait indefinitely
-			double kTimeOut = 3 * 60.0;
-			double StartSeconds = FPlatformTime::Seconds();
-			for(;;)
+			// If we're reporting non-crash, try to avoid spinning here and instead do that in the tick.
+			// However, if there was already a crash reporter running (i.e. we hit ensure() too quickly), take a hitch here
+			if (FPlatformProcess::IsProcRunning(LinuxCrashReporterTracker::CurrentlyRunningCrashReporter))
 			{
-				if (!FPlatformProcess::IsProcRunning(RunningProc))
+				// do not wait indefinitely, allow 45 second hitch (anticipating callstack parsing)
+				const double kEnsureTimeOut = 45.0;
+				const double kEnsureSleepInterval = 0.1;
+				if (!LinuxCrashReporterTracker::WaitForProcWithTimeout(LinuxCrashReporterTracker::CurrentlyRunningCrashReporter, kEnsureTimeOut, kEnsureSleepInterval))
 				{
-					break;
+					FPlatformProcess::TerminateProc(LinuxCrashReporterTracker::CurrentlyRunningCrashReporter);
 				}
 
-				if (FPlatformTime::Seconds() - StartSeconds > kTimeOut)
-				{
-					break;
-				}
+				LinuxCrashReporterTracker::Tick(0.001f);	// tick one more time to make it clean up after itself
+			}
+
+			LinuxCrashReporterTracker::CurrentlyRunningCrashReporter = FPlatformProcess::CreateProc(RelativePathToCrashReporter, *CrashReportClientArguments, true, false, false, NULL, 0, NULL, NULL);
+			LinuxCrashReporterTracker::CurrentTicker = FTicker::GetCoreTicker().AddTicker(FTickerDelegate::CreateStatic(&LinuxCrashReporterTracker::Tick), 1.f);
+		}
+		else
+		{
+			// spin here until CrashReporter exits
+			FProcHandle RunningProc = FPlatformProcess::CreateProc(RelativePathToCrashReporter, *CrashReportClientArguments, true, false, false, NULL, 0, NULL, NULL);
+			// do not wait indefinitely - can be more generous about the hitch than in ensure() case
+			const double kCrashTimeOut = 3 * 60.0;
+			const double kCrashSleepInterval = 1.0;
+			if (!LinuxCrashReporterTracker::WaitForProcWithTimeout(RunningProc, kCrashTimeOut, kCrashSleepInterval))
+			{
+				FPlatformProcess::TerminateProc(RunningProc);
+			}
 
-				FPlatformProcess::Sleep(1.0f);
-			};
+			FPlatformProcess::CloseProc(RunningProc);
 		}
 	}
 
diff --git a/Engine/Source/Runtime/Core/Private/Linux/LinuxPlatformProcess.cpp b/Engine/Source/Runtime/Core/Private/Linux/LinuxPlatformProcess.cpp
--- a/Engine/Source/Runtime/Core/Private/Linux/LinuxPlatformProcess.cpp
+++ b/Engine/Source/Runtime/Core/Private/Linux/LinuxPlatformProcess.cpp
@@ -981,7 +981,7 @@ bool FProcState::IsRunning()
 		// which is a dubious, but valid behavior. We don't want to keep zombie around though.
 		if (!bIsRunning)
 		{
-			UE_LOG(LogHAL, Log, TEXT("Child %d is no more running (zombie), Wait()ing immediately."), GetProcessId() );
+			UE_LOG(LogHAL, Log, TEXT("Child %d is no longer running (zombie), Wait()ing immediately."), GetProcessId() );
 			Wait();
 		}
 	}
diff --git a/Engine/Source/Runtime/Core/Private/Linux/LinuxPlatformStackWalk.cpp b/Engine/Source/Runtime/Core/Private/Linux/LinuxPlatformStackWalk.cpp
--- a/Engine/Source/Runtime/Core/Private/Linux/LinuxPlatformStackWalk.cpp
+++ b/Engine/Source/Runtime/Core/Private/Linux/LinuxPlatformStackWalk.cpp
@@ -694,7 +694,8 @@ bool FLinuxPlatformStackWalk::ProgramCounterToHumanReadableString( int32 Current
 				{
 					const char * SourceFilename = NULL;
 					int LineNumber;
-					if (LinuxStackWalkHelpers::GetBacktraceSymbols()->GetInfoForAddress(reinterpret_cast< void* >( ProgramCounter ), NULL, &SourceFilename, &LineNumber) && FunctionName != NULL)
+					// do not symbolicate if we're handling ensure()
+					if (!LinuxContext->bHandlingEnsure && LinuxStackWalkHelpers::GetBacktraceSymbols()->GetInfoForAddress(reinterpret_cast< void* >( ProgramCounter ), NULL, &SourceFilename, &LineNumber) && FunctionName != nullptr)
 					{
 						FCStringAnsi::Sprintf(TempArray, " [%s, line %d]", SourceFilename, LineNumber);
 						LinuxStackWalkHelpers::AppendToString(HumanReadableString, HumanReadableStringSize, Context, TempArray);
@@ -724,7 +725,7 @@ bool FLinuxPlatformStackWalk::ProgramCounterToHumanReadableString( int32 Current
 
 void FLinuxPlatformStackWalk::StackWalkAndDump( ANSICHAR* HumanReadableString, SIZE_T HumanReadableStringSize, int32 IgnoreCount, void* Context )
 {
-	if ( Context == nullptr )
+	if (Context == nullptr)
 	{
 		FLinuxCrashContext CrashContext;
 		CrashContext.InitFromSignal(0, nullptr, nullptr);
@@ -736,6 +737,23 @@ void FLinuxPlatformStackWalk::StackWalkAndDump( ANSICHAR* HumanReadableString, S
 	}
 }
 
+void FLinuxPlatformStackWalk::StackWalkAndDumpEx(ANSICHAR* HumanReadableString, SIZE_T HumanReadableStringSize, int32 IgnoreCount, uint32 Flags, void* Context)
+{
+	const bool bHandlingEnsure = (Flags & EStackWalkFlags::FlagsUsedWhenHandlingEnsure) == EStackWalkFlags::FlagsUsedWhenHandlingEnsure;
+	if (Context == nullptr)
+	{
+		FLinuxCrashContext CrashContext;
+		CrashContext.InitFromSignal(0, nullptr, nullptr);
+		CrashContext.bHandlingEnsure = bHandlingEnsure;
+		FPlatformStackWalk::StackWalkAndDump(HumanReadableString, HumanReadableStringSize, IgnoreCount + 1, &CrashContext);
+	}
+	else
+	{
+		TGuardValue<bool> Guard(reinterpret_cast<FLinuxCrashContext *>(Context)->bHandlingEnsure, bHandlingEnsure);
+		FPlatformStackWalk::StackWalkAndDump(HumanReadableString, HumanReadableStringSize, IgnoreCount + 1, Context);
+	}
+}
+
 void FLinuxPlatformStackWalk::CaptureStackBackTrace( uint64* BackTrace, uint32 MaxDepth, void* Context )
 {
 	// Make sure we have place to store the information before we go through the process of raising
@@ -782,6 +800,7 @@ void NewReportEnsure(const TCHAR* ErrorMessage)
 
 	FLinuxCrashContext EnsureContext;
 	EnsureContext.InitFromSignal(SIGTRAP, &Signal, nullptr);
+	EnsureContext.bHandlingEnsure = true;
 	EnsureContext.CaptureStackTrace();
 	EnsureContext.GenerateCrashInfoAndLaunchReporter(true);
 
diff --git a/Engine/Source/Runtime/Core/Private/Misc/CallbackDevice.cpp b/Engine/Source/Runtime/Core/Private/Misc/CallbackDevice.cpp
--- a/Engine/Source/Runtime/Core/Private/Misc/CallbackDevice.cpp
+++ b/Engine/Source/Runtime/Core/Private/Misc/CallbackDevice.cpp
@@ -74,4 +74,6 @@ FCoreDelegates::FVRHeadsetRecenter FCoreDelegates::VRHeadsetRecenter;
 
 FCoreDelegates::FOnUserActivityStringChanged FCoreDelegates::UserActivityStringChanged;
 
+FCoreDelegates::FOnGameSessionIDChange FCoreDelegates::GameSessionIDChanged;
+
 FCoreDelegates::FOnAsyncLoadingFlush FCoreDelegates::OnAsyncLoadingFlush;
\ No newline at end of file
diff --git a/Engine/Source/Runtime/Core/Private/Misc/CoreMisc.cpp b/Engine/Source/Runtime/Core/Private/Misc/CoreMisc.cpp
--- a/Engine/Source/Runtime/Core/Private/Misc/CoreMisc.cpp
+++ b/Engine/Source/Runtime/Core/Private/Misc/CoreMisc.cpp
@@ -560,6 +560,7 @@ bool FCommandLine::bIsInitialized = false;
 TCHAR FCommandLine::CmdLine[FCommandLine::MaxCommandLineSize] = TEXT("");
 TCHAR FCommandLine::OriginalCmdLine[FCommandLine::MaxCommandLineSize] = TEXT("");
 TCHAR FCommandLine::LoggingCmdLine[FCommandLine::MaxCommandLineSize] = TEXT("");
+TCHAR FCommandLine::LoggingOriginalCmdLine[FCommandLine::MaxCommandLineSize] = TEXT("");
 FString FCommandLine::SubprocessCommandLine(TEXT(" -Multiprocess"));
 
 bool FCommandLine::IsInitialized()
@@ -585,11 +586,18 @@ const TCHAR* FCommandLine::GetOriginal()
 	return OriginalCmdLine;
 }
 
+const TCHAR* FCommandLine::GetOriginalForLogging()
+{
+	UE_CLOG(!bIsInitialized, LogInit, Fatal, TEXT("Attempting to get the command line but it hasn't been initialized yet."));
+	return LoggingOriginalCmdLine;
+}
+
 bool FCommandLine::Set(const TCHAR* NewCommandLine)
 {
 	if (!bIsInitialized)
 	{
-		FCString::Strncpy(OriginalCmdLine, NewCommandLine, ARRAY_COUNT(CmdLine));
+		FCString::Strncpy(OriginalCmdLine, NewCommandLine, ARRAY_COUNT(OriginalCmdLine));
+		FCString::Strncpy(LoggingOriginalCmdLine, NewCommandLine, ARRAY_COUNT(LoggingOriginalCmdLine));
 	}
 
 	FCString::Strncpy( CmdLine, NewCommandLine, ARRAY_COUNT(CmdLine) );
@@ -673,6 +681,9 @@ void FCommandLine::WhitelistCommandLines()
 	// Process the command line for logging purposes
 	TArray<FString> LoggingCmdList = FilterCommandLineForLogging(LoggingCmdLine);
 	BuildWhitelistCommandLine(LoggingCmdLine, ARRAY_COUNT(LoggingCmdLine), LoggingCmdList);
+	// Process the original command line for logging purposes
+	TArray<FString> LoggingOriginalCmdList = FilterCommandLineForLogging(LoggingOriginalCmdLine);
+	BuildWhitelistCommandLine(LoggingOriginalCmdLine, ARRAY_COUNT(LoggingOriginalCmdLine), LoggingOriginalCmdList);
 }
 
 TArray<FString> FCommandLine::FilterCommandLine(TCHAR* CommandLine)
diff --git a/Engine/Source/Runtime/Core/Private/Misc/DateTime.cpp b/Engine/Source/Runtime/Core/Private/Misc/DateTime.cpp
--- a/Engine/Source/Runtime/Core/Private/Misc/DateTime.cpp
+++ b/Engine/Source/Runtime/Core/Private/Misc/DateTime.cpp
@@ -181,6 +181,40 @@ FString FDateTime::ToIso8601() const
 	return ToString(TEXT("%Y-%m-%dT%H:%M:%S.%sZ"));
 }
 
+FString FDateTime::ToHttpDate() const
+{
+	FString DayStr;
+	switch (GetDayOfWeek())
+	{
+		case EDayOfWeek::Monday:	DayStr = TEXT("Mon");	break;
+		case EDayOfWeek::Tuesday:	DayStr = TEXT("Tue");	break;
+		case EDayOfWeek::Wednesday:	DayStr = TEXT("Wed");	break;
+		case EDayOfWeek::Thursday:	DayStr = TEXT("Thu");	break;
+		case EDayOfWeek::Friday:	DayStr = TEXT("Fri");	break;
+		case EDayOfWeek::Saturday:	DayStr = TEXT("Sat");	break;
+		case EDayOfWeek::Sunday:	DayStr = TEXT("Sun");	break;
+	}
+
+	FString MonthStr;
+	switch (GetMonthOfYear())
+	{
+		case EMonthOfYear::January:		MonthStr = TEXT("Jan");	break;
+		case EMonthOfYear::February:	MonthStr = TEXT("Feb");	break;
+		case EMonthOfYear::March:		MonthStr = TEXT("Mar");	break;
+		case EMonthOfYear::April:		MonthStr = TEXT("Apr");	break;
+		case EMonthOfYear::May:			MonthStr = TEXT("May");	break;
+		case EMonthOfYear::June:		MonthStr = TEXT("Jun");	break;
+		case EMonthOfYear::July:		MonthStr = TEXT("Jul");	break;
+		case EMonthOfYear::August:		MonthStr = TEXT("Aug");	break;
+		case EMonthOfYear::September:	MonthStr = TEXT("Sep");	break;
+		case EMonthOfYear::October:		MonthStr = TEXT("Oct");	break;
+		case EMonthOfYear::November:	MonthStr = TEXT("Nov");	break;
+		case EMonthOfYear::December:	MonthStr = TEXT("Dec");	break;
+	}
+
+	FString Time = FString::Printf(TEXT("%02i:%02i:%02i"), GetHour(), GetMinute(), GetSecond());
+	return FString::Printf(TEXT("%s, %02d %s %d %s GMT"), *DayStr, GetDay(), *MonthStr, GetYear(), *Time);
+}
 
 FString FDateTime::ToString() const
 {
@@ -447,6 +481,310 @@ bool FDateTime::ParseIso8601( const TCHAR* DateTimeString, FDateTime& OutDateTim
 	return true;
 }
 
+/**
+ * Parse a string of the HTTP-date format from a web server
+ * https://www.w3.org/Protocols/rfc2616/rfc2616-sec3.html#sec3.3.1
+ *
+ *   HTTP-date    = rfc1123-date | rfc850-date | asctime-date
+ *	 rfc1123-date = wkday "," SP date1 SP time SP "GMT"
+ *	 rfc850-date  = weekday "," SP date2 SP time SP "GMT"
+ *   asctime-date = wkday SP date3 SP time SP 4DIGIT
+ *   date1        = 2DIGIT SP month SP 4DIGIT
+ *				  ; day month year (e.g., 02 Jun 1982)
+ *   date2        = 2DIGIT "-" month "-" 2DIGIT
+ *				  ; day-month-year (e.g., 02-Jun-82)
+ *   date3        = month SP ( 2DIGIT | ( SP 1DIGIT ))
+ *				  ; month day (e.g., Jun  2)
+ *   time         = 2DIGIT ":" 2DIGIT ":" 2DIGIT
+ *				  ; 00:00:00 - 23:59:59
+ *   wkday        = "Mon" | "Tue" | "Wed"
+ *				  | "Thu" | "Fri" | "Sat" | "Sun"
+ *   weekday      = "Monday" | "Tuesday" | "Wednesday"
+ *				  | "Thursday" | "Friday" | "Saturday" | "Sunday"
+ *   month        = "Jan" | "Feb" | "Mar" | "Apr"
+ *				  | "May" | "Jun" | "Jul" | "Aug"
+ *				  | "Sep" | "Oct" | "Nov" | "Dec"
+ */
+bool FDateTime::ParseHttpDate(const FString& HttpDate, FDateTime& OutDateTime)
+{
+	auto ParseTime = [](const FString& Time, int32& Hour, int32& Minute, int32& Second) -> bool
+	{
+		// 2DIGIT ":" 2DIGIT ":" 2DIGIT
+		// ; 00:00 : 00 - 23 : 59 : 59
+		TArray<FString> Tokens;
+
+		// split up on a single delimiter
+		int32 NumTokens = Time.ParseIntoArray(Tokens, TEXT(":"), true);
+		if (NumTokens == 3)
+		{
+			Hour = FCString::Atoi(*Tokens[0]);
+			Minute = FCString::Atoi(*Tokens[1]);
+			Second = FCString::Atoi(*Tokens[2]);
+			return (Hour >= 0 && Hour < 24) && (Minute >= 0 && Minute <= 59) && (Second >= 0 && Second <= 59);
+		}
+
+		return false;
+	};
+
+	auto ParseWkday = [](const FString& WkDay) -> int32
+	{
+		const int32 NumChars = WkDay.Len();
+		if (NumChars == 3)
+		{
+			if (WkDay.Equals(TEXT("Mon")))
+			{
+				return 1;
+			}
+			else if (WkDay.Equals(TEXT("Tue")))
+			{
+				return 2;
+			}
+			else if (WkDay.Equals(TEXT("Wed")))
+			{
+				return 3;
+			}
+			else if (WkDay.Equals(TEXT("Thu")))
+			{
+				return 4;
+			}
+			else if (WkDay.Equals(TEXT("Fri")))
+			{
+				return 5;
+			}
+			else if (WkDay.Equals(TEXT("Sat")))
+			{
+				return 6;
+			}
+			else if (WkDay.Equals(TEXT("Sun")))
+			{
+				return 7;
+			}
+		}
+
+		return -1;
+	};
+
+	auto ParseWeekday = [](const FString& WeekDay) -> int32
+	{
+		const int32 NumChars = WeekDay.Len();
+		if (NumChars >= 6 && NumChars <= 9)
+		{
+			if (WeekDay.Equals(TEXT("Monday")))
+			{
+				return 1;
+			}
+			else if (WeekDay.Equals(TEXT("Tueday")))
+			{
+				return 2;
+			}
+			else if (WeekDay.Equals(TEXT("Wednesday")))
+			{
+				return 3;
+			}
+			else if (WeekDay.Equals(TEXT("Thursday")))
+			{
+				return 4;
+			}
+			else if (WeekDay.Equals(TEXT("Friday")))
+			{
+				return 5;
+			}
+			else if (WeekDay.Equals(TEXT("Saturday")))
+			{
+				return 6;
+			}
+			else if (WeekDay.Equals(TEXT("Sunday")))
+			{
+				return 7;
+			}
+		}
+
+		return -1;
+	};
+
+	auto ParseMonth = [](const FString& Month) -> int32
+	{
+		const int32 NumChars = Month.Len();
+		if (NumChars == 3)
+		{
+			if (Month.Equals(TEXT("Jan")))
+			{
+				return 1;
+			}
+			else if (Month.Equals(TEXT("Feb")))
+			{
+				return 2;
+			}
+			else if (Month.Equals(TEXT("Mar")))
+			{
+				return 3;
+			}
+			else if (Month.Equals(TEXT("Apr")))
+			{
+				return 4;
+			}
+			else if (Month.Equals(TEXT("May")))
+			{
+				return 5;
+			}
+			else if (Month.Equals(TEXT("Jun")))
+			{
+				return 6;
+			}
+			else if (Month.Equals(TEXT("Jul")))
+			{
+				return 7;
+			}
+			else if (Month.Equals(TEXT("Aug")))
+			{
+				return 8;
+			}
+			else if (Month.Equals(TEXT("Sep")))
+			{
+				return 9;
+			}
+			else if (Month.Equals(TEXT("Oct")))
+			{
+				return 10;
+			}
+			else if (Month.Equals(TEXT("Nov")))
+			{
+				return 11;
+			}
+			else if (Month.Equals(TEXT("Dec")))
+			{
+				return 12;
+			}
+		}
+
+		return -1;
+	};
+
+	auto ParseDate1 = [ParseMonth](const FString& DayStr, const FString& MonStr, const FString& YearStr, int32& Month, int32& Day, int32& Year) -> bool
+	{
+		// date1 = 2DIGIT SP month SP 4DIGIT
+		// ; day month year(e.g., 02 Jun 1982)
+
+		Day = FCString::Atoi(*DayStr);
+		Month = ParseMonth(MonStr);
+		Year = (YearStr.Len() == 4) ? FCString::Atoi(*YearStr) : -1;
+
+		return (Day > 0 && Day <= 31) && (Month > 0 && Month <= 12) && (Year > 0 && Year <= 9999);
+	};
+
+	auto ParseDate2 = [ParseMonth](const FString& Date2, int32& Month, int32& Day, int32& Year) -> bool
+	{
+		// date2 = 2DIGIT "-" month "-" 2DIGIT
+		// ; day - month - year(e.g., 02 - Jun - 82)
+		TArray<FString> Tokens;
+
+		// split up on a single delimiter
+		int32 NumTokens = Date2.ParseIntoArray(Tokens, TEXT("-"), true);
+		if (NumTokens == 3)
+		{
+			Day = FCString::Atoi(*Tokens[0]);
+			Month = ParseMonth(Tokens[1]);
+			Year = FCString::Atoi(*Tokens[2]);
+			
+			// Horrible assumption here, but this is a deprecated part of the spec
+			Year += 1900;
+		}
+
+		return (Day > 0 && Day <= 31) && (Month > 0 && Month <= 12) && (Year > 0 && Year <= 9999);
+	};
+
+	auto ParseDate3 = [ParseMonth](const FString& MonStr, const FString& DayStr, int32& Month, int32& Day) -> bool
+	{
+		// date3 = month SP(2DIGIT | (SP 1DIGIT))
+		// ; month day(e.g., Jun  2)
+		const int32 NumDigits = DayStr.Len();
+		Day = (NumDigits > 0 && NumDigits <= 2) ? FCString::Atoi(*DayStr) : -1;
+		Month = ParseMonth(MonStr);
+
+		return (Day > 0 && Day <= 31) && (Month > 0 && Month <= 12);
+	};
+
+	if (!HttpDate.IsEmpty())
+	{
+		TArray<FString> Tokens;
+
+		// split up on a single delimiter
+		int32 NumTokens = HttpDate.ParseIntoArray(Tokens, TEXT(" "), true);
+
+		if (NumTokens > 0 && ensure(Tokens.Num() == NumTokens))
+		{
+			int32 Month = 0;
+			int32 Day = 0;
+			int32 Year = 0;
+			int32 Hour = 0;
+			int32 Minute = 0;
+			int32 Second = 0;
+			int32 Millisecond = 0;
+
+			if (Tokens[0].EndsWith(TEXT(",")))
+			{
+				Tokens[0].RemoveAt(Tokens[0].Len() - 1, 1);
+			}
+
+			if (Tokens[Tokens.Num() - 1].Equals(TEXT("GMT")))
+			{
+				// rfc1123 - date | rfc850 - date 
+				if (Tokens.Num() == 6)
+				{
+					int32 WkDay = ParseWkday(Tokens[0]);
+					if (WkDay > 0)
+					{
+						// rfc1123 - date = wkday "," SP date1 SP time SP "GMT"
+						if (ParseDate1(Tokens[1], Tokens[2], Tokens[3], Month, Day, Year))
+						{
+							ParseTime(Tokens[4], Hour, Minute, Second);
+						}
+					}
+				}
+				else if (Tokens.Num() == 4)
+				{
+					// rfc850 - date = weekday "," SP date2 SP time SP "GMT"
+					int32 WeekDay = ParseWeekday(Tokens[0]);
+					if (WeekDay > 0)
+					{
+						if (ParseDate2(Tokens[1], Month, Day, Year))
+						{
+							ParseTime(Tokens[2], Hour, Minute, Second);
+						}
+					}
+				}
+
+			}
+			else if (Tokens.Num() == 5)
+			{
+				// asctime - date = wkday SP date3 SP time SP 4DIGIT
+				int32 WkDay = ParseWkday(Tokens[0]);
+				if (WkDay > 0)
+				{
+					if (ParseDate3(Tokens[1], Tokens[2], Month, Day))
+					{
+						if (ParseTime(Tokens[3], Hour, Minute, Second))
+						{
+							if (Tokens[4].Len() == 4)
+							{
+								Year = FCString::Atoi(*Tokens[4]);
+							}
+						}
+					}
+				}
+			}
+
+			if (Validate(Year, Month, Day, Hour, Minute, Second, Millisecond))
+			{
+				// convert the tokens to numbers
+				OutDateTime = FDateTime(Year, Month, Day, Hour, Minute, Second, Millisecond);
+				return true;
+			}
+		}
+	}
+
+	return false;
+}
 
 FDateTime FDateTime::UtcNow()
 {
diff --git a/Engine/Source/Runtime/Core/Private/Misc/OutputDevice.cpp b/Engine/Source/Runtime/Core/Private/Misc/OutputDevice.cpp
--- a/Engine/Source/Runtime/Core/Private/Misc/OutputDevice.cpp
+++ b/Engine/Source/Runtime/Core/Private/Misc/OutputDevice.cpp
@@ -314,7 +314,7 @@ void FDebug::EnsureFailed(const ANSICHAR* Expr, const ANSICHAR* File, int32 Line
 			SCOPE_LOG_TIME_IN_SECONDS(*StackWalkPerfMessage, nullptr)
 #endif
 			StackTrace[0] = 0;
-			FPlatformStackWalk::StackWalkAndDump( StackTrace, StackTraceSize, CALLSTACK_IGNOREDEPTH );
+			FPlatformStackWalk::StackWalkAndDumpEx( StackTrace, StackTraceSize, CALLSTACK_IGNOREDEPTH, FGenericPlatformStackWalk::EStackWalkFlags::FlagsUsedWhenHandlingEnsure );
 		}
 
 		// Create a final string that we'll output to the log (and error history buffer)
diff --git a/Engine/Source/Runtime/Core/Private/ProfilingDebugging/LoadTimeTracker.cpp b/Engine/Source/Runtime/Core/Private/ProfilingDebugging/LoadTimeTracker.cpp
new file mode 100644
--- /dev/null
+++ b/Engine/Source/Runtime/Core/Private/ProfilingDebugging/LoadTimeTracker.cpp
@@ -0,0 +1,62 @@
+// Copyright 1998-2016 Epic Games, Inc. All Rights Reserved.
+
+/**
+* Here are a number of profiling helper functions so we do not have to duplicate a lot of the glue
+* code everywhere.  And we can have consistent naming for all our files.
+*
+*/
+
+// Core includes.
+#include "CorePrivatePCH.h"
+#include "LoadTimeTracker.h"
+
+#if ENABLE_LOADTIME_TRACKING
+
+void FLoadTimeTracker::ReportScopeTime(double ScopeTime, const FName ScopeLabel)
+{
+	check(IsInGameThread());
+	TArray<double>& LoadTimes = TimeInfo.FindOrAdd(ScopeLabel);
+	LoadTimes.Add(ScopeTime);
+}
+
+
+void FLoadTimeTracker::DumpLoadTimes()
+{
+	double TotalTime = 0.0;
+	UE_LOG(LogLoad, Log, TEXT("------------- Load times -------------"));
+	for(auto Itr = TimeInfo.CreateIterator(); Itr; ++Itr)
+	{
+		const FString KeyName = Itr.Key().ToString();
+		const TArray<double>& LoadTimes = Itr.Value();
+		if(LoadTimes.Num() == 1)
+		{
+			TotalTime += Itr.Value()[0];
+			UE_LOG(LogLoad, Log, TEXT("%s: %f"), *KeyName, Itr.Value()[0]);
+		}
+		else
+		{
+			double InnerTotal = 0.0;
+			for(int Index = 0; Index < LoadTimes.Num(); ++Index)
+			{
+				InnerTotal += Itr.Value()[Index];
+				UE_LOG(LogLoad, Log, TEXT("%s[%d]: %f"), *KeyName, Index, LoadTimes[Index]);
+			}
+
+			UE_LOG(LogLoad, Log, TEXT("    Sub-Total: %f"), InnerTotal);
+
+			TotalTime += InnerTotal;
+		}
+		
+	}
+	UE_LOG(LogLoad, Log, TEXT("------------- ---------- -------------"));
+	UE_LOG(LogLoad, Log, TEXT("Total Load times: %f"), TotalTime);
+}
+
+static FAutoConsoleCommand LoadTimerDumpCmd(
+	TEXT("LoadTimes.DumpTracking"),
+	TEXT("Dump high level load times being tracked"),
+	FConsoleCommandDelegate::CreateStatic(&FLoadTimeTracker::DumpLoadTimesStatic)
+	);
+
+#endif
+
diff --git a/Engine/Source/Runtime/Core/Private/Windows/WindowsPlatformMisc.cpp b/Engine/Source/Runtime/Core/Private/Windows/WindowsPlatformMisc.cpp
--- a/Engine/Source/Runtime/Core/Private/Windows/WindowsPlatformMisc.cpp
+++ b/Engine/Source/Runtime/Core/Private/Windows/WindowsPlatformMisc.cpp
@@ -2851,13 +2851,21 @@ IPlatformChunkInstall* FWindowsPlatformMisc::GetPlatformChunkInstall()
 	if (!ChunkInstall)
 	{
 #if !(WITH_EDITORONLY_DATA || IS_PROGRAM)
-		static IPlatformChunkInstallModule* PlatformChunkInstallModule = FModuleManager::LoadModulePtr<IPlatformChunkInstallModule>("HTTPChunkInstaller");
-		if (PlatformChunkInstallModule != NULL)
+
+		IPlatformChunkInstallModule* PlatformChunkInstallModule = nullptr;
+
+		FModuleStatus Status;
+		if (FModuleManager::Get().QueryModule("HTTPChunkInstaller", Status))
 		{
-			// Attempt to grab the platform installer
-			ChunkInstall = PlatformChunkInstallModule->GetPlatformChunkInstall();
+			PlatformChunkInstallModule = FModuleManager::LoadModulePtr<IPlatformChunkInstallModule>("HTTPChunkInstaller");
+			if (PlatformChunkInstallModule != nullptr)
+			{
+				// Attempt to grab the platform installer
+				ChunkInstall = PlatformChunkInstallModule->GetPlatformChunkInstall();
+			}
 		}
-		else
+
+		if (PlatformChunkInstallModule == nullptr)
 #endif
 		{
 			// Placeholder instance
diff --git a/Engine/Source/Runtime/Core/Public/Async/AsyncWork.h b/Engine/Source/Runtime/Core/Public/Async/AsyncWork.h
--- a/Engine/Source/Runtime/Core/Public/Async/AsyncWork.h
+++ b/Engine/Source/Runtime/Core/Public/Async/AsyncWork.h
@@ -50,7 +50,7 @@ template<typename TTask>
 class FAutoDeleteAsyncTask
 	: private IQueuedWork
 {
-	/** User job embedded in this task */ 
+	/** User job embedded in this task */
 	TTask Task;
 
 	/* Generic start function, not called directly
@@ -68,25 +68,25 @@ class FAutoDeleteAsyncTask
 		{
 			QueuedPool->AddQueuedWork(this);
 		}
-		else 
+		else
 		{
 			// we aren't doing async stuff
 			DoWork();
 		}
 	}
 
-	/** 
+	/**
 	* Tells the user job to do the work, sometimes called synchronously, sometimes from the thread pool. Calls the event tracker.
 	**/
 	void DoWork()
-	{		
-		FScopeCycleCounter Scope(Task.GetStatId(), true); 
+	{
+		FScopeCycleCounter Scope(Task.GetStatId(), true);
 
-		Task.DoWork();		
+		Task.DoWork();
 		delete this;
 	}
 
-	/** 
+	/**
 	* Always called from the thread pool. Just passes off to DoWork
 	**/
 	virtual void DoThreadedWork()
@@ -112,65 +112,72 @@ class FAutoDeleteAsyncTask
 	}
 
 public:
+#if PLATFORM_COMPILER_HAS_VARIADIC_TEMPLATES
+	/** Forwarding constructor. */
+	template<typename...T>
+	explicit FAutoDeleteAsyncTask(T&&... Args) : Task(Forward<T>(Args)...)
+	{
+	}
+#else
 	/** Default constructor. */
 	FAutoDeleteAsyncTask( )
 	{
 	}
-	/** Passthrough constructor. Generally speaking references will not pass through; use pointers */
+	/** Forwarding constructor. */
 	template<typename T>
-	FAutoDeleteAsyncTask( T Arg )
-		: Task(Arg)
+	FAutoDeleteAsyncTask( T&& Arg )
+		: Task(Forward<T>(Arg))
 	{
 	}
-	/** Passthrough constructor. Generally speaking references will not pass through; use pointers */
+	/** Forwarding constructor. */
 	template<typename T1,typename T2>
-	FAutoDeleteAsyncTask( T1 Arg1, T2 Arg2 )
-		: Task(Arg1,Arg2)
+	FAutoDeleteAsyncTask( T1&& Arg1, T2&& Arg2 )
+		: Task(Forward<T1>(Arg1), Forward<T2>(Arg2))
 	{
 	}
-	/** Passthrough constructor. Generally speaking references will not pass through; use pointers */
+	/** Forwarding constructor. */
 	template<typename T1,typename T2, typename T3>
-	FAutoDeleteAsyncTask( T1 Arg1, T2 Arg2, T3 Arg3 )
-		: Task(Arg1,Arg2,Arg3)
+	FAutoDeleteAsyncTask( T1&& Arg1, T2&& Arg2, T3&& Arg3 )
+		: Task(Forward<T1>(Arg1), Forward<T2>(Arg2), Forward<T3>(Arg3))
 	{
 	}
-	/** Passthrough constructor. Generally speaking references will not pass through; use pointers */
+	/** Forwarding constructor. */
 	template<typename T1,typename T2, typename T3, typename T4>
-	FAutoDeleteAsyncTask( T1 Arg1, T2 Arg2, T3 Arg3, T4 Arg4 )
-		: Task(Arg1,Arg2,Arg3,Arg4)
+	FAutoDeleteAsyncTask( T1&& Arg1, T2&& Arg2, T3&& Arg3, T4&& Arg4 )
+		: Task(Forward<T1>(Arg1), Forward<T2>(Arg2), Forward<T3>(Arg3), Forward<T4>(Arg4))
 	{
 	}
-	/** Passthrough constructor. Generally speaking references will not pass through; use pointers */
+	/** Forwarding constructor. */
 	template<typename T1,typename T2, typename T3, typename T4, typename T5>
-	FAutoDeleteAsyncTask( T1 Arg1, T2 Arg2, T3 Arg3, T4 Arg4, T5 Arg5 )
-		: Task(Arg1,Arg2,Arg3,Arg4,Arg5)
+	FAutoDeleteAsyncTask( T1&& Arg1, T2&& Arg2, T3&& Arg3, T4&& Arg4, T5&& Arg5 )
+		: Task(Forward<T1>(Arg1), Forward<T2>(Arg2), Forward<T3>(Arg3), Forward<T4>(Arg4), Forward<T5>(Arg5))
 	{
 	}
-	/** Passthrough constructor. Generally speaking references will not pass through; use pointers */
+	/** Forwarding constructor. */
 	template<typename T1,typename T2, typename T3, typename T4, typename T5, typename T6>
-	FAutoDeleteAsyncTask( T1 Arg1, T2 Arg2, T3 Arg3, T4 Arg4, T5 Arg5, T6 Arg6 )
-		: Task(Arg1,Arg2,Arg3,Arg4,Arg5,Arg6)
+	FAutoDeleteAsyncTask( T1&& Arg1, T2&& Arg2, T3&& Arg3, T4&& Arg4, T5&& Arg5, T6&& Arg6 )
+		: Task(Forward<T1>(Arg1), Forward<T2>(Arg2), Forward<T3>(Arg3), Forward<T4>(Arg4), Forward<T5>(Arg5), Forward<T6>(Arg6))
 	{
 	}
-	/** Passthrough constructor. Generally speaking references will not pass through; use pointers */
+	/** Forwarding constructor. */
 	template<typename T1,typename T2, typename T3, typename T4, typename T5, typename T6, typename T7>
-	FAutoDeleteAsyncTask( T1 Arg1, T2 Arg2, T3 Arg3, T4 Arg4, T5 Arg5, T6 Arg6, T7 Arg7 )
-		: Task(Arg1,Arg2,Arg3,Arg4,Arg5,Arg6,Arg7)
+	FAutoDeleteAsyncTask( T1&& Arg1, T2&& Arg2, T3&& Arg3, T4&& Arg4, T5&& Arg5, T6&& Arg6, T7&& Arg7 )
+		: Task(Forward<T1>(Arg1), Forward<T2>(Arg2), Forward<T3>(Arg3), Forward<T4>(Arg4), Forward<T5>(Arg5), Forward<T6>(Arg6), Forward<T7>(Arg7))
 	{
 	}
-	/** Passthrough constructor. Generally speaking references will not pass through; use pointers */
+	/** Forwarding constructor. */
 	template<typename T1,typename T2, typename T3, typename T4, typename T5, typename T6, typename T7, typename T8>
-	FAutoDeleteAsyncTask( T1 Arg1, T2 Arg2, T3 Arg3, T4 Arg4, T5 Arg5, T6 Arg6, T7 Arg7, T8 Arg8 )
-		: Task(Arg1,Arg2,Arg3,Arg4,Arg5,Arg6,Arg7,Arg8)
+	FAutoDeleteAsyncTask( T1&& Arg1, T2&& Arg2, T3&& Arg3, T4&& Arg4, T5&& Arg5, T6&& Arg6, T7&& Arg7, T8&& Arg8 )
+		: Task(Forward<T1>(Arg1), Forward<T2>(Arg2), Forward<T3>(Arg3), Forward<T4>(Arg4), Forward<T5>(Arg5), Forward<T6>(Arg6), Forward<T7>(Arg7), Forward<T8>(Arg8))
 	{
 	}
-	/** Passthrough constructor. Generally speaking references will not pass through; use pointers */
+	/** Forwarding constructor. */
 	template<typename T1,typename T2, typename T3, typename T4, typename T5, typename T6, typename T7, typename T8, typename T9>
-	FAutoDeleteAsyncTask( T1 Arg1, T2 Arg2, T3 Arg3, T4 Arg4, T5 Arg5, T6 Arg6, T7 Arg7, T8 Arg8, T9 Arg9 )
-		: Task(Arg1,Arg2,Arg3,Arg4,Arg5,Arg6,Arg7,Arg8,Arg9)
+	FAutoDeleteAsyncTask( T1&& Arg1, T2&& Arg2, T3&& Arg3, T4&& Arg4, T5&& Arg5, T6&& Arg6, T7&& Arg7, T8&& Arg8, T9&& Arg9 )
+		: Task(Forward<T1>(Arg1), Forward<T2>(Arg2), Forward<T3>(Arg3), Forward<T4>(Arg4), Forward<T5>(Arg5), Forward<T6>(Arg6), Forward<T7>(Arg7), Forward<T8>(Arg8), Forward<T9>(Arg9))
 	{
 	}
-
+#endif
 
 	/** 
 	* Run this task on this thread, now. Will end up destroying myself, so it is not safe to use this object after this call.
diff --git a/Engine/Source/Runtime/Core/Public/GenericPlatform/GenericApplication.h b/Engine/Source/Runtime/Core/Public/GenericPlatform/GenericApplication.h
--- a/Engine/Source/Runtime/Core/Public/GenericPlatform/GenericApplication.h
+++ b/Engine/Source/Runtime/Core/Public/GenericPlatform/GenericApplication.h
@@ -352,6 +352,16 @@ struct FDisplayMetrics
 
 	/** Logs out display metrics */
 	CORE_API void PrintToLog() const;
+
+protected:
+	// The title safe zone ratio that will be returned by FDisplayMetrics::GetDisplayMetrics on platforms that don't have a defined safe zone
+	static float GetDebugTitleSafeZoneRatio();
+
+	// The action safe zone ratio that will be returned by FDisplayMetrics::GetDisplayMetrics on platforms that don't have a defined safe zone
+	static float GetDebugActionSafeZoneRatio();
+
+	// Apply the debug/default safe zones
+	void ApplyDefaultSafeZones();
 };
 
 
diff --git a/Engine/Source/Runtime/Core/Public/GenericPlatform/GenericPlatformDriver.h b/Engine/Source/Runtime/Core/Public/GenericPlatform/GenericPlatformDriver.h
--- a/Engine/Source/Runtime/Core/Public/GenericPlatform/GenericPlatformDriver.h
+++ b/Engine/Source/Runtime/Core/Public/GenericPlatform/GenericPlatformDriver.h
@@ -348,6 +348,22 @@ struct FBlackListEntry
 	{
 		return !DriverVersionString.IsEmpty();
 	}
+
+	/**
+	 * Returns true if the latest version of the driver is blacklisted by this entry,
+	 * i.e. the comparison op is > or >=.
+	 */
+	bool IsLatestBlacklisted() const
+	{
+		bool bLatestBlacklisted = false;
+		if (IsValid())
+		{
+			const TCHAR* DriverVersionTchar = *DriverVersionString;
+			EComparisonOp ComparisonOp = ParseComparisonOp(DriverVersionTchar);
+			bLatestBlacklisted = (ComparisonOp == ECO_Larger) || (ComparisonOp == ECO_LargerThan);
+		}
+		return bLatestBlacklisted;
+	}
 };
 
 struct FGPUHardware
@@ -485,6 +501,35 @@ struct FGPUHardware
 		return FBlackListEntry();
 	}
 
+	/**
+	 * Returns true if the latest version of the driver has been blacklisted.
+	 */
+	bool IsLatestBlacklisted() const
+	{
+		bool bLatestBlacklisted = false;
+		const TCHAR* Section = GetVendorSectionName();
+
+		if(Section)
+		{
+			TArray<FString> BlacklistStrings;
+			GConfig->GetArray(GetVendorSectionName(), TEXT("Blacklist"), BlacklistStrings, GHardwareIni);
+
+			for(int32 i = 0; !bLatestBlacklisted && i < BlacklistStrings.Num(); ++i)
+			{
+				FBlackListEntry Entry;
+
+				const TCHAR* Line = *BlacklistStrings[i];
+			
+				ensure(Line[0] == TCHAR('('));
+
+				Entry.LoadFromINIString(&Line[1]);
+
+				bLatestBlacklisted |= Entry.IsLatestBlacklisted();
+			}
+		}
+		return bLatestBlacklisted;
+	}
+
 	// to get a section name in the Hardware.ini file
 	// @return 0 if not found
 	const TCHAR* GetVendorSectionName() const
diff --git a/Engine/Source/Runtime/Core/Public/GenericPlatform/GenericPlatformStackWalk.h b/Engine/Source/Runtime/Core/Public/GenericPlatform/GenericPlatformStackWalk.h
--- a/Engine/Source/Runtime/Core/Public/GenericPlatform/GenericPlatformStackWalk.h
+++ b/Engine/Source/Runtime/Core/Public/GenericPlatform/GenericPlatformStackWalk.h
@@ -78,6 +78,21 @@ struct CORE_API FGenericPlatformStackWalk
 {
 	typedef FGenericPlatformStackWalk Base;
 
+	struct EStackWalkFlags
+	{
+		enum
+		{
+			/** Default value (empty set of flags). */
+			AccurateStackWalk				=	0,
+
+			/** Used when preferring speed over more information. Platform-specific, may be ignored, or may result in non-symbolicated callstacks or missing some other information. */
+			FastStackWalk					=	(1 << 0),
+
+			/** This is a set of flags that will be passed when unwinding the callstack for ensure(). */
+			FlagsUsedWhenHandlingEnsure		=	(FastStackWalk)
+		};
+	};
+
 	/**
 	* Initializes options related to stack walking from ini, i.e. how detailed the stack walking should be, performance settings etc.
 	*/
@@ -167,6 +182,18 @@ struct CORE_API FGenericPlatformStackWalk
 	 */ 
 	static void StackWalkAndDump( ANSICHAR* HumanReadableString, SIZE_T HumanReadableStringSize, int32 IgnoreCount, void* Context = nullptr );
 
+	/**
+	 * Walks the stack and appends the human readable string to the passed in one.
+	 * @warning: The code assumes that HumanReadableString is large enough to contain the information.
+	 *
+	 * @param	HumanReadableString	String to concatenate information with
+	 * @param	HumanReadableStringSize size of string in characters
+	 * @param	IgnoreCount			Number of stack entries to ignore (some are guaranteed to be in the stack walking code)
+	 * @param   Flags				Used to pass additional information (see StackWalkFlags)
+	 * @param	Context				Optional thread context information
+	 */ 
+	static void StackWalkAndDumpEx( ANSICHAR* HumanReadableString, SIZE_T HumanReadableStringSize, int32 IgnoreCount, uint32 Flags, void* Context = nullptr );
+
 	/**
 	 * Returns the number of modules loaded by the currently running process.
 	 */
diff --git a/Engine/Source/Runtime/Core/Public/HAL/Platform.h b/Engine/Source/Runtime/Core/Public/HAL/Platform.h
--- a/Engine/Source/Runtime/Core/Public/HAL/Platform.h
+++ b/Engine/Source/Runtime/Core/Public/HAL/Platform.h
@@ -95,7 +95,7 @@
 
 //@port Identify the platform here and include the platform header to setup the platform types, etc
 #if PLATFORM_WINDOWS
-	#include "Windows/WIndowsPlatform.h"
+	#include "Windows/WindowsPlatform.h"
 #elif PLATFORM_PS4
 	#include "PS4/PS4Platform.h"
 #elif PLATFORM_XBOXONE
@@ -282,6 +282,10 @@
 	#define PLATFORM_SUPPORTS_STACK_SYMBOLS 0
 #endif
 
+#ifndef PLATFORM_HAS_64BIT_ATOMICS
+	#define PLATFORM_HAS_64BIT_ATOMICS 1
+#endif
+
 #ifndef PLATFORM_HAS_128BIT_ATOMICS
 	#define PLATFORM_HAS_128BIT_ATOMICS 0
 #endif
diff --git a/Engine/Source/Runtime/Core/Public/HAL/ThreadingBase.h b/Engine/Source/Runtime/Core/Public/HAL/ThreadingBase.h
--- a/Engine/Source/Runtime/Core/Public/HAL/ThreadingBase.h
+++ b/Engine/Source/Runtime/Core/Public/HAL/ThreadingBase.h
@@ -773,6 +773,140 @@ private:
 	volatile int32 Counter;
 };
 
+// This class cannot be implemented on platforms that don't define 64bit atomic functions
+#if PLATFORM_HAS_64BIT_ATOMICS
+/** Thread safe counter for 64bit ints */
+class FThreadSafeCounter64
+{
+public:
+
+	/**
+	* Default constructor.
+	*
+	* Initializes the counter to 0.
+	*/
+	FThreadSafeCounter64()
+	{
+		Counter = 0;
+	}
+
+	/**
+	* Copy Constructor.
+	*
+	* If the counter in the Other parameter is changing from other threads, there are no
+	* guarantees as to which values you will get up to the caller to not care, synchronize
+	* or other way to make those guarantees.
+	*
+	* @param Other The other thread safe counter to copy
+	*/
+	FThreadSafeCounter64(const FThreadSafeCounter& Other)
+	{
+		Counter = Other.GetValue();
+	}
+
+	/** Assignment has the same caveats as the copy ctor. */
+	FThreadSafeCounter64& operator=(const FThreadSafeCounter64& Other)
+	{
+		if (this == &Other) return *this;
+		Set(Other.GetValue());
+		return *this;
+	}
+
+
+	/**
+	* Constructor, initializing counter to passed in value.
+	*
+	* @param Value	Value to initialize counter to
+	*/
+	FThreadSafeCounter64(int64 Value)
+	{
+		Counter = Value;
+	}
+
+	/**
+	* Increment and return new value.
+	*
+	* @return the new, incremented value
+	* @see Add, Decrement, Reset, Set, Subtract
+	*/
+	int64 Increment()
+	{
+		return FPlatformAtomics::InterlockedIncrement(&Counter);
+	}
+
+	/**
+	* Adds an amount and returns the old value.
+	*
+	* @param Amount Amount to increase the counter by
+	* @return the old value
+	* @see Decrement, Increment, Reset, Set, Subtract
+	*/
+	int64 Add(int64 Amount)
+	{
+		return FPlatformAtomics::InterlockedAdd(&Counter, Amount);
+	}
+
+	/**
+	* Decrement and return new value.
+	*
+	* @return the new, decremented value
+	* @see Add, Increment, Reset, Set, Subtract
+	*/
+	int64 Decrement()
+	{
+		return FPlatformAtomics::InterlockedDecrement(&Counter);
+	}
+
+	/**
+	* Subtracts an amount and returns the old value.
+	*
+	* @param Amount Amount to decrease the counter by
+	* @return the old value
+	* @see Add, Decrement, Increment, Reset, Set
+	*/
+	int64 Subtract(int64 Amount)
+	{
+		return FPlatformAtomics::InterlockedAdd(&Counter, -Amount);
+	}
+
+	/**
+	* Sets the counter to a specific value and returns the old value.
+	*
+	* @param Value	Value to set the counter to
+	* @return The old value
+	* @see Add, Decrement, Increment, Reset, Subtract
+	*/
+	int64 Set(int64 Value)
+	{
+		return FPlatformAtomics::InterlockedExchange(&Counter, Value);
+	}
+
+	/**
+	* Resets the counter's value to zero.
+	*
+	* @return the old value.
+	* @see Add, Decrement, Increment, Set, Subtract
+	*/
+	int64 Reset()
+	{
+		return FPlatformAtomics::InterlockedExchange(&Counter, 0);
+	}
+
+	/**
+	* Gets the current value.
+	*
+	* @return the current value
+	*/
+	int64 GetValue() const
+	{
+		return Counter;
+	}
+
+private:
+	/** Thread-safe counter */
+	volatile int64 Counter;
+};
+#endif
 
 /**
  * Thread safe bool, wraps FThreadSafeCounter
diff --git a/Engine/Source/Runtime/Core/Public/Linux/LinuxPlatformCrashContext.h b/Engine/Source/Runtime/Core/Public/Linux/LinuxPlatformCrashContext.h
--- a/Engine/Source/Runtime/Core/Public/Linux/LinuxPlatformCrashContext.h
+++ b/Engine/Source/Runtime/Core/Public/Linux/LinuxPlatformCrashContext.h
@@ -21,6 +21,9 @@ struct CORE_API FLinuxCrashContext : public FGenericCrashContext
 	/** Whether backtrace was already captured */
 	bool bCapturedBacktrace;
 
+	/** Whether we're handling ensure() (prefer speed in this case). */
+	bool bHandlingEnsure;
+
 	/** Symbols received via backtrace_symbols(), if any (note that we will need to clean it up) */
 	char ** BacktraceSymbols;
 
@@ -35,6 +38,7 @@ struct CORE_API FLinuxCrashContext : public FGenericCrashContext
 		,	Info(nullptr)
 		,	Context(nullptr)
 		,	bCapturedBacktrace(false)
+		,	bHandlingEnsure(false)
 		,	BacktraceSymbols(nullptr)
 	{
 		SignalDescription[ 0 ] = 0;
diff --git a/Engine/Source/Runtime/Core/Public/Linux/LinuxPlatformStackWalk.h b/Engine/Source/Runtime/Core/Public/Linux/LinuxPlatformStackWalk.h
--- a/Engine/Source/Runtime/Core/Public/Linux/LinuxPlatformStackWalk.h
+++ b/Engine/Source/Runtime/Core/Public/Linux/LinuxPlatformStackWalk.h
@@ -14,6 +14,7 @@ struct CORE_API FLinuxPlatformStackWalk : public FGenericPlatformStackWalk
 	static bool ProgramCounterToHumanReadableString( int32 CurrentCallDepth, uint64 ProgramCounter, ANSICHAR* HumanReadableString, SIZE_T HumanReadableStringSize, FGenericCrashContext* Context = nullptr );
 	static void CaptureStackBackTrace( uint64* BackTrace, uint32 MaxDepth, void* Context = nullptr );
 	static void StackWalkAndDump( ANSICHAR* HumanReadableString, SIZE_T HumanReadableStringSize, int32 IgnoreCount, void* Context = nullptr );
+	static void StackWalkAndDumpEx(ANSICHAR* HumanReadableString, SIZE_T HumanReadableStringSize, int32 IgnoreCount, uint32 Flags, void* Context = nullptr);
 };
 
 typedef FLinuxPlatformStackWalk FPlatformStackWalk;
diff --git a/Engine/Source/Runtime/Core/Public/Misc/CallbackDevice.h b/Engine/Source/Runtime/Core/Public/Misc/CallbackDevice.h
--- a/Engine/Source/Runtime/Core/Public/Misc/CallbackDevice.h
+++ b/Engine/Source/Runtime/Core/Public/Misc/CallbackDevice.h
@@ -226,6 +226,10 @@ public:
 	DECLARE_MULTICAST_DELEGATE_OneParam(FOnUserActivityStringChanged, const FString&);
 	static FOnUserActivityStringChanged UserActivityStringChanged;
 
+	/** Sent when application code changes the currently active game session. The exact semantics of this will vary between games but it is useful for analytics, crash reports, etc  */
+	DECLARE_MULTICAST_DELEGATE_OneParam(FOnGameSessionIDChange, const FString&);
+	static FOnGameSessionIDChange GameSessionIDChanged;
+	
 
 private:
 
diff --git a/Engine/Source/Runtime/Core/Public/Misc/CoreMisc.h b/Engine/Source/Runtime/Core/Public/Misc/CoreMisc.h
--- a/Engine/Source/Runtime/Core/Public/Misc/CoreMisc.h
+++ b/Engine/Source/Runtime/Core/Public/Misc/CoreMisc.h
@@ -82,6 +82,11 @@ struct CORE_API FCommandLine
 	 */
 	static const TCHAR* GetOriginal();
 
+	/**
+	 * Returns an edited version of the command line originally passed to the executable.
+	 */
+	static const TCHAR* GetOriginalForLogging();
+
 	/** 
 	 * Checks if the command line has been initialized. 
 	 */
@@ -150,6 +155,8 @@ private:
 	static TCHAR OriginalCmdLine[MaxCommandLineSize];
 	/** character buffer containing the command line filtered for logging purposes */
 	static TCHAR LoggingCmdLine[MaxCommandLineSize];
+	/** character buffer containing the original command line filtered for logging purposes */
+	static TCHAR LoggingOriginalCmdLine[MaxCommandLineSize];
 	/** subprocess command line */
 	static FString SubprocessCommandLine;
 };
diff --git a/Engine/Source/Runtime/Core/Public/Misc/DateTime.h b/Engine/Source/Runtime/Core/Public/Misc/DateTime.h
--- a/Engine/Source/Runtime/Core/Public/Misc/DateTime.h
+++ b/Engine/Source/Runtime/Core/Public/Misc/DateTime.h
@@ -466,6 +466,16 @@ public:
 	 */
 	CORE_API FString ToIso8601() const;
 
+	/**
+	 * Returns the RFC 1123 string representation of the FDateTime.
+	 *
+	 * The resulting string assumes that the FDateTime is in UTC.
+	 * 
+	 * @return String representation.
+	 * @see ParseHttpDate, ToString
+	 */
+	CORE_API FString ToHttpDate() const;
+
 	/**
 	 * Returns the string representation of this date using a default format.
 	 *
@@ -612,6 +622,17 @@ public:
 	 */
 	static CORE_API bool ParseIso8601( const TCHAR* DateTimeString, FDateTime& OutDateTime );
 
+	/**
+	 * Parses a date string in HTTP-date format (rfc1123-date | rfc850-date | asctime-date)
+	 * https://www.w3.org/Protocols/rfc2616/rfc2616-sec3.html#sec3.3.1
+	 * 
+	 * @param HttpDate The string to be parsed
+	 * @param OutDateTime FDateTime object (assumes UTC) corresponding to the input string.
+	 * @return true if the string was converted successfully, false otherwise.
+	 * @see Parse
+	 */
+	static CORE_API bool ParseHttpDate( const FString& HttpDate, FDateTime& OutDateTime );
+
 	/**
 	 * Gets the local date on this computer.
 	 *
diff --git a/Engine/Source/Runtime/Core/Public/ProfilingDebugging/LoadTimeTracker.h b/Engine/Source/Runtime/Core/Public/ProfilingDebugging/LoadTimeTracker.h
new file mode 100644
--- /dev/null
+++ b/Engine/Source/Runtime/Core/Public/ProfilingDebugging/LoadTimeTracker.h
@@ -0,0 +1,81 @@
+// Copyright 1998-2016 Epic Games, Inc. All Rights Reserved.
+
+/**
+* Declarations for LoadTimer which helps get load times for various parts of the game.
+*/
+
+#pragma once
+
+#if !UE_BUILD_SHIPPING
+#define ENABLE_LOADTIME_TRACKING 1
+#else
+#define ENABLE_LOADTIME_TRACKING 0
+#endif
+
+#include "ScopedTimers.h"
+
+#if ENABLE_LOADTIME_TRACKING
+
+/** High level load time tracker utility (such as initial engine startup or game specific timings) */
+class CORE_API FLoadTimeTracker
+{
+public:
+	static FLoadTimeTracker& Get()
+	{
+		static FLoadTimeTracker Singleton;
+		return Singleton;
+	}
+
+	/** Adds a scoped time for a given label.  Records each instance individually */
+	void ReportScopeTime(double ScopeTime, const FName ScopeLabel);
+
+	/** Prints out total time and individual times */
+	void DumpLoadTimes();
+
+
+	static void DumpLoadTimesStatic()
+	{
+		Get().DumpLoadTimes();
+	}
+
+private:
+	TMap<FName, TArray<double>> TimeInfo;
+
+private:
+	FLoadTimeTracker()
+	{
+	}
+};
+
+/** Scoped duration timer that automatically reports to FLoadTimeTracker */
+class CORE_API FScopedLoadTimer : public FDurationTimer
+{
+public:
+	explicit FScopedLoadTimer(const FName InLabel) 
+		: FDurationTimer(TimerData)
+		, TimerData(0)
+		, Label(InLabel)
+	{
+		Start();
+	}
+
+	/** Dtor, updating seconds with time delta. */
+	~FScopedLoadTimer()
+	{
+		Stop();
+		FLoadTimeTracker::Get().ReportScopeTime(GetAccumulatedTime(), Label);
+	}
+private:
+	double TimerData;
+	const FName Label;
+};
+
+
+#define SCOPED_LOADTIMER(TimerName) FScopedLoadTimer TimerName##_ScopedTimer_ (FName(TEXT(#TimerName)))
+#define ACCUM_LOADTIME(TimerName, Time) FLoadTimeTracker::Get().ReportScopeTime(Time, FName(TimerName));
+#define SET_LOADTIME(TimerName, Time) FLoadTimeTracker::Get().SetTime(Time, FName(TimerName));
+#else
+#define SCOPED_LOADTIMER(TimerName)
+#define ACCUM_LOADTIME(TimerName, Time)
+#define SET_LOADTIME(TimerName, Time)
+#endif
\ No newline at end of file
