Merging CL# 2863184 //UE4/Release-4.11 to Dev-Main (//UE4/Dev-Main)

diff --git a/Engine/Source/Runtime/Core/Public/Templates/Function.h b/Engine/Source/Runtime/Core/Public/Templates/Function.h
--- a/Engine/Source/Runtime/Core/Public/Templates/Function.h
+++ b/Engine/Source/Runtime/Core/Public/Templates/Function.h
@@ -679,7 +679,14 @@ public:
 		: Super  (NoInit)
 		, Storage(MoveTemp(Other.Storage))
 	{
-		Super::CopyAndReseat(Other, Storage.GetBoundObject()->GetAddress());
+		if (UE4Function_Private::IFunction_OwnedObject* Func = Storage.GetBoundObject())
+		{
+			Super::CopyAndReseat(Other, Func->GetAddress());
+		}
+		else
+		{
+			Super::Unset();
+		}
 
 		Other.Unset();
 	}
