UE4 - Fix no logging on some platforms.

diff --git a/Engine/Source/Runtime/Core/Private/GenericPlatform/GenericPlatformMisc.cpp b/Engine/Source/Runtime/Core/Private/GenericPlatform/GenericPlatformMisc.cpp
--- a/Engine/Source/Runtime/Core/Private/GenericPlatform/GenericPlatformMisc.cpp
+++ b/Engine/Source/Runtime/Core/Private/GenericPlatform/GenericPlatformMisc.cpp
@@ -392,7 +392,7 @@ void FGenericPlatformMisc::LocalPrint( const TCHAR* Str )
 
 bool FGenericPlatformMisc::HasSeparateChannelForDebugOutput()
 {
-	return false;
+	return true;
 }
 
 void FGenericPlatformMisc::RequestMinimize()
diff --git a/Engine/Source/Runtime/Core/Private/GenericPlatform/GenericPlatformOutputDevices.cpp b/Engine/Source/Runtime/Core/Private/GenericPlatform/GenericPlatformOutputDevices.cpp
--- a/Engine/Source/Runtime/Core/Private/GenericPlatform/GenericPlatformOutputDevices.cpp
+++ b/Engine/Source/Runtime/Core/Private/GenericPlatform/GenericPlatformOutputDevices.cpp
@@ -12,7 +12,8 @@ void FGenericPlatformOutputDevices::SetupOutputDevices()
 
 	GLog->AddOutputDevice(FPlatformOutputDevices::GetLog());
 
-	if (!FParse::Param(FCommandLine::Get(), TEXT("NOCONSOLE")))
+	bool bHasConsole = !FParse::Param(FCommandLine::Get(), TEXT("NOCONSOLE"));
+	if (bHasConsole)
 	{
 		GLog->AddOutputDevice(GLogConsole);
 	}
@@ -22,7 +23,7 @@ void FGenericPlatformOutputDevices::SetupOutputDevices()
 	if (!FPlatformProperties::SupportsWindowedMode() || FPlatformMisc::IsDebuggerPresent() || (UE_BUILD_SHIPPING && !NO_LOGGING) || GIsBuildMachine)
 	{
 		// Only need to do this if it's actually going to go to a different place than GLogConsole
-		if(FPlatformMisc::HasSeparateChannelForDebugOutput())
+		if(!bHasConsole || FPlatformMisc::HasSeparateChannelForDebugOutput())
 		{
 			GLog->AddOutputDevice(new FOutputDeviceDebug());
 		}
diff --git a/Engine/Source/Runtime/Core/Private/Mac/MacPlatformMisc.cpp b/Engine/Source/Runtime/Core/Private/Mac/MacPlatformMisc.cpp
--- a/Engine/Source/Runtime/Core/Private/Mac/MacPlatformMisc.cpp
+++ b/Engine/Source/Runtime/Core/Private/Mac/MacPlatformMisc.cpp
@@ -1119,6 +1119,11 @@ bool FMacPlatformMisc::GetDiskTotalAndFreeSpace(const FString& InPath, uint64& T
 	return (Err == 0);
 }
 
+bool FMacPlatformMisc::HasSeparateChannelForDebugOutput();
+{
+	return false;
+}
+
 #include "ModuleManager.h"
 
 void FMacPlatformMisc::LoadPreInitModules()
diff --git a/Engine/Source/Runtime/Core/Private/Windows/WindowsPlatformMisc.cpp b/Engine/Source/Runtime/Core/Private/Windows/WindowsPlatformMisc.cpp
--- a/Engine/Source/Runtime/Core/Private/Windows/WindowsPlatformMisc.cpp
+++ b/Engine/Source/Runtime/Core/Private/Windows/WindowsPlatformMisc.cpp
@@ -1062,11 +1062,6 @@ void FWindowsPlatformMisc::LocalPrint( const TCHAR *Message )
 	OutputDebugString(Message);
 }
 
-bool FWindowsPlatformMisc::HasSeparateChannelForDebugOutput()
-{
-	return true;
-}
-
 void FWindowsPlatformMisc::RequestExit( bool Force )
 {
 	UE_LOG(LogWindows, Log,  TEXT("FPlatformMisc::RequestExit(%i)"), Force );
diff --git a/Engine/Source/Runtime/Core/Public/Mac/MacPlatformMisc.h b/Engine/Source/Runtime/Core/Public/Mac/MacPlatformMisc.h
--- a/Engine/Source/Runtime/Core/Public/Mac/MacPlatformMisc.h
+++ b/Engine/Source/Runtime/Core/Public/Mac/MacPlatformMisc.h
@@ -102,7 +102,7 @@ struct CORE_API FMacPlatformMisc : public FGenericPlatformMisc
 	static void GetOSVersions( FString& out_OSVersionLabel, FString& out_OSSubVersionLabel );
 	static bool HasPlatformFeature(const TCHAR* FeatureName);
 	static bool GetDiskTotalAndFreeSpace(const FString& InPath, uint64& TotalNumberOfBytes, uint64& NumberOfFreeBytes);
-
+	static bool HasSeparateChannelForDebugOutput();
 
 	/** 
 	 * Determines the shader format for the plarform
diff --git a/Engine/Source/Runtime/Core/Public/Windows/WindowsPlatformMisc.h b/Engine/Source/Runtime/Core/Public/Windows/WindowsPlatformMisc.h
--- a/Engine/Source/Runtime/Core/Public/Windows/WindowsPlatformMisc.h
+++ b/Engine/Source/Runtime/Core/Public/Windows/WindowsPlatformMisc.h
@@ -85,7 +85,6 @@ struct CORE_API FWindowsPlatformMisc
 	static uint32 GetCharKeyMap(uint32* KeyCodes, FString* KeyNames, uint32 MaxMappings);
 	static void SetUTF8Output();
 	static void LocalPrint(const TCHAR *Message);
-	static bool HasSeparateChannelForDebugOutput();
 	static void RequestExit(bool Force);
 	static void RequestMinimize();
 	static const TCHAR* GetSystemErrorMessage(TCHAR* OutBuffer, int32 BufferCount, int32 Error);
