Fix Mac compile error.

diff --git a/Engine/Source/Runtime/Core/Private/Mac/MacPlatformMisc.cpp b/Engine/Source/Runtime/Core/Private/Mac/MacPlatformMisc.cpp
--- a/Engine/Source/Runtime/Core/Private/Mac/MacPlatformMisc.cpp
+++ b/Engine/Source/Runtime/Core/Private/Mac/MacPlatformMisc.cpp
@@ -1119,7 +1119,7 @@ bool FMacPlatformMisc::GetDiskTotalAndFreeSpace(const FString& InPath, uint64& T
 	return (Err == 0);
 }
 
-bool FMacPlatformMisc::HasSeparateChannelForDebugOutput();
+bool FMacPlatformMisc::HasSeparateChannelForDebugOutput()
 {
 	return false;
 }
