Copying //UE4/Orion-Staging to //UE4/Main (Origin //Orion/Dev-General @ 2870388)

diff --git a/Engine/Source/Runtime/Core/Private/Linux/LinuxPlatformCrashContext.cpp b/Engine/Source/Runtime/Core/Private/Linux/LinuxPlatformCrashContext.cpp
--- a/Engine/Source/Runtime/Core/Private/Linux/LinuxPlatformCrashContext.cpp
+++ b/Engine/Source/Runtime/Core/Private/Linux/LinuxPlatformCrashContext.cpp
@@ -5,6 +5,7 @@
 #include "Misc/App.h"
 #include "EngineVersion.h"
 #include "PlatformMallocCrash.h"
+#include "LinuxPlatformRunnableThread.h"
 
 #include <sys/utsname.h>	// for uname()
 #include <signal.h>
@@ -204,7 +205,7 @@ void FLinuxCrashContext::GenerateReport(const FString & DiagnosticsPath) const
 /** 
  * Mimics Windows WER format
  */
-void GenerateWindowsErrorReport(const FString & WERPath)
+void GenerateWindowsErrorReport(const FString & WERPath, bool bReportingNonCrash)
 {
 	FArchive* ReportFile = IFileManager::Get().CreateFileWriter(*WERPath);
 	if (ReportFile != NULL)
@@ -252,6 +253,7 @@ void GenerateWindowsErrorReport(const FString & WERPath)
 		WriteLine(ReportFile, TEXT("\t\t<Parameter1>6.1.7601.2.1.0.256.48</Parameter1>"));
 		WriteLine(ReportFile, TEXT("\t\t<Parameter2>1033</Parameter2>"));
 		WriteLine(ReportFile, *FString::Printf(TEXT("\t\t<DeploymentName>%s</DeploymentName>"), FApp::GetDeploymentName()));
+		WriteLine(ReportFile, *FString::Printf(TEXT("\t\t<IsEnsure>%s</IsEnsure>"), bReportingNonCrash ? TEXT("1") : TEXT("0")));
 		WriteLine(ReportFile, TEXT("\t</DynamicSignatures>"));
 
 		WriteLine(ReportFile, TEXT("\t<SystemInformation>"));
@@ -381,7 +383,7 @@ void FLinuxCrashContext::GenerateCrashInfoAndLaunchReporter(bool bReportingNonCr
 		GenerateReport(FPaths::Combine(*CrashInfoFolder, TEXT("Diagnostics.txt")));
 
 		// generate "WER"
-		GenerateWindowsErrorReport(FPaths::Combine(*CrashInfoFolder, TEXT("wermeta.xml")));
+		GenerateWindowsErrorReport(FPaths::Combine(*CrashInfoFolder, TEXT("wermeta.xml")), bReportingNonCrash);
 
 		// generate "minidump" (just >1 byte)
 		GenerateMinidump(FPaths::Combine(*CrashInfoFolder, TEXT("minidump.dmp")));
@@ -457,12 +459,14 @@ void FLinuxCrashContext::GenerateCrashInfoAndLaunchReporter(bool bReportingNonCr
 	if (!bReportingNonCrash)
 	{
 		// remove the handler for this signal and re-raise it (which should generate the proper core dump)
-		UE_LOG(LogLinux, Log, TEXT("Engine crash handling finished; re-raising signal %d for the default handler. Good bye."), Signal);
+		// print message to stdout directly, it may be too late for the log (doesn't seem to be printed during a crash in the thread) 
+		printf("Engine crash handling finished; re-raising signal %d for the default handler. Good bye.\n", Signal);
+		fflush(stdout);
 
 		struct sigaction ResetToDefaultAction;
 		FMemory::Memzero(ResetToDefaultAction);
 		ResetToDefaultAction.sa_handler = SIG_DFL;
-		sigemptyset(&ResetToDefaultAction.sa_mask);
+		sigfillset(&ResetToDefaultAction.sa_mask);
 		sigaction(Signal, &ResetToDefaultAction, nullptr);
 
 		raise(Signal);
@@ -524,28 +528,96 @@ void PlatformCrashHandler(int32 Signal, siginfo_t* Info, void* Context)
 void FLinuxPlatformMisc::SetGracefulTerminationHandler()
 {
 	struct sigaction Action;
-	FMemory::Memzero(&Action, sizeof(struct sigaction));
+	FMemory::Memzero(Action);
 	Action.sa_sigaction = GracefulTerminationHandler;
-	sigemptyset(&Action.sa_mask);
+	sigfillset(&Action.sa_mask);
 	Action.sa_flags = SA_SIGINFO | SA_RESTART | SA_ONSTACK;
-	sigaction(SIGINT, &Action, NULL);
-	sigaction(SIGTERM, &Action, NULL);
-	sigaction(SIGHUP, &Action, NULL);	//  this should actually cause the server to just re-read configs (restart?)
+	sigaction(SIGINT, &Action, nullptr);
+	sigaction(SIGTERM, &Action, nullptr);
+	sigaction(SIGHUP, &Action, nullptr);	//  this should actually cause the server to just re-read configs (restart?)
 }
 
+// reserve stack for the main thread in BSS
+char FRunnableThreadLinux::MainThreadSignalHandlerStack[FRunnableThreadLinux::EConstants::CrashHandlerStackSize];
+
 void FLinuxPlatformMisc::SetCrashHandler(void (* CrashHandler)(const FGenericCrashContext & Context))
 {
 	GCrashHandlerPointer = CrashHandler;
 
+	// This table lists all signals that we handle. 0 is not a valid signal, it is used as a separator: everything 
+	// before is considered a crash and handled by the crash handler; everything above it is handled elsewhere 
+	// and also omitted from setting to ignore
+	int HandledSignals[] = 
+	{
+		// signals we consider crashes
+		SIGQUIT, 
+		SIGILL, 
+		SIGFPE, 
+		SIGBUS, 
+		SIGSEGV, 
+		SIGSYS,
+		0,	// marks the end of crash signals
+		SIGINT,
+		SIGTERM,
+		SIGHUP,
+		SIGCHLD
+	};
+
 	struct sigaction Action;
-	FMemory::Memzero(&Action, sizeof(struct sigaction));
+	FMemory::Memzero(Action);
+	sigfillset(&Action.sa_mask);
+	Action.sa_flags = SA_SIGINFO | SA_RESTART | SA_ONSTACK;
 	Action.sa_sigaction = PlatformCrashHandler;
-	sigemptyset(&Action.sa_mask);
+
+	// install this handler for all the "crash" signals
+	for (int Signal : HandledSignals)
+	{
+		if (!Signal)
+		{
+			// hit the end of crash signals, the rest is already handled elsewhere
+			break;
+		}
+		sigaction(Signal, &Action, nullptr);
+	}
+
+	// reinitialize the structure, since assigning to both sa_handler and sa_sigacton is ill-advised
+	FMemory::Memzero(Action);
+	sigfillset(&Action.sa_mask);
 	Action.sa_flags = SA_SIGINFO | SA_RESTART | SA_ONSTACK;
-	sigaction(SIGQUIT, &Action, NULL);	// SIGQUIT is a user-initiated "crash".
-	sigaction(SIGILL, &Action, NULL);
-	sigaction(SIGFPE, &Action, NULL);
-	sigaction(SIGBUS, &Action, NULL);
-	sigaction(SIGSEGV, &Action, NULL);
-	sigaction(SIGSYS, &Action, NULL);
-}
\ No newline at end of file
+	Action.sa_handler = SIG_IGN;
+
+	// set all the signals except ones we know we are handling to be ignored
+	for (int Signal = 1; Signal < NSIG; ++Signal)
+	{
+		bool bSignalShouldBeIgnored = true;
+		for (int HandledSignal : HandledSignals)
+		{
+			if (Signal == HandledSignal)
+			{
+				bSignalShouldBeIgnored = false;
+				break;
+			}
+		}
+
+		if (bSignalShouldBeIgnored)
+		{
+			sigaction(Signal, &Action, nullptr);
+		}
+	}
+
+	checkf(IsInGameThread(), TEXT("Crash handler should be set from the main thread only."));
+	stack_t SignalHandlerStack;
+	
+	FMemory::Memzero(SignalHandlerStack);
+	SignalHandlerStack.ss_sp = FRunnableThreadLinux::MainThreadSignalHandlerStack;
+	SignalHandlerStack.ss_size = sizeof(FRunnableThreadLinux::MainThreadSignalHandlerStack);
+
+	if (sigaltstack(&SignalHandlerStack, nullptr) < 0)
+	{
+		int ErrNo = errno;
+		UE_LOG(LogLinux, Warning, TEXT("Unable to set alternate stack for crash handler, sigaltstack() failed with errno=%d (%s)"),
+			ErrNo,
+			UTF8_TO_TCHAR(strerror(ErrNo))
+			);
+	}
+}
diff --git a/Engine/Source/Runtime/Core/Private/Linux/LinuxPlatformMisc.cpp b/Engine/Source/Runtime/Core/Private/Linux/LinuxPlatformMisc.cpp
--- a/Engine/Source/Runtime/Core/Private/Linux/LinuxPlatformMisc.cpp
+++ b/Engine/Source/Runtime/Core/Private/Linux/LinuxPlatformMisc.cpp
@@ -48,11 +48,11 @@ namespace
 	void InstallChildExitedSignalHanlder()
 	{
 		struct sigaction Action;
-		FMemory::Memzero(&Action, sizeof(struct sigaction));
+		FMemory::Memzero(Action);
 		Action.sa_sigaction = EmptyChildHandler;
-		sigemptyset(&Action.sa_mask);
+		sigfillset(&Action.sa_mask);
 		Action.sa_flags = SA_SIGINFO | SA_RESTART | SA_ONSTACK;
-		sigaction(SIGCHLD, &Action, NULL);
+		sigaction(SIGCHLD, &Action, nullptr);
 	}
 }
 
diff --git a/Engine/Source/Runtime/Core/Private/Linux/LinuxPlatformProcess.cpp b/Engine/Source/Runtime/Core/Private/Linux/LinuxPlatformProcess.cpp
--- a/Engine/Source/Runtime/Core/Private/Linux/LinuxPlatformProcess.cpp
+++ b/Engine/Source/Runtime/Core/Private/Linux/LinuxPlatformProcess.cpp
@@ -788,23 +788,44 @@ FProcHandle FLinuxPlatformProcess::CreateProc(const TCHAR* URL, const TCHAR* Par
 	extern char ** environ;	// provided by libc
 	pid_t ChildPid = -1;
 
-	posix_spawn_file_actions_t FileActions;
-
-	posix_spawn_file_actions_init(&FileActions);
-	if (PipeWriteChild)
+	int PosixSpawnErrNo = -1;
+	if (PipeWriteChild || PipeReadChild)
 	{
-		const FPipeHandle* PipeWriteHandle = reinterpret_cast< const FPipeHandle* >(PipeWriteChild);
-		posix_spawn_file_actions_adddup2(&FileActions, PipeWriteHandle->GetHandle(), STDOUT_FILENO);
-	}
+		posix_spawn_file_actions_t FileActions;
+		posix_spawn_file_actions_init(&FileActions);
+
+		if (PipeWriteChild)
+		{
+			const FPipeHandle* PipeWriteHandle = reinterpret_cast<const FPipeHandle*>(PipeWriteChild);
+			posix_spawn_file_actions_adddup2(&FileActions, PipeWriteHandle->GetHandle(), STDOUT_FILENO);
+		}
 
-	if (PipeReadChild)
+		if (PipeReadChild)
+		{
+			const FPipeHandle* PipeReadHandle = reinterpret_cast<const FPipeHandle*>(PipeReadChild);
+			posix_spawn_file_actions_adddup2(&FileActions, PipeReadHandle->GetHandle(), STDIN_FILENO);
+		}
+
+		PosixSpawnErrNo = posix_spawn(&ChildPid, TCHAR_TO_UTF8(*ProcessPath), &FileActions, nullptr, Argv, environ);
+		posix_spawn_file_actions_destroy(&FileActions);
+	}
+	else
 	{
-		const FPipeHandle* PipeReadHandle = reinterpret_cast< const FPipeHandle* >(PipeReadChild);
-		posix_spawn_file_actions_adddup2(&FileActions, PipeReadHandle->GetHandle(), STDIN_FILENO);
+		// if we don't have any actions to do, use a faster route that will use vfork() instead.
+		// This is not just faster, it is crucial when spawning a crash reporter to report a crash due to stack overflow in a thread
+		// since otherwise atfork handlers will get called and posix_spawn() will crash (in glibc's __reclaim_stacks()).
+		// However, it has its problems, see:
+		//		http://ewontfix.com/7/
+		//		https://sourceware.org/bugzilla/show_bug.cgi?id=14750
+		//		https://sourceware.org/bugzilla/show_bug.cgi?id=14749
+		posix_spawnattr_t SpawnAttr;
+		posix_spawnattr_init(&SpawnAttr);
+		posix_spawnattr_setflags(&SpawnAttr, POSIX_SPAWN_USEVFORK);
+
+		PosixSpawnErrNo = posix_spawn(&ChildPid, TCHAR_TO_UTF8(*ProcessPath), nullptr, &SpawnAttr, Argv, environ);
+		posix_spawnattr_destroy(&SpawnAttr);
 	}
 
-	int PosixSpawnErrNo = posix_spawn(&ChildPid, TCHAR_TO_UTF8(*ProcessPath), &FileActions, nullptr, Argv, environ);
-	posix_spawn_file_actions_destroy(&FileActions);
 	if (PosixSpawnErrNo != 0)
 	{
 		UE_LOG(LogHAL, Fatal, TEXT("FLinuxPlatformProcess::CreateProc: posix_spawn() failed (%d, %s)"), PosixSpawnErrNo, ANSI_TO_TCHAR(strerror(PosixSpawnErrNo)));
diff --git a/Engine/Source/Runtime/Core/Private/Linux/LinuxPlatformStackWalk.cpp b/Engine/Source/Runtime/Core/Private/Linux/LinuxPlatformStackWalk.cpp
--- a/Engine/Source/Runtime/Core/Private/Linux/LinuxPlatformStackWalk.cpp
+++ b/Engine/Source/Runtime/Core/Private/Linux/LinuxPlatformStackWalk.cpp
@@ -695,7 +695,7 @@ bool FLinuxPlatformStackWalk::ProgramCounterToHumanReadableString( int32 Current
 					const char * SourceFilename = NULL;
 					int LineNumber;
 					// do not symbolicate if we're handling ensure()
-					if (!LinuxContext->bHandlingEnsure && LinuxStackWalkHelpers::GetBacktraceSymbols()->GetInfoForAddress(reinterpret_cast< void* >( ProgramCounter ), NULL, &SourceFilename, &LineNumber) && FunctionName != nullptr)
+					if (!LinuxContext->GetIsEnsure() && LinuxStackWalkHelpers::GetBacktraceSymbols()->GetInfoForAddress(reinterpret_cast< void* >( ProgramCounter ), NULL, &SourceFilename, &LineNumber) && FunctionName != nullptr)
 					{
 						FCStringAnsi::Sprintf(TempArray, " [%s, line %d]", SourceFilename, LineNumber);
 						LinuxStackWalkHelpers::AppendToString(HumanReadableString, HumanReadableStringSize, Context, TempArray);
@@ -742,14 +742,31 @@ void FLinuxPlatformStackWalk::StackWalkAndDumpEx(ANSICHAR* HumanReadableString,
 	const bool bHandlingEnsure = (Flags & EStackWalkFlags::FlagsUsedWhenHandlingEnsure) == EStackWalkFlags::FlagsUsedWhenHandlingEnsure;
 	if (Context == nullptr)
 	{
-		FLinuxCrashContext CrashContext;
+		FLinuxCrashContext CrashContext(bHandlingEnsure);
 		CrashContext.InitFromSignal(0, nullptr, nullptr);
-		CrashContext.bHandlingEnsure = bHandlingEnsure;
 		FPlatformStackWalk::StackWalkAndDump(HumanReadableString, HumanReadableStringSize, IgnoreCount + 1, &CrashContext);
 	}
 	else
 	{
-		TGuardValue<bool> Guard(reinterpret_cast<FLinuxCrashContext *>(Context)->bHandlingEnsure, bHandlingEnsure);
+		/** Helper sets the ensure value in the context and guarantees it gets reset afterwards (even if an exception is thrown) */
+		struct FLocalGuardHelper
+		{
+			FLocalGuardHelper(FLinuxCrashContext* InContext, bool bNewEnsureValue)
+				: Context(InContext), bOldEnsureValue(Context->GetIsEnsure())
+			{
+				Context->SetIsEnsure(bNewEnsureValue);
+			}
+			~FLocalGuardHelper()
+			{
+				Context->SetIsEnsure(bOldEnsureValue);
+			}
+
+		private:
+			FLinuxCrashContext* Context;
+			bool bOldEnsureValue;
+		};
+
+		FLocalGuardHelper Guard(reinterpret_cast<FLinuxCrashContext*>(Context), bHandlingEnsure);
 		FPlatformStackWalk::StackWalkAndDump(HumanReadableString, HumanReadableStringSize, IgnoreCount + 1, Context);
 	}
 }
@@ -798,9 +815,9 @@ void NewReportEnsure(const TCHAR* ErrorMessage)
 	Signal.si_code = TRAP_TRACE;
 	Signal.si_addr = __builtin_return_address(0);
 
-	FLinuxCrashContext EnsureContext;
+	const bool bIsEnsure = true;
+	FLinuxCrashContext EnsureContext(bIsEnsure);
 	EnsureContext.InitFromSignal(SIGTRAP, &Signal, nullptr);
-	EnsureContext.bHandlingEnsure = true;
 	EnsureContext.CaptureStackTrace();
 	EnsureContext.GenerateCrashInfoAndLaunchReporter(true);
 
diff --git a/Engine/Source/Runtime/Core/Private/Mac/MacPlatformMisc.cpp b/Engine/Source/Runtime/Core/Private/Mac/MacPlatformMisc.cpp
--- a/Engine/Source/Runtime/Core/Private/Mac/MacPlatformMisc.cpp
+++ b/Engine/Source/Runtime/Core/Private/Mac/MacPlatformMisc.cpp
@@ -1412,7 +1412,7 @@ void FMacPlatformMisc::SetCrashHandler(void (* CrashHandler)(const FGenericCrash
 	}
 }
 
-void FMacCrashContext::GenerateWindowsErrorReport(char const* WERPath) const
+void FMacCrashContext::GenerateWindowsErrorReport(char const* WERPath, bool bIsEnsure) const
 {
 	int ReportFile = open(WERPath, O_CREAT|O_WRONLY, 0766);
 	if (ReportFile != -1)
@@ -1570,6 +1570,7 @@ void FMacCrashContext::GenerateWindowsErrorReport(char const* WERPath) const
 		WriteUTF16String(ReportFile, *GMacAppInfo.LCID);
 		WriteLine(ReportFile, TEXT("</Parameter2>"));
 		WriteLine(ReportFile, *FString::Printf(TEXT("\t\t<DeploymentName>%s</DeploymentName>"), FApp::GetDeploymentName()));
+		WriteLine(ReportFile, *FString::Printf(TEXT("\t\t<IsEnsure>%s</IsEnsure>"), bIsEnsure ? TEXT("1") : TEXT("0")));
 
 		WriteLine(ReportFile, TEXT("\t</DynamicSignatures>"));
 		
@@ -1624,7 +1625,7 @@ void FMacCrashContext::CopyMinidump(char const* OutputPath, char const* InputPat
 	}
 }
 
-void FMacCrashContext::GenerateInfoInFolder(char const* const InfoFolder) const
+void FMacCrashContext::GenerateInfoInFolder(char const* const InfoFolder, bool bIsEnsure) const
 {
 	// create a crash-specific directory
 	char CrashInfoFolder[PATH_MAX] = {};
@@ -1652,7 +1653,7 @@ void FMacCrashContext::GenerateInfoInFolder(char const* const InfoFolder) const
 		// generate "WER"
 		FCStringAnsi::Strncpy(FilePath, CrashInfoFolder, PATH_MAX);
 		FCStringAnsi::Strcat(FilePath, PATH_MAX, "/wermeta.xml");
-		GenerateWindowsErrorReport(FilePath);
+		GenerateWindowsErrorReport(FilePath, bIsEnsure);
 		
 		// generate "minidump" (Apple crash log format)
 		FCStringAnsi::Strncpy(FilePath, CrashInfoFolder, PATH_MAX);
@@ -1772,7 +1773,8 @@ void FMacCrashContext::GenerateCrashInfoAndLaunchReporter() const
 		FCStringAnsi::Strcat(CrashInfoFolder, PATH_MAX, ItoANSI(GMacAppInfo.RunUUID.C, 16));
 		FCStringAnsi::Strcat(CrashInfoFolder, PATH_MAX, ItoANSI(GMacAppInfo.RunUUID.D, 16));
 		
-		GenerateInfoInFolder(CrashInfoFolder);
+		const bool bIsEnsure = false;
+		GenerateInfoInFolder(CrashInfoFolder, bIsEnsure);
 
 		// try launching the tool and wait for its exit, if at all
 		// Use vfork() & execl() as they are async-signal safe, CreateProc can fail in Cocoa
@@ -1836,7 +1838,8 @@ void FMacCrashContext::GenerateEnsureInfoAndLaunchReporter() const
 		FString GameName = FApp::GetGameName();
 		FString EnsureLogFolder = FString(GMacAppInfo.CrashReportPath) / FString::Printf(TEXT("EnsureReport-%s-%s"), *GameName, *Guid.ToString(EGuidFormats::Digits));
 		
-		GenerateInfoInFolder(TCHAR_TO_UTF8(*EnsureLogFolder));
+		const bool bIsEnsure = true;
+		GenerateInfoInFolder(TCHAR_TO_UTF8(*EnsureLogFolder), bIsEnsure);
 		
 		FString Arguments = FString::Printf(TEXT("\"%s/\" -Unattended"), *EnsureLogFolder);
 		FString ReportClient = FPaths::ConvertRelativePathToFull(FPlatformProcess::GenerateApplicationPath(TEXT("CrashReportClient"), EBuildConfigurations::Development));
diff --git a/Engine/Source/Runtime/Core/Private/ProfilingDebugging/CookStats.cpp b/Engine/Source/Runtime/Core/Private/ProfilingDebugging/CookStats.cpp
new file mode 100644
--- /dev/null
+++ b/Engine/Source/Runtime/Core/Private/ProfilingDebugging/CookStats.cpp
@@ -0,0 +1,14 @@
+// Copyright 1998-2016 Epic Games, Inc. All Rights Reserved.
+
+#include "CorePrivatePCH.h"
+
+#include "CookStats.h"
+
+#if ENABLE_COOK_STATS
+CORE_API FCookStatsManager::FGatherCookStatsDelegate FCookStatsManager::CookStatsCallbacks;
+
+CORE_API void FCookStatsManager::LogCookStats(AddStatFuncRef AddStat)
+{
+	CookStatsCallbacks.Broadcast(AddStat);
+}
+#endif
\ No newline at end of file
diff --git a/Engine/Source/Runtime/Core/Private/ProfilingDebugging/LoadTimeTracker.cpp b/Engine/Source/Runtime/Core/Private/ProfilingDebugging/LoadTimeTracker.cpp
--- a/Engine/Source/Runtime/Core/Private/ProfilingDebugging/LoadTimeTracker.cpp
+++ b/Engine/Source/Runtime/Core/Private/ProfilingDebugging/LoadTimeTracker.cpp
@@ -52,6 +52,15 @@ void FLoadTimeTracker::DumpLoadTimes()
 	UE_LOG(LogLoad, Log, TEXT("Total Load times: %f"), TotalTime);
 }
 
+void FLoadTimeTracker::ResetLoadTimes()
+{
+	static bool bActuallyReset = !FParse::Param(FCommandLine::Get(), TEXT("-NoLoadTrackClear"));
+	if(bActuallyReset)
+	{
+		TimeInfo.Reset();
+	}
+}
+
 static FAutoConsoleCommand LoadTimerDumpCmd(
 	TEXT("LoadTimes.DumpTracking"),
 	TEXT("Dump high level load times being tracked"),
diff --git a/Engine/Source/Runtime/Core/Private/Windows/WindowsPlatformCrashContext.cpp b/Engine/Source/Runtime/Core/Private/Windows/WindowsPlatformCrashContext.cpp
--- a/Engine/Source/Runtime/Core/Private/Windows/WindowsPlatformCrashContext.cpp
+++ b/Engine/Source/Runtime/Core/Private/Windows/WindowsPlatformCrashContext.cpp
@@ -50,7 +50,7 @@ bool WriteMinidump( const TCHAR* Path, LPEXCEPTION_POINTERS ExceptionInfo, bool
 	DumpExceptionInfo.ClientPointers	= FALSE;
 
 	// CrashContext.runtime-xml is now a part of the minidump file.
-	FWindowsPlatformCrashContext CrashContext;
+	FWindowsPlatformCrashContext CrashContext(bIsEnsure);
 	CrashContext.SerializeContentToBuffer();
 	
 	MINIDUMP_USER_STREAM CrashContextStream ={0};
@@ -299,7 +299,7 @@ int32 ReportCrashUsingCrashReportClient(EXCEPTION_POINTERS* ExceptionInfo, const
 			{
 				// No super safe due to dynamic memory allocations, but at least enables new functionality.
 				// Introduces a new runtime crash context. Will replace all Windows related crash reporting.
-				FWindowsPlatformCrashContext CrashContext;
+				FWindowsPlatformCrashContext CrashContext(bIsEnsure);
 
 				const FString CrashContextXMLPath = FPaths::Combine( *FPaths::GameLogDir(), *CrashContext.GetUniqueCrashName(), FPlatformCrashContext::CrashContextRuntimeXMLNameW );
 				CrashContext.SerializeAsXML( *CrashContextXMLPath );
diff --git a/Engine/Source/Runtime/Core/Private/Windows/WindowsWindow.cpp b/Engine/Source/Runtime/Core/Private/Windows/WindowsWindow.cpp
--- a/Engine/Source/Runtime/Core/Private/Windows/WindowsWindow.cpp
+++ b/Engine/Source/Runtime/Core/Private/Windows/WindowsWindow.cpp
@@ -380,7 +380,11 @@ void FWindowsWindow::ReshapeWindow( int32 NewX, int32 NewY, int32 NewWidth, int3
 		NewHeight = FMath::Max( NewHeight, FMath::Min( OldHeight, MinRetainedHeight ) );
 	}
 
-		
+	if (IsMaximized())
+	{
+		Restore();
+	}
+
 	// We use SWP_NOSENDCHANGING when in fullscreen mode to prevent Windows limiting our window size to the current resolution, as that 
 	// prevents us being able to change to a higher resolution while in fullscreen mode
 	::SetWindowPos( HWnd, nullptr, WindowX, WindowY, NewWidth, NewHeight, SWP_NOZORDER | SWP_NOACTIVATE | ((WindowMode == EWindowMode::Fullscreen) ? SWP_NOSENDCHANGING : 0) );
diff --git a/Engine/Source/Runtime/Core/Public/GenericPlatform/GenericPlatformCrashContext.h b/Engine/Source/Runtime/Core/Public/GenericPlatform/GenericPlatformCrashContext.h
--- a/Engine/Source/Runtime/Core/Public/GenericPlatform/GenericPlatformCrashContext.h
+++ b/Engine/Source/Runtime/Core/Public/GenericPlatform/GenericPlatformCrashContext.h
@@ -145,6 +145,14 @@ public:
 	/** Unescapes a specified XML string, naive implementation. */
 	static FString UnescapeXMLString( const FString& Text );
 
+	/**
+	 * @return whether this crash is a non-crash event
+	 */
+	bool GetIsEnsure() const { return bIsEnsure; }
+
+protected:
+	bool bIsEnsure;
+
 private:
 	/** Serializes platform specific properties to the buffer. */
 	virtual void AddPlatformSpecificProperties();
diff --git a/Engine/Source/Runtime/Core/Public/HAL/IPlatformFileLogWrapper.h b/Engine/Source/Runtime/Core/Public/HAL/IPlatformFileLogWrapper.h
--- a/Engine/Source/Runtime/Core/Public/HAL/IPlatformFileLogWrapper.h
+++ b/Engine/Source/Runtime/Core/Public/HAL/IPlatformFileLogWrapper.h
@@ -37,7 +37,7 @@ public:
 		FILE_LOG(LogPlatformFile, VeryVerbose, TEXT("Tell %s"), *Filename);
 		double StartTime = FPlatformTime::Seconds();
 		int64 Result = FileHandle->Tell();
-		float ThisTime = 1000.0f * float(FPlatformTime::Seconds() - StartTime);
+		float ThisTime = (FPlatformTime::Seconds() - StartTime) / 1000.0;
 		FILE_LOG(LogPlatformFile, VeryVerbose, TEXT("Tell return %lld [%fms]"), Result, ThisTime);
 		return Result;
 	}
@@ -46,7 +46,7 @@ public:
 		FILE_LOG(LogPlatformFile, VeryVerbose, TEXT("Seek %s %lld"), *Filename, NewPosition);
 		double StartTime = FPlatformTime::Seconds();
 		bool Result = FileHandle->Seek(NewPosition);
-		float ThisTime = 1000.0f * float(FPlatformTime::Seconds() - StartTime);
+		float ThisTime = (FPlatformTime::Seconds() - StartTime) / 1000.0;
 		FILE_LOG(LogPlatformFile, VeryVerbose, TEXT("Seek return %d [%fms]"), int32(Result), ThisTime);
 		return Result;
 	}
@@ -55,7 +55,7 @@ public:
 		FILE_LOG(LogPlatformFile, VeryVerbose, TEXT("SeekFromEnd %s %lld"), *Filename, NewPositionRelativeToEnd);
 		double StartTime = FPlatformTime::Seconds();
 		bool Result = FileHandle->SeekFromEnd(NewPositionRelativeToEnd);
-		float ThisTime = 1000.0f * float(FPlatformTime::Seconds() - StartTime);
+		float ThisTime = (FPlatformTime::Seconds() - StartTime) / 1000.0;
 		FILE_LOG(LogPlatformFile, VeryVerbose, TEXT("SeekFromEnd return %d [%fms]"), int32(Result), ThisTime);
 		return Result;
 	}
@@ -64,7 +64,7 @@ public:
 		FILE_LOG(LogPlatformFile, VeryVerbose, TEXT("Read %s %lld"), *Filename, BytesToRead);
 		double StartTime = FPlatformTime::Seconds();
 		bool Result = FileHandle->Read(Destination, BytesToRead);
-		float ThisTime = 1000.0f * float(FPlatformTime::Seconds() - StartTime);
+		float ThisTime = (FPlatformTime::Seconds() - StartTime) / 1000.0;
 		FILE_LOG(LogPlatformFile, VeryVerbose, TEXT("Read return %d [%fms]"), int32(Result), ThisTime);
 		return Result;
 	}
@@ -73,7 +73,7 @@ public:
 		FILE_LOG(LogPlatformFile, VeryVerbose, TEXT("Write %s %lld"), *Filename, BytesToWrite);
 		double StartTime = FPlatformTime::Seconds();
 		bool Result = FileHandle->Write(Source, BytesToWrite);
-		float ThisTime = 1000.0f * float(FPlatformTime::Seconds() - StartTime);
+		float ThisTime = (FPlatformTime::Seconds() - StartTime) / 1000.0;
 		FILE_LOG(LogPlatformFile, VeryVerbose, TEXT("Write return %d [%fms]"), int32(Result), ThisTime);
 		return Result;
 	}
@@ -82,7 +82,7 @@ public:
 		FILE_LOG(LogPlatformFile, Verbose, TEXT("Size %s"), *Filename);
 		double StartTime = FPlatformTime::Seconds();
 		int64 Result = FileHandle->Size();
-		float ThisTime = 1000.0f * float(FPlatformTime::Seconds() - StartTime);
+		float ThisTime = (FPlatformTime::Seconds() - StartTime) / 1000.0;
 		FILE_LOG(LogPlatformFile, Verbose, TEXT("Size return %lld [%fms]"), Result, ThisTime);
 		return Result;
 	}
@@ -124,145 +124,177 @@ public:
 
 	virtual bool		FileExists(const TCHAR* Filename) override
 	{
-		FILE_LOG(LogPlatformFile, Log, TEXT("FileExists %s"), Filename);
+		FString DataStr = FString::Printf(TEXT("FileExists %s"), Filename);
+		FScopedNamedEvent NamedEvent(FColor::Emerald, *DataStr);
+		FILE_LOG(LogPlatformFile, Log, TEXT("%s"), *DataStr);
 		double StartTime = FPlatformTime::Seconds();
 		bool Result = LowerLevel->FileExists(Filename);
-		float ThisTime = 1000.0f * float(FPlatformTime::Seconds() - StartTime);
+		float ThisTime = (FPlatformTime::Seconds() - StartTime) / 1000.0;
 		FILE_LOG(LogPlatformFile, Log, TEXT("FileExists return %d [%fms]"), int32(Result), ThisTime);
 		return Result;
 	}
 	virtual int64		FileSize(const TCHAR* Filename) override
 	{
-		FILE_LOG(LogPlatformFile, Log, TEXT("FileSize %s"), Filename);
+		FString DataStr = FString::Printf(TEXT("FileSize %s"), Filename);
+		FScopedNamedEvent NamedEvent(FColor::Emerald, *DataStr);
+		FILE_LOG(LogPlatformFile, Log, TEXT("%s"), *DataStr);
 		double StartTime = FPlatformTime::Seconds();
 		int64 Result = LowerLevel->FileSize(Filename);
-		float ThisTime = 1000.0f * float(FPlatformTime::Seconds() - StartTime);
+		float ThisTime = (FPlatformTime::Seconds() - StartTime) / 1000.0;
 		FILE_LOG(LogPlatformFile, Log, TEXT("FileSize return %lld [%fms]"), Result, ThisTime);
 		return Result;
 	}
 	virtual bool		DeleteFile(const TCHAR* Filename) override
 	{
-		FILE_LOG(LogPlatformFile, Log, TEXT("DeleteFile %s"), Filename);
+		FString DataStr = FString::Printf(TEXT("DeleteFile %s"), Filename);
+		FScopedNamedEvent NamedEvent(FColor::Emerald, *DataStr);
+		FILE_LOG(LogPlatformFile, Log, TEXT("%s"), *DataStr);
 		double StartTime = FPlatformTime::Seconds();
 		bool Result = LowerLevel->DeleteFile(Filename);
-		float ThisTime = 1000.0f * float(FPlatformTime::Seconds() - StartTime);
+		float ThisTime = (FPlatformTime::Seconds() - StartTime) / 1000.0;
 		FILE_LOG(LogPlatformFile, Log, TEXT("DeleteFile return %d [%fms]"), int32(Result), ThisTime);
 		return Result;
 	}
 	virtual bool		IsReadOnly(const TCHAR* Filename) override
 	{
-		FILE_LOG(LogPlatformFile, Log, TEXT("IsReadOnly %s"), Filename);
+		FString DataStr = FString::Printf(TEXT("IsReadOnly %s"), Filename);
+		FScopedNamedEvent NamedEvent(FColor::Emerald, *DataStr);
+		FILE_LOG(LogPlatformFile, Log, TEXT("%s"), *DataStr);
 		double StartTime = FPlatformTime::Seconds();
 		bool Result = LowerLevel->IsReadOnly(Filename);
-		float ThisTime = 1000.0f * float(FPlatformTime::Seconds() - StartTime);
+		float ThisTime = (FPlatformTime::Seconds() - StartTime) / 1000.0;
 		FILE_LOG(LogPlatformFile, Log, TEXT("IsReadOnly return %d [%fms]"), int32(Result), ThisTime);
 		return Result;
 	}
 	virtual bool		MoveFile(const TCHAR* To, const TCHAR* From) override
 	{
-		FILE_LOG(LogPlatformFile, Log, TEXT("MoveFile %s %s"), To, From);
+		FString DataStr = FString::Printf(TEXT("MoveFile %s %s"), To, From);
+		FScopedNamedEvent NamedEvent(FColor::Emerald, *DataStr);
+		FILE_LOG(LogPlatformFile, Log, TEXT("%s"), *DataStr);
 		double StartTime = FPlatformTime::Seconds();
 		bool Result = LowerLevel->MoveFile(To, From);
-		float ThisTime = 1000.0f * float(FPlatformTime::Seconds() - StartTime);
+		float ThisTime = (FPlatformTime::Seconds() - StartTime) / 1000.0;
 		FILE_LOG(LogPlatformFile, Log, TEXT("MoveFile return %d [%fms]"), int32(Result), ThisTime);
 		return Result;
 	}
 	virtual bool		SetReadOnly(const TCHAR* Filename, bool bNewReadOnlyValue) override
 	{
-		FILE_LOG(LogPlatformFile, Log, TEXT("SetReadOnly %s %d"), Filename, int32(bNewReadOnlyValue));
+		FString DataStr = FString::Printf(TEXT("SetReadOnly %s %d"), Filename, int32(bNewReadOnlyValue));
+		FScopedNamedEvent NamedEvent(FColor::Emerald, *DataStr);
+		FILE_LOG(LogPlatformFile, Log, TEXT("%s"), *DataStr);
 		double StartTime = FPlatformTime::Seconds();
 		bool Result = LowerLevel->SetReadOnly(Filename, bNewReadOnlyValue);
-		float ThisTime = 1000.0f * float(FPlatformTime::Seconds() - StartTime);
+		float ThisTime = (FPlatformTime::Seconds() - StartTime) / 1000.0;
 		FILE_LOG(LogPlatformFile, Log, TEXT("SetReadOnly return %d [%fms]"), int32(Result), ThisTime);
 		return Result;
 	}
 	virtual FDateTime	GetTimeStamp(const TCHAR* Filename) override
 	{
-		FILE_LOG(LogPlatformFile, Log, TEXT("GetTimeStamp %s"), Filename);
+		FString DataStr = FString::Printf(TEXT("GetTimeStamp %s"), Filename);
+		FScopedNamedEvent NamedEvent(FColor::Emerald, *DataStr);
+		FILE_LOG(LogPlatformFile, Log, TEXT("%s"), *DataStr);
 		double StartTime = FPlatformTime::Seconds();
 		FDateTime Result = LowerLevel->GetTimeStamp(Filename);
-		float ThisTime = 1000.0f * float(FPlatformTime::Seconds() - StartTime);
+		float ThisTime = (FPlatformTime::Seconds() - StartTime) / 1000.0;
 		FILE_LOG(LogPlatformFile, Log, TEXT("GetTimeStamp return %llx [%fms]"), Result.GetTicks(), ThisTime);
 		return Result;
 	}
 	virtual void		SetTimeStamp(const TCHAR* Filename, FDateTime DateTime) override
 	{
-		FILE_LOG(LogPlatformFile, Log, TEXT("SetTimeStamp %s"), Filename);
+		FString DataStr = FString::Printf(TEXT("SetTimeStamp %s"), Filename);
+		FScopedNamedEvent NamedEvent(FColor::Emerald, *DataStr);
+		FILE_LOG(LogPlatformFile, Log, TEXT("%s"), *DataStr);
 		double StartTime = FPlatformTime::Seconds();
 		LowerLevel->SetTimeStamp(Filename, DateTime);
-		float ThisTime = 1000.0f * float(FPlatformTime::Seconds() - StartTime);
+		float ThisTime = (FPlatformTime::Seconds() - StartTime) / 1000.0;
 		FILE_LOG(LogPlatformFile, Log, TEXT("SetTimeStamp [%fms]"), ThisTime);
 	}
 	virtual FDateTime	GetAccessTimeStamp(const TCHAR* Filename) override
 	{
-		FILE_LOG(LogPlatformFile, Log, TEXT("GetAccessTimeStamp %s"), Filename);
+		FString DataStr = FString::Printf(TEXT("GetAccessTimeStamp %s"), Filename);
+		FScopedNamedEvent NamedEvent(FColor::Emerald, *DataStr);
+		FILE_LOG(LogPlatformFile, Log, TEXT("%s"), *DataStr);
 		double StartTime = FPlatformTime::Seconds();
 		FDateTime Result = LowerLevel->GetAccessTimeStamp(Filename);
-		float ThisTime = 1000.0f * float(FPlatformTime::Seconds() - StartTime);
+		float ThisTime = (FPlatformTime::Seconds() - StartTime) / 1000.0;
 		FILE_LOG(LogPlatformFile, Log, TEXT("GetAccessTimeStamp return %llx [%fms]"), Result.GetTicks(), ThisTime);
 		return Result;
 	}
 	virtual FString	GetFilenameOnDisk(const TCHAR* Filename) override
 	{
-		FILE_LOG(LogPlatformFile, Log, TEXT("GetFilenameOnDisk %s"), Filename);
+		FString DataStr = FString::Printf(TEXT("GetFilenameOnDisk %s"), Filename);
+		FScopedNamedEvent NamedEvent(FColor::Emerald, *DataStr);
+		FILE_LOG(LogPlatformFile, Log, TEXT("%s"), *DataStr);
 		double StartTime = FPlatformTime::Seconds();
 		FString Result = LowerLevel->GetFilenameOnDisk(Filename);
-		float ThisTime = 1000.0f * float(FPlatformTime::Seconds() - StartTime);
+		float ThisTime = (FPlatformTime::Seconds() - StartTime) / 1000.0;
 		FILE_LOG(LogPlatformFile, Log, TEXT("GetFilenameOnDisk return %s [%fms]"), *Result, ThisTime);
 		return Result;
 	}
 	virtual IFileHandle*	OpenRead(const TCHAR* Filename, bool bAllowWrite) override
 	{
-		FILE_LOG(LogPlatformFile, Log, TEXT("OpenRead %s"), Filename);
+		FString DataStr = FString::Printf(TEXT("OpenRead %s"), Filename);
+		FScopedNamedEvent NamedEvent(FColor::Emerald, *DataStr);
+		FILE_LOG(LogPlatformFile, Log, TEXT("%s"), *DataStr);
 		double StartTime = FPlatformTime::Seconds();
 		IFileHandle* Result = LowerLevel->OpenRead(Filename, bAllowWrite);
-		float ThisTime = 1000.0f * float(FPlatformTime::Seconds() - StartTime);
+		float ThisTime = (FPlatformTime::Seconds() - StartTime) / 1000.0;
 		FILE_LOG(LogPlatformFile, Log, TEXT("OpenRead return %llx [%fms]"), uint64(Result), ThisTime);
 		return Result ? (new FLoggedFileHandle(Result, Filename, *this)) : Result;
 	}
 	virtual IFileHandle*	OpenWrite(const TCHAR* Filename, bool bAppend = false, bool bAllowRead = false) override
 	{
-		FILE_LOG(LogPlatformFile, Log, TEXT("OpenWrite %s %d %d"), Filename, int32(bAppend), int32(bAllowRead));
+		FString DataStr = FString::Printf(TEXT("OpenWrite %s %d %d"), Filename, int32(bAppend), int32(bAllowRead));
+		FScopedNamedEvent NamedEvent(FColor::Emerald, *DataStr);
+		FILE_LOG(LogPlatformFile, Log, TEXT("%s"), *DataStr);
 		double StartTime = FPlatformTime::Seconds();
 		IFileHandle* Result = LowerLevel->OpenWrite(Filename, bAppend, bAllowRead);
-		float ThisTime = 1000.0f * float(FPlatformTime::Seconds() - StartTime);
+		float ThisTime = (FPlatformTime::Seconds() - StartTime) / 1000.0;
 		FILE_LOG(LogPlatformFile, Log, TEXT("OpenWrite return %llx [%fms]"), uint64(Result), ThisTime);
 		return Result ? (new FLoggedFileHandle(Result, Filename, *this)) : Result;
 	}
 
 	virtual bool		DirectoryExists(const TCHAR* Directory) override
 	{
-		FILE_LOG(LogPlatformFile, Log, TEXT("DirectoryExists %s"), Directory);
+		FString DataStr = FString::Printf(TEXT("DirectoryExists %s"), Directory);
+		FScopedNamedEvent NamedEvent(FColor::Emerald, *DataStr);
+		FILE_LOG(LogPlatformFile, Log, TEXT("%s"), *DataStr);
 		double StartTime = FPlatformTime::Seconds();
 		bool Result = LowerLevel->DirectoryExists(Directory);
-		float ThisTime = 1000.0f * float(FPlatformTime::Seconds() - StartTime);
+		float ThisTime = (FPlatformTime::Seconds() - StartTime) / 1000.0;
 		FILE_LOG(LogPlatformFile, Log, TEXT("DirectoryExists return %d [%fms]"), int32(Result), ThisTime);
 		return Result;
 	}
 	virtual bool		CreateDirectory(const TCHAR* Directory) override
 	{
-		FILE_LOG(LogPlatformFile, Log, TEXT("CreateDirectory %s"), Directory);
+		FString DataStr = FString::Printf(TEXT("CreateDirectory %s"), Directory);
+		FScopedNamedEvent NamedEvent(FColor::Emerald, *DataStr);
+		FILE_LOG(LogPlatformFile, Log, TEXT("%s"), *DataStr);
 		double StartTime = FPlatformTime::Seconds();
 		bool Result = LowerLevel->CreateDirectory(Directory);
-		float ThisTime = 1000.0f * float(FPlatformTime::Seconds() - StartTime);
+		float ThisTime = (FPlatformTime::Seconds() - StartTime) / 1000.0;
 		FILE_LOG(LogPlatformFile, Log, TEXT("CreateDirectory return %d [%fms]"), int32(Result), ThisTime);
 		return Result;
 	}
 	virtual bool		DeleteDirectory(const TCHAR* Directory) override
 	{
-		FILE_LOG(LogPlatformFile, Log, TEXT("DeleteDirectory %s"), Directory);
+		FString DataStr = FString::Printf(TEXT("DeleteDirectory %s"), Directory);
+		FScopedNamedEvent NamedEvent(FColor::Emerald, *DataStr);
+		FILE_LOG(LogPlatformFile, Log, TEXT("%s"), *DataStr);
 		double StartTime = FPlatformTime::Seconds();
 		bool Result = LowerLevel->DeleteDirectory(Directory);
-		float ThisTime = 1000.0f * float(FPlatformTime::Seconds() - StartTime);
+		float ThisTime = (FPlatformTime::Seconds() - StartTime) / 1000.0;
 		FILE_LOG(LogPlatformFile, Log, TEXT("DeleteDirectory return %d [%fms]"), int32(Result), ThisTime);
 		return Result;
 	}
 	virtual FFileStatData GetStatData(const TCHAR* FilenameOrDirectory) override
 	{
-		FILE_LOG(LogPlatformFile, Log, TEXT("GetStatData %s"), FilenameOrDirectory);
+		FString DataStr = FString::Printf(TEXT("GetStatData %s"), FilenameOrDirectory);
+		FScopedNamedEvent NamedEvent(FColor::Emerald, *DataStr);
+		FILE_LOG(LogPlatformFile, Log, TEXT("%s"), *DataStr);
 		double StartTime = FPlatformTime::Seconds();
 		FFileStatData Result = LowerLevel->GetStatData(FilenameOrDirectory);
-		float ThisTime = 1000.0f * float(FPlatformTime::Seconds() - StartTime);
+		float ThisTime = (FPlatformTime::Seconds() - StartTime) / 1000.0;
 		FILE_LOG(LogPlatformFile, Log, TEXT("GetStatData return %d [%fms]"), int32(Result.bIsValid), ThisTime);
 		return Result;
 	}
@@ -280,7 +312,7 @@ public:
 			FILE_LOG(LogPlatformFile, Verbose, TEXT("Visit %s %d"), FilenameOrDirectory, int32(bIsDirectory));
 			double StartTime = FPlatformTime::Seconds();
 			bool Result = Visitor.Visit(FilenameOrDirectory, bIsDirectory);
-			float ThisTime = 1000.0f * float(FPlatformTime::Seconds() - StartTime);
+			float ThisTime = (FPlatformTime::Seconds() - StartTime) / 1000.0;
 			FILE_LOG(LogPlatformFile, Verbose, TEXT("Visit return %d [%fms]"), int32(Result), ThisTime);
 			return Result;
 		}
@@ -288,21 +320,25 @@ public:
 
 	virtual bool		IterateDirectory(const TCHAR* Directory, IPlatformFile::FDirectoryVisitor& Visitor) override
 	{
-		FILE_LOG(LogPlatformFile, Log, TEXT("IterateDirectory %s"), Directory);
+		FString DataStr = FString::Printf(TEXT("IterateDirectory %s"), Directory);
+		FScopedNamedEvent NamedEvent(FColor::Emerald, *DataStr);
+		FILE_LOG(LogPlatformFile, Log, TEXT("%s"), *DataStr);
 		double StartTime = FPlatformTime::Seconds();
 		FLogVisitor LogVisitor(Visitor);
 		bool Result = LowerLevel->IterateDirectory(Directory, LogVisitor);
-		float ThisTime = 1000.0f * float(FPlatformTime::Seconds() - StartTime);
+		float ThisTime = (FPlatformTime::Seconds() - StartTime) / 1000.0;
 		FILE_LOG(LogPlatformFile, Log, TEXT("IterateDirectory return %d [%fms]"), int32(Result), ThisTime);
 		return Result;
 	}
 	virtual bool		IterateDirectoryRecursively(const TCHAR* Directory, IPlatformFile::FDirectoryVisitor& Visitor) override
 	{
-		FILE_LOG(LogPlatformFile, Log, TEXT("IterateDirectoryRecursively %s"), Directory);
+		FString DataStr = FString::Printf(TEXT("IterateDirectoryRecursively %s"), Directory);
+		FScopedNamedEvent NamedEvent(FColor::Emerald, *DataStr);
+		FILE_LOG(LogPlatformFile, Log, TEXT("%s"), *DataStr);
 		double StartTime = FPlatformTime::Seconds();
 		FLogVisitor LogVisitor(Visitor);
 		bool Result = LowerLevel->IterateDirectoryRecursively(Directory, LogVisitor);
-		float ThisTime = 1000.0f * float(FPlatformTime::Seconds() - StartTime);
+		float ThisTime = (FPlatformTime::Seconds() - StartTime) / 1000.0;
 		FILE_LOG(LogPlatformFile, Log, TEXT("IterateDirectoryRecursively return %d [%fms]"), int32(Result), ThisTime);
 		return Result;
 	}
@@ -320,7 +356,7 @@ public:
 			FILE_LOG(LogPlatformFile, Verbose, TEXT("Visit %s %d"), FilenameOrDirectory, int32(StatData.bIsDirectory));
 			double StartTime = FPlatformTime::Seconds();
 			bool Result = Visitor.Visit(FilenameOrDirectory, StatData);
-			float ThisTime = 1000.0f * float(FPlatformTime::Seconds() - StartTime);
+			float ThisTime = (FPlatformTime::Seconds() - StartTime) / 1000.0;
 			FILE_LOG(LogPlatformFile, Verbose, TEXT("Visit return %d [%fms]"), int32(Result), ThisTime);
 			return Result;
 		}
@@ -328,40 +364,48 @@ public:
 
 	virtual bool		IterateDirectoryStat(const TCHAR* Directory, IPlatformFile::FDirectoryStatVisitor& Visitor) override
 	{
-		FILE_LOG(LogPlatformFile, Log, TEXT("IterateDirectoryStat %s"), Directory);
+		FString DataStr = FString::Printf(TEXT("IterateDirectoryStat %s"), Directory);
+		FScopedNamedEvent NamedEvent(FColor::Emerald, *DataStr);
+		FILE_LOG(LogPlatformFile, Log, TEXT("%s"), *DataStr);
 		double StartTime = FPlatformTime::Seconds();
 		FLogStatVisitor LogVisitor(Visitor);
 		bool Result = LowerLevel->IterateDirectoryStat(Directory, LogVisitor);
-		float ThisTime = 1000.0f * float(FPlatformTime::Seconds() - StartTime);
+		float ThisTime = (FPlatformTime::Seconds() - StartTime) / 1000.0;
 		FILE_LOG(LogPlatformFile, Log, TEXT("IterateDirectoryStat return %d [%fms]"), int32(Result), ThisTime);
 		return Result;
 	}
 	virtual bool		IterateDirectoryStatRecursively(const TCHAR* Directory, IPlatformFile::FDirectoryStatVisitor& Visitor) override
 	{
-		FILE_LOG(LogPlatformFile, Log, TEXT("IterateDirectoryStatRecursively %s"), Directory);
+		FString DataStr = FString::Printf(TEXT("IterateDirectoryStatRecursively %s"), Directory);
+		FScopedNamedEvent NamedEvent(FColor::Emerald, *DataStr);
+		FILE_LOG(LogPlatformFile, Log, TEXT("%s"), *DataStr);
 		double StartTime = FPlatformTime::Seconds();
 		FLogStatVisitor LogVisitor(Visitor);
 		bool Result = LowerLevel->IterateDirectoryStatRecursively(Directory, LogVisitor);
-		float ThisTime = 1000.0f * float(FPlatformTime::Seconds() - StartTime);
+		float ThisTime = (FPlatformTime::Seconds() - StartTime) / 1000.0;
 		FILE_LOG(LogPlatformFile, Log, TEXT("IterateDirectoryStatRecursively return %d [%fms]"), int32(Result), ThisTime);
 		return Result;
 	}
 
 	virtual bool		DeleteDirectoryRecursively(const TCHAR* Directory) override
 	{
-		FILE_LOG(LogPlatformFile, Log, TEXT("DeleteDirectoryRecursively %s"), Directory);
+		FString DataStr = FString::Printf(TEXT("DeleteDirectoryRecursively %s"), Directory);
+		FScopedNamedEvent NamedEvent(FColor::Emerald, *DataStr);
+		FILE_LOG(LogPlatformFile, Log, TEXT("%s"), *DataStr);
 		double StartTime = FPlatformTime::Seconds();
 		bool Result = LowerLevel->DeleteDirectoryRecursively(Directory);
-		float ThisTime = 1000.0f * float(FPlatformTime::Seconds() - StartTime);
+		float ThisTime = (FPlatformTime::Seconds() - StartTime) / 1000.0;
 		FILE_LOG(LogPlatformFile, Log, TEXT("DeleteDirectoryRecursively return %d [%fms]"), int32(Result), ThisTime);
 		return Result;
 	}
 	virtual bool		CopyFile(const TCHAR* To, const TCHAR* From) override
 	{
-		FILE_LOG(LogPlatformFile, Log, TEXT("CopyFile %s %s"), To, From);
+		FString DataStr = FString::Printf(TEXT("CopyFile %s %s"), To, From);
+		FScopedNamedEvent NamedEvent(FColor::Emerald, *DataStr);
+		FILE_LOG(LogPlatformFile, Log, TEXT("%s"), *DataStr);
 		double StartTime = FPlatformTime::Seconds();
 		bool Result = LowerLevel->CopyFile(To, From);
-		float ThisTime = 1000.0f * float(FPlatformTime::Seconds() - StartTime);
+		float ThisTime = (FPlatformTime::Seconds() - StartTime) / 1000.0;
 		FILE_LOG(LogPlatformFile, Log, TEXT("CopyFile return %d [%fms]"), int32(Result), ThisTime);
 		return Result;
 	}
diff --git a/Engine/Source/Runtime/Core/Public/Linux/LinuxPlatformCrashContext.h b/Engine/Source/Runtime/Core/Public/Linux/LinuxPlatformCrashContext.h
--- a/Engine/Source/Runtime/Core/Public/Linux/LinuxPlatformCrashContext.h
+++ b/Engine/Source/Runtime/Core/Public/Linux/LinuxPlatformCrashContext.h
@@ -21,9 +21,6 @@ struct CORE_API FLinuxCrashContext : public FGenericCrashContext
 	/** Whether backtrace was already captured */
 	bool bCapturedBacktrace;
 
-	/** Whether we're handling ensure() (prefer speed in this case). */
-	bool bHandlingEnsure;
-
 	/** Symbols received via backtrace_symbols(), if any (note that we will need to clean it up) */
 	char ** BacktraceSymbols;
 
@@ -38,13 +35,24 @@ struct CORE_API FLinuxCrashContext : public FGenericCrashContext
 		,	Info(nullptr)
 		,	Context(nullptr)
 		,	bCapturedBacktrace(false)
-		,	bHandlingEnsure(false)
 		,	BacktraceSymbols(nullptr)
 	{
 		SignalDescription[ 0 ] = 0;
 		MinidumpCallstackInfo[ 0 ] = 0;
 	}
 
+	FLinuxCrashContext(bool bInIsEnsure)
+		: Signal(0)
+		, Info(nullptr)
+		, Context(nullptr)
+		, bCapturedBacktrace(false)
+		, BacktraceSymbols(nullptr)
+	{
+		SignalDescription[0] = 0;
+		MinidumpCallstackInfo[0] = 0;
+		bIsEnsure = bInIsEnsure;
+	}
+
 	~FLinuxCrashContext();
 
 	/**
@@ -71,6 +79,11 @@ struct CORE_API FLinuxCrashContext : public FGenericCrashContext
 	 */
 	void GenerateCrashInfoAndLaunchReporter(bool bReportingNonCrash = false) const;
 
+	/**
+	 * Sets whether this crash represents a non-crash event like an ensure
+	 */
+	void SetIsEnsure(bool bInIsEnsure) { bIsEnsure = bInIsEnsure; }
+
 protected:
 	/**
 	 * Dumps all the data from crash context to the "minidump" report.
diff --git a/Engine/Source/Runtime/Core/Public/Linux/LinuxPlatformRunnableThread.h b/Engine/Source/Runtime/Core/Public/Linux/LinuxPlatformRunnableThread.h
--- a/Engine/Source/Runtime/Core/Public/Linux/LinuxPlatformRunnableThread.h
+++ b/Engine/Source/Runtime/Core/Public/Linux/LinuxPlatformRunnableThread.h
@@ -13,16 +13,29 @@
 **/
 class FRunnableThreadLinux : public FRunnableThreadPThread
 {
-	enum
+	struct EConstants
 	{
-		LinuxThreadNameLimit = 15			// the limit for thread name is just 15 chars :( http://man7.org/linux/man-pages/man3/pthread_setname_np.3.html
+		enum
+		{
+			LinuxThreadNameLimit = 15,			// the limit for thread name is just 15 chars :( http://man7.org/linux/man-pages/man3/pthread_setname_np.3.html
+			CrashHandlerStackSize = SIGSTKSZ + 128 * 1024	// should be at least SIGSTKSIZE, plus 128K because we do logging in crash handler
+		};
 	};
 
+	/** Each thread needs a separate stack for the signal handler, so possible stack overflows in the thread are handled */
+	void* ThreadCrashHandlingStack;
+
 public:
-    FRunnableThreadLinux() : FRunnableThreadPThread()
-    {
-    }
-    
+
+	/** Separate stack for the signal handler (so possible stack overflows don't go unnoticed), shared between threads. */
+	static char MainThreadSignalHandlerStack[EConstants::CrashHandlerStackSize];
+
+	FRunnableThreadLinux()
+		:	FRunnableThreadPThread()
+		,	ThreadCrashHandlingStack(nullptr)
+	{
+	}
+
 private:
 
 	/**
@@ -32,27 +45,27 @@ private:
 	{
 		FString SizeLimitedThreadName = ThreadName;
 
-		if (SizeLimitedThreadName.Len() > LinuxThreadNameLimit)
+		if (SizeLimitedThreadName.Len() > EConstants::LinuxThreadNameLimit)
 		{
 			// first, attempt to cut out common and meaningless substrings
 			SizeLimitedThreadName = SizeLimitedThreadName.Replace(TEXT("Thread"), TEXT(""));
 			SizeLimitedThreadName = SizeLimitedThreadName.Replace(TEXT("Runnable"), TEXT(""));
 
 			// if still larger
-			if (SizeLimitedThreadName.Len() > LinuxThreadNameLimit)
+			if (SizeLimitedThreadName.Len() > EConstants::LinuxThreadNameLimit)
 			{
 				FString Temp = SizeLimitedThreadName;
 
 				// cut out the middle and replace with a substitute
 				const TCHAR Dash[] = TEXT("-");
 				const int32 DashLen = ARRAY_COUNT(Dash) - 1;
-				int NumToLeave = (LinuxThreadNameLimit - DashLen) / 2;
+				int NumToLeave = (EConstants::LinuxThreadNameLimit - DashLen) / 2;
 
-				SizeLimitedThreadName = Temp.Left(LinuxThreadNameLimit - (NumToLeave + DashLen));
+				SizeLimitedThreadName = Temp.Left(EConstants::LinuxThreadNameLimit - (NumToLeave + DashLen));
 				SizeLimitedThreadName += Dash;
 				SizeLimitedThreadName += Temp.Right(NumToLeave);
 
-				check(SizeLimitedThreadName.Len() <= LinuxThreadNameLimit);
+				check(SizeLimitedThreadName.Len() <= EConstants::LinuxThreadNameLimit);
 			}
 		}
 
@@ -61,6 +74,38 @@ private:
 		{
 			UE_LOG(LogHAL, Warning, TEXT("pthread_setname_np(, '%s') failed with error %d (%s)."), *ThreadName, ErrCode, ANSI_TO_TCHAR(strerror(ErrCode)));
 		}
+
+		// set the alternate stack for handling crashes due to stack overflow
+		ThreadCrashHandlingStack = FMemory::Malloc(EConstants::CrashHandlerStackSize);
+		if (ThreadCrashHandlingStack)
+		{
+			stack_t SignalHandlerStack;
+			FMemory::Memzero(SignalHandlerStack);
+			SignalHandlerStack.ss_sp = ThreadCrashHandlingStack;
+			SignalHandlerStack.ss_size = EConstants::CrashHandlerStackSize;
+
+			if (sigaltstack(&SignalHandlerStack, nullptr) < 0)
+			{
+				int ErrNo = errno;
+				UE_LOG(LogLinux, Warning, TEXT("Unable to set alternate stack for crash handler, sigaltstack() failed with errno=%d (%s)"),
+					ErrNo,
+					UTF8_TO_TCHAR(strerror(ErrNo))
+					);
+
+				// free the memory immediately
+				FMemory::Free(ThreadCrashHandlingStack);
+				ThreadCrashHandlingStack = nullptr;
+			}
+		}
+	}
+
+	virtual void PostRun()
+	{
+		if (ThreadCrashHandlingStack)
+		{
+			FMemory::Free(ThreadCrashHandlingStack);
+			ThreadCrashHandlingStack = nullptr;
+		}
 	}
 
 	/**
diff --git a/Engine/Source/Runtime/Core/Public/Mac/MacPlatformCrashContext.h b/Engine/Source/Runtime/Core/Public/Mac/MacPlatformCrashContext.h
--- a/Engine/Source/Runtime/Core/Public/Mac/MacPlatformCrashContext.h
+++ b/Engine/Source/Runtime/Core/Public/Mac/MacPlatformCrashContext.h
@@ -7,13 +7,13 @@
 struct CORE_API FMacCrashContext : public FApplePlatformCrashContext
 {
 	/** Mimics Windows WER format */
-	void GenerateWindowsErrorReport(char const* WERPath) const;
+	void GenerateWindowsErrorReport(char const* WERPath, bool bIsEnsure = false) const;
 
 	/** Copies the PLCrashReporter minidump */
 	void CopyMinidump(char const* OutputPath, char const* InputPath) const;
 
 	/** Generates the ensure/crash info into the given folder */
-	void GenerateInfoInFolder(char const* const InfoFolder) const;
+	void GenerateInfoInFolder(char const* const InfoFolder, bool bIsEnsure = false) const;
 	
 	/** Generates information for crash reporter */
 	void GenerateCrashInfoAndLaunchReporter() const;
diff --git a/Engine/Source/Runtime/Core/Public/ProfilingDebugging/CookStats.h b/Engine/Source/Runtime/Core/Public/ProfilingDebugging/CookStats.h
new file mode 100644
--- /dev/null
+++ b/Engine/Source/Runtime/Core/Public/ProfilingDebugging/CookStats.h
@@ -0,0 +1,107 @@
+// Copyright 1998-2016 Epic Games, Inc. All Rights Reserved.
+
+#pragma once
+
+#ifndef ENABLE_COOK_STATS
+#define ENABLE_COOK_STATS WITH_EDITOR
+#endif
+
+#if ENABLE_COOK_STATS
+#include "Core.h"
+
+/**
+ * Centralizes the system to gather stats from a cook that need to be collected at the core/engine level.
+ * Essentially, add a delegate to the CookStatsCallbacks member. When a cook a complete that is configured to use stats (ENABLE_COOK_STATS),
+ * it will broadcast this delegate, which, when called, gives you a TFunction to call to report each of your stats.
+ * 
+ * For simplicity, FAutoRegisterCallback is provided to auto-register your callback on startup. Usage is like:
+ * 
+ * static void ReportCookStats(AddStatFuncRef AddStat)
+ * {
+ *		AddStat("MySubsystem.Event1", CreateKeyValueArray("Attr1", "Value1", "Attr2", "Value2" ... );
+ *		AddStat("MySubsystem.Event2", CreateKeyValueArray("Attr1", "Value1", "Attr2", "Value2" ... );
+ * }
+ * 
+ * FAutoRegisterCallback GMySubsystemCookRegistration(&ReportCookStats);
+ * 
+ * When a cook is complete, your callback will be called, and the stats will be either logged, sent to an
+ * analytics provider, logged to an external file, etc. You don't care how they are added, you just call 
+ * AddStat for each stat you want to add.
+ */
+class FCookStatsManager
+{
+public:
+	/** Every stat is essentially a name followed by an array of key/value "attributes" associated with the stat. */
+	struct KeyValue
+	{
+		FString Key;
+		FString Value;
+		KeyValue(const FString& InKey, const FString& InValue) : Key(InKey), Value(InValue) {}
+	};
+	/** When you register a callback for reporting your stats, you will be given a TFunctionRef to call to add stats. This is the signature of the function you will call. */
+	typedef TFunctionRef<void(const FString& StatName, const TArray<KeyValue>& StatAttributes)> AddStatFuncRef;
+
+	/** To register a callback for reporting stats, create an instance of this delegate type and add it to the delegate variable below. */
+	DECLARE_MULTICAST_DELEGATE_OneParam(FGatherCookStatsDelegate, AddStatFuncRef);
+
+	/** Use this to register a callback to gather cook stats for a translation unit. When a cook is finished, this delegate will be fired. */
+	static CORE_API FGatherCookStatsDelegate CookStatsCallbacks;
+
+	/** Called after the cook is finished to gather the stats. */
+	static CORE_API void LogCookStats(AddStatFuncRef AddStat);
+
+	/** Helper struct to auto-register your STATIC FUNCTION with CookStatsCallbacks */
+	struct FAutoRegisterCallback
+	{
+		template<typename Func>
+		explicit FAutoRegisterCallback(Func Callback)
+		{
+			CookStatsCallbacks.AddLambda(Callback);
+		}
+	};
+
+	//-------------------------------------------------------------------------
+	// Helper to initialize an array of KeyValues on one line.
+	//-------------------------------------------------------------------------
+
+	/** Helper to create an array of KeyValues using a single expression. If there is an odd number of arguments, the last value is considered to be an empty string. */
+	template <typename...ArgTypes>
+	static TArray<FCookStatsManager::KeyValue> CreateKeyValueArray(ArgTypes&&...Args)
+	{
+		TArray<FCookStatsManager::KeyValue> Attrs;
+		Attrs.Empty(sizeof...(Args) / 2);
+		ImplCreateKeyValueArray(Attrs, Args...);
+		return Attrs;
+	}
+private:
+	/** ToString abstraction that can convert anything into a string. Differs from LexicalConversion in that it is legal to pass a string or TCHAR* to it. */
+	static const FString& ToString(const FString& InStr) { return InStr; }
+	static FString ToString(const TCHAR* InStr) { return FString(InStr); }
+	template <typename T> static FString ToString(T&&Val) { return LexicalConversion::ToString(Forward<T>(Val)); }
+
+	template <typename KeyType>
+	static void ImplCreateKeyValueArray(TArray<FCookStatsManager::KeyValue>& Attrs, KeyType&& Key)
+	{
+		Attrs.Emplace(ToString(Forward<KeyType>(Key)), TEXT(""));
+	}
+
+	template <typename KeyType, typename ValueType>
+	static void ImplCreateKeyValueArray(TArray<FCookStatsManager::KeyValue>& Attrs, KeyType&& Key, ValueType&& Value)
+	{
+		Attrs.Emplace(ToString(Forward<KeyType>(Key)), ToString(Forward<ValueType>(Value)));
+	}
+
+	template <typename KeyType, typename ValueType, typename...ArgTypes>
+	static void ImplCreateKeyValueArray(TArray<FCookStatsManager::KeyValue>& Attrs, KeyType&& Key, ValueType&& Value, ArgTypes&&...Args)
+	{
+		Attrs.Emplace(ToString(Forward<KeyType>(Key)), ToString(Forward<ValueType>(Value)));
+		ImplCreateKeyValueArray(Attrs, Forward<ArgTypes>(Args)...);
+	}
+
+	};
+
+/** any expression inside this block will be compiled out if ENABLE_COOK_STATS is not true. Shorthand for the multi-line #if ENABLE_COOK_STATS guard. */
+#define COOK_STAT(expr) expr
+#else
+#define COOK_STAT(expr)
+#endif
diff --git a/Engine/Source/Runtime/Core/Public/ProfilingDebugging/LoadTimeTracker.h b/Engine/Source/Runtime/Core/Public/ProfilingDebugging/LoadTimeTracker.h
--- a/Engine/Source/Runtime/Core/Public/ProfilingDebugging/LoadTimeTracker.h
+++ b/Engine/Source/Runtime/Core/Public/ProfilingDebugging/LoadTimeTracker.h
@@ -6,11 +6,7 @@
 
 #pragma once
 
-#if !UE_BUILD_SHIPPING
 #define ENABLE_LOADTIME_TRACKING 1
-#else
-#define ENABLE_LOADTIME_TRACKING 0
-#endif
 
 #include "ScopedTimers.h"
 
@@ -38,6 +34,13 @@ public:
 		Get().DumpLoadTimes();
 	}
 
+	const TMap<FName, TArray<double>>& GetData() const
+	{
+		return TimeInfo;
+	}
+
+	void ResetLoadTimes();
+
 private:
 	TMap<FName, TArray<double>> TimeInfo;
 
diff --git a/Engine/Source/Runtime/Core/Public/Windows/WindowsPlatformCrashContext.h b/Engine/Source/Runtime/Core/Public/Windows/WindowsPlatformCrashContext.h
--- a/Engine/Source/Runtime/Core/Public/Windows/WindowsPlatformCrashContext.h
+++ b/Engine/Source/Runtime/Core/Public/Windows/WindowsPlatformCrashContext.h
@@ -16,6 +16,13 @@ struct CORE_API FWindowsPlatformCrashContext : public FGenericCrashContext
 		UE4_MINIDUMP_CRASHCONTEXT = LastReservedStream + 1,
 	};
 
+	FWindowsPlatformCrashContext() {}
+
+	FWindowsPlatformCrashContext(bool bInIsEnsure)
+	{
+		bIsEnsure = bInIsEnsure;
+	}
+
 	virtual void AddPlatformSpecificProperties() override
 	{
 		AddCrashProperty( TEXT( "PlatformIsRunningWindows" ), 1 );
diff --git a/Engine/Source/Runtime/Engine/Classes/Animation/AnimationAsset.h b/Engine/Source/Runtime/Engine/Classes/Animation/AnimationAsset.h
--- a/Engine/Source/Runtime/Engine/Classes/Animation/AnimationAsset.h
+++ b/Engine/Source/Runtime/Engine/Classes/Animation/AnimationAsset.h
@@ -421,6 +421,7 @@ public:
 		ActivePlayers.Reset();
 		bCanUseMarkerSync = false;
 		MontageLeaderWeight = 0.f;
+		MarkerTickContext = FMarkerTickContext();
 	}
 
 	// Checks the last tick record in the ActivePlayers array to see if it's a better leader than the current candidate.
diff --git a/Engine/Source/Runtime/Engine/Classes/Components/SkeletalMeshComponent.h b/Engine/Source/Runtime/Engine/Classes/Components/SkeletalMeshComponent.h
--- a/Engine/Source/Runtime/Engine/Classes/Components/SkeletalMeshComponent.h
+++ b/Engine/Source/Runtime/Engine/Classes/Components/SkeletalMeshComponent.h
@@ -1239,7 +1239,7 @@ public:
 	/** should call this method if occurred any changes in clothing assets */
 	void RecreateClothingActors();
 	/** add bounding box for cloth */
-	void AddClothingBounds(FBoxSphereBounds& InOutBounds) const;
+	void AddClothingBounds(FBoxSphereBounds& InOutBounds, const FTransform& LocalToWorld) const;
 	/** changes clothing LODs, if clothing LOD is disabled or LODIndex is greater than apex clothing LODs, simulation will be disabled */
 	void SetClothingLOD(int32 LODIndex);
 	/** check whether clothing teleport is needed or not to avoid a weird simulation result */
diff --git a/Engine/Source/Runtime/Engine/Classes/Engine/DemoNetConnection.h b/Engine/Source/Runtime/Engine/Classes/Engine/DemoNetConnection.h
--- a/Engine/Source/Runtime/Engine/Classes/Engine/DemoNetConnection.h
+++ b/Engine/Source/Runtime/Engine/Classes/Engine/DemoNetConnection.h
@@ -32,6 +32,7 @@ public:
 	virtual void FlushNet( bool bIgnoreSimulation = false ) override;
 	virtual void HandleClientPlayer( APlayerController* PC, class UNetConnection* NetConnection ) override;
 	virtual bool ClientHasInitializedLevelFor( const UObject* TestObject ) const override;
+	virtual TSharedPtr<FObjectReplicator> CreateReplicatorForNewActorChannel(UObject* Object);
 
 public:
 
diff --git a/Engine/Source/Runtime/Engine/Classes/Engine/DemoNetDriver.h b/Engine/Source/Runtime/Engine/Classes/Engine/DemoNetDriver.h
--- a/Engine/Source/Runtime/Engine/Classes/Engine/DemoNetDriver.h
+++ b/Engine/Source/Runtime/Engine/Classes/Engine/DemoNetDriver.h
@@ -90,10 +90,13 @@ class ENGINE_API UDemoNetDriver : public UNetDriver
 	/** Public delegate for external systems to be notified when scrubbing is complete. Only called for successful scrub. */
 	FOnGotoTimeMCDelegate OnGotoTimeDelegate;
 
+	bool		IsLoadingCheckpoint() const { return bIsLoadingCheckpoint; }
+
 private:
 	bool		bIsFastForwarding;
 	bool		bIsFastForwardingForCheckpoint;
 	bool		bWasStartStreamingSuccessful;
+	bool		bIsLoadingCheckpoint;
 
 	TArray<FNetworkGUID> NonQueuedGUIDsForScrubbing;
 
@@ -131,6 +134,7 @@ public:
 	virtual bool ShouldQueueBunchesForActorGUID(FNetworkGUID InGUID) const override;
 	virtual FNetworkGUID GetGUIDForActor(const AActor* InActor) const override;
 	virtual AActor* GetActorForGUID(FNetworkGUID InGUID) const override;
+	virtual bool ShouldReceiveRepNotifiesForObject(UObject* Object) const override;
 
 	/** 
 	 * Scrubs playback to the given time. 
diff --git a/Engine/Source/Runtime/Engine/Classes/Engine/GameInstance.h b/Engine/Source/Runtime/Engine/Classes/Engine/GameInstance.h
--- a/Engine/Source/Runtime/Engine/Classes/Engine/GameInstance.h
+++ b/Engine/Source/Runtime/Engine/Classes/Engine/GameInstance.h
@@ -24,6 +24,16 @@ namespace GameInstanceState
 
 class FOnlineSessionSearchResult;
 
+/**
+ * Notification that the client is about to travel to a new URL
+ *
+ * @param PendingURL the travel URL
+ * @param TravelType type of travel that will occur (absolute, relative, etc)
+ * @param bIsSeamlessTravel is traveling seamlessly
+ */
+DECLARE_MULTICAST_DELEGATE_ThreeParams(FOnPreClientTravel, const FString& /*PendingURL*/, ETravelType /*TravelType*/, bool /*bIsSeamlessTravel*/);
+typedef FOnPreClientTravel::FDelegate FOnPreClientTravelDelegate;
+
 /**
  * GameInstance: high-level manager object for an instance of the running game.
  * Spawned at game creation and not destroyed until game instance is shut down.
@@ -53,6 +63,9 @@ protected:
 	UPROPERTY()
 	class UOnlineSession* OnlineSession;
 
+	/** Listeners to PreClientTravel call */
+	FOnPreClientTravel NotifyPreClientTravelDelegates;
+
 public:
 
 	FString PIEMapName;
@@ -230,4 +243,9 @@ public:
 
 	/** Returns true if this instance is for a dedicated server world */
 	bool IsDedicatedServerInstance() const;
+
+	/** Broadcast a notification that travel is occurring */
+	void NotifyPreClientTravel(const FString& PendingURL, ETravelType TravelType, bool bIsSeamlessTravel);
+	/** @return delegate fired when client travel occurs */
+	FOnPreClientTravel& OnNotifyPreClientTravel() { return NotifyPreClientTravelDelegates; }
 };
\ No newline at end of file
diff --git a/Engine/Source/Runtime/Engine/Classes/Engine/NetConnection.h b/Engine/Source/Runtime/Engine/Classes/Engine/NetConnection.h
--- a/Engine/Source/Runtime/Engine/Classes/Engine/NetConnection.h
+++ b/Engine/Source/Runtime/Engine/Classes/Engine/NetConnection.h
@@ -533,6 +533,9 @@ public:
 	ENGINE_API virtual FString RemoteAddressToString() PURE_VIRTUAL(UNetConnection::RemoteAddressToString, return TEXT("Error"););
 	
 	
+	/** Called by UActorChannel. Handles creating a new replicator for an actor */
+	ENGINE_API virtual TSharedPtr<FObjectReplicator> CreateReplicatorForNewActorChannel(UObject* Object);
+
 	// Functions.
 
 	/** Resend any pending acks. */
diff --git a/Engine/Source/Runtime/Engine/Classes/Engine/NetDriver.h b/Engine/Source/Runtime/Engine/Classes/Engine/NetDriver.h
--- a/Engine/Source/Runtime/Engine/Classes/Engine/NetDriver.h
+++ b/Engine/Source/Runtime/Engine/Classes/Engine/NetDriver.h
@@ -583,6 +583,9 @@ public:
 	/** Returns the actor that corresponds to InGUID, if one can be found. */
 	virtual AActor* GetActorForGUID(FNetworkGUID InGUID) const { return nullptr; }
 
+	/** Returns true if RepNotifies should be checked and generated when receiving properties for the given object. */
+	virtual bool ShouldReceiveRepNotifiesForObject(UObject* Object) const { return true; }
+
 	/** Returns the object that manages the list of replicated UObjects. */
 	ENGINE_API FNetworkObjectList& GetNetworkObjectList() { return *NetworkObjects; }
 
diff --git a/Engine/Source/Runtime/Engine/Classes/Engine/NetSerialization.h b/Engine/Source/Runtime/Engine/Classes/Engine/NetSerialization.h
--- a/Engine/Source/Runtime/Engine/Classes/Engine/NetSerialization.h
+++ b/Engine/Source/Runtime/Engine/Classes/Engine/NetSerialization.h
@@ -855,7 +855,8 @@ bool FFastArraySerializer::FastArrayDeltaSerialize( TArray<Type> &Items, FNetDel
 			Type& Item = Items[idx];
 			if (Item.MostRecentArrayReplicationKey < ArrayReplicationKey && Item.MostRecentArrayReplicationKey > BaseReplicationKey)
 			{
-				UE_LOG( LogNetFastTArray, Warning, TEXT( "Adding implicit delete for ElementID: %d. MostRecentArrayReplicationKey: %d. Current Payload: [%d/%d]"), Item.ReplicationID, Item.MostRecentArrayReplicationKey, ArrayReplicationKey, BaseReplicationKey);
+				// This will happen in normal conditions in network replays.
+				UE_LOG( LogNetFastTArray, Log, TEXT( "Adding implicit delete for ElementID: %d. MostRecentArrayReplicationKey: %d. Current Payload: [%d/%d]"), Item.ReplicationID, Item.MostRecentArrayReplicationKey, ArrayReplicationKey, BaseReplicationKey);
 				DeleteIndices.Add(idx);
 			}
 		}
diff --git a/Engine/Source/Runtime/Engine/Classes/GameFramework/Character.h b/Engine/Source/Runtime/Engine/Classes/GameFramework/Character.h
--- a/Engine/Source/Runtime/Engine/Classes/GameFramework/Character.h
+++ b/Engine/Source/Runtime/Engine/Classes/GameFramework/Character.h
@@ -297,6 +297,10 @@ protected:
 	/** Event called after actor's base changes (if SetBase was requested to notify us with bNotifyPawn). */
 	virtual void BaseChange();
 
+	/** CharacterMovement ServerLastTransformUpdateTimeStamp value, replicated to simulated proxies. */
+	UPROPERTY(Replicated)
+	float ReplicatedServerLastTransformUpdateTimeStamp;
+
 	/** CharacterMovement MovementMode (and custom mode) replicated for simulated proxies. Use CharacterMovementComponent::UnpackNetworkMovementMode() to translate it. */
 	UPROPERTY(Replicated)
 	uint8 ReplicatedMovementMode;
@@ -307,6 +311,9 @@ protected:
 
 public:	
 
+	/** Accessor for ReplicatedServerLastTransformUpdateTimeStamp. */
+	FORCEINLINE float GetServerLastTransformUpdateTimeStamp() const { return ReplicatedServerLastTransformUpdateTimeStamp; }
+
 	/** Accessor for BasedMovement */
 	FORCEINLINE const FBasedMovementInfo& GetBasedMovement() const { return BasedMovement; }
 	
diff --git a/Engine/Source/Runtime/Engine/Classes/GameFramework/CharacterMovementComponent.h b/Engine/Source/Runtime/Engine/Classes/GameFramework/CharacterMovementComponent.h
--- a/Engine/Source/Runtime/Engine/Classes/GameFramework/CharacterMovementComponent.h
+++ b/Engine/Source/Runtime/Engine/Classes/GameFramework/CharacterMovementComponent.h
@@ -534,6 +534,10 @@ protected:
 	UPROPERTY()
 	FVector LastUpdateVelocity;
 
+	/** Timestamp when location or rotation last changed during an update. Only valid on the server. */
+	UPROPERTY(Transient)
+	float ServerLastTransformUpdateTimeStamp;
+
 	/** Accumulated impulse to be added next tick. */
 	UPROPERTY()
 	FVector PendingImpulseToApply;
@@ -553,6 +557,9 @@ protected:
 
 public:
 
+	/** Get the value of ServerLastTransformUpdateTimeStamp. */
+	FORCEINLINE float GetServerLastTransformUpdateTimeStamp() const { return ServerLastTransformUpdateTimeStamp;  }
+
 	/**
 	 * Compute remaining time step given remaining time and current iterations.
 	 * The last iteration (limited by MaxSimulationIterations) always returns the remaining time, which may violate MaxSimulationTimeStep.
@@ -2268,8 +2275,14 @@ public:
 	/** Used to track time of last correction */
 	float LastCorrectionTime;
 
+	/** Used to track the timestamp of the last server move. */
+	double SmoothingServerTimeStamp;
+
+	/** Used to track the client time as we try to match the server.*/
+	double SmoothingClientTimeStamp;
+
 	/** Used to track how much time has elapsed since last correction. It can be computed as World->TimeSince(LastCorrectionTime). */
-	DEPRECATED(4.11, "bUseLinearSmoothing will be removed, use LastCorrectionTime instead.")
+	DEPRECATED(4.11, "CurrentSmoothTime will be removed, use LastCorrectionTime instead.")
 	float CurrentSmoothTime;
 
 	/** Used to signify that linear smoothing is desired */
diff --git a/Engine/Source/Runtime/Engine/Classes/Particles/ParticleSystemComponent.h b/Engine/Source/Runtime/Engine/Classes/Particles/ParticleSystemComponent.h
--- a/Engine/Source/Runtime/Engine/Classes/Particles/ParticleSystemComponent.h
+++ b/Engine/Source/Runtime/Engine/Classes/Particles/ParticleSystemComponent.h
@@ -379,8 +379,10 @@ public:
 	/** Time in seconds since we were last considered significant. */
 	float LastSignificantTime;
 
-	/** If this component is considered significant or not. */
-	uint32 bIsSignificant : 1;
+	/** If this component is having it's significance managed by gameplay code. */
+	uint32 bIsManagingSignificance : 1;
+	/** If this component was previously having it's significance managed by gameplay code. Allows us to refresh render data when this changes. */
+	uint32 bWasManagingSignificance : 1;
 	
 private:
 	/** Did we auto attach during activation? Used to determine if we should restore the relative transform during detachment. */
@@ -438,8 +440,6 @@ public:
 	 */
 	UPROPERTY(EditAnywhere, BlueprintReadWrite, Category=Particles)
 	float SecondsBeforeInactive;
-	/** Allow the above or not. This feature can interfere with some other systems. */
-	bool bAllowSecondsBeforeInactive;
 
 private:
 	/** Tracks the time since the last forced UpdateTransform. */
@@ -611,8 +611,24 @@ public:
 	/** Whether this component should have it's significance managed by game code. */
 	bool ShouldManageSignificance()const;
 	/** When the overall significance for the component is changed. */
-	void OnSignificanceChanged(bool bSignificant, bool bApplyToEmitters);
+	void OnSignificanceChanged(bool bSignificant, bool bApplyToEmitters, bool bAsync=false);
+
+	/** Called from game code when the component begins having it's significance managed. */
+	FORCEINLINE void SetManagingSignificance(bool bManageSignificance)
+	{
+		bWasManagingSignificance = bIsManagingSignificance;
+		bIsManagingSignificance = bManageSignificance;
+	}
+
+	/** Returns the approximate distance squared from this component to the passed location. */
+	float GetApproxDistanceSquared(FVector Point)const;
 	
+	/** True if this component can be considered invisible and potentially culled. */
+	bool CanConsiderInvisible()const;
+
+	/** true if this component can be occluded. */
+	bool CanBeOccluded()const;
+
 	void Complete();
 
 	FORCEINLINE const FTransform& GetAsyncComponentToWorld()
diff --git a/Engine/Source/Runtime/Engine/Private/Animation/AnimSequence.cpp b/Engine/Source/Runtime/Engine/Private/Animation/AnimSequence.cpp
--- a/Engine/Source/Runtime/Engine/Private/Animation/AnimSequence.cpp
+++ b/Engine/Source/Runtime/Engine/Private/Animation/AnimSequence.cpp
@@ -4432,7 +4432,11 @@ void UAnimSequence::ValidateCurrentPosition(const FMarkerSyncAnimPosition& Posit
 	checkSlow(MarkerMatchesPosition(this, PreviousMarker.MarkerIndex, Position.PreviousMarkerName));
 	checkSlow(MarkerMatchesPosition(this, NextMarker.MarkerIndex, Position.NextMarkerName));
 
-	CurrentTime = GetCurrentTimeFromMarkers(PreviousMarker, NextMarker, Position.PositionBetweenMarkers);
+	// Only reset position if we found valid markers. Otherwise stay where we are to not pop.
+	if ((PreviousMarker.MarkerIndex != MarkerIndexSpecialValues::AnimationBoundary) && (NextMarker.MarkerIndex != MarkerIndexSpecialValues::AnimationBoundary))
+	{
+		CurrentTime = GetCurrentTimeFromMarkers(PreviousMarker, NextMarker, Position.PositionBetweenMarkers);
+	}
 }
 
 void UAnimSequence::AdvanceMarkerPhaseAsFollower(const FMarkerTickContext& Context, float DeltaRemaining, bool bLooping, float& CurrentTime, FMarkerPair& PreviousMarker, FMarkerPair& NextMarker) const
@@ -4482,8 +4486,13 @@ void UAnimSequence::AdvanceMarkerPhaseAsFollower(const FMarkerTickContext& Conte
 		{
 			check(AuthoredSyncMarkers[NextMarker.MarkerIndex].MarkerName == End.NextMarkerName);
 		}
-		//End Validation
-		CurrentTime = GetCurrentTimeFromMarkers(PreviousMarker, NextMarker, End.PositionBetweenMarkers);
+
+		// End Validation
+		// Only reset position if we found valid markers. Otherwise stay where we are to not pop.
+		if ((PreviousMarker.MarkerIndex != MarkerIndexSpecialValues::AnimationBoundary) && (NextMarker.MarkerIndex != MarkerIndexSpecialValues::AnimationBoundary))
+		{
+			CurrentTime = GetCurrentTimeFromMarkers(PreviousMarker, NextMarker, End.PositionBetweenMarkers);
+		}
 	}
 	else
 	{
@@ -4492,7 +4501,7 @@ void UAnimSequence::AdvanceMarkerPhaseAsFollower(const FMarkerTickContext& Conte
 		{
 			if (PreviousMarker.MarkerIndex == -1)
 			{
-				check(!bLooping || Context.GetMarkerSyncEndPosition().PreviousMarkerName == NAME_None); // shouldnt have an end of anim marker if looping
+				check(!bLooping || Context.GetMarkerSyncEndPosition().PreviousMarkerName == NAME_None); // shouldn't have an end of anim marker if looping
 				CurrentTime = FMath::Max(CurrentTime + DeltaRemaining, 0.f);
 				break;
 			}
@@ -4527,8 +4536,13 @@ void UAnimSequence::AdvanceMarkerPhaseAsFollower(const FMarkerTickContext& Conte
 		{
 			check(AuthoredSyncMarkers[PreviousMarker.MarkerIndex].MarkerName == End.PreviousMarkerName);
 		}
-		//End Validation
-		CurrentTime = GetCurrentTimeFromMarkers(PreviousMarker, NextMarker, End.PositionBetweenMarkers);
+
+		// End Validation
+		// Only reset position if we found valid markers. Otherwise stay where we are to not pop.
+		if ((PreviousMarker.MarkerIndex != MarkerIndexSpecialValues::AnimationBoundary) && (NextMarker.MarkerIndex != MarkerIndexSpecialValues::AnimationBoundary))
+		{
+			CurrentTime = GetCurrentTimeFromMarkers(PreviousMarker, NextMarker, End.PositionBetweenMarkers);
+		}
 	}
 }
 
@@ -4634,6 +4648,31 @@ float UAnimSequence::GetCurrentTimeFromMarkers(FMarkerPair& PrevMarker, FMarkerP
 
 void UAnimSequence::GetMarkerIndicesForPosition(const FMarkerSyncAnimPosition& SyncPosition, bool bLooping, FMarkerPair& OutPrevMarker, FMarkerPair& OutNextMarker, float& OutCurrentTime) const
 {
+	// If we're not looping, assume we're playing a transition and we need to stay where we are.
+	if (!bLooping)
+	{
+		OutPrevMarker.MarkerIndex = -1;
+		OutNextMarker.MarkerIndex = -1;
+
+		for (int32 Idx = 0; Idx<AuthoredSyncMarkers.Num(); Idx++)
+		{
+			const float MarkerTime = AuthoredSyncMarkers[Idx].Time;
+			if (OutCurrentTime > MarkerTime)
+			{
+				OutPrevMarker.MarkerIndex = Idx;
+				OutPrevMarker.TimeToMarker = MarkerTime - OutCurrentTime;
+			}
+			else if (OutCurrentTime < MarkerTime)
+			{
+				OutNextMarker.MarkerIndex = Idx;
+				OutNextMarker.TimeToMarker = MarkerTime - OutCurrentTime;
+				break;
+			}
+		}
+
+		return;
+	}
+
 	if (SyncPosition.PreviousMarkerName == NAME_None)
 	{
 		OutPrevMarker.MarkerIndex = -1;
diff --git a/Engine/Source/Runtime/Engine/Private/Animation/AnimSequenceBase.cpp b/Engine/Source/Runtime/Engine/Private/Animation/AnimSequenceBase.cpp
--- a/Engine/Source/Runtime/Engine/Private/Animation/AnimSequenceBase.cpp
+++ b/Engine/Source/Runtime/Engine/Private/Animation/AnimSequenceBase.cpp
@@ -242,17 +242,25 @@ void UAnimSequenceBase::TickAssetPlayer(FAnimTickRecord& Instance, struct FAnimN
 	else
 	{
 		// Follow the leader
-		if (Instance.bCanUseMarkerSync && Context.CanUseMarkerPosition() && Context.MarkerTickContext.IsMarkerSyncStartValid())
+		if (Instance.bCanUseMarkerSync)
 		{
-			TickByMarkerAsFollower(*Instance.MarkerTickRecord, Context.MarkerTickContext, CurrentTime, PreviousTime, Context.GetLeaderDelta(), Instance.bLooping);
+			if (Context.CanUseMarkerPosition() && Context.MarkerTickContext.IsMarkerSyncStartValid())
+			{
+				TickByMarkerAsFollower(*Instance.MarkerTickRecord, Context.MarkerTickContext, CurrentTime, PreviousTime, Context.GetLeaderDelta(), Instance.bLooping);
+			}
+			else
+			{
+				// If leader is not valid, advance time as normal, do not jump position and pop.
+				FAnimationRuntime::AdvanceTime(Instance.bLooping, MoveDelta, CurrentTime, SequenceLength);
+				UE_LOG(LogAnimMarkerSync, Log, TEXT("Follower (%s) (normal advance)  - PreviousTime (%0.2f), CurrentTime (%0.2f), MoveDelta (%0.2f), Looping (%d) "), *GetName(), PreviousTime, CurrentTime, MoveDelta, Instance.bLooping ? 1 : 0);
+			}
 		}
 		else
 		{
 			CurrentTime = Context.GetAnimationPositionRatio() * SequenceLength;
-			UE_LOG(LogAnimMarkerSync, Log, TEXT("Follower (%s) (normal advance) - PreviousTime (%0.2f), CurrentTime (%0.2f), MoveDelta (%0.2f), Looping (%d) "), *GetName(), PreviousTime, CurrentTime, MoveDelta, Instance.bLooping ? 1 : 0);
+			UE_LOG(LogAnimMarkerSync, Log, TEXT("Follower (%s) (normalized position advance) - PreviousTime (%0.2f), CurrentTime (%0.2f), MoveDelta (%0.2f), Looping (%d) "), *GetName(), PreviousTime, CurrentTime, MoveDelta, Instance.bLooping ? 1 : 0);
 		}
 
-
 		//@TODO: NOTIFIES: Calculate AdvanceType based on what the new delta time is
 
 		if( CurrentTime != PreviousTime )
diff --git a/Engine/Source/Runtime/Engine/Private/Character.cpp b/Engine/Source/Runtime/Engine/Private/Character.cpp
--- a/Engine/Source/Runtime/Engine/Private/Character.cpp
+++ b/Engine/Source/Runtime/Engine/Private/Character.cpp
@@ -1308,6 +1308,7 @@ void ACharacter::PreReplication( IRepChangedPropertyTracker & ChangedPropertyTra
 		DOREPLIFETIME_ACTIVE_OVERRIDE( ACharacter, RepRootMotion, false );
 	}
 
+	ReplicatedServerLastTransformUpdateTimeStamp = CharacterMovement->GetServerLastTransformUpdateTimeStamp();
 	ReplicatedMovementMode = CharacterMovement->PackNetworkMovementMode();	
 	ReplicatedBasedMovement = BasedMovement;
 
@@ -1331,6 +1332,7 @@ void ACharacter::GetLifetimeReplicatedProps( TArray< FLifetimeProperty > & OutLi
 
 	DOREPLIFETIME_CONDITION( ACharacter, RepRootMotion,						COND_SimulatedOnly );
 	DOREPLIFETIME_CONDITION( ACharacter, ReplicatedBasedMovement,			COND_SimulatedOnly );
+	DOREPLIFETIME_CONDITION( ACharacter, ReplicatedServerLastTransformUpdateTimeStamp, COND_SimulatedOnly);
 	DOREPLIFETIME_CONDITION( ACharacter, ReplicatedMovementMode,			COND_SimulatedOnly );
 	DOREPLIFETIME_CONDITION( ACharacter, bIsCrouched,						COND_SimulatedOnly );
 	DOREPLIFETIME_CONDITION( ACharacter, AnimRootMotionTranslationScale,	COND_SimulatedOnly );
diff --git a/Engine/Source/Runtime/Engine/Private/Components/CharacterMovementComponent.cpp b/Engine/Source/Runtime/Engine/Private/Components/CharacterMovementComponent.cpp
--- a/Engine/Source/Runtime/Engine/Private/Components/CharacterMovementComponent.cpp
+++ b/Engine/Source/Runtime/Engine/Private/Components/CharacterMovementComponent.cpp
@@ -1878,8 +1878,10 @@ void UCharacterMovementComponent::PerformMovement(float DeltaSeconds)
 	SaveBaseLocation();
 	UpdateComponentVelocity();
 
+	const bool bHasAuthority = CharacterOwner && CharacterOwner->HasAuthority();
+
 	// If we move we want to avoid a long delay before replication catches up to notice this change, especially if it's throttling our rate.
-	if (UNetDriver::IsAdaptiveNetUpdateFrequencyEnabled() && UpdatedComponent && CharacterOwner && CharacterOwner->HasAuthority())
+	if (bHasAuthority && UNetDriver::IsAdaptiveNetUpdateFrequencyEnabled() && UpdatedComponent)
 	{
 		const UWorld* MyWorld = GetWorld();
 		if (MyWorld && MyWorld->GetTimeSeconds() <= CharacterOwner->NetUpdateTime)
@@ -1899,6 +1901,17 @@ void UCharacterMovementComponent::PerformMovement(float DeltaSeconds)
 		}
 	}
 
+	if (bHasAuthority && UpdatedComponent && GetNetMode() < NM_Client)
+	{
+		const bool bLocationChanged = (UpdatedComponent->GetComponentLocation() != LastUpdateLocation);
+		const bool bRotationChanged = (UpdatedComponent->GetComponentQuat() != LastUpdateRotation);
+		if (bLocationChanged || bRotationChanged)
+		{
+			const UWorld* MyWorld = GetWorld();
+			ServerLastTransformUpdateTimeStamp = MyWorld ? MyWorld->GetTimeSeconds() : 0.f;
+		}
+	}
+
 	LastUpdateLocation = UpdatedComponent ? UpdatedComponent->GetComponentLocation() : FVector::ZeroVector;
 	LastUpdateRotation = UpdatedComponent ? UpdatedComponent->GetComponentQuat() : FQuat::Identity;
 	LastUpdateVelocity = Velocity;
@@ -6190,9 +6203,35 @@ void UCharacterMovementComponent::SmoothCorrection(const FVector& OldLocation, c
 			const FScopedPreventAttachedComponentMove PreventMeshMove(CharacterOwner->GetMesh());
 			UpdatedComponent->SetWorldLocationAndRotation(NewLocation, NewRotation);
 		}
+	
+		// If running ahead, pull back slightly. This will cause the next delta to seem slightly longer, and cause us to lerp to it slightly slower.
+		if (ClientData->SmoothingClientTimeStamp > ClientData->SmoothingServerTimeStamp)
+		{
+			const double OldClientTimeStamp = ClientData->SmoothingClientTimeStamp;
+			ClientData->SmoothingClientTimeStamp = FMath::LerpStable(ClientData->SmoothingServerTimeStamp, OldClientTimeStamp, 0.5);
+
+			UE_LOG(LogCharacterNetSmoothing, VeryVerbose, TEXT("SmoothCorrection: Pull back client from ClientTimeStamp: %.6f to %.6f, ServerTimeStamp: %.6f for %s"),
+				OldClientTimeStamp, ClientData->SmoothingClientTimeStamp, ClientData->SmoothingServerTimeStamp, *GetNameSafe(CharacterOwner));
+		}
+
+		// Using server timestamp lets us know how much time actually elapsed, regardless of packet lag variance.
+		ClientData->SmoothingServerTimeStamp = CharacterOwner->GetServerLastTransformUpdateTimeStamp();
+
+		// Initial update has no delta.
+		if (ClientData->LastCorrectionTime == 0)
+		{
+			ClientData->SmoothingClientTimeStamp = ClientData->SmoothingServerTimeStamp;
+		}
+
+		// Don't let the client fall too far behind or run ahead.
+		ClientData->SmoothingClientTimeStamp = FMath::Clamp<double>(ClientData->SmoothingClientTimeStamp, ClientData->SmoothingServerTimeStamp - ClientData->MaxResponseTime, ClientData->SmoothingServerTimeStamp);
+
+		// Compute actual delta between new server timestamp and client simulation.
+		ClientData->LastCorrectionDelta = ClientData->SmoothingServerTimeStamp - ClientData->SmoothingClientTimeStamp;
+		ClientData->LastCorrectionTime = MyWorld->GetTimeSeconds();
 
-		ClientData->LastCorrectionDelta = FMath::Clamp(MyWorld->TimeSince(ClientData->LastCorrectionTime), 1.0f / 120.0f, 1.0f / 5.0f);
-		ClientData->LastCorrectionTime  = MyWorld->GetTimeSeconds();
+		UE_LOG(LogCharacterNetSmoothing, VeryVerbose, TEXT("SmoothCorrection: WorldTime: %.6f, ServerTimeStamp: %.6f, ClientTimeStamp: %.6f, Delta: %.6f for %s"),
+			MyWorld->GetTimeSeconds(), ClientData->SmoothingServerTimeStamp, ClientData->SmoothingClientTimeStamp, ClientData->LastCorrectionDelta, *GetNameSafe(CharacterOwner));
 
 #if !(UE_BUILD_SHIPPING || UE_BUILD_TEST)
 		if ( CVarNetVisualizeSimulatedCorrections.GetValueOnGameThread() >= 2 )
@@ -6268,21 +6307,53 @@ void UCharacterMovementComponent::SmoothClientPosition_Interpolate(float DeltaSe
 		if (NetworkSmoothingMode == ENetworkSmoothingMode::Linear)
 		{
 			const UWorld* MyWorld = GetWorld();
-			const float CurrentSmoothTime = MyWorld->TimeSince(ClientData->LastCorrectionTime);
 
-			// Linearly interpolate between correction updates
-			const bool bAtEnd = (ClientData->LastCorrectionDelta < SMALL_NUMBER) || (CurrentSmoothTime >= ClientData->LastCorrectionDelta);
-			const float LerpPercent	= bAtEnd ? 1.0f : (CurrentSmoothTime / ClientData->LastCorrectionDelta);
+			// Increment client position.
+			ClientData->SmoothingClientTimeStamp += DeltaSeconds;
 
-			if (LerpPercent >= 1.0f)
+			float LerpPercent = 0.f;
+			const float LerpLimit = 1.15f;
+			const float TargetDelta = ClientData->LastCorrectionDelta;
+			if (TargetDelta > SMALL_NUMBER)
 			{
-				ClientData->MeshTranslationOffset = FVector::ZeroVector;
+				// Don't let the client get too far ahead (happens on spikes). But we do want a buffer for variable network conditions.
+				const float MaxClientTimeAheadPercent = 0.15f;
+				const float MaxTimeAhead = TargetDelta * MaxClientTimeAheadPercent;
+				ClientData->SmoothingClientTimeStamp = FMath::Min<float>(ClientData->SmoothingClientTimeStamp, ClientData->SmoothingServerTimeStamp + MaxTimeAhead);
+
+				// Compute interpolation alpha based on our client position within the server delta. We should take TargetDelta seconds to reach alpha of 1.
+				const float RemainingTime = ClientData->SmoothingServerTimeStamp - ClientData->SmoothingClientTimeStamp;
+				const float CurrentSmoothTime = TargetDelta - RemainingTime;
+				LerpPercent = FMath::Clamp(CurrentSmoothTime / TargetDelta, 0.0f, LerpLimit);
+
+				UE_LOG(LogCharacterNetSmoothing, VeryVerbose, TEXT("Interpolate: WorldTime: %.6f, ServerTimeStamp: %.6f, ClientTimeStamp: %.6f, Elapsed: %.6f, Alpha: %.6f for %s"),
+					MyWorld->GetTimeSeconds(), ClientData->SmoothingServerTimeStamp, ClientData->SmoothingClientTimeStamp, CurrentSmoothTime, LerpPercent, *GetNameSafe(CharacterOwner));
+			}
+			else
+			{
+				LerpPercent = 1.0f;
+			}
+
+			if (LerpPercent >= 1.0f - KINDA_SMALL_NUMBER)
+			{
+				if (Velocity.IsNearlyZero())
+				{
+					ClientData->MeshTranslationOffset = FVector::ZeroVector;
+					ClientData->SmoothingClientTimeStamp = ClientData->SmoothingServerTimeStamp;
+					bNetworkSmoothingComplete = true;
+				}
+				else
+				{
+					// Allow limited forward prediction.
+					ClientData->MeshTranslationOffset = FMath::LerpStable(ClientData->OriginalMeshTranslationOffset, FVector::ZeroVector, LerpPercent);
+					bNetworkSmoothingComplete = (LerpPercent >= LerpLimit);
+				}
+
 				ClientData->MeshRotationOffset = ClientData->MeshRotationTarget;
-				bNetworkSmoothingComplete = true;
 			}
 			else
 			{
-				ClientData->MeshTranslationOffset = FMath::Lerp(ClientData->OriginalMeshTranslationOffset, FVector::ZeroVector, LerpPercent);
+				ClientData->MeshTranslationOffset = FMath::LerpStable(ClientData->OriginalMeshTranslationOffset, FVector::ZeroVector, LerpPercent);
 				ClientData->MeshRotationOffset = FQuat::FastLerp(ClientData->OriginalMeshRotationOffset, ClientData->MeshRotationTarget, LerpPercent).GetNormalized();
 			}
 
@@ -6339,8 +6410,8 @@ void UCharacterMovementComponent::SmoothClientPosition_Interpolate(float DeltaSe
 		}
 
 #if !(UE_BUILD_SHIPPING || UE_BUILD_TEST)
-		UE_LOG(LogCharacterNetSmoothing, Verbose, TEXT("SmoothClientPosition_Interpolate %s: Translation: %s Rotation: %s"),
-			*GetNameSafe(CharacterOwner), *ClientData->MeshTranslationOffset.ToString(), *ClientData->MeshRotationOffset.ToString());
+		//UE_LOG(LogCharacterNetSmoothing, VeryVerbose, TEXT("SmoothClientPosition_Interpolate %s: Translation: %s Rotation: %s"),
+		//	*GetNameSafe(CharacterOwner), *ClientData->MeshTranslationOffset.ToString(), *ClientData->MeshRotationOffset.ToString());
 
 		if ( CVarNetVisualizeSimulatedCorrections.GetValueOnGameThread() >= 1 )
 		{
@@ -8227,6 +8298,8 @@ FNetworkPredictionData_Client_Character::FNetworkPredictionData_Client_Character
 	, MeshRotationTarget(FQuat::Identity)
 	, LastCorrectionDelta(0.f)
 	, LastCorrectionTime(0.f)
+	, SmoothingServerTimeStamp(0.f)
+	, SmoothingClientTimeStamp(0.f)
 	, CurrentSmoothTime(0.f) // Deprecated
 	, bUseLinearSmoothing(false) // Deprecated
 	, MaxSmoothNetUpdateDist(0.f)
diff --git a/Engine/Source/Runtime/Engine/Private/Components/SkeletalMeshComponent.cpp b/Engine/Source/Runtime/Engine/Private/Components/SkeletalMeshComponent.cpp
--- a/Engine/Source/Runtime/Engine/Private/Components/SkeletalMeshComponent.cpp
+++ b/Engine/Source/Runtime/Engine/Private/Components/SkeletalMeshComponent.cpp
@@ -1393,7 +1393,7 @@ FBoxSphereBounds USkeletalMeshComponent::CalcBounds(const FTransform& LocalToWor
 		FBoxSphereBounds NewBounds = CalcMeshBound( RootBoneOffset, bHasValidBodies, LocalToWorld );
 
 #if WITH_APEX_CLOTHING
-		AddClothingBounds(NewBounds);
+		AddClothingBounds(NewBounds, LocalToWorld);
 #endif// #if WITH_APEX_CLOTHING
 
 		bCachedLocalBoundsUpToDate = true;
diff --git a/Engine/Source/Runtime/Engine/Private/DataChannel.cpp b/Engine/Source/Runtime/Engine/Private/DataChannel.cpp
--- a/Engine/Source/Runtime/Engine/Private/DataChannel.cpp
+++ b/Engine/Source/Runtime/Engine/Private/DataChannel.cpp
@@ -2720,8 +2720,7 @@ TSharedRef< FObjectReplicator > & UActorChannel::FindOrCreateReplicator( UObject
 			// Still didn't find one, need to create
 			UE_LOG( LogNetTraffic, Log, TEXT( "Creating Replicator for %s" ), *Obj->GetName() );
 
-			NewReplicator = TSharedRef<FObjectReplicator>(new FObjectReplicator());
-			NewReplicator->InitWithObject( Obj, Connection, true );
+			NewReplicator = Connection->CreateReplicatorForNewActorChannel(Obj);
 		}
 		else
 		{
diff --git a/Engine/Source/Runtime/Engine/Private/DataReplication.cpp b/Engine/Source/Runtime/Engine/Private/DataReplication.cpp
--- a/Engine/Source/Runtime/Engine/Private/DataReplication.cpp
+++ b/Engine/Source/Runtime/Engine/Private/DataReplication.cpp
@@ -155,7 +155,7 @@ bool FObjectReplicator::ValidateAgainstState( const UObject* ObjectState )
 		return false;
 	}
 
-	if ( RepLayout->DiffProperties( RepState, ObjectState, false ) )
+	if ( RepLayout->DiffProperties( &(RepState->RepNotifies), RepState->StaticBuffer.GetData(), ObjectState, false ) )
 	{
 		UE_LOG(LogRep, Warning, TEXT("ValidateAgainstState: Properties changed for %s"), *ObjectState->GetName());
 		return false;
@@ -533,7 +533,8 @@ bool FObjectReplicator::ReceivedBunch( FInBunch& Bunch, const FReplicationFlags&
 			{
 				bool bLocalHasUnmapped = false;
 				// We hijack a non custom delta property to signify we are using FRepLayout to read the entire property block
-				if ( !RepLayout->ReceiveProperties( ObjectClass, RepState, (void*)Object, Bunch, bLocalHasUnmapped ) )
+				const bool bShouldReceiveRepNotifies = Connection->Driver->ShouldReceiveRepNotifiesForObject( Object );
+				if ( !RepLayout->ReceiveProperties( ObjectClass, RepState, (void*)Object, Bunch, bLocalHasUnmapped, bShouldReceiveRepNotifies ) )
 				{
 					UE_LOG(LogRep, Error, TEXT("ReceiveProperties FAILED %s in %s"), *ReplicatedProp->GetName(), *Object->GetFullName());
 					return false;
diff --git a/Engine/Source/Runtime/Engine/Private/DemoNetDriver.cpp b/Engine/Source/Runtime/Engine/Private/DemoNetDriver.cpp
--- a/Engine/Source/Runtime/Engine/Private/DemoNetDriver.cpp
+++ b/Engine/Source/Runtime/Engine/Private/DemoNetDriver.cpp
@@ -307,6 +307,7 @@ bool UDemoNetDriver::InitBase( bool bInitAsClient, FNetworkNotify* InNotify, con
 		bWasStartStreamingSuccessful	= true;
 		SavedReplicatedWorldTimeSeconds	= 0.0f;
 		SavedSecondsToSkip				= 0.0f;
+		bIsLoadingCheckpoint			= false;
 
 		ResetDemoState();
 
@@ -1802,21 +1803,6 @@ void UDemoNetDriver::FinalizeFastForward( const float StartTime )
 	// This must be set before we CallRepNotifies or they might be skipped again
 	bIsFastForwarding = false;
 
-	// Flush all pending RepNotifies that were built up during the fast-forward.
-	if ( ServerConnection != nullptr )
-	{
-		for ( auto& ChannelPair : ServerConnection->ActorChannels )
-		{
-			if ( ChannelPair.Value != NULL )
-			{
-				for ( auto& ReplicatorPair : ChannelPair.Value->ReplicationMap )
-				{
-					ReplicatorPair.Value->CallRepNotifies( true );
-				}
-			}
-		}
-	}
-
 	AGameState* const GameState = World != nullptr ? World->GetGameState() : nullptr;
 
 	// Correct server world time for fast-forwarding after a checkpoint
@@ -1835,6 +1821,48 @@ void UDemoNetDriver::FinalizeFastForward( const float StartTime )
 		GameState->OnRep_ReplicatedWorldTimeSeconds();
 	}
 
+	if (bIsFastForwardingForCheckpoint)
+	{
+		// Make a pass at OnReps for startup actors, since they were skipped during checkpoint loading.
+		// At this point the shadow state of these actors should be the actual state from before the checkpoint,
+		// and the current state is the CDO state evolved by any changes that occurred during checkpoint loading and fast-forwarding.
+		for (UChannel* Channel : ServerConnection->OpenChannels)
+		{
+			UActorChannel* const ActorChannel = Cast<UActorChannel>(Channel);
+			if (ActorChannel == nullptr)
+			{
+				continue;
+			}
+
+			const AActor* const Actor = ActorChannel->GetActor();
+			if (Actor == nullptr)
+			{
+				continue;
+			}
+
+			const FObjectReplicator* const ActorReplicator = ActorChannel->ActorReplicator;
+			if (Actor->IsNetStartupActor() && ActorReplicator)
+			{
+				ActorReplicator->RepLayout->DiffProperties(&(ActorReplicator->RepState->RepNotifies), ActorReplicator->RepState->StaticBuffer.GetData(), Actor, true);
+			}
+		}
+	}
+
+	// Flush all pending RepNotifies that were built up during the fast-forward.
+	if ( ServerConnection != nullptr )
+	{
+		for ( auto& ChannelPair : ServerConnection->ActorChannels )
+		{
+			if ( ChannelPair.Value != NULL )
+			{
+				for ( auto& ReplicatorPair : ChannelPair.Value->ReplicationMap )
+				{
+					ReplicatorPair.Value->CallRepNotifies( true );
+				}
+			}
+		}
+	}
+
 	// We may have been fast-forwarding immediately after loading a checkpoint
 	// for fine-grained scrubbing. If so, at this point we are no longer loading a checkpoint.
 	bIsFastForwardingForCheckpoint = false;
@@ -2000,6 +2028,8 @@ void UDemoNetDriver::LoadCheckpoint( FArchive* GotoCheckpointArchive, int64 Goto
 
 	FNetworkReplayDelegates::OnPreScrub.Broadcast(GetWorld());
 
+	bIsLoadingCheckpoint = true;
+
 #if 1
 	// Destroy all non startup actors. They will get restored with the checkpoint
 	for ( FActorIterator It( GetWorld() ); It; ++It )
@@ -2156,6 +2186,7 @@ void UDemoNetDriver::LoadCheckpoint( FArchive* GotoCheckpointArchive, int64 Goto
 	ReadDemoFrame( GotoCheckpointArchive );
 
 	bDemoPlaybackDone = false;
+	bIsLoadingCheckpoint = false;
 
 	// Save the replicated server time here
 	if (World != nullptr)
@@ -2231,6 +2262,23 @@ AActor* UDemoNetDriver::GetActorForGUID(FNetworkGUID InGUID) const
 
 }
 
+bool UDemoNetDriver::ShouldReceiveRepNotifiesForObject(UObject* Object) const
+{
+	// Return false for startup actors during checkpoint loading, since they are
+	// not destroyed and re-created like dynamic actors. Startup actors will
+	// have their properties diffed and RepNotifies called after the checkpoint is loaded.
+
+	if (!bIsLoadingCheckpoint && !bIsFastForwardingForCheckpoint)
+	{
+		return true;
+	}
+
+	const AActor* const Actor = Cast<AActor>(Object);
+	const bool bIsStartupActor = Actor != nullptr && Actor->IsNetStartupActor();
+
+	return !bIsStartupActor;
+}
+
 void UDemoNetDriver::AddNonQueuedActorForScrubbing(AActor const* Actor)
 {
 	UActorChannel const* const* const FoundChannel = ServerConnection->ActorChannels.Find(Actor);
@@ -2404,6 +2452,30 @@ bool UDemoNetConnection::ClientHasInitializedLevelFor(const UObject* TestObject)
 	return ( GetDriver()->DemoFrameNum > 2 || Super::ClientHasInitializedLevelFor( TestObject ) );
 }
 
+TSharedPtr<FObjectReplicator> UDemoNetConnection::CreateReplicatorForNewActorChannel(UObject* Object)
+{
+	TSharedPtr<FObjectReplicator> NewReplicator = MakeShareable(new FObjectReplicator());
+
+	// To handle rewinding net startup actors in replays properly, we need to
+	// initialize the shadow state with the object's current state.
+	// Afterwards, we will copy the CDO state to object's current state with repnotifies disabled.
+	UDemoNetDriver* Driver = GetDriver();
+	AActor* Actor = Cast<AActor>(Object);
+
+	const bool bIsCheckpointStartupActor = Driver && Driver->IsLoadingCheckpoint() && Actor && Actor->IsNetStartupActor();
+	const bool bUseDefaultState = !bIsCheckpointStartupActor;
+
+	NewReplicator->InitWithObject(Object, this, bUseDefaultState);
+
+	// Now that the shadow state is initialized, copy the CDO state into the actor state.
+	if (bIsCheckpointStartupActor && NewReplicator->RepLayout.IsValid() && Object->GetClass())
+	{
+		NewReplicator->RepLayout->DiffProperties(nullptr, Object, Object->GetClass()->GetDefaultObject(), true);
+	}
+
+	return NewReplicator;
+}
+
 bool UDemoNetDriver::IsLevelInitializedForActor( const AActor* InActor, const UNetConnection* InConnection ) const
 {
 	return ( DemoFrameNum > 2 || Super::IsLevelInitializedForActor( InActor, InConnection ) );
diff --git a/Engine/Source/Runtime/Engine/Private/GameEngine.cpp b/Engine/Source/Runtime/Engine/Private/GameEngine.cpp
--- a/Engine/Source/Runtime/Engine/Private/GameEngine.cpp
+++ b/Engine/Source/Runtime/Engine/Private/GameEngine.cpp
@@ -542,6 +542,7 @@ void UGameEngine::Init(IEngineLoop* InEngineLoop)
 		bool bWindowAlreadyExists = GameViewportWindow.IsValid();
 		if (!bWindowAlreadyExists)
 		{
+			UE_LOG(LogEngine, Log, TEXT("GameWindow did not exist.  Was created"));
 			GameViewportWindow = CreateGameWindow();
 		}
 
diff --git a/Engine/Source/Runtime/Engine/Private/GameInstance.cpp b/Engine/Source/Runtime/Engine/Private/GameInstance.cpp
--- a/Engine/Source/Runtime/Engine/Private/GameInstance.cpp
+++ b/Engine/Source/Runtime/Engine/Private/GameInstance.cpp
@@ -906,3 +906,9 @@ bool UGameInstance::IsDedicatedServerInstance() const
 		return WorldContext ? WorldContext->RunAsDedicated : false;
 	}
 }
+
+void UGameInstance::NotifyPreClientTravel(const FString& PendingURL, ETravelType TravelType, bool bIsSeamlessTravel)
+{
+	OnNotifyPreClientTravel().Broadcast(PendingURL, TravelType, bIsSeamlessTravel);
+}
+
diff --git a/Engine/Source/Runtime/Engine/Private/NetConnection.cpp b/Engine/Source/Runtime/Engine/Private/NetConnection.cpp
--- a/Engine/Source/Runtime/Engine/Private/NetConnection.cpp
+++ b/Engine/Source/Runtime/Engine/Private/NetConnection.cpp
@@ -1024,7 +1024,8 @@ void UNetConnection::ReceivedPacket( FBitReader& Reader )
 					}
 					else
 					{
-						UE_LOG( LogNetTraffic, Warning, TEXT( "      Received unreliable bunch before open (Channel %d Current Sequence %i)" ), Bunch.ChIndex, InReliable[Bunch.ChIndex] );
+						// Simply a log (not a warning, since this can happen under normal conditions, like from a re-join, etc)
+						UE_LOG( LogNetTraffic, Log, TEXT( "      Received unreliable bunch before open (Channel %d Current Sequence %i)" ), Bunch.ChIndex, InReliable[Bunch.ChIndex] );
 					}
 
 					// Since we won't be processing this packet, don't ack it
@@ -1189,6 +1190,13 @@ void UNetConnection::PopLastStart()
 	NETWORK_PROFILER(GNetworkProfiler.PopSendBunch(this));
 }
 
+TSharedPtr<FObjectReplicator> UNetConnection::CreateReplicatorForNewActorChannel(UObject* Object)
+{
+	TSharedPtr<FObjectReplicator> NewReplicator = MakeShareable(new FObjectReplicator());
+	NewReplicator->InitWithObject( Object, this, true );
+	return NewReplicator;
+}
+
 void UNetConnection::PurgeAcks()
 {
 	for ( int32 i = 0; i < ResendAcks.Num(); i++ )
