Copying //UE4/Dev-Core to //UE4/Main

==========================
MAJOR FEATURES + CHANGES
==========================

Change 2860239 on 2016/02/09 by [REDACTED_PERSON][REDACTED_HANDLE]

	Strip unnecessary UMaterialExpressions from cooked packages.

Change 2872941 on 2016/02/18 by [REDACTED_PERSON][REDACTED_HANDLE]

	Fixing missing materials in cooked builds:

	- FExpressionInput::Expression can now only be used in editor builds
	- Added FExpressionInput::ExpressionName which can be used in cooked builds to check if the input was connected
	- FExpressionInput and derived structs will now use native serialization
	- Keeping a few more material expressions types in cooked builds

Change 2863269 on 2016/02/11 by [REDACTED_PERSON][REDACTED_HANDLE]

	MallocBinned2 is now internally threadsafe.

Change 2866771 on 2016/02/15 by [REDACTED_PERSON][REDACTED_HANDLE]

	Temp HACK for slow cook times (waiting for Anim Team to fix it properly)

Change 2873372 on 2016/02/19 by [REDACTED_PERSON][REDACTED_HANDLE]

	Print the max number of UObjects when someone tries to allocate too many of them.

Change 2861582 on 2016/02/10 by [REDACTED_PERSON][REDACTED_HANDLE]

	Fix for illegal access of StaticClass() at module shutdown.

Change 2862906 on 2016/02/11 by [REDACTED_PERSON][REDACTED_HANDLE]

	Per-thread caches and global recycler.

Change 2863366 on 2016/02/11 by [REDACTED_PERSON][REDACTED_HANDLE]

	Fixing ReferenceChainSearch not using the right garbage collection keep flags.
	Added gc.FindStaleClusters console command to list any stale clusters.

Change 2864959 on 2016/02/12 by [REDACTED_PERSON][REDACTED_HANDLE]

	cvar for enabling/disabling MallocBinned2's per-thread free caches.

Change 2865192 on 2016/02/12 by [REDACTED_PERSON][REDACTED_HANDLE]

	Bundle sizes and recycle bundle counts are now configurable via cvars.

Change 2867187 on 2016/02/15 by [REDACTED_PERSON][REDACTED_HANDLE]

	Fix for double-buffered free batch recycling.

Change 2868533 on 2016/02/16 by [REDACTED_PERSON][REDACTED_HANDLE]

	FBundle class for better bundle management and clarity.

Change 2870302 on 2016/02/17 by [REDACTED_PERSON][REDACTED_HANDLE]

	Visualizer fixes for VS2015 and VS2013 Update 5.
	Quotes removed from invalid value nodes on PS4.

Change 2870778 on 2016/02/17 by [REDACTED_PERSON][REDACTED_HANDLE]

	FMalloc::FlushCurrentThreadCache for returning thread-local cached memory back to the global allocator.

Change 2873917 on 2016/02/19 by [REDACTED_PERSON][REDACTED_HANDLE]

	Fixing UHT not respecting #if WITH_EDITORONLY_DATA properties inside of #if !CPP block.
	Removed temp hack from MaterialExpression.IO related to the bugfix.

Change 2871812 on 2016/02/18 by [REDACTED_PERSON][REDACTED_HANDLE]

	Fix crash in UnrealPak when signing is enabled

diff --git a/Engine/Plugins/Media/MediaPlayerEditor/Source/MediaPlayerEditor/Private/MediaPlayerEditorModule.cpp b/Engine/Plugins/Media/MediaPlayerEditor/Source/MediaPlayerEditor/Private/MediaPlayerEditorModule.cpp
--- a/Engine/Plugins/Media/MediaPlayerEditor/Source/MediaPlayerEditor/Private/MediaPlayerEditorModule.cpp
+++ b/Engine/Plugins/Media/MediaPlayerEditor/Source/MediaPlayerEditor/Private/MediaPlayerEditorModule.cpp
@@ -104,18 +104,21 @@ protected:
 	void RegisterCustomizations()
 	{
 		FPropertyEditorModule& PropertyModule = FModuleManager::LoadModuleChecked<FPropertyEditorModule>("PropertyEditor");
-		PropertyModule.RegisterCustomClassLayout(UMediaPlayer::StaticClass()->GetFName(), FOnGetDetailCustomizationInstance::CreateStatic(&FMediaPlayerCustomization::MakeInstance));
-		PropertyModule.RegisterCustomClassLayout(UMediaSoundWave::StaticClass()->GetFName(), FOnGetDetailCustomizationInstance::CreateStatic(&FMediaSoundWaveCustomization::MakeInstance));
-		PropertyModule.RegisterCustomClassLayout(UMediaTexture::StaticClass()->GetFName(), FOnGetDetailCustomizationInstance::CreateStatic(&FMediaTextureCustomization::MakeInstance));
+		MediaPlayerName    = UMediaPlayer   ::StaticClass()->GetFName();
+		MediaSoundWaveName = UMediaSoundWave::StaticClass()->GetFName();
+		MediaTextureName   = UMediaTexture  ::StaticClass()->GetFName();
+		PropertyModule.RegisterCustomClassLayout(MediaPlayerName, FOnGetDetailCustomizationInstance::CreateStatic(&FMediaPlayerCustomization::MakeInstance));
+		PropertyModule.RegisterCustomClassLayout(MediaSoundWaveName, FOnGetDetailCustomizationInstance::CreateStatic(&FMediaSoundWaveCustomization::MakeInstance));
+		PropertyModule.RegisterCustomClassLayout(MediaTextureName, FOnGetDetailCustomizationInstance::CreateStatic(&FMediaTextureCustomization::MakeInstance));
 	}
 
 	/** Unregister details view customizations. */
 	void UnregisterCustomizations()
 	{
 		FPropertyEditorModule& PropertyModule = FModuleManager::LoadModuleChecked<FPropertyEditorModule>("PropertyEditor");
-		PropertyModule.UnregisterCustomClassLayout(UMediaPlayer::StaticClass()->GetFName());
-		PropertyModule.UnregisterCustomClassLayout(UMediaSoundWave::StaticClass()->GetFName());
-		PropertyModule.UnregisterCustomClassLayout(UMediaTexture::StaticClass()->GetFName());
+		PropertyModule.UnregisterCustomClassLayout(MediaPlayerName);
+		PropertyModule.UnregisterCustomClassLayout(MediaSoundWaveName);
+		PropertyModule.UnregisterCustomClassLayout(MediaTextureName);
 	}
 
 protected:
@@ -164,6 +167,11 @@ private:
 
 	/** Holds the tool bar extensibility manager. */
 	TSharedPtr<FExtensibilityManager> ToolBarExtensibilityManager;
+
+	/** Class names */
+	FName MediaPlayerName;
+	FName MediaSoundWaveName;
+	FName MediaTextureName;
 };
 
 
diff --git a/Engine/Source/Developer/SimplygonSwarm/Private/SimplygonRESTClient.cpp b/Engine/Source/Developer/SimplygonSwarm/Private/SimplygonRESTClient.cpp
--- a/Engine/Source/Developer/SimplygonSwarm/Private/SimplygonRESTClient.cpp
+++ b/Engine/Source/Developer/SimplygonSwarm/Private/SimplygonRESTClient.cpp
@@ -453,7 +453,7 @@ void FSimplygonRESTClient::DownloadAsset_Response(FHttpRequestPtr Request, FHttp
 		
 		if (!FFileHelper::SaveArrayToFile(data, *SwarmTask->TaskData.OutputZipFilePath))
 		{
-			UE_LOG(LogSimplygonRESTClient, Log, TEXT("Unable to save file %S"), *SwarmTask->TaskData.OutputZipFilePath);
+			UE_LOG(LogSimplygonRESTClient, Log, TEXT("Unable to save file %s"), *SwarmTask->TaskData.OutputZipFilePath);
 			SwarmTask->SetState(SRS_FAILED);
 			
 		}
diff --git a/Engine/Source/Programs/UnrealHeaderTool/Private/CodeGenerator.cpp b/Engine/Source/Programs/UnrealHeaderTool/Private/CodeGenerator.cpp
--- a/Engine/Source/Programs/UnrealHeaderTool/Private/CodeGenerator.cpp
+++ b/Engine/Source/Programs/UnrealHeaderTool/Private/CodeGenerator.cpp
@@ -29,7 +29,7 @@ static bool bVerifyContents = false;
 static const bool bMultiLineUFUNCTION = true;
 static const bool bMultiLineUPROPERTY = true;
 
-static TSharedRef<FUnrealSourceFile> PerformInitialParseOnHeader(UPackage* InParent, const FString& FileName, EObjectFlags Flags, const TCHAR* Buffer, FUHTMakefile& UHTMakefile);
+static TSharedRef<FUnrealSourceFile> PerformInitialParseOnHeader(UPackage* InParent, const TCHAR* FileName, EObjectFlags Flags, const TCHAR* Buffer, FUHTMakefile& UHTMakefile);
 
 FCompilerMetadataManager GScriptHelper;
 
@@ -3144,7 +3144,8 @@ void FNativeClassHeaderGenerator::ExportMirrorsForNoexportStructs(const TArray<U
 		UScriptStruct* Struct = NativeStructs[i];
 
 		// Export struct.
-		HeaderOutput.Logf(TEXT("%sstruct %s"), FCString::Tab(TextIndent), NameLookupCPP.GetNameCPP(Struct));
+		const TCHAR* StructName = NameLookupCPP.GetNameCPP(Struct);
+		HeaderOutput.Logf(TEXT("%sstruct %s"), FCString::Tab(TextIndent), StructName);
 		if (Struct->GetSuperStruct() != NULL)
 		{
 			HeaderOutput.Logf(TEXT(" : public %s"), NameLookupCPP.GetNameCPP(Struct->GetSuperStruct()));
@@ -5605,7 +5606,7 @@ ECompilationResult::Type PreparseModules(FUHTMakefile& UHTMakefile, const FStrin
 						FError::Throwf(TEXT("UnrealHeaderTool was unable to load source file '%s'"), *FullFilename);
 					}
 
-					TSharedRef<FUnrealSourceFile> UnrealSourceFile = PerformInitialParseOnHeader(Package, RawFilename, RF_Public | RF_Standalone, *HeaderFile, UHTMakefile);
+					TSharedRef<FUnrealSourceFile> UnrealSourceFile = PerformInitialParseOnHeader(Package, *RawFilename, RF_Public | RF_Standalone, *HeaderFile, UHTMakefile);
 					FUnrealSourceFile* UnrealSourceFilePtr = &UnrealSourceFile.Get();
 					TArray<UClass*> DefinedClasses = UnrealSourceFile->GetDefinedClasses();
 					for (UClass* DefinedClass : DefinedClasses)
@@ -6019,7 +6020,7 @@ UClass* ProcessParsedClass(bool bClassIsAnInterface, TArray<FHeaderProvider> &De
 }
 
 
-TSharedRef<FUnrealSourceFile> PerformInitialParseOnHeader(UPackage* InParent, const FString& FileName, EObjectFlags Flags, const TCHAR* Buffer, FUHTMakefile& UHTMakefile)
+TSharedRef<FUnrealSourceFile> PerformInitialParseOnHeader(UPackage* InParent, const TCHAR* FileName, EObjectFlags Flags, const TCHAR* Buffer, FUHTMakefile& UHTMakefile)
 {
 	const TCHAR* InBuffer = Buffer;
 
diff --git a/Engine/Source/Programs/UnrealHeaderTool/Private/HeaderParser.cpp b/Engine/Source/Programs/UnrealHeaderTool/Private/HeaderParser.cpp
--- a/Engine/Source/Programs/UnrealHeaderTool/Private/HeaderParser.cpp
+++ b/Engine/Source/Programs/UnrealHeaderTool/Private/HeaderParser.cpp
@@ -7303,10 +7303,10 @@ void FHeaderParser::SimplifiedClassParse(const TCHAR* InBuffer, TArray<FSimplifi
 		if( bIf || FParse::Command(&Str,TEXT("#ifdef")) || FParse::Command(&Str,TEXT("#ifndef")) )
 		{
 			FStringOutputDevice TextDumpDummy;
-			int32 PreprocessorNest = 1;
 			FStringOutputDevice* Target = NULL;
 			FStringOutputDevice* SpacerTarget = NULL;
-			bool bKeepPreprocessorDirectives = true;
+			TArray<bool, TInlineAllocator<8>> KeepPreprocessorDirectiveStack;
+			KeepPreprocessorDirectiveStack.Push(true);
 			bool bNotCPP = false;
 			bool bCPP = false;
 			bool bUnknownDirective = false;
@@ -7320,14 +7320,14 @@ void FHeaderParser::SimplifiedClassParse(const TCHAR* InBuffer, TArray<FSimplifi
 			else if( bIf && FParse::Command(&Str,TEXT("!CPP")) )
 			{
 				Target = &ClassHeaderTextStrippedOfCppText;
-				bKeepPreprocessorDirectives = false;
+				KeepPreprocessorDirectiveStack.Top() = false;
 				bNotCPP = true;
 			}
 #if WITH_HOT_RELOAD_CTORS
 			else if (bIf && FParse::Command(&Str, TEXT("WITH_HOT_RELOAD")))
 			{
 				Target = &ClassHeaderTextStrippedOfCppText;
-				bKeepPreprocessorDirectives = false;
+				KeepPreprocessorDirectiveStack.Top() = false;
 			}
 #endif // WITH_HOT_RELOAD_CTORS
 			else if (bIf && (FParse::Command(&Str,TEXT("WITH_EDITORONLY_DATA")) || FParse::Command(&Str,TEXT("WITH_EDITOR"))))
@@ -7343,7 +7343,7 @@ void FHeaderParser::SimplifiedClassParse(const TCHAR* InBuffer, TArray<FSimplifi
 				SpacerTarget = &ClassHeaderTextStrippedOfCppText;
 			}
 
-			if (bKeepPreprocessorDirectives)
+			if (KeepPreprocessorDirectiveStack.Top())
 			{
 				Target->Logf( TEXT("%s\r\n"), *StrLine );
 			}
@@ -7358,7 +7358,7 @@ void FHeaderParser::SimplifiedClassParse(const TCHAR* InBuffer, TArray<FSimplifi
 				SpacerTarget->Logf( TEXT("\r\n") );
 			}
 
-			while ((PreprocessorNest > 0) && FParse::Line(&Buffer, StrLine, 1))
+			while ((KeepPreprocessorDirectiveStack.Num() > 0) && FParse::Line(&Buffer, StrLine, 1))
 			{
 				if (SpacerTarget != NULL)
 				{
@@ -7368,22 +7368,29 @@ void FHeaderParser::SimplifiedClassParse(const TCHAR* InBuffer, TArray<FSimplifi
 
 				CurrentLine++;
 				Str = *StrLine;
+				bool bKeepPreprocessorDirective = KeepPreprocessorDirectiveStack.Top();
 				bool bIsPrep = false;
 				if( FParse::Command(&Str,TEXT("#endif")) )
 				{
-					PreprocessorNest--;
+					KeepPreprocessorDirectiveStack.Pop();
 					bIsPrep = true;
 				}
 				else if( FParse::Command(&Str,TEXT("#if")) || FParse::Command(&Str,TEXT("#ifdef")) || FParse::Command(&Str,TEXT("#ifndef")) )
 				{
-					PreprocessorNest++;
+					bKeepPreprocessorDirective = FParse::Command(&Str, TEXT("WITH_EDITORONLY_DATA")) || FParse::Command(&Str, TEXT("WITH_EDITOR"));
+					KeepPreprocessorDirectiveStack.Push(bKeepPreprocessorDirective);
+					bIsPrep = true;
+				}
+				else if (FParse::Command(&Str, TEXT("#ifndef")))
+				{
+					KeepPreprocessorDirectiveStack.Push(false);
 					bIsPrep = true;
 				}
 				else if( FParse::Command(&Str,TEXT("#elif")) )
 				{
 					bIsPrep = true;
 				}
-				else if(PreprocessorNest == 1 && FParse::Command(&Str,TEXT("#else")))
+				else if (KeepPreprocessorDirectiveStack.Num() == 1 && FParse::Command(&Str, TEXT("#else")))
 				{
 					if (!bUnknownDirective)
 					{
@@ -7396,12 +7403,12 @@ void FHeaderParser::SimplifiedClassParse(const TCHAR* InBuffer, TArray<FSimplifi
 						{
 							Target = &TextDumpDummy;
 							SpacerTarget = &ClassHeaderTextStrippedOfCppText;
-							bKeepPreprocessorDirectives = true;
+							KeepPreprocessorDirectiveStack.Top() = true;
 							StrLine = TEXT("#if CPP\r\n");
 						}
 						else
 						{
-							bKeepPreprocessorDirectives = false;
+							KeepPreprocessorDirectiveStack.Top() = false;
 							Target = &ClassHeaderTextStrippedOfCppText;
 							SpacerTarget = NULL;
 						}
@@ -7409,7 +7416,7 @@ void FHeaderParser::SimplifiedClassParse(const TCHAR* InBuffer, TArray<FSimplifi
 					bIsPrep = true;
 				}
 
-				if (bKeepPreprocessorDirectives || !bIsPrep)
+				if (bKeepPreprocessorDirective || !bIsPrep)
 				{
 					Target->Logf( TEXT("%s\r\n"), *StrLine );
 				}
diff --git a/Engine/Source/Programs/UnrealPak/Private/UnrealPak.cpp b/Engine/Source/Programs/UnrealPak/Private/UnrealPak.cpp
--- a/Engine/Source/Programs/UnrealPak/Private/UnrealPak.cpp
+++ b/Engine/Source/Programs/UnrealPak/Private/UnrealPak.cpp
@@ -758,10 +758,12 @@ bool CreatePakFile(const TCHAR* Filename, TArray<FPakInputPair>& FilesToAdd, con
 	FCompressedFileBuffer CompressedFileBuffer;
 
 	uint8* PaddingBuffer = nullptr;
+	int64 PaddingBufferSize = 0;
 	if (CmdLineParameters.PatchFilePadAlign > 0)
 	{
-		PaddingBuffer = (uint8*)FMemory::Malloc(CmdLineParameters.PatchFilePadAlign);
-		FMemory::Memset(PaddingBuffer, 0, CmdLineParameters.PatchFilePadAlign);
+		PaddingBufferSize = CmdLineParameters.PatchFilePadAlign;
+		PaddingBuffer = (uint8*)FMemory::Malloc(PaddingBufferSize);
+		FMemory::Memset(PaddingBuffer, 0, PaddingBufferSize);
 	}
 
 	for (int32 FileIndex = 0; FileIndex < FilesToAdd.Num(); FileIndex++)
@@ -803,8 +805,29 @@ bool CreatePakFile(const TCHAR* Filename, TArray<FPakInputPair>& FilesToAdd, con
 			if ((NewEntryOffset / CmdLineParameters.FileSystemBlockSize) != ((NewEntryOffset+RealFileSize) / CmdLineParameters.FileSystemBlockSize))
 			{
 				//File crosses a block boundary, so align it to the beginning of the next boundary
+				int64 OldOffset = NewEntryOffset;
 				NewEntryOffset = AlignArbitrary(NewEntryOffset, CmdLineParameters.FileSystemBlockSize);
-				PakFileHandle->Seek(NewEntryOffset);
+				int64 PaddingRequired = NewEntryOffset - OldOffset;
+				
+				if (PaddingRequired > 0)
+				{
+					// If we don't already have a padding buffer, create one
+					if (PaddingBuffer == nullptr)
+					{
+						PaddingBufferSize = 64 * 1024;
+						PaddingBuffer = (uint8*)FMemory::Malloc(PaddingBufferSize);
+						FMemory::Memset(PaddingBuffer, 0, PaddingBufferSize);
+					}
+
+					while (PaddingRequired > 0)
+					{
+						int64 AmountToWrite = FMath::Min(PaddingRequired, PaddingBufferSize);
+						PakFileHandle->Serialize(PaddingBuffer, AmountToWrite);
+						PaddingRequired -= AmountToWrite;
+					}
+					
+					check(PakFileHandle->Tell() == NewEntryOffset);
+				}
 			}
 		}
 
@@ -815,7 +838,7 @@ bool CreatePakFile(const TCHAR* Filename, TArray<FPakInputPair>& FilesToAdd, con
 			NewEntryOffset = AlignArbitrary(NewEntryOffset, CmdLineParameters.PatchFilePadAlign);
 			int64 CurrentLoc = PakFileHandle->Tell();
 			int64 PaddingSize = NewEntryOffset - CurrentLoc;
-			check(PaddingSize <= CmdLineParameters.PatchFilePadAlign);
+			check(PaddingSize <= PaddingBufferSize);
 
 			//have to pad manually with 0's.  File locations skipped by Seek and never written are uninitialized which would defeat the whole purpose
 			//of padding for certain platforms patch diffing systems.
diff --git a/Engine/Source/Runtime/Core/Private/HAL/MallocBinned2.cpp b/Engine/Source/Runtime/Core/Private/HAL/MallocBinned2.cpp
--- a/Engine/Source/Runtime/Core/Private/HAL/MallocBinned2.cpp
+++ b/Engine/Source/Runtime/Core/Private/HAL/MallocBinned2.cpp
@@ -10,9 +10,31 @@
 #include "MemoryMisc.h"
 #include "HAL/PlatformAtomics.h"
 
-// Alignment of OS-allocated pointer - pool-allocated pointers will have a non-aligned pointer
-#define BINNED2_LARGE_ALLOC 65536
-
+#define BINNED2_LARGE_ALLOC         65536 // Alignment of OS-allocated pointer - pool-allocated pointers will have a non-aligned pointer
+#define BINNED2_MINIMUM_ALIGNMENT   16    // Alignment of blocks
+#define BINNED2_MAX_SMALL_POOL_SIZE 32768 // Maximum block size in GMallocBinned2SmallBlockSizes
+
+
+static int32 GMallocBinned2PerThreadCaches = 1;
+static FAutoConsoleVariableRef GMallocBinned2PerThreadCachesCVar(
+	TEXT("MallocBinned2.PerThreadCaches"),
+	GMallocBinned2PerThreadCaches,
+	TEXT("Enables per-thread caches of small (<= 32768 byte) allocations from FMallocBinned2")
+);
+
+static int32 GMallocBinned2BundleSize = BINNED2_LARGE_ALLOC * 2;
+static FAutoConsoleVariableRef GMallocBinned2BundleSizeCVar(
+	TEXT("MallocBinned2.BundleSize"),
+	GMallocBinned2BundleSize,
+	TEXT("Size in bytes of per-block bundles used in the recycling process")
+);
+
+static int32 GMallocBinned2MaxBundlesBeforeRecycle = 4;
+static FAutoConsoleVariableRef GMallocBinned2MaxBundlesBeforeRecycleCVar(
+	TEXT("MallocBinned2.BundleRecycleCount"),
+	GMallocBinned2MaxBundlesBeforeRecycle,
+	TEXT("Number of freed bundles in the global recycler before it returns them to the system, per-block size")
+);
 
 // Block sizes are based around getting the maximum amount of allocations per pool, with as little alignment waste as possible.
 // Block sizes should be close to even divisors of the system page size, and well distributed.
@@ -27,12 +49,61 @@ const uint32 GMallocBinned2SmallBlockSizes[] =
 	32768
 };
 
+// Mapping of sizes to small table indices
+uint32 GMallocBinned2MemSizeToIndex[BINNED2_MAX_SMALL_POOL_SIZE];
+
 /** Information about a piece of free memory. */
 struct FMallocBinned2::FFreeBlock
 {
-	void*  NextFreeBlock; // Next free block in another pool
-	uint32 NumFreeBlocks; // Number of consecutive free blocks here, at least 1.
-	uint32 BlockSize;     // Size of the blocks that this list points to
+	FFreeBlock(uint32 InPageSize, uint32 InBlockSize)
+		: NextFreeBlock       (nullptr)
+		, NumFreeBlocks       (InPageSize / InBlockSize - 1)
+		, BlockSize           (InBlockSize)
+		, bShortBlockAllocated(false)
+	{
+	}
+
+	uint32 GetNumFreeRegularBlocks() const
+	{
+		return NumFreeBlocks;
+	}
+
+	uint32 GetShortBlockSize() const
+	{
+		return BlockSize - BINNED2_MINIMUM_ALIGNMENT;
+	}
+
+	bool HasFreeShortBlock() const
+	{
+		return !bShortBlockAllocated;
+	}
+
+	void* AllocateRegularBlock()
+	{
+		--NumFreeBlocks;
+		void* Result = (uint8*)this + (NumFreeBlocks + IsAligned(this, BINNED2_LARGE_ALLOC)) * BlockSize;
+
+		return Result;
+	}
+
+	void* AllocateShortBlock()
+	{
+		// Make sure this isn't a fake block
+		checkSlow(IsAligned(this, BINNED2_LARGE_ALLOC));
+
+		if (bShortBlockAllocated)
+		{
+			return nullptr;
+		}
+
+		bShortBlockAllocated = true;
+		return (uint8*)this + BINNED2_MINIMUM_ALIGNMENT;
+	}
+
+ public:	void*  NextFreeBlock;             // Next free block in another pool
+ public:	uint32 NumFreeBlocks;             // Number of consecutive free blocks here, at least 1.
+ public:	uint32 BlockSize            : 31; // Size of the blocks that this list points to
+ public:	uint32 bShortBlockAllocated : 1;  // Whether the short block is allocated or not
 };
 
 FMallocBinned2::FPoolList::FPoolList()
@@ -47,31 +118,43 @@ FMallocBinned2::FPoolTable::FPoolTable()
 
 struct FMallocBinned2::FPoolInfo
 {
- public:	uint16      Taken;          // Number of allocated elements in this pool, when counts down to zero can free the entire pool.
+ public:	uint16      Taken;          // Number of allocated elements in this pool, when counts down to zero can free the entire pool
 private:	uint32      AllocSize;      // Number of bytes allocated
  public:	FFreeBlock* FirstFreeBlock; // Pointer to first free memory in this pool or the OS Allocation Size in bytes if this allocation is not binned
  public:	FPoolInfo*  Next;           // Pointer to next pool
  public:	FPoolInfo** PtrToPrevNext;  // Pointer to whichever pointer points to this pool
 
 public:
-	bool IsExhausted() const
+	bool HasFreeRegularBlock() const
 	{
-		return !FirstFreeBlock || FirstFreeBlock->NumFreeBlocks == 0;
+		return FirstFreeBlock && FirstFreeBlock->GetNumFreeRegularBlocks() != 0;
 	}
 
-	void* AllocateBlock(uint32 BlockSize)
+	bool HasFreeShortBlock() const
 	{
-		check(!IsExhausted());
-		++Taken;
-		--FirstFreeBlock->NumFreeBlocks;
-		void* Result = (uint8*)FirstFreeBlock + (FirstFreeBlock->NumFreeBlocks + IsAligned(FirstFreeBlock, BINNED2_LARGE_ALLOC)) * BlockSize;
+		return FirstFreeBlock && FirstFreeBlock->HasFreeShortBlock();
+	}
 
-		if (FirstFreeBlock->NumFreeBlocks == 0)
-		{
-			FirstFreeBlock = (FFreeBlock*)FirstFreeBlock->NextFreeBlock;
-		}
-		checkSlow(!FirstFreeBlock || FirstFreeBlock->NumFreeBlocks != 0);
+	uint32 GetShortBlockSize() const
+	{
+		return FirstFreeBlock ? FirstFreeBlock->GetShortBlockSize() : 0;
+	}
+
+	void* AllocateRegularBlock()
+	{
+		checkSlow(HasFreeRegularBlock());
+		++Taken;
+		void* Result = FirstFreeBlock->AllocateRegularBlock();
+		ExhaustPoolIfNecessary();
+		return Result;
+	}
 
+	void* AllocateShortBlock()
+	{
+		checkSlow(HasFreeShortBlock());
+		++Taken;
+		void* Result = FirstFreeBlock->AllocateShortBlock();
+		ExhaustPoolIfNecessary();
 		return Result;
 	}
 
@@ -113,6 +196,16 @@ public:
 		}
 		*PtrToPrevNext = Next;
 	}
+
+private:
+	void ExhaustPoolIfNecessary()
+	{
+		if (FirstFreeBlock->GetNumFreeRegularBlocks() == 0)// && (FirstFreeBlock->GetShortBlockSize() == 0 || !FirstFreeBlock->HasFreeShortBlock()))
+		{
+			FirstFreeBlock = (FFreeBlock*)FirstFreeBlock->NextFreeBlock;
+		}
+		checkSlow(!FirstFreeBlock || FirstFreeBlock->GetNumFreeRegularBlocks() != 0 || FirstFreeBlock->HasFreeShortBlock());
+	}
 };
 
 template <>
@@ -156,18 +249,13 @@ struct FMallocBinned2::PoolHashBucket
 
 struct FMallocBinned2::Private
 {
-	enum
-	{
-		DEFAULT_BINNED_ALLOCATOR_ALIGNMENT = 16, // Default alignment for binned allocator
-	};
-
 	FORCEINLINE static bool IsOSAllocation(const void* Ptr)
 	{
 		return IsAligned(Ptr, BINNED2_LARGE_ALLOC);
 	}
 
 	static_assert(ARRAY_COUNT(GMallocBinned2SmallBlockSizes) == SMALL_POOL_COUNT, "Small block size array size must match SMALL_POOL_COUNT");
-	static_assert(sizeof(FFreeBlock) <= DEFAULT_BINNED_ALLOCATOR_ALIGNMENT,       "Free block struct must be small enough to fit into a block.");
+	static_assert(sizeof(FFreeBlock) <= BINNED2_MINIMUM_ALIGNMENT,                "Free block struct must be small enough to fit into a block.");
 
 	// Implementation. 
 	static CA_NO_RETURN void OutOfMemory(uint64 Size, uint32 Alignment=0)
@@ -287,8 +375,311 @@ struct FMallocBinned2::Private
 	{
 		return (FFreeBlock*)AlignDown(Ptr, BINNED2_LARGE_ALLOC);
 	}
+
+	struct FBundleNode
+	{
+		FBundleNode* NextNodeInCurrentBundle;
+		FBundleNode* NextBundle;
+	};
+
+	struct FBundle
+	{
+		FBundle()
+		{
+			Reset();
+		}
+
+		void Reset()
+		{
+			Head        = nullptr;
+			ApproxBytes = 0;
+		}
+
+		void PushHead(FBundleNode* Node, uint32 NodeSize)
+		{
+			Node->NextNodeInCurrentBundle = Head;
+			Head         = Node;
+			ApproxBytes += NodeSize;
+		}
+
+		FBundleNode* PopHead(uint32 NodeSize)
+		{
+			check(Head);
+			FBundleNode* Result = Head;
+
+			ApproxBytes = FMath::Max<int32>(ApproxBytes - NodeSize, 0);
+			Head        = Head->NextNodeInCurrentBundle;
+
+			return Result;
+		}
+
+		FBundleNode* Head;
+		uint32       ApproxBytes;
+	};
+
+	static_assert(sizeof(FBundleNode) <= BINNED2_MINIMUM_ALIGNMENT, "Bundle nodes must fit into the smallest block size");
+
+	struct FGlobalRecycler
+	{
+		FGlobalRecycler()
+		{
+			DefaultConstructItems<FBundleNode*>(CurrentBundle, SMALL_POOL_COUNT);
+		}
+
+		void PushBundle(uint32 InPoolIndex, FBundleNode* InBundle, FBundleNode*& OutBundlesToRecycle)
+		{
+			FBundleNode** CurrentHeadPtr = &CurrentBundle[InPoolIndex];
+			FBundleNode*  CurrentHead;
+			do
+			{
+				CurrentHead          = *CurrentHeadPtr;
+				InBundle->NextBundle = CurrentHead;
+			}
+			while (FPlatformAtomics::InterlockedCompareExchangePointer((void**)CurrentHeadPtr, InBundle, CurrentHead) != CurrentHead);
+
+			if (BundleCount[InPoolIndex].Increment() >= GMallocBinned2MaxBundlesBeforeRecycle)
+			{
+				OutBundlesToRecycle = PopBundle(InPoolIndex);
+				OutBundlesToRecycle->NextBundle = nullptr;
+			}
+			else
+			{
+				OutBundlesToRecycle = nullptr;
+			}
+		}
+
+		FBundleNode* PopBundle(uint32 InPoolIndex)
+		{
+			FBundleNode** CurrentHeadPtr = &CurrentBundle[InPoolIndex];
+			FBundleNode*  CurrentHead;
+			FBundleNode*  NextHead;
+			do 
+			{
+				CurrentHead = *CurrentHeadPtr;
+				if (!CurrentHead)
+				{
+					return nullptr;
+				}
+
+				NextHead = CurrentHead->NextBundle;
+			}
+			while (FPlatformAtomics::InterlockedCompareExchangePointer((void**)CurrentHeadPtr, NextHead, CurrentHead) != CurrentHead);
+
+			BundleCount[InPoolIndex].Decrement();
+
+			return CurrentHead;
+		}
+
+	private:
+		FBundleNode*       CurrentBundle[SMALL_POOL_COUNT];
+		FThreadSafeCounter BundleCount  [SMALL_POOL_COUNT];
+	};
+
+	static FGlobalRecycler GGlobalRecycler;
+
+	struct FFreeBlockList
+	{
+		void PushToFront(void* InPtr, uint32 InPoolIndex, FBundleNode*& OutBundlesToRecycle)
+		{
+			checkSlow(InPtr);
+
+			OutBundlesToRecycle = nullptr;
+			if (PartialBundle.ApproxBytes >= (uint32)GMallocBinned2BundleSize)
+			{
+				if (FullBundle.Head)
+				{
+					GGlobalRecycler.PushBundle(InPoolIndex, FullBundle.Head, OutBundlesToRecycle);
+				}
+
+				FullBundle = PartialBundle;
+				PartialBundle.Reset();
+			}
+
+			PartialBundle.PushHead((FBundleNode*)InPtr, GMallocBinned2SmallBlockSizes[InPoolIndex]);
+		}
+
+		void* PopFromFront(uint32 InPoolIndex)
+		{
+			if (!PartialBundle.Head)
+			{
+				FBundle NewPartialBundle = FullBundle;
+				if (!NewPartialBundle.Head)
+				{
+					NewPartialBundle.Head = GGlobalRecycler.PopBundle(InPoolIndex);
+					if (!NewPartialBundle.Head)
+					{
+						return nullptr;
+					}
+
+					NewPartialBundle.ApproxBytes = (uint32)GMallocBinned2BundleSize;
+				}
+
+				PartialBundle = NewPartialBundle;
+				FullBundle.Reset();
+			}
+
+			void* Result = PartialBundle.PopHead(GMallocBinned2SmallBlockSizes[InPoolIndex]);
+			return Result;
+		}
+
+		FBundleNode* PopBundles(uint32 InPoolIndex)
+		{
+			FBundleNode* Partial = PartialBundle.Head;
+			if (Partial)
+			{
+				PartialBundle.Reset();
+				Partial->NextBundle = nullptr;
+			}
+
+			FBundleNode* Full = FullBundle.Head;
+			if (Full)
+			{
+				FullBundle.Reset();
+				Full->NextBundle = nullptr;
+			}
+
+			FBundleNode* Result = Partial;
+			if (Result)
+			{
+				Result->NextBundle = Full;
+			}
+			else
+			{
+				Result = Full;
+			}
+
+			return Result;
+		}
+
+		bool IsEmpty() const
+		{
+			return !PartialBundle.Head && !FullBundle.Head;
+		}
+
+	private:
+		FBundle PartialBundle;
+		FBundle FullBundle;
+	};
+
+	// We've reinvented TThreadSingleton here because we don't want it to be re-entering our allocation
+	// functions!
+	template <typename T>
+	struct TThreadSingleton2
+	{
+		FORCEINLINE static T& Get(uint32 PageSize)
+		{
+			static uint32 TlsSlot = 0;
+
+			if (TlsSlot == 0)
+			{
+				const uint32 ThisTlsSlot = FPlatformTLS::AllocTlsSlot();
+				const uint32 PrevTlsSlot = FPlatformAtomics::InterlockedCompareExchange((int32*)&TlsSlot, (int32)ThisTlsSlot, 0);
+				if (PrevTlsSlot != 0)
+				{
+					FPlatformTLS::FreeTlsSlot(ThisTlsSlot);
+				}
+			}
+
+			T* ThreadSingleton = (T*)FPlatformTLS::GetTlsValue(TlsSlot);
+			if (!ThreadSingleton)
+			{
+				ThreadSingleton = new (FPlatformMemory::BinnedAllocFromOS(Align(sizeof(T), PageSize))) T();
+				FPlatformTLS::SetTlsValue(TlsSlot, ThreadSingleton);
+			}
+
+			return *ThreadSingleton;
+		}
+	};
+
+	struct FPerThreadFreeBlockLists : TThreadSingleton2<FPerThreadFreeBlockLists>
+	{
+		FPerThreadFreeBlockLists()
+		{
+		}
+
+		void* Malloc(uint32 InPoolIndex)
+		{
+			checkSlow(InPoolIndex >= 0 && InPoolIndex < SMALL_POOL_COUNT);
+
+			FFreeBlockList& List = FreeLists[InPoolIndex];
+			if (List.IsEmpty())
+			{
+				return nullptr;
+			}
+
+			void* Result = List.PopFromFront(InPoolIndex);
+			return Result;
+		}
+
+		void Free(void* InPtr, uint32 InPoolIndex, FBundleNode*& OutBundlesToRecycle)
+		{
+			checkSlow(InPoolIndex >= 0 && InPoolIndex < SMALL_POOL_COUNT);
+
+			FreeLists[InPoolIndex].PushToFront(InPtr, InPoolIndex, OutBundlesToRecycle);
+		}
+
+		FBundleNode* PopBundles(uint32 InPoolIndex)
+		{
+			checkSlow(InPoolIndex >= 0 && InPoolIndex < SMALL_POOL_COUNT);
+
+			FBundleNode* Result = FreeLists[InPoolIndex].PopBundles(InPoolIndex);
+			return Result;
+		}
+
+	private:
+		FFreeBlockList FreeLists[SMALL_POOL_COUNT];
+	};
+
+	static void FreeBundles(FMallocBinned2& Allocator, FBundleNode* BundlesToRecycle, uint32 InBlockSize, uint32 InPoolIndex)
+	{
+		FPoolTable& Table = Allocator.SmallPoolTables[InPoolIndex];
+
+		FBundleNode* Bundle = BundlesToRecycle;
+		while (Bundle)
+		{
+			FBundleNode* NextBundle = Bundle->NextBundle;
+
+			FBundleNode* Node = Bundle->NextNodeInCurrentBundle;
+			while (Node)
+			{
+				FBundleNode* NextNode = Node->NextNodeInCurrentBundle;
+				FPoolInfo*   NodePool = FindPoolInfo(Allocator, Node);
+
+				// If this pool was exhausted, move to available list.
+				if (!NodePool->FirstFreeBlock)
+				{
+					Table.ActivePools.LinkToFront(NodePool);
+				}
+
+				// Free a pooled allocation.
+				FFreeBlock* Free = (FFreeBlock*)Node;
+				Free->NumFreeBlocks = 1;
+				Free->NextFreeBlock = NodePool->FirstFreeBlock;
+				Free->BlockSize     = InBlockSize;
+				Free->bShortBlockAllocated = false;
+				NodePool->FirstFreeBlock   = Free;
+
+				// Free this pool.
+				checkSlow(NodePool->Taken >= 1);
+				if (--NodePool->Taken == 0)
+				{
+					FFreeBlock* BasePtrOfNode = GetPoolHeaderFromPointer(Node);
+
+					// Free the OS memory.
+					NodePool->Unlink();
+					Allocator.CachedOSPageAllocator.Free(BasePtrOfNode, Allocator.PageSize);
+				}
+
+				Node = NextNode;
+			}
+
+			Bundle = NextBundle;
+		}
+	}
 };
 
+FMallocBinned2::Private::FGlobalRecycler FMallocBinned2::Private::GGlobalRecycler;
+
 FORCEINLINE bool FMallocBinned2::FPoolList::IsEmpty() const
 {
 	return Front == nullptr;
@@ -317,7 +708,7 @@ FMallocBinned2::FPoolInfo& FMallocBinned2::FPoolList::PushNewPoolToFront(FMalloc
 	const uint32 LocalPageSize = Allocator.PageSize;
 
 	// Allocate memory.
-	FFreeBlock* Free = (FFreeBlock*)Allocator.CachedOSPageAllocator.Allocate(LocalPageSize);
+	FFreeBlock* Free = new (Allocator.CachedOSPageAllocator.Allocate(LocalPageSize)) FFreeBlock(LocalPageSize, InBlockSize);
 	if (!Free)
 	{
 		Private::OutOfMemory(LocalPageSize);
@@ -329,11 +720,6 @@ FMallocBinned2::FPoolInfo& FMallocBinned2::FPoolList::PushNewPoolToFront(FMalloc
 	Result->Taken          = 0;
 	Result->FirstFreeBlock = Free;
 
-	// Create first free item.
-	Free->NumFreeBlocks = LocalPageSize / InBlockSize - 1;
-	Free->NextFreeBlock = nullptr;
-	Free->BlockSize     = InBlockSize;
-
 	return *Result;
 }
 
@@ -343,34 +729,34 @@ FMallocBinned2::FMallocBinned2(uint32 InPageSize, uint64 AddressLimit)
 	, PtrToPoolMapping  (PageSize, NumPoolsPerPage)
 	, HashBucketFreeList(nullptr)
 {
-	checkf(FMath::IsPowerOfTwo(PageSize),                                                           TEXT("OS page size must be a power of two"));
-	checkf(FMath::IsPowerOfTwo(AddressLimit),                                                       TEXT("OS address limit must be a power of two"));
-	checkf(AddressLimit > PageSize,                                                                 TEXT("OS address limit must be greater than the page size")); // Check to catch 32 bit overflow in AddressLimit
-	checkf(GMallocBinned2SmallBlockSizes[SMALL_POOL_COUNT - 1] == MAX_SMALL_POOLED_ALLOCATION_SIZE, TEXT("MAX_SMALL_POOLED_ALLOCATION_SIZE must equal the smallest block size"));
-	checkf(PageSize % BINNED2_LARGE_ALLOC == 0,                                                     TEXT("OS page size must be a multiple of BINNED2_LARGE_ALLOC"));
-	checkf(sizeof(FMallocBinned2::FFreeBlock) <= GMallocBinned2SmallBlockSizes[0],                  TEXT("Pool header must be able to fit into the smallest block"));
+	checkf(FMath::IsPowerOfTwo(PageSize),                                                      TEXT("OS page size must be a power of two"));
+	checkf(FMath::IsPowerOfTwo(AddressLimit),                                                  TEXT("OS address limit must be a power of two"));
+	checkf(AddressLimit > PageSize,                                                            TEXT("OS address limit must be greater than the page size")); // Check to catch 32 bit overflow in AddressLimit
+	checkf(GMallocBinned2SmallBlockSizes[SMALL_POOL_COUNT - 1] == BINNED2_MAX_SMALL_POOL_SIZE, TEXT("BINNED2_MAX_SMALL_POOL_SIZE must equal the smallest block size"));
+	checkf(PageSize % BINNED2_LARGE_ALLOC == 0,                                                TEXT("OS page size must be a multiple of BINNED2_LARGE_ALLOC"));
+	checkf(sizeof(FMallocBinned2::FFreeBlock) <= GMallocBinned2SmallBlockSizes[0],             TEXT("Pool header must be able to fit into the smallest block"));
 
 	// Init pool tables.
 	for (uint32 Index = 0; Index != SMALL_POOL_COUNT; ++Index)
 	{
 		checkf(Index == 0 || GMallocBinned2SmallBlockSizes[Index - 1] < GMallocBinned2SmallBlockSizes[Index], TEXT("Small block sizes must be strictly increasing"));
 		checkf(GMallocBinned2SmallBlockSizes[Index] <= PageSize,                                              TEXT("Small block size must be small enough to fit into a page"));
-		checkf(GMallocBinned2SmallBlockSizes[Index] % Private::DEFAULT_BINNED_ALLOCATOR_ALIGNMENT == 0,       TEXT("Small block size must be a multiple of DEFAULT_BINNED_ALLOCATOR_ALIGNMENT"));
+		checkf(GMallocBinned2SmallBlockSizes[Index] % BINNED2_MINIMUM_ALIGNMENT == 0,                         TEXT("Small block size must be a multiple of BINNED2_MINIMUM_ALIGNMENT"));
 
 		SmallPoolTables[Index].BlockSize = GMallocBinned2SmallBlockSizes[Index];
 	}
 
 	// Set up pool mappings
-	FPoolTable** TableEntry = MemSizeToPoolTable;
-	FPoolTable*  Pool       = SmallPoolTables;
-	for (uint32 Index = 0; Index != MAX_SMALL_POOLED_ALLOCATION_SIZE; ++Index)
+	uint32* IndexEntry = GMallocBinned2MemSizeToIndex;
+	uint32  PoolIndex  = 0;
+	for (uint32 Index = 0; Index != BINNED2_MAX_SMALL_POOL_SIZE; ++Index)
 	{
-		while (Pool->BlockSize - 1 < Index)
+		while (GMallocBinned2SmallBlockSizes[PoolIndex] - 1 < Index)
 		{
-			++Pool;
-			check(Pool != SmallPoolTables + SMALL_POOL_COUNT);
+			++PoolIndex;
+			check(PoolIndex != SMALL_POOL_COUNT);
 		}
-		*TableEntry++ = Pool;
+		*IndexEntry++ = PoolIndex;
 	}
 
 	uint64 MaxHashBuckets = PtrToPoolMapping.GetMaxHashBuckets(AddressLimit);
@@ -385,12 +771,12 @@ FMallocBinned2::~FMallocBinned2()
 
 bool FMallocBinned2::IsInternallyThreadSafe() const
 { 
-	return false;
+	return true;
 }
 
 void* FMallocBinned2::Malloc(SIZE_T Size, uint32 Alignment)
 {
-	Alignment = FMath::Max<uint32>(Alignment, Private::DEFAULT_BINNED_ALLOCATOR_ALIGNMENT);
+	Alignment = FMath::Max<uint32>(Alignment, BINNED2_MINIMUM_ALIGNMENT);
 	Size      = Align(FMath::Max((SIZE_T)1, Size), Alignment);
 
 	checkf(FMath::IsPowerOfTwo(Alignment), TEXT("Alignment must be a power of two"));
@@ -398,10 +784,23 @@ void* FMallocBinned2::Malloc(SIZE_T Size, uint32 Alignment)
 
 	// Only allocate from the small pools if the size is small enough and the alignment isn't crazy large.
 	// With large alignments, we'll waste a lot of memory allocating an entire page, but such alignments are highly unlikely in practice.
-	if (Size <= MAX_SMALL_POOLED_ALLOCATION_SIZE && Alignment <= Private::DEFAULT_BINNED_ALLOCATOR_ALIGNMENT)
+	if (Size <= BINNED2_MAX_SMALL_POOL_SIZE && Alignment <= BINNED2_MINIMUM_ALIGNMENT)
 	{
+		uint32 PoolIndex = GMallocBinned2MemSizeToIndex[Size - 1];
+
+		if (GMallocBinned2PerThreadCaches)
+		{
+			Private::FPerThreadFreeBlockLists& Lists = Private::FPerThreadFreeBlockLists::Get(PageSize);
+			if (void* Result = Lists.Malloc(PoolIndex))
+			{
+				return Result;
+			}
+		}
+
+		FScopeLock Lock(&Mutex);
+
 		// Allocate from small object pool.
-		FPoolTable& Table = *MemSizeToPoolTable[Size - 1];
+		FPoolTable& Table = SmallPoolTables[PoolIndex];
 
 		FPoolInfo* Pool;
 		if (!Table.ActivePools.IsEmpty())
@@ -413,8 +812,8 @@ void* FMallocBinned2::Malloc(SIZE_T Size, uint32 Alignment)
 			Pool = &Table.ActivePools.PushNewPoolToFront(*this, Table.BlockSize);
 		}
 
-		void* Result = Pool->AllocateBlock(Table.BlockSize);
-		if (Pool->IsExhausted())
+		void* Result = Pool->AllocateRegularBlock();
+		if (!Pool->HasFreeRegularBlock())
 		{
 			Table.ExhaustedPools.LinkToFront(Pool);
 		}
@@ -422,6 +821,8 @@ void* FMallocBinned2::Malloc(SIZE_T Size, uint32 Alignment)
 		return Result;
 	}
 
+	FScopeLock Lock(&Mutex);
+
 	// Use OS for non-pooled allocations.
 	UPTRINT AlignedSize = Align(Size, PageSize);
 	void* Result = CachedOSPageAllocator.Allocate(AlignedSize);
@@ -451,7 +852,7 @@ void* FMallocBinned2::Realloc(void* Ptr, SIZE_T NewSize, uint32 Alignment)
 		return Result;
 	}
 
-	Alignment = FMath::Max<uint32>(Alignment, Private::DEFAULT_BINNED_ALLOCATOR_ALIGNMENT);
+	Alignment = FMath::Max<uint32>(Alignment, BINNED2_MINIMUM_ALIGNMENT);
 	NewSize   = Align(NewSize, Alignment);
 
 	checkf(FMath::IsPowerOfTwo(Alignment), TEXT("Alignment must be a power of two"));
@@ -460,9 +861,10 @@ void* FMallocBinned2::Realloc(void* Ptr, SIZE_T NewSize, uint32 Alignment)
 	if (!Private::IsOSAllocation(Ptr))
 	{
 		// Reallocate to a smaller/bigger pool if necessary
-		FFreeBlock* Free  = Private::GetPoolHeaderFromPointer(Ptr);
-		FPoolTable& Table = *MemSizeToPoolTable[Free->BlockSize - 1];
-		if (NewSize <= Table.BlockSize && (&Table == &SmallPoolTables[0] || NewSize > (&Table)[-1].BlockSize) && IsAligned(Ptr, Alignment))
+		FFreeBlock* Free       = Private::GetPoolHeaderFromPointer(Ptr);
+		uint32      TableIndex = GMallocBinned2MemSizeToIndex[Free->BlockSize - 1];
+		FPoolTable& Table      = SmallPoolTables[TableIndex];
+		if (NewSize <= Table.BlockSize && (TableIndex == 0 || NewSize > SmallPoolTables[TableIndex - 1].BlockSize) && IsAligned(Ptr, Alignment))
 		{
 			return Ptr;
 		}
@@ -474,6 +876,8 @@ void* FMallocBinned2::Realloc(void* Ptr, SIZE_T NewSize, uint32 Alignment)
 		return Result;
 	}
 
+	FScopeLock Lock(&Mutex);
+
 	// Allocated from OS.
 	FPoolInfo* Pool        = Private::FindPoolInfo(*this, Ptr);
 	UPTRINT    PoolOsBytes = Pool->GetOsAllocatedBytes();
@@ -500,7 +904,7 @@ void FMallocBinned2::Free(void* Ptr)
 
 	FPoolInfo* Pool = Private::FindPoolInfo(*this, Ptr);
 #if PLATFORM_IOS || PLATFORM_MAC
-	if (Pool == NULL)
+	if (!Pool)
 	{
 		UE_LOG(LogMemory, Warning, TEXT("Attempting to free a pointer we didn't allocate!"));
 		return;
@@ -510,33 +914,33 @@ void FMallocBinned2::Free(void* Ptr)
 	checkSlow(Pool);
 	if (!Private::IsOSAllocation(Ptr))
 	{
-		FFreeBlock* BasePtr = Private::GetPoolHeaderFromPointer(Ptr);
-		FPoolTable& Table   = *MemSizeToPoolTable[BasePtr->BlockSize - 1];
+		FFreeBlock* BasePtr   = Private::GetPoolHeaderFromPointer(Ptr);
+		uint32      BlockSize = BasePtr->BlockSize;
+		uint32      PoolIndex = GMallocBinned2MemSizeToIndex[BlockSize - 1];
 
-		// If this pool was exhausted, move to available list.
-		if (!Pool->FirstFreeBlock)
+		Private::FBundleNode* BundlesToRecycle;
+		if (GMallocBinned2PerThreadCaches)
 		{
-			Table.ActivePools.LinkToFront(Pool);
+			Private::FPerThreadFreeBlockLists& Lists = Private::FPerThreadFreeBlockLists::Get(PageSize);
+			Lists.Free(Ptr, PoolIndex, BundlesToRecycle);
+		}
+		else
+		{
+			BundlesToRecycle = (Private::FBundleNode*)Ptr;
+			BundlesToRecycle->NextNodeInCurrentBundle = nullptr;
 		}
 
-		// Free a pooled allocation.
-		FFreeBlock* Free     = (FFreeBlock*)Ptr;
-		Free->NumFreeBlocks  = 1;
-		Free->NextFreeBlock  = Pool->FirstFreeBlock;
-		Free->BlockSize      = Table.BlockSize;
-		Pool->FirstFreeBlock = Free;
-
-		// Free this pool.
-		checkSlow(Pool->Taken >= 1);
-		if (--Pool->Taken == 0)
+		if (BundlesToRecycle)
 		{
-			// Free the OS memory.
-			Pool->Unlink();
-			CachedOSPageAllocator.Free(BasePtr, PageSize);
+			FScopeLock Lock(&Mutex);
+
+			Private::FreeBundles(*this, BundlesToRecycle, BlockSize, PoolIndex);
 		}
 	}
 	else
 	{
+		FScopeLock Lock(&Mutex);
+
 		// Free an OS allocation.
 		CachedOSPageAllocator.Free(Ptr, Pool->GetOsAllocatedBytes());
 	}
@@ -551,11 +955,13 @@ bool FMallocBinned2::GetAllocationSize(void* Ptr, SIZE_T& SizeOut)
 
 	if (!Private::IsOSAllocation(Ptr))
 	{
-		FFreeBlock* Free = Private::GetPoolHeaderFromPointer(Ptr);
-		SizeOut = MemSizeToPoolTable[Free->BlockSize - 1]->BlockSize;
+		const FFreeBlock* Free = Private::GetPoolHeaderFromPointer(Ptr);
+		SizeOut = SmallPoolTables[GMallocBinned2MemSizeToIndex[Free->BlockSize - 1]].BlockSize;
 		return true;
 	}
 
+	FScopeLock Lock(&Mutex);
+
 	FPoolInfo* Pool = Private::FindPoolInfo(*this, Ptr);
 	SizeOut = Pool->GetOSRequestedBytes();
 	return true;
@@ -573,6 +979,8 @@ void FMallocBinned2::FPoolList::Validate(TFunctionRef<void(FMallocBinned2::FPool
 
 bool FMallocBinned2::ValidateHeap()
 {
+	FScopeLock Lock(&Mutex);
+
 	for (FPoolTable& Table : SmallPoolTables)
 	{
 		Table.ActivePools.Validate(
@@ -581,7 +989,7 @@ bool FMallocBinned2::ValidateHeap()
 				check(Pool.FirstFreeBlock);
 				for (FFreeBlock* Free = Pool.FirstFreeBlock; Free; Free = (FFreeBlock*)Free->NextFreeBlock)
 				{
-					check(Free->NumFreeBlocks > 0);
+					check(Free->GetNumFreeRegularBlocks() > 0);
 				}
 			}
 		);
@@ -601,3 +1009,18 @@ const TCHAR* FMallocBinned2::GetDescriptiveName()
 {
 	return TEXT("binned2");
 }
+
+void FMallocBinned2::FlushCurrentThreadCache()
+{
+	FScopeLock Lock(&Mutex);
+
+	Private::FPerThreadFreeBlockLists& Lists = Private::FPerThreadFreeBlockLists::Get(PageSize);
+	for (int32 PoolIndex = 0; PoolIndex != SMALL_POOL_COUNT; ++PoolIndex)
+	{
+		Private::FBundleNode* Bundles = Lists.PopBundles(PoolIndex);
+		if (Bundles)
+		{
+			Private::FreeBundles(*this, Bundles, GMallocBinned2SmallBlockSizes[PoolIndex], PoolIndex);
+		}
+	}
+}
diff --git a/Engine/Source/Runtime/Core/Public/HAL/MallocBinned2.h b/Engine/Source/Runtime/Core/Public/HAL/MallocBinned2.h
--- a/Engine/Source/Runtime/Core/Public/HAL/MallocBinned2.h
+++ b/Engine/Source/Runtime/Core/Public/HAL/MallocBinned2.h
@@ -23,10 +23,6 @@ private:
 	// Counts.
 	enum { SMALL_POOL_COUNT = 41 };
 
-	/** Maximum allocation for the pooled allocator */
-	enum { MAX_SMALL_POOLED_ALLOCATION_SIZE = 32768 };
-	enum { MAX_LARGE_POOLED_PAGE_FRACTIONS  = 2 };
-
 	// Forward declares.
 	struct FFreeBlock;
 	struct FPoolInfo;
@@ -103,14 +99,13 @@ private:
 	// Pool tables for different pool sizes
 	FPoolTable SmallPoolTables[SMALL_POOL_COUNT];
 
-	// Mapping of indices to both large and small tables
-	FPoolTable* MemSizeToPoolTable[MAX_SMALL_POOLED_ALLOCATION_SIZE];
-
 	PoolHashBucket* HashBuckets;
 	PoolHashBucket* HashBucketFreeList;
 
 	TCachedOSPageAllocator<BINNED2_MAX_CACHED_OS_FREES, BINNED2_MAX_CACHED_OS_FREES_BYTE_LIMIT> CachedOSPageAllocator;
 
+	FCriticalSection Mutex;
+
 public:
 	// FMalloc interface.
 	// InPageSize - First parameter is page size, all allocs from BinnedAllocFromOS() MUST be aligned to this size
@@ -157,4 +152,6 @@ public:
 	virtual bool ValidateHeap() override;
 
 	virtual const TCHAR* GetDescriptiveName() override;
+
+	virtual void FlushCurrentThreadCache() override;
 };
diff --git a/Engine/Source/Runtime/Core/Public/HAL/MemoryBase.h b/Engine/Source/Runtime/Core/Public/HAL/MemoryBase.h
--- a/Engine/Source/Runtime/Core/Public/HAL/MemoryBase.h
+++ b/Engine/Source/Runtime/Core/Public/HAL/MemoryBase.h
@@ -155,6 +155,13 @@ public:
 		return TEXT("Unspecified allocator");
 	}
 
+	/**
+	 * Triggers a cache flush of any memory being held by the allocator on the current thread.
+	 */
+	virtual void FlushCurrentThreadCache()
+	{
+	}
+
 protected:
 	friend struct FCurrentFrameCalls;
 
diff --git a/Engine/Source/Runtime/Core/Public/UObject/CoreObjectVersion.h b/Engine/Source/Runtime/Core/Public/UObject/CoreObjectVersion.h
--- a/Engine/Source/Runtime/Core/Public/UObject/CoreObjectVersion.h
+++ b/Engine/Source/Runtime/Core/Public/UObject/CoreObjectVersion.h
@@ -8,7 +8,7 @@ struct CORE_API FCoreObjectVersion
 	{
 		// Before any version changes were made
 		BeforeCustomVersionWasAdded = 0,
-
+		MaterialInputNativeSerialize,
 
 		// -----<new versions can be added above this line>-------------------------------------------------
 		VersionPlusOne,
diff --git a/Engine/Source/Runtime/CoreUObject/Private/UObject/ReferenceChainSearch.cpp b/Engine/Source/Runtime/CoreUObject/Private/UObject/ReferenceChainSearch.cpp
--- a/Engine/Source/Runtime/CoreUObject/Private/UObject/ReferenceChainSearch.cpp
+++ b/Engine/Source/Runtime/CoreUObject/Private/UObject/ReferenceChainSearch.cpp
@@ -101,7 +101,7 @@ void FReferenceChainSearch::FFindReferencerCollector::HandleObjectReference(UObj
 }
 
 
-void FReferenceChainSearch::PrintReferencers( FReferenceChain& Referencer )
+void FReferenceChainSearch::PrintReferencers(FReferenceChain& Referencer)
 {
 	UE_LOG(LogReferenceChain, Log, TEXT("  "));
 
@@ -242,29 +242,37 @@ static void LinkToParents(TMultiMap<UObject*, FRefGraphItem*>& InputGraphNodeLis
 	}
 }
 
-// Recursively creates reference chains from the internal reference graph relationship
-void CreateReferenceChain(FRefGraphItem* Node, FReferenceChainSearch::FReferenceChain& ThisChain, TArray<FReferenceChainSearch::FReferenceChain>& ChainArray, UObject* ObjectToFind, int32 Levels )
+// Returns true if the object can't be collected by GC
+static FORCEINLINE bool IsNonGCObject(UObject* Object)
+{
+	FUObjectItem* ObjectItem = GUObjectArray.ObjectToObjectItem(Object);
+	return (ObjectItem->IsRootSet() ||
+		ObjectItem->HasAnyFlags(EInternalObjectFlags::GarbageCollectionKeepFlags) ||
+		(GARBAGE_COLLECTION_KEEPFLAGS != RF_NoFlags && Object->HasAnyFlags(GARBAGE_COLLECTION_KEEPFLAGS)));
+}
+
+void FReferenceChainSearch::CreateReferenceChain(FRefGraphItem* Node, FReferenceChainSearch::FReferenceChain& ThisChain, TArray<FReferenceChainSearch::FReferenceChain>& ChainArray, UObject* InObjectToFind, int32 Levels)
 {
 	ThisChain.RefChain.Push(Node->Link);
 	if (Levels <= 0)
 	{
-
-		UE_LOG(LogReferenceChain, Log, TEXT("Chain is too long!"));
+		UE_CLOG(ShouldOutputToLog(), LogReferenceChain, Log, TEXT("Chain is too long!"));
 		for (int32 i=0; i < ThisChain.RefChain.Num(); ++i)
 		{
-			UE_LOG(LogReferenceChain, Log, TEXT("%s -> %s"), *ThisChain.RefChain[i].ReferencedBy->GetPathName(), *ThisChain.RefChain[i].ReferencedObj->GetPathName());
+			UE_CLOG(ShouldOutputToLog(), LogReferenceChain, Log, TEXT("%s -> %s"), *ThisChain.RefChain[i].ReferencedBy->GetPathName(), *ThisChain.RefChain[i].ReferencedObj->GetPathName());
 		}
 	}
 	check(Levels > 0);
 	Levels--;
 
 	// If we encounter the target object or another root-object, we stop here
-	if (Node->Link.ReferencedObj == ObjectToFind)
+	if (Node->Link.ReferencedObj == InObjectToFind)
 	{
 		ChainArray.Push(ThisChain);
 		return;
 	}
-	if (Node->Link.ReferencedObj->HasAnyFlags(GARBAGE_COLLECTION_KEEPFLAGS) || Node->Link.ReferencedObj->IsRooted())
+
+	if (IsNonGCObject(Node->Link.ReferencedObj))
 	{
 		return;
 	}
@@ -285,13 +293,13 @@ void CreateReferenceChain(FRefGraphItem* Node, FReferenceChainSearch::FReference
 		if (bSkip) { continue; };
 
 		FReferenceChainSearch::FReferenceChain ChildChain = ThisChain;
-		CreateReferenceChain(Node->Children[i], ChildChain, ChainArray, ObjectToFind, Levels);
+		CreateReferenceChain(Node->Children[i], ChildChain, ChainArray, InObjectToFind, Levels);
 	}
 }
 
 void FReferenceChainSearch::BuildRefGraph()
 {
-	UE_LOG(LogReferenceChain, Log, TEXT("Generating reference graph ..."));
+	UE_CLOG(ShouldOutputToLog(), LogReferenceChain, Log, TEXT("Generating reference graph ..."));
 
 	bool bContinue = true;
 
@@ -320,15 +328,13 @@ void FReferenceChainSearch::BuildRefGraph()
 	}
 
 	int32 Level = 0;
-	UE_LOG(LogReferenceChain, Log, TEXT("Level 0 has %d nodes ..."), GraphNodes.Num());
+	UE_CLOG(ShouldOutputToLog(), LogReferenceChain, Log, TEXT("Level 0 has %d nodes ..."), GraphNodes.Num());
 
 	while(bContinue)
 	{
 		int32 AddedNodes = 0;
-
 		TArray<FRefGraphItem*> NewGraphNodes;
 
-
 		for (FRawObjectIterator It;It;++It)
 		{
 			FUObjectItem* ObjItem = *It;
@@ -339,7 +345,7 @@ void FReferenceChainSearch::BuildRefGraph()
 			for (int32 i=0; i < RefList.Num(); ++i)
 			{
 				if (RefList[i].ReferenceType == EReferenceType::Invalid ||
-					RefList[i].ReferencedObj->HasAnyFlags(GARBAGE_COLLECTION_KEEPFLAGS) || RefList[i].ReferencedObj->IsRooted()) // references to rooted objects are not important
+					IsNonGCObject(RefList[i].ReferencedObj)) // references to rooted objects are not important
 				{ 
 					continue; 
 				}
@@ -370,7 +376,7 @@ void FReferenceChainSearch::BuildRefGraph()
 			}
 		}
 		Level++;
-		UE_LOG(LogReferenceChain, Log, TEXT("Level %d added %d nodes ..."), Level, NewGraphNodes.Num());
+		UE_CLOG(ShouldOutputToLog(), LogReferenceChain, Log, TEXT("Level %d added %d nodes ..."), Level, NewGraphNodes.Num());
 
 		for (int32 i = 0; i < NewGraphNodes.Num(); ++i)
 		{
@@ -385,12 +391,12 @@ void FReferenceChainSearch::BuildRefGraph()
 
 	TArray<FReferenceChain> Chains;
 
-	UE_LOG(LogReferenceChain, Log, TEXT("Generating reference chains ..."));
+	UE_CLOG(ShouldOutputToLog(), LogReferenceChain, Log, TEXT("Generating reference chains ..."));
 	for (auto It = GraphNodes.CreateConstIterator(); It; ++It)
 	{
 		FRefGraphItem* Node = It.Value();
 
-		if (Node->Link.ReferencedBy->HasAnyFlags(GARBAGE_COLLECTION_KEEPFLAGS) || Node->Link.ReferencedBy->IsRooted())
+		if (IsNonGCObject(Node->Link.ReferencedBy))
 		{
 			FReferenceChain CurChain;
 			CreateReferenceChain(Node, CurChain, Chains, ObjectToFind, Level);
@@ -405,7 +411,7 @@ void FReferenceChainSearch::BuildRefGraph()
 
 void FReferenceChainSearch::PerformSearch()
 {
-	UE_LOG(LogReferenceChain, Log, TEXT("Searching referencers for %s. This may take several minutes."), *ObjectToFind->GetName());
+	UE_CLOG(ShouldOutputToLog(), LogReferenceChain, Log, TEXT("Searching referencers for %s. This may take several minutes."), *ObjectToFind->GetName());
 	
 	for (FRawObjectIterator It;It;++It)
 	{
@@ -674,7 +680,7 @@ bool FReferenceChainSearch::ProcessObject( UObject* CurrentObject )
 		}
 		else
 		{
-			UE_LOG(LogReferenceChain, Fatal,TEXT("Unknown token"));
+			UE_CLOG(ShouldOutputToLog(), LogReferenceChain, Fatal,TEXT("Unknown token"));
 		}
 	}
 	check(StackEntry == Stack.GetData());
@@ -684,49 +690,57 @@ bool FReferenceChainSearch::ProcessObject( UObject* CurrentObject )
 FReferenceChainSearch::FReferenceChainSearch( UObject* InObjectToFind, uint32 Mode ) 
 	:ObjectToFind(InObjectToFind), SearchMode(Mode)
 {
-	if (ObjectToFind == NULL) { return; }
+	if (ObjectToFind == NULL) 
+	{ 
+		return; 
+	}
 
 	PerformSearch();
 	
-	if ( !!(Mode & ESearchMode::PrintResults) )
+	if (ShouldOutputToLog())
 	{
-		bool bIsFirst = true;
-		
-		for (int32 i=0; i < Referencers.Num(); ++i)
-		{
-			UObject* Obj = Referencers[i].RefChain[0].ReferencedBy;
+		PrintResults();
+	}
+}
 
-			if (!Obj->IsIn(ObjectToFind) && Obj != ObjectToFind)
-			{
-				if (bIsFirst)
-				{
-					UE_LOG(LogReferenceChain, Log, TEXT("  "));
-					UE_LOG(LogReferenceChain, Log, TEXT("External Referencers:"));
-					bIsFirst = false;
-				}
+void FReferenceChainSearch::PrintResults()
+{
+	bool bIsFirst = true;
 
-				PrintReferencers(Referencers[i]);
+	for (int32 i = 0; i < Referencers.Num(); ++i)
+	{
+		UObject* Obj = Referencers[i].RefChain[0].ReferencedBy;
+
+		if (!Obj->IsIn(ObjectToFind) && Obj != ObjectToFind)
+		{
+			if (bIsFirst)
+			{
+				UE_LOG(LogReferenceChain, Log, TEXT("  "));
+				UE_LOG(LogReferenceChain, Log, TEXT("External Referencers:"));
+				bIsFirst = false;
 			}
+
+			PrintReferencers(Referencers[i]);
 		}
+	}
 
-		bIsFirst = true;
-		
-		for (int32 i=0; i < Referencers.Num(); ++i)
-		{
-			UObject* Obj = Referencers[i].RefChain[0].ReferencedBy;
+	bIsFirst = true;
 
-			CA_SUPPRESS(6011)
-			if (Obj->IsIn(ObjectToFind) || Obj == ObjectToFind)
-			{
-				if (bIsFirst)
-				{
-					UE_LOG(LogReferenceChain, Log, TEXT("  "));
-					UE_LOG(LogReferenceChain, Log, TEXT("Internal Referencers:"));
-					bIsFirst = false;
-				}
+	for (int32 i = 0; i < Referencers.Num(); ++i)
+	{
+		UObject* Obj = Referencers[i].RefChain[0].ReferencedBy;
 
-				PrintReferencers(Referencers[i]);
+		CA_SUPPRESS(6011)
+		if (Obj->IsIn(ObjectToFind) || Obj == ObjectToFind)
+		{
+			if (bIsFirst)
+			{
+				UE_LOG(LogReferenceChain, Log, TEXT("  "));
+				UE_LOG(LogReferenceChain, Log, TEXT("Internal Referencers:"));
+				bIsFirst = false;
 			}
+
+			PrintReferencers(Referencers[i]);
 		}
 	}
 }
diff --git a/Engine/Source/Runtime/CoreUObject/Private/UObject/UObjectClusters.cpp b/Engine/Source/Runtime/CoreUObject/Private/UObject/UObjectClusters.cpp
--- a/Engine/Source/Runtime/CoreUObject/Private/UObject/UObjectClusters.cpp
+++ b/Engine/Source/Runtime/CoreUObject/Private/UObject/UObjectClusters.cpp
@@ -9,6 +9,7 @@
 #include "ModuleManager.h"
 #include "FastReferenceCollector.h"
 #include "GarbageCollection.h"
+#include "ReferenceChainSearch.h"
 
 int32 GCreateGCClusters = 1;
 static FAutoConsoleVariableRef CCreateGCClusters(
@@ -26,6 +27,8 @@ static FAutoConsoleVariableRef CMergeGCClusters(
 	ECVF_Default
 	);
 
+#if !UE_BUILD_SHIPPING
+
 // Dumps all clusters to log.
 static void ListClusters(const TArray<FString>& Args)
 {
@@ -68,12 +71,61 @@ static void ListClusters(const TArray<FString>& Args)
 	UE_LOG(LogObj, Display, TEXT("Average number of custer-to-cluster references: %d"), GUObjectClusters.Num() ? (TotalInterClusterReferences / GUObjectClusters.Num()) : 0);
 }
 
+static void FindStaleClusters(const TArray<FString>& Args)
+{
+	// This is seriously slow.
+	UE_LOG(LogObj, Display, TEXT("Searching for stale clusters. This may take a while...")); 
+	int32 NumStaleClusters = 0;
+	int32 TotalNumClusters = 0;
+	for (FRawObjectIterator It(true); It; ++It)
+	{
+		FUObjectItem* ObjectItem = *It;
+		if (ObjectItem->HasAnyFlags(EInternalObjectFlags::ClusterRoot))
+		{
+			TotalNumClusters++;
+
+			UObject* ClusterRootObject = static_cast<UObject*>(ObjectItem->Object);
+			FReferenceChainSearch SearchRefs(ClusterRootObject, FReferenceChainSearch::ESearchMode::Shortest);
+			
+			bool bReferenced = false;
+			if (SearchRefs.GetReferenceChains().Num() > 0)
+			{
+				for (const FReferenceChainSearch::FReferenceChain& ReferenceChain : SearchRefs.GetReferenceChains())
+				{
+					UObject* ReferencingObj = ReferenceChain.RefChain[0].ReferencedBy;
+					// Ignore internal references
+					if (!ReferencingObj->IsIn(ClusterRootObject) && ReferencingObj != ClusterRootObject)
+					{
+						bReferenced = true;
+						break;
+					}
+				}
+			}
+			if (!bReferenced)
+			{
+				NumStaleClusters++;
+				UE_LOG(LogObj, Display, TEXT("Cluster %s has no external references:"), *ClusterRootObject->GetFullName());
+				SearchRefs.PrintResults();
+			}
+		}
+	}
+	UE_LOG(LogObj, Display, TEXT("Found %d clusters, including %d stale."), TotalNumClusters, NumStaleClusters);
+}
+
 static FAutoConsoleCommand ListClustersCommand(
 	TEXT("gc.ListClusters"),
 	TEXT("Dumps all clusters do output log. When 'Hiearchy' argument is specified lists all objects inside clusters."),
 	FConsoleCommandWithArgsDelegate::CreateStatic(ListClusters)
 	);
 
+static FAutoConsoleCommand FindStaleClustersCommand(
+	TEXT("gc.FindStaleClusters"),
+	TEXT("Dumps all clusters do output log that are not referenced by anything."),
+	FConsoleCommandWithArgsDelegate::CreateStatic(FindStaleClusters)
+	);
+
+#endif // !UE_BUILD_SHIPPING
+
 /**
 * Pool for reducing allocations when constructing clusters
 */
diff --git a/Engine/Source/Runtime/CoreUObject/Public/UObject/ReferenceChainSearch.h b/Engine/Source/Runtime/CoreUObject/Public/UObject/ReferenceChainSearch.h
--- a/Engine/Source/Runtime/CoreUObject/Public/UObject/ReferenceChainSearch.h
+++ b/Engine/Source/Runtime/CoreUObject/Public/UObject/ReferenceChainSearch.h
@@ -136,7 +136,10 @@ public:
 	 *
 	 *	@return		Array of reference chains
 	 **/
-	const TArray<FReferenceChain> GetReferenceChains() const { return Referencers; }
+	const TArray<FReferenceChain>& GetReferenceChains() const { return Referencers; }
+
+	/** Dumps results to log */
+	void PrintResults();
 
 private:
 	/** Object we want to find all the referencers to */
@@ -159,4 +162,11 @@ private:
 	void InsertReferenceChain(FReferenceChain& Referencer);
 	/** Generates internal reference graph representation */
 	void BuildRefGraph();
+	/** Recursively creates reference chains from the internal reference graph relationship */
+	void CreateReferenceChain(struct FRefGraphItem* Node, FReferenceChain& ThisChain, TArray<FReferenceChain>& ChainArray, UObject* ObjectToFind, int32 Levels);
+	/** Returns true if the search should log anything */
+	bool ShouldOutputToLog() const
+	{
+		return !!(SearchMode & ESearchMode::PrintResults);
+	}
 };
diff --git a/Engine/Source/Runtime/CoreUObject/Public/UObject/UObjectArray.h b/Engine/Source/Runtime/CoreUObject/Public/UObject/UObjectArray.h
--- a/Engine/Source/Runtime/CoreUObject/Public/UObject/UObjectArray.h
+++ b/Engine/Source/Runtime/CoreUObject/Public/UObject/UObjectArray.h
@@ -201,7 +201,7 @@ public:
 	int32 AddSingle()
 	{
 		int32 Result = NumElements;
-		checkf(NumElements + 1 <= MaxElements, TEXT("Maximum number of Objects exceeded, make sure you update MaxObjectsInGame/MaxObjectsInEditor in project settings."));
+		checkf(NumElements + 1 <= MaxElements, TEXT("Maximum number of UObjects (%d) exceeded, make sure you update MaxObjectsInGame/MaxObjectsInEditor in project settings."), MaxElements);
 		check(Result == NumElements);
 		++NumElements;
 		FPlatformMisc::MemoryBarrier();
@@ -212,7 +212,7 @@ public:
 	int32 AddRange(int32 Count)
 	{
 		int32 Result = NumElements + Count - 1;
-		checkf(NumElements + Count <= MaxElements, TEXT("Maximum number of Objects exceeded, make sure you update MaxObjectsInGame/MaxObjectsInEditor in project settings."));
+		checkf(NumElements + Count <= MaxElements, TEXT("Maximum number of UObjects (%d) exceeded, make sure you update MaxObjectsInGame/MaxObjectsInEditor in project settings."), MaxElements);
 		check(Result == (NumElements + Count - 1));
 		NumElements += Count;
 		FPlatformMisc::MemoryBarrier();
diff --git a/Engine/Source/Runtime/Engine/Classes/Animation/AnimSequence.h b/Engine/Source/Runtime/Engine/Classes/Animation/AnimSequence.h
--- a/Engine/Source/Runtime/Engine/Classes/Animation/AnimSequence.h
+++ b/Engine/Source/Runtime/Engine/Classes/Animation/AnimSequence.h
@@ -907,6 +907,10 @@ private:
 
 	friend class UAnimationAsset;
 	friend struct FScopedAnimSequenceCompressedDataCache;
+
+#if WITH_EDITORONLY_DATA
+	bool bCompressedOnCook;
+#endif
 };
 
 
diff --git a/Engine/Source/Runtime/Engine/Classes/Materials/Material.h b/Engine/Source/Runtime/Engine/Classes/Materials/Material.h
--- a/Engine/Source/Runtime/Engine/Classes/Materials/Material.h
+++ b/Engine/Source/Runtime/Engine/Classes/Materials/Material.h
@@ -144,9 +144,11 @@ enum EMaterialDecalResponse
 USTRUCT(noexport)
 struct FMaterialInput
 {
+#if WITH_EDITORONLY_DATA
 	/** Material expression that this input is connected to, or NULL if not connected. */
 	UPROPERTY()
 	class UMaterialExpression* Expression;
+#endif
 
 	/** Index into Expression's outputs array that this input is connected to. */
 	UPROPERTY()
@@ -174,8 +176,9 @@ struct FMaterialInput
 	UPROPERTY()
 	int32 MaskA;
 
+	/** Material expression name that this input is connected to, or None if not connected. Used only in cooked builds */
 	UPROPERTY()
-	int32 GCC64_Padding;    // @todo 64: if the C++ didn't mismirror this structure (with ExpressionInput), we might not need this
+	FName ExpressionName;
 
 };
 #endif
@@ -817,8 +820,10 @@ public:
 	ENGINE_API virtual bool GetStaticComponentMaskParameterValue(FName ParameterName, bool &R, bool &G, bool &B, bool &A, FGuid &OutExpressionGuid) const override;
 	ENGINE_API virtual bool GetTerrainLayerWeightParameterValue(FName ParameterName, int32& OutWeightmapIndex, FGuid &OutExpressionGuid) const override;
 	ENGINE_API virtual bool UpdateLightmassTextureTracking() override;
+#if WITH_EDITOR
 	ENGINE_API virtual bool GetTexturesInPropertyChain(EMaterialProperty InProperty, TArray<UTexture*>& OutTextures, 
 		TArray<FName>* OutTextureParamNames, class FStaticParameterSet* InStaticParameterSet) override;
+#endif
 	ENGINE_API virtual void RecacheUniformExpressions() const override;
 
 	ENGINE_API virtual float GetOpacityMaskClipValue(bool bIsGameThread = IsInGameThread()) const override;
@@ -835,8 +840,10 @@ public:
 
 	/** Checks to see if an input property should be active, based on the state of the material */
 	ENGINE_API virtual bool IsPropertyActive(EMaterialProperty InProperty) const override;
+#if WITH_EDITOR
 	/** Allows material properties to be compiled with the option of being overridden by the material attributes input. */
 	ENGINE_API virtual int32 CompilePropertyEx( class FMaterialCompiler* Compiler, EMaterialProperty Property ) override;
+#endif // WITH_EDITOR
 	ENGINE_API virtual void ForceRecompileForRendering() override;
 	//~ End UMaterialInterface Interface.
 
@@ -1250,6 +1257,7 @@ public:
 	/* Returns any UMaterialExpressionCustomOutput expressions */
 	ENGINE_API void GetAllCustomOutputExpressions(TArray<class UMaterialExpressionCustomOutput*>& OutCustomOutputs) const;
 
+#if WITH_EDITOR
 	/**
 	 *	Get all referenced expressions (returns the chains for all properties).
 	 *
@@ -1271,11 +1279,14 @@ public:
 	 */
 	ENGINE_API virtual bool GetExpressionsInPropertyChain(EMaterialProperty InProperty, 
 		TArray<UMaterialExpression*>& OutExpressions, class FStaticParameterSet* InStaticParameterSet);
+#endif
 
 	/** Appends textures referenced by expressions, including nested functions. */
 	ENGINE_API void AppendReferencedTextures(TArray<UTexture*>& InOutTextures) const;
 
 protected:
+
+#if WITH_EDITOR
 	/**
 	 *	Recursively retrieve the expressions contained in the chain of the given expression.
 	 *
@@ -1298,6 +1309,7 @@ protected:
 	*
 	*/
 	void RecursiveUpdateRealtimePreview(UMaterialExpression* InExpression, TArray<UMaterialExpression*>& InOutExpressionsToProcess);
+#endif
 
 public:
 	bool HasNormalConnected() const { return Normal.IsConnected(); }
diff --git a/Engine/Source/Runtime/Engine/Classes/Materials/MaterialExpression.h b/Engine/Source/Runtime/Engine/Classes/Materials/MaterialExpression.h
--- a/Engine/Source/Runtime/Engine/Classes/Materials/MaterialExpression.h
+++ b/Engine/Source/Runtime/Engine/Classes/Materials/MaterialExpression.h
@@ -14,9 +14,11 @@ class UEdGraphNode;
 USTRUCT(noexport)
 struct FExpressionInput
 {
+#if WITH_EDITORONLY_DATA
 	/** UMaterial expression that this input is connected to, or NULL if not connected. */
 	UPROPERTY()
 	class UMaterialExpression* Expression;
+#endif
 
 	/** Index into Expression's outputs array that this input is connected to. */
 	UPROPERTY()
@@ -44,14 +46,17 @@ struct FExpressionInput
 	UPROPERTY()
 	int32 MaskA;
 
+	/** Material expression name that this input is connected to, or None if not connected. Used only in cooked builds */
 	UPROPERTY()
-	int32 GCC64_Padding;    // @todo 64: if the C++ didn't mismirror this structure (with MaterialInput), we might not need this
+	FName ExpressionName;
 
 };
 
 USTRUCT(noexport)
 struct FMaterialAttributesInput : public FExpressionInput
 {
+	UPROPERTY(transient)
+	int32 PropertyConnectedBitmask;
 };
 
 #endif
@@ -184,16 +189,14 @@ class ENGINE_API UMaterialExpression : public UObject
 	virtual void PostEditImport() override;
 	virtual bool CanEditChange( const UProperty* InProperty ) const override;
 #endif // WITH_EDITOR
-
+	
 	virtual bool Modify( bool bAlwaysMarkDirty=true ) override;
 	virtual void Serialize( FArchive& Ar ) override;
-	virtual bool NeedsLoadForServer() const override
-	{
-		// [RCL] FIXME 2016-02-01: should not be loaded for servers, but currently causes problems with GC clusters (see OR-14209).
-		return Super::NeedsLoadForServer();
-	}
+	virtual bool NeedsLoadForClient() const override;
+	virtual bool NeedsLoadForServer() const override;	
 	//~ End UObject Interface.
 
+#if WITH_EDITOR
 	/**
 	 * Create the new shader code chunk needed for the Abs expression
 	 *
@@ -203,6 +206,7 @@ class ENGINE_API UMaterialExpression : public UObject
 	 */	
 	virtual int32 Compile(class FMaterialCompiler* Compiler, int32 OutputIndex, int32 MultiplexIndex) { return INDEX_NONE; }
 	virtual int32 CompilePreview(class FMaterialCompiler* Compiler, int32 OutputIndex, int32 MultiplexIndex) { return Compile(Compiler, OutputIndex, MultiplexIndex); }
+#endif
 
 	/** 
 	 * Callback to get any texture reference this expression emits.
@@ -239,13 +243,13 @@ class ENGINE_API UMaterialExpression : public UObject
 	virtual bool UsesLeftGutter() const;
 	virtual bool UsesRightGutter() const;
 
+#if WITH_EDITOR
 	/**
 	 *	Returns the text to display on the material expression (in the material editor).
 	 *
 	 *	@return	FString		The text to display.
 	 */
 	virtual void GetCaption(TArray<FString>& OutCaptions) const;
-#if WITH_EDITOR
 	/** Get a single line description of the material expression (used for lists) */
 	virtual FString GetDescription() const;
 	/** Get a tooltip for the specified connector. */
@@ -260,13 +264,16 @@ class ENGINE_API UMaterialExpression : public UObject
 	 *	@return int32			The padding (in pixels).
 	 */
 	virtual int GetLabelPadding() { return 0; }
+#if WITH_EDITOR
 	virtual int32 CompilerError(class FMaterialCompiler* Compiler, const TCHAR* pcMessage);
-
+#endif // WITH_EDITOR
 
 	/**
 	 * @return whether the expression preview needs realtime update
 	 */
+#if WITH_EDITOR
 	virtual bool NeedsRealtimePreview() { return false; }
+#endif
 
 	/**
 	 * MatchesSearchQuery: Check this expression to see if it matches the search query
@@ -275,6 +282,7 @@ class ENGINE_API UMaterialExpression : public UObject
 	 */
 	virtual bool MatchesSearchQuery( const TCHAR* SearchQuery );
 
+#if WITH_EDITOR
 	/**
 	 * Copy the SrcExpressions into the specified material, preserving internal references.
 	 * New material expressions are created within the specified material.
@@ -291,12 +299,15 @@ class ENGINE_API UMaterialExpression : public UObject
 	 * Connects the specified output to the passed material for previewing. 
 	 */
 	void ConnectToPreviewMaterial(UMaterial* InMaterial, int32 OutputIndex);
+#endif // WITH_EDITOR
 
+#if WITH_EDITOR
 	/**
 	 * Connects the specified input expression to the specified output of this expression.
 	 */
 	void ConnectExpression(FExpressionInput* Input, int32 OutputIndex);
-	
+#endif // WITH_EDITOR
+
 	/** 
 	* Generates a GUID for the parameter expression if one doesn't already exist and we are one.
 	*
diff --git a/Engine/Source/Runtime/Engine/Classes/Materials/MaterialExpressionAbs.h b/Engine/Source/Runtime/Engine/Classes/Materials/MaterialExpressionAbs.h
--- a/Engine/Source/Runtime/Engine/Classes/Materials/MaterialExpressionAbs.h
+++ b/Engine/Source/Runtime/Engine/Classes/Materials/MaterialExpressionAbs.h
@@ -20,8 +20,10 @@ class UMaterialExpressionAbs : public UMaterialExpression
 
 
 	//~ Begin UMaterialExpression Interface
+#if WITH_EDITOR
 	virtual int32 Compile(class FMaterialCompiler* Compiler, int32 OutputIndex, int32 MultiplexIndex) override;
 	virtual void GetCaption(TArray<FString>& OutCaptions) const override;
+#endif
 	//~ End UMaterialExpression Interface
 
 };
diff --git a/Engine/Source/Runtime/Engine/Classes/Materials/MaterialExpressionActorPositionWS.h b/Engine/Source/Runtime/Engine/Classes/Materials/MaterialExpressionActorPositionWS.h
--- a/Engine/Source/Runtime/Engine/Classes/Materials/MaterialExpressionActorPositionWS.h
+++ b/Engine/Source/Runtime/Engine/Classes/Materials/MaterialExpressionActorPositionWS.h
@@ -12,8 +12,10 @@ class UMaterialExpressionActorPositionWS : public UMaterialExpression
 
 
 	//~ Begin UMaterialExpression Interface
+#if WITH_EDITOR
 	virtual int32 Compile(class FMaterialCompiler* Compiler, int32 OutputIndex, int32 MultiplexIndex) override;
 	virtual void GetCaption(TArray<FString>& OutCaptions) const override;
+#endif
 	//~ End UMaterialExpression Interface
 };
 
diff --git a/Engine/Source/Runtime/Engine/Classes/Materials/MaterialExpressionCollectionParameter.h b/Engine/Source/Runtime/Engine/Classes/Materials/MaterialExpressionCollectionParameter.h
--- a/Engine/Source/Runtime/Engine/Classes/Materials/MaterialExpressionCollectionParameter.h
+++ b/Engine/Source/Runtime/Engine/Classes/Materials/MaterialExpressionCollectionParameter.h
@@ -30,11 +30,14 @@ class UMaterialExpressionCollectionParameter : public UMaterialExpression
 	virtual void PostEditChangeProperty(FPropertyChangedEvent& PropertyChangedEvent) override;
 #endif // WITH_EDITOR
 	virtual void PostLoad() override;
+	virtual bool NeedsLoadForClient() const override;
 	//~ End UObject Interface.
 
 	//~ Begin UMaterialExpression Interface
+#if WITH_EDITOR
 	virtual int32 Compile(class FMaterialCompiler* Compiler, int32 OutputIndex, int32 MultiplexIndex) override;
 	virtual void GetCaption(TArray<FString>& OutCaptions) const override;
+#endif
 	virtual bool MatchesSearchQuery( const TCHAR* SearchQuery ) override;
 	//~ End UMaterialExpression Interface
 };
diff --git a/Engine/Source/Runtime/Engine/Classes/Materials/MaterialExpressionComment.h b/Engine/Source/Runtime/Engine/Classes/Materials/MaterialExpressionComment.h
--- a/Engine/Source/Runtime/Engine/Classes/Materials/MaterialExpressionComment.h
+++ b/Engine/Source/Runtime/Engine/Classes/Materials/MaterialExpressionComment.h
@@ -31,7 +31,9 @@ class UMaterialExpressionComment : public UMaterialExpression
 	//~ End UObject Interface
 
 	//~ Begin UMaterialExpression Interface
+#if WITH_EDITOR
 	virtual void GetCaption(TArray<FString>& OutCaptions) const override;
+#endif // WITH_EDITOR
 	virtual bool MatchesSearchQuery( const TCHAR* SearchQuery ) override;
 	//~ End UMaterialExpression Interface
 };
diff --git a/Engine/Source/Runtime/Engine/Classes/Materials/MaterialExpressionDynamicParameter.h b/Engine/Source/Runtime/Engine/Classes/Materials/MaterialExpressionDynamicParameter.h
--- a/Engine/Source/Runtime/Engine/Classes/Materials/MaterialExpressionDynamicParameter.h
+++ b/Engine/Source/Runtime/Engine/Classes/Materials/MaterialExpressionDynamicParameter.h
@@ -32,11 +32,14 @@ class UMaterialExpressionDynamicParameter : public UMaterialExpression
 #endif // WITH_EDITOR
 
 	virtual void PostLoad() override;
+	virtual bool NeedsLoadForClient() const override;
 	//~ End UObject Interface
 
 	//~ Begin UMaterialExpression Interface
+#if WITH_EDITOR
 	virtual int32 Compile(class FMaterialCompiler* Compiler, int32 OutputIndex, int32 MultiplexIndex) override;
 	virtual void GetCaption(TArray<FString>& OutCaptions) const override;
+#endif
 	virtual TArray<FExpressionOutput>& GetOutputs() override;
 	virtual int32 GetWidth() const override;
 	virtual int32 GetLabelPadding() override { return 8; }
diff --git a/Engine/Source/Runtime/Engine/Classes/Materials/MaterialExpressionFeatureLevelSwitch.h b/Engine/Source/Runtime/Engine/Classes/Materials/MaterialExpressionFeatureLevelSwitch.h
--- a/Engine/Source/Runtime/Engine/Classes/Materials/MaterialExpressionFeatureLevelSwitch.h
+++ b/Engine/Source/Runtime/Engine/Classes/Materials/MaterialExpressionFeatureLevelSwitch.h
@@ -19,20 +19,23 @@ class UMaterialExpressionFeatureLevelSwitch : public UMaterialExpression
 
 	//~ Begin UObject Interface.
 	virtual void Serialize(FArchive& Ar) override;
+	virtual bool NeedsLoadForClient() const override;
 	//~ End UObject Interface.
 
 	//~ Begin UMaterialExpression Interface
+#if WITH_EDITOR
 	virtual int32 Compile(class FMaterialCompiler* Compiler, int32 OutputIndex, int32 MultiplexIndex) override;
 	virtual void GetCaption(TArray<FString>& OutCaptions) const override;
-	
+#endif // WITH_EDITOR
+
 	virtual const TArray<FExpressionInput*> GetInputs() override;
 	virtual FExpressionInput* GetInput(int32 InputIndex) override;
 	virtual FString GetInputName(int32 InputIndex) const override;
 	virtual bool IsInputConnectionRequired(int32 InputIndex) const override;
-	virtual bool IsResultMaterialAttributes(int32 OutputIndex) override;
 #if WITH_EDITOR
+	virtual bool IsResultMaterialAttributes(int32 OutputIndex) override;
 	virtual uint32 GetInputType(int32 InputIndex) override {return MCT_Unknown;}
 	virtual uint32 GetOutputType(int32 OutputIndex) override {return MCT_Unknown;}
-#endif
+#endif // WITH_EDITOR
 	//~ End UMaterialExpression Interface
 };
diff --git a/Engine/Source/Runtime/Engine/Classes/Materials/MaterialExpressionFunctionInput.h b/Engine/Source/Runtime/Engine/Classes/Materials/MaterialExpressionFunctionInput.h
--- a/Engine/Source/Runtime/Engine/Classes/Materials/MaterialExpressionFunctionInput.h
+++ b/Engine/Source/Runtime/Engine/Classes/Materials/MaterialExpressionFunctionInput.h
@@ -82,17 +82,14 @@ class UMaterialExpressionFunctionInput : public UMaterialExpression
 	//~ Begin UMaterialExpression Interface
 #if WITH_EDITOR
 	virtual void GetExpressionToolTip(TArray<FString>& OutToolTip) override;
-#endif
-
 	virtual int32 CompilePreview(class FMaterialCompiler* Compiler, int32 OutputIndex, int32 MultiplexIndex) override;
 	virtual int32 Compile(class FMaterialCompiler* Compiler, int32 OutputIndex, int32 MultiplexIndex) override;
-
 	virtual void GetCaption(TArray<FString>& OutCaptions) const override;
+
 	virtual bool IsResultMaterialAttributes(int32 OutputIndex) override;
-#if WITH_EDITOR
 	virtual uint32 GetInputType(int32 InputIndex) override;
 	virtual uint32 GetOutputType(int32 OutputIndex) override;
-#endif
+#endif // WITH_EDITOR
 	//~ End UMaterialExpression Interface
 
 	/** Generate the Id for this input. */
@@ -104,8 +101,10 @@ class UMaterialExpressionFunctionInput : public UMaterialExpression
 
 private:
 
+#if WITH_EDITOR
 	/** Helper function which compiles this expression for previewing. */
 	int32 CompilePreviewValue(FMaterialCompiler* Compiler, int32 MultiplexIndex=INDEX_NONE);
+#endif // WITH_EDITOR
 };
 
 
diff --git a/Engine/Source/Runtime/Engine/Classes/Materials/MaterialExpressionMakeMaterialAttributes.h b/Engine/Source/Runtime/Engine/Classes/Materials/MaterialExpressionMakeMaterialAttributes.h
--- a/Engine/Source/Runtime/Engine/Classes/Materials/MaterialExpressionMakeMaterialAttributes.h
+++ b/Engine/Source/Runtime/Engine/Classes/Materials/MaterialExpressionMakeMaterialAttributes.h
@@ -65,9 +65,11 @@ class UMaterialExpressionMakeMaterialAttributes : public UMaterialExpression
 	FExpressionInput PixelDepthOffset;
 
 	//~ Begin UMaterialExpression Interface
+#if WITH_EDITOR
 	virtual int32 Compile(class FMaterialCompiler* Compiler, int32 OutputIndex, int32 MultiplexIndex) override;
 	virtual void GetCaption(TArray<FString>& OutCaptions) const override;
 	virtual bool IsResultMaterialAttributes(int32 OutputIndex) override {return true;}
+#endif // WITH_EDITOR
 	//~ End UMaterialExpression Interface
 };
 
diff --git a/Engine/Source/Runtime/Engine/Classes/Materials/MaterialExpressionMaterialFunctionCall.h b/Engine/Source/Runtime/Engine/Classes/Materials/MaterialExpressionMaterialFunctionCall.h
--- a/Engine/Source/Runtime/Engine/Classes/Materials/MaterialExpressionMaterialFunctionCall.h
+++ b/Engine/Source/Runtime/Engine/Classes/Materials/MaterialExpressionMaterialFunctionCall.h
@@ -90,11 +90,14 @@ class UMaterialExpressionMaterialFunctionCall : public UMaterialExpression
 	virtual void PostEditChangeProperty(FPropertyChangedEvent& PropertyChangedEvent) override;
 #endif // WITH_EDITOR
 	virtual void PostLoad() override;
+	virtual bool NeedsLoadForClient() const override;
 	//~ End UObject Interface.
 
 	//~ Begin UMaterialExpression Interface
+#if WITH_EDITOR
 	virtual int32 Compile(class FMaterialCompiler* Compiler, int32 OutputIndex, int32 MultiplexIndex) override;
 	virtual void GetCaption(TArray<FString>& OutCaptions) const override;
+#endif
 	virtual const TArray<FExpressionInput*> GetInputs() override;
 	virtual FExpressionInput* GetInput(int32 InputIndex) override;
 	virtual FString GetInputName(int32 InputIndex) const override;
@@ -103,15 +106,16 @@ class UMaterialExpressionMaterialFunctionCall : public UMaterialExpression
 	virtual FString GetDescription() const override;
 	virtual void GetConnectorToolTip(int32 InputIndex, int32 OutputIndex, TArray<FString>& OutToolTip) override;
 	virtual void GetExpressionToolTip(TArray<FString>& OutToolTip) override;
-#endif
+#endif // WITH_EDITOR
 	virtual bool MatchesSearchQuery( const TCHAR* SearchQuery ) override;
-	virtual bool IsResultMaterialAttributes(int32 OutputIndex) override;
 #if WITH_EDITOR
+	virtual bool IsResultMaterialAttributes(int32 OutputIndex) override;
 	virtual uint32 GetInputType(int32 InputIndex) override;
 #endif // WITH_EDITOR
 	//~ End UMaterialExpression Interface
 
 
+#if WITH_EDITOR
 	/** 
 	 * Set a new material function, given an old function so that links can be passed over if the name matches. 
 	 *
@@ -122,6 +126,7 @@ class UMaterialExpressionMaterialFunctionCall : public UMaterialExpression
 	 *	@return									true if setting the function was a success, false if it failed.
 	 */
 	ENGINE_API bool SetMaterialFunction(UMaterialFunction* ThisFunctionResource, UMaterialFunction* OldFunctionResource, UMaterialFunction* NewResource);
+#endif // WITH_EDITOR
 
 	/** 
 	 * Update FunctionInputs and FunctionOutputs from the MaterialFunction.  
@@ -131,6 +136,7 @@ class UMaterialExpressionMaterialFunctionCall : public UMaterialExpression
 
 private:
 	
+#if WITH_EDITOR
 	/** Helper that fixes up expression links where possible. */
 	void FixupReferencingExpressions(
 		const TArray<FFunctionExpressionOutput>& NewOutputs,
@@ -138,6 +144,7 @@ private:
 		TArray<UMaterialExpression*>& Expressions, 
 		TArray<FExpressionInput*>& MaterialInputs,
 		bool bMatchByName);
+#endif // WITH_EDITOR
 };
 
 
diff --git a/Engine/Source/Runtime/Engine/Classes/Materials/MaterialExpressionParameter.h b/Engine/Source/Runtime/Engine/Classes/Materials/MaterialExpressionParameter.h
--- a/Engine/Source/Runtime/Engine/Classes/Materials/MaterialExpressionParameter.h
+++ b/Engine/Source/Runtime/Engine/Classes/Materials/MaterialExpressionParameter.h
@@ -31,6 +31,10 @@ class UMaterialExpressionParameter : public UMaterialExpression
 #endif
 	//~ End UMaterialExpression Interface
 
+	//~ Begin UObject Interface
+	virtual bool NeedsLoadForClient() const override;
+	//~ End UObject Interface
+
 	ENGINE_API virtual FGuid& GetParameterExpressionId() override
 	{
 		return ExpressionGUID;
diff --git a/Engine/Source/Runtime/Engine/Classes/Materials/MaterialExpressionParticleSubUV.h b/Engine/Source/Runtime/Engine/Classes/Materials/MaterialExpressionParticleSubUV.h
--- a/Engine/Source/Runtime/Engine/Classes/Materials/MaterialExpressionParticleSubUV.h
+++ b/Engine/Source/Runtime/Engine/Classes/Materials/MaterialExpressionParticleSubUV.h
@@ -15,9 +15,11 @@ class UMaterialExpressionParticleSubUV : public UMaterialExpressionTextureSample
 
 
 	//~ Begin UMaterialExpression Interface
+#if WITH_EDITOR
 	virtual int32 Compile(class FMaterialCompiler* Compiler, int32 OutputIndex, int32 MultiplexIndex) override;
-	virtual int32 GetWidth() const override;
 	virtual void GetCaption(TArray<FString>& OutCaptions) const override;
+#endif
+	virtual int32 GetWidth() const override;	
 	virtual int32 GetLabelPadding() override { return 8; }
 	//~ End UMaterialExpression Interface
 };
diff --git a/Engine/Source/Runtime/Engine/Classes/Materials/MaterialExpressionQualitySwitch.h b/Engine/Source/Runtime/Engine/Classes/Materials/MaterialExpressionQualitySwitch.h
--- a/Engine/Source/Runtime/Engine/Classes/Materials/MaterialExpressionQualitySwitch.h
+++ b/Engine/Source/Runtime/Engine/Classes/Materials/MaterialExpressionQualitySwitch.h
@@ -18,17 +18,23 @@ class UMaterialExpressionQualitySwitch : public UMaterialExpression
 	FExpressionInput Inputs[EMaterialQualityLevel::Num];
 
 	//~ Begin UMaterialExpression Interface
+#if WITH_EDITOR
 	virtual int32 Compile(class FMaterialCompiler* Compiler, int32 OutputIndex, int32 MultiplexIndex) override;
 	virtual void GetCaption(TArray<FString>& OutCaptions) const override;
-	
+#endif // WITH_EDITOR
+
 	virtual const TArray<FExpressionInput*> GetInputs() override;
 	virtual FExpressionInput* GetInput(int32 InputIndex) override;
 	virtual FString GetInputName(int32 InputIndex) const override;
 	virtual bool IsInputConnectionRequired(int32 InputIndex) const override;
-	virtual bool IsResultMaterialAttributes(int32 OutputIndex) override;
 #if WITH_EDITOR
+	virtual bool IsResultMaterialAttributes(int32 OutputIndex) override;
 	virtual uint32 GetInputType(int32 InputIndex) override {return MCT_Unknown;}
 	virtual uint32 GetOutputType(int32 InputIndex) override {return MCT_Unknown;}
-#endif
+#endif // WITH_EDITOR
 	//~ End UMaterialExpression Interface
+
+	//~ Begin UObject Interface
+	virtual bool NeedsLoadForClient() const override;
+	//~ End UObject Interface
 };
diff --git a/Engine/Source/Runtime/Engine/Classes/Materials/MaterialExpressionStaticSwitch.h b/Engine/Source/Runtime/Engine/Classes/Materials/MaterialExpressionStaticSwitch.h
--- a/Engine/Source/Runtime/Engine/Classes/Materials/MaterialExpressionStaticSwitch.h
+++ b/Engine/Source/Runtime/Engine/Classes/Materials/MaterialExpressionStaticSwitch.h
@@ -24,14 +24,16 @@ class UMaterialExpressionStaticSwitch : public UMaterialExpression
 
 
 	//~ Begin UMaterialExpression Interface
+#if WITH_EDITOR
 	virtual int32 Compile(class FMaterialCompiler* Compiler, int32 OutputIndex, int32 MultiplexIndex) override;
 	virtual void GetCaption(TArray<FString>& OutCaptions) const override;
+#endif // WITH_EDITOR
 	virtual FString GetInputName(int32 InputIndex) const override;
-	virtual bool IsResultMaterialAttributes(int32 OutputIndex) override;
 #if WITH_EDITOR
+	virtual bool IsResultMaterialAttributes(int32 OutputIndex) override;
 	virtual uint32 GetInputType(int32 InputIndex) override;
 	virtual uint32 GetOutputType(int32 OutputIndex) override {return MCT_Unknown;}
-#endif
+#endif // WITH_EDITOR
 	//~ End UMaterialExpression Interface
 };
 
diff --git a/Engine/Source/Runtime/Engine/Classes/Materials/MaterialExpressionStaticSwitchParameter.h b/Engine/Source/Runtime/Engine/Classes/Materials/MaterialExpressionStaticSwitchParameter.h
--- a/Engine/Source/Runtime/Engine/Classes/Materials/MaterialExpressionStaticSwitchParameter.h
+++ b/Engine/Source/Runtime/Engine/Classes/Materials/MaterialExpressionStaticSwitchParameter.h
@@ -19,13 +19,13 @@ class UMaterialExpressionStaticSwitchParameter : public UMaterialExpressionStati
 
 	//~ Begin UMaterialExpression Interface
 	virtual FString GetInputName(int32 InputIndex) const override;
+#if WITH_EDITOR
 	virtual int32 Compile(class FMaterialCompiler* Compiler, int32 OutputIndex, int32 MultiplexIndex) override;
 	virtual void GetCaption(TArray<FString>& OutCaptions) const override;
 	virtual bool IsResultMaterialAttributes(int32 OutputIndex) override;
-#if WITH_EDITOR
 	virtual uint32 GetInputType(int32 InputIndex) override {return MCT_Unknown;}
 	virtual uint32 GetOutputType(int32 OutputIndex) override {return MCT_Unknown;}
-#endif
+#endif // WITH_EDITOR
 	//~ End UMaterialExpression Interface
 };
 
diff --git a/Engine/Source/Runtime/Engine/Classes/Materials/MaterialExpressionTextureObject.h b/Engine/Source/Runtime/Engine/Classes/Materials/MaterialExpressionTextureObject.h
--- a/Engine/Source/Runtime/Engine/Classes/Materials/MaterialExpressionTextureObject.h
+++ b/Engine/Source/Runtime/Engine/Classes/Materials/MaterialExpressionTextureObject.h
@@ -18,13 +18,13 @@ class UMaterialExpressionTextureObject : public UMaterialExpressionTextureBase
 	virtual void PostEditChangeProperty(FPropertyChangedEvent& PropertyChangedEvent) override;
 #endif // WITH_EDITOR
 	//~ Begin UMaterialExpression Interface
+#if WITH_EDITOR
 	virtual int32 Compile(class FMaterialCompiler* Compiler, int32 OutputIndex, int32 MultiplexIndex) override;
+	virtual int32 CompilePreview(class FMaterialCompiler* Compiler, int32 OutputIndex, int32 MultiplexIndex) override;
 	virtual void GetCaption(TArray<FString>& OutCaptions) const override;
-#if WITH_EDITOR
 	virtual uint32 GetOutputType(int32 OutputIndex) override;
-#endif // WITH_EDITOR
 	//~ End UMaterialExpression Interface
-	virtual int32 CompilePreview(class FMaterialCompiler* Compiler, int32 OutputIndex, int32 MultiplexIndex) override;
+#endif // WITH_EDITOR
 };
 
 
diff --git a/Engine/Source/Runtime/Engine/Classes/Materials/MaterialExpressionTextureObjectParameter.h b/Engine/Source/Runtime/Engine/Classes/Materials/MaterialExpressionTextureObjectParameter.h
--- a/Engine/Source/Runtime/Engine/Classes/Materials/MaterialExpressionTextureObjectParameter.h
+++ b/Engine/Source/Runtime/Engine/Classes/Materials/MaterialExpressionTextureObjectParameter.h
@@ -16,13 +16,13 @@ class UMaterialExpressionTextureObjectParameter : public UMaterialExpressionText
 
 
 	//~ Begin UMaterialExpression Interface
-	virtual void GetCaption(TArray<FString>& OutCaptions) const override;
-	virtual const TArray<FExpressionInput*> GetInputs() override;
+#if WITH_EDITOR
+	virtual void GetCaption(TArray<FString>& OutCaptions) const override;	
 	virtual int32 Compile(class FMaterialCompiler* Compiler, int32 OutputIndex, int32 MultiplexIndex) override;
 	virtual int32 CompilePreview(class FMaterialCompiler* Compiler, int32 OutputIndex, int32 MultiplexIndex) override;
-#if WITH_EDITOR
 	virtual uint32 GetOutputType(int32 InputIndex) override {return MCT_Texture;}
 #endif
+	virtual const TArray<FExpressionInput*> GetInputs() override;
 	//~ End UMaterialExpression Interface
 
 	//~ Begin UMaterialExpressionTextureSampleParameter Interface
diff --git a/Engine/Source/Runtime/Engine/Private/Animation/AnimSequence.cpp b/Engine/Source/Runtime/Engine/Private/Animation/AnimSequence.cpp
--- a/Engine/Source/Runtime/Engine/Private/Animation/AnimSequence.cpp
+++ b/Engine/Source/Runtime/Engine/Private/Animation/AnimSequence.cpp
@@ -122,6 +122,9 @@ UAnimSequence::UAnimSequence(const FObjectInitializer& ObjectInitializer)
 	, bCalcAdditiveDynamically(true)
 {
 	RateScale = 1.0;
+#if WITH_EDITORONLY_DATA
+	bCompressedOnCook = false;
+#endif
 }
 
 void UAnimSequence::PostInitProperties()
@@ -323,12 +326,14 @@ void UAnimSequence::Serialize(FArchive& Ar)
 #if WITH_EDITORONLY_DATA
 	FScopedAnimSequenceCompressedDataCache CompressedDataCache;
 
-	if (Ar.IsCooking() && Ar.IsSaving() && CanBakeAdditive())
+	if (!bCompressedOnCook && Ar.IsCooking() && Ar.IsSaving() && CanBakeAdditive())
 	{
 		CompressedDataCache.InitFrom(this);
 
 		//Make copies of original compressed data
 		BakeOutAdditiveIntoRawData();
+		// Don't do this again
+		bCompressedOnCook = true;
 	}
 #endif
 
