Memory leak fix.

diff --git a/Engine/Source/Runtime/Core/Private/HAL/MallocBinned2.cpp b/Engine/Source/Runtime/Core/Private/HAL/MallocBinned2.cpp
--- a/Engine/Source/Runtime/Core/Private/HAL/MallocBinned2.cpp
+++ b/Engine/Source/Runtime/Core/Private/HAL/MallocBinned2.cpp
@@ -639,8 +639,8 @@ struct FMallocBinned2::Private
 		{
 			FBundleNode* NextBundle = Bundle->NextBundle;
 
-			FBundleNode* Node = Bundle->NextNodeInCurrentBundle;
-			while (Node)
+			FBundleNode* Node = Bundle;
+			do
 			{
 				FBundleNode* NextNode = Node->NextNodeInCurrentBundle;
 				FPoolInfo*   NodePool = FindPoolInfo(Allocator, Node);
@@ -671,7 +671,7 @@ struct FMallocBinned2::Private
 				}
 
 				Node = NextNode;
-			}
+			} while (Node);
 
 			Bundle = NextBundle;
 		}
