Copying //UE4/Dev-Framework to Dev-Main (//UE4/Dev-Main) @ 2894246

==========================
MAJOR FEATURES + CHANGES
==========================

Change 2879756 on 2016/02/24 by [REDACTED_PERSON]

	OR-15722 Fix of crash during stop voice

	- changes Stop to not call XAUDIO2_PLAY_TAILS
	- calls FlushSourceBuffers() before calling stop() so no more buffers are queued
	- Caches realtime buffers so they aren't ever free'd until SoundSource destructor

Change 2879758 on 2016/02/24 by [REDACTED_PERSON]

	UE-27307 Fix for ogg-vorbis file handle getting access after deletion.

Change 2880678 on 2016/02/24 by [REDACTED_PERSON]

	UE-26093 Supporting preview sounds being excempt from editor volume and transient master volume

	- Creating flag for preview sounds from content browser that result in them being exempt from master volume and always playing at full volume.
	- Removing setting the master volume in all the platform devices (because that's silly) and adding it as a normal gain stage in the platform independent code.

Change 2882199 on 2016/02/25 by [REDACTED_PERSON]

	UE-27474 Adding ClearSoundMixClassOverride

	- Added new method to FDynamicParameter to query if the parameter has finished lerping to target value (used to trigger removal override after clearing)
	- New BP function to clear a sound class override to a sound mix
	- If sound mix is not active, will immediately remove override entry since there's no need for complex lerping
	- If sound mix is active, will set the override target value to the value that it would have it wasn't overriden, and lerp to that value. Once the the target value is reached and it was clearing, it will be removed from the sound mix override entry
	- If a sound mix override is set by BP while in the process of clearing, the clearing action is reset and the sound class override will not be cleared.

Change 2886288 on 2016/02/29 by [REDACTED_PERSON]

	#ue4 - Fix requests to crouch to a half height smaller than the radius triggering OnStartCrouch with wrong HeightAdjust value.

Change 2886363 on 2016/02/29 by mason.seay

	Adjusted lightmap resolution on floor mesh and rebuilt lighting to fix lighting errors

Change 2886860 on 2016/02/29 by [REDACTED_PERSON]

	UE-27607 ShooterGame Crashing in Standalone and Launch when pausing/unpausing

	- Fix is to more carefully account for the state of a source getting paused/unpaused while async loading

Change 2887786 on 2016/03/01 by [REDACTED_PERSON]

	Preload HUD and keep it in memory.

Change 2887855 on 2016/03/01 by [REDACTED_PERSON]

	Fix cached modifier key state being able to be incorrect if modifier and other key pressed in same frame
	Don't update cached modifier key state every tick, just when key pressed or app activated

Change 2889984 on 2016/03/02 by [REDACTED_PERSON]

	Fix tooltip typos in 'Find Path to Actor Synchronously' blueprint node

Change 2892730 on 2016/03/03 by [REDACTED_PERSON]

	Added a simple loop for small memcpy.

diff --git a/Engine/Source/Editor/UnrealEd/Private/EditorEngine.cpp b/Engine/Source/Editor/UnrealEd/Private/EditorEngine.cpp
--- a/Engine/Source/Editor/UnrealEd/Private/EditorEngine.cpp
+++ b/Engine/Source/Editor/UnrealEd/Private/EditorEngine.cpp
@@ -1942,7 +1942,7 @@ void UEditorEngine::PlayPreviewSound( USoundBase* Sound,  USoundNode* SoundNode
 		AudioComponent->bAllowSpatialization = false;
 		AudioComponent->bReverb = false;
 		AudioComponent->bCenterChannelOnly = false;
-
+		AudioComponent->bIsPreviewSound = true;
 		AudioComponent->Play();	
 	}
 }
diff --git a/Engine/Source/Runtime/Android/AndroidAudio/Private/AndroidAudioSource.cpp b/Engine/Source/Runtime/Android/AndroidAudio/Private/AndroidAudioSource.cpp
--- a/Engine/Source/Runtime/Android/AndroidAudio/Private/AndroidAudioSource.cpp
+++ b/Engine/Source/Runtime/Android/AndroidAudio/Private/AndroidAudioSource.cpp
@@ -412,7 +412,6 @@ void FSLESSoundSource::Update( void )
 		// Emulate the bleed to rear speakers followed by stereo fold down
 		Volume *= 1.25f;
 	}
-	Volume *= FApp::GetVolumeMultiplier();
 	Volume *= AudioDevice->PlatformAudioHeadroom;
 	Volume = FMath::Clamp(Volume, 0.0f, MAX_VOLUME);
 	
diff --git a/Engine/Source/Runtime/Core/Private/Windows/WindowsApplication.cpp b/Engine/Source/Runtime/Core/Private/Windows/WindowsApplication.cpp
--- a/Engine/Source/Runtime/Core/Private/Windows/WindowsApplication.cpp
+++ b/Engine/Source/Runtime/Core/Private/Windows/WindowsApplication.cpp
@@ -272,6 +272,31 @@ FModifierKeysState FWindowsApplication::GetModifierKeys() const
 	return CachedModifierKeyState;
 }
 
+void FWindowsApplication::UpdateModifierKeyState(int32 ModifierKey, bool bNewState)
+{
+	const bool bIsLeftShiftDown = (ModifierKey == VK_LSHIFT ? bNewState : CachedModifierKeyState.IsLeftShiftDown());
+	const bool bIsRightShiftDown = (ModifierKey == VK_RSHIFT ? bNewState : CachedModifierKeyState.IsRightShiftDown());
+	const bool bIsLeftControlDown = (ModifierKey == VK_LCONTROL ? bNewState : CachedModifierKeyState.IsLeftControlDown());
+	const bool bIsRightControlDown = (ModifierKey == VK_RCONTROL ? bNewState : CachedModifierKeyState.IsRightControlDown());
+	const bool bIsLeftAltDown = (ModifierKey == VK_LMENU ? bNewState : CachedModifierKeyState.IsLeftAltDown());
+	const bool bIsRightAltDown = (ModifierKey == VK_RMENU ? bNewState : CachedModifierKeyState.IsRightAltDown());
+	const bool bAreCapsLocked = (ModifierKey == VK_CAPITAL ? bNewState : CachedModifierKeyState.AreCapsLocked());
+
+	CachedModifierKeyState = FModifierKeysState(bIsLeftShiftDown, bIsRightShiftDown, bIsLeftControlDown, bIsRightControlDown, bIsLeftAltDown, bIsRightAltDown, false, false, bAreCapsLocked); // Win key is ignored
+}
+
+void FWindowsApplication::UpdateAllModifierKeyStates()
+{
+	const bool bIsLeftShiftDown = (::GetAsyncKeyState(VK_LSHIFT) & 0x8000) != 0;
+	const bool bIsRightShiftDown = (::GetAsyncKeyState(VK_RSHIFT) & 0x8000) != 0;
+	const bool bIsLeftControlDown = (::GetAsyncKeyState(VK_LCONTROL) & 0x8000) != 0;
+	const bool bIsRightControlDown = (::GetAsyncKeyState(VK_RCONTROL) & 0x8000) != 0;
+	const bool bIsLeftAltDown = (::GetAsyncKeyState(VK_LMENU) & 0x8000) != 0;
+	const bool bIsRightAltDown = (::GetAsyncKeyState(VK_RMENU) & 0x8000) != 0;
+	const bool bAreCapsLocked = (::GetKeyState(VK_CAPITAL) & 0x0001) != 0;
+
+	CachedModifierKeyState = FModifierKeysState(bIsLeftShiftDown, bIsRightShiftDown, bIsLeftControlDown, bIsRightControlDown, bIsLeftAltDown, bIsRightAltDown, false, false, bAreCapsLocked); // Win key is ignored
+}
 
 static TSharedPtr< FWindowsWindow > FindWindowByHWND(const TArray< TSharedRef< FWindowsWindow > >& WindowsToSearch, HWND HandleToFind)
 {
@@ -1239,6 +1264,7 @@ int32 FWindowsApplication::ProcessDeferredMessage( const FDeferredWindowsMessage
 							PressedModifierKeys.Add( VK_RMENU );
 						}
 					}
+					UpdateModifierKeyState(ActualKey, true);
 					break;
 				case VK_CONTROL:
 					// Differentiate between left and right control
@@ -1258,7 +1284,7 @@ int32 FWindowsApplication::ProcessDeferredMessage( const FDeferredWindowsMessage
 							PressedModifierKeys.Add( VK_RCONTROL );
 						}
 					}
-
+					UpdateModifierKeyState(ActualKey, true);
 					break;
 				case VK_SHIFT:
 					// Differentiate between left and right shift
@@ -1267,6 +1293,10 @@ int32 FWindowsApplication::ProcessDeferredMessage( const FDeferredWindowsMessage
 					{
 						PressedModifierKeys.Add( ActualKey );
 					}
+					UpdateModifierKeyState(ActualKey, true);
+					break;
+				case VK_CAPITAL:
+					UpdateModifierKeyState(VK_CAPITAL, true);
 					break;
 				default:
 					// No translation needed
@@ -1314,6 +1344,7 @@ int32 FWindowsApplication::ProcessDeferredMessage( const FDeferredWindowsMessage
 						ActualKey = VK_RMENU;
 					}
 					PressedModifierKeys.Remove( ActualKey );
+					UpdateModifierKeyState(ActualKey, false);
 					break;
 				case VK_CONTROL:
 					// Differentiate between left and right control
@@ -1326,11 +1357,16 @@ int32 FWindowsApplication::ProcessDeferredMessage( const FDeferredWindowsMessage
 						ActualKey = VK_RCONTROL;
 					}
 					PressedModifierKeys.Remove( ActualKey );
+					UpdateModifierKeyState(ActualKey, false);
 					break;
 				case VK_SHIFT:
 					// Differentiate between left and right shift
 					ActualKey = MapVirtualKey( (lParam & 0x00ff0000) >> 16, MAPVK_VSC_TO_VK_EX);
 					PressedModifierKeys.Remove( ActualKey );
+					UpdateModifierKeyState(ActualKey, false);
+					break;
+				case VK_CAPITAL:
+					UpdateModifierKeyState(VK_CAPITAL, false);
 					break;
 				default:
 					// No translation needed
@@ -1524,6 +1560,8 @@ int32 FWindowsApplication::ProcessDeferredMessage( const FDeferredWindowsMessage
 				}
 				bForceActivateByMouse = false;
 
+				UpdateAllModifierKeyStates();
+
 				if ( CurrentNativeEventWindowPtr.IsValid() )
 				{
 					BOOL Result = false;
@@ -1536,6 +1574,7 @@ int32 FWindowsApplication::ProcessDeferredMessage( const FDeferredWindowsMessage
 			break;
 
 		case WM_ACTIVATEAPP:
+			UpdateAllModifierKeyStates();
 			MessageHandler->OnApplicationActivationChanged( !!wParam );
 			break;
 
@@ -1788,19 +1827,6 @@ void FWindowsApplication::ProcessDeferredEvents( const float TimeDelta )
 	}
 }
 
-void FWindowsApplication::Tick( const float TimeDelta )
-{
-	const bool bIsLeftShiftDown = (::GetAsyncKeyState(VK_LSHIFT) & 0x8000) != 0;
-	const bool bIsRightShiftDown = (::GetAsyncKeyState(VK_RSHIFT) & 0x8000) != 0;
-	const bool bIsLeftControlDown = (::GetAsyncKeyState(VK_LCONTROL) & 0x8000) != 0;
-	const bool bIsRightControlDown = (::GetAsyncKeyState(VK_RCONTROL) & 0x8000) != 0;
-	const bool bIsLeftAltDown = (::GetAsyncKeyState(VK_LMENU) & 0x8000) != 0;
-	const bool bIsRightAltDown = (::GetAsyncKeyState(VK_RMENU) & 0x8000) != 0;
-	const bool bAreCapsLocked = (::GetKeyState(VK_CAPITAL) & 0x0001) != 0;
-
-	CachedModifierKeyState = FModifierKeysState(bIsLeftShiftDown, bIsRightShiftDown, bIsLeftControlDown, bIsRightControlDown, bIsLeftAltDown, bIsRightAltDown, false, false, bAreCapsLocked); // Win key is ignored
-}
-
 void FWindowsApplication::PollGameDeviceState( const float TimeDelta )
 {
 	// initialize any externally-implemented input devices (we delay load initialize the array so any plugins have had time to load)
diff --git a/Engine/Source/Runtime/Core/Public/Windows/WindowsApplication.h b/Engine/Source/Runtime/Core/Public/Windows/WindowsApplication.h
--- a/Engine/Source/Runtime/Core/Public/Windows/WindowsApplication.h
+++ b/Engine/Source/Runtime/Core/Public/Windows/WindowsApplication.h
@@ -333,7 +333,6 @@ public:
 	virtual void PollGameDeviceState( const float TimeDelta ) override;
 	virtual void PumpMessages( const float TimeDelta ) override;
 	virtual void ProcessDeferredEvents( const float TimeDelta ) override;
-	virtual void Tick( const float TimeDelta ) override;
 	virtual TSharedRef< FGenericWindow > MakeWindow() override;
 	virtual void InitializeWindow( const TSharedRef< FGenericWindow >& Window, const TSharedRef< FGenericWindowDefinition >& InDefinition, const TSharedPtr< FGenericWindow >& InParent, const bool bShowImmediately ) override;
 	virtual void SetCapture( const TSharedPtr< FGenericWindow >& InWindow ) override;
@@ -422,6 +421,12 @@ private:
 	/** Queries and caches the number of connected mouse devices. */
 	void QueryConnectedMice();
 
+	/** Helper function to update the cached states of all modifier keys */
+	void UpdateAllModifierKeyStates();
+
+	/** Helper function to update the cached modifier key state when a key is pressed or released */
+	void UpdateModifierKeyState(int32 ModifierKey, bool bNewState);
+
 private:
 
 	static const FIntPoint MinimizedWindowPosition;
diff --git a/Engine/Source/Runtime/Engine/Classes/AI/Navigation/NavigationSystem.h b/Engine/Source/Runtime/Engine/Classes/AI/Navigation/NavigationSystem.h
--- a/Engine/Source/Runtime/Engine/Classes/AI/Navigation/NavigationSystem.h
+++ b/Engine/Source/Runtime/Engine/Classes/AI/Navigation/NavigationSystem.h
@@ -300,7 +300,7 @@ public:
 	static UNavigationPath* FindPathToLocationSynchronously(UObject* WorldContext, const FVector& PathStart, const FVector& PathEnd, AActor* PathfindingContext = NULL, TSubclassOf<UNavigationQueryFilter> FilterClass = NULL);
 
 	/** Finds path instantly, in a FindPath Synchronously. Main advantage over FindPathToLocationSynchronously is that 
-	 *	the resulting path with automatically get updated if goal actor moves more then TetherDistance away from last path node
+	 *	the resulting path will automatically get updated if goal actor moves more than TetherDistance away from last path node
 	 *	@param PathfindingContext could be one of following: NavigationData (like Navmesh actor), Pawn or Controller. This parameter determines parameters of specific pathfinding query */
 	UFUNCTION(BlueprintCallable, Category = "AI|Navigation", meta = (WorldContext="WorldContext"))
 	static UNavigationPath* FindPathToActorSynchronously(UObject* WorldContext, const FVector& PathStart, AActor* GoalActor, float TetherDistance = 50.f, AActor* PathfindingContext = NULL, TSubclassOf<UNavigationQueryFilter> FilterClass = NULL);
diff --git a/Engine/Source/Runtime/Engine/Classes/Components/AudioComponent.h b/Engine/Source/Runtime/Engine/Classes/Components/AudioComponent.h
--- a/Engine/Source/Runtime/Engine/Classes/Components/AudioComponent.h
+++ b/Engine/Source/Runtime/Engine/Classes/Components/AudioComponent.h
@@ -287,6 +287,9 @@ public:
 	/** Whether or not this sound class forces sounds to the center channel */
 	uint32 bCenterChannelOnly:1;
 
+	/** Whether or not this sound is a preview sound */
+	uint32 bIsPreviewSound:1;
+
 	/** Used by the subtitle manager to prioritize subtitles wave instances spawned by this component. */
 	float SubtitlePriority;
 
diff --git a/Engine/Source/Runtime/Engine/Classes/Kismet/GameplayStatics.h b/Engine/Source/Runtime/Engine/Classes/Kismet/GameplayStatics.h
--- a/Engine/Source/Runtime/Engine/Classes/Kismet/GameplayStatics.h
+++ b/Engine/Source/Runtime/Engine/Classes/Kismet/GameplayStatics.h
@@ -549,6 +549,14 @@ class ENGINE_API UGameplayStatics : public UBlueprintFunctionLibrary
 	UFUNCTION(BlueprintCallable, Category = "Audio", meta = (WorldContext = "WorldContextObject"))
 	static void SetSoundMixClassOverride(UObject* WorldContextObject, class USoundMix* InSoundMixModifier, class USoundClass* InSoundClass, float Volume = 1.0f, float Pitch = 1.0f, float FadeInTime = 1.0f, bool bApplyToChildren = true);
 
+	/** Clears the override of the sound class adjuster in the given sound mix. If the override did not exist in the sound mix, this will do nothing.
+	 * @param InSoundMixModifier The sound mix to modify.
+	 * @param InSoundClass The sound class to override (or add) in the sound mix.
+	 * @param FadeOutTime The interpolation time to use to go from the current sound class adjuster override values to the non-override values.
+	*/
+	UFUNCTION(BlueprintCallable, Category = "Audio", meta = (WorldContext = "WorldContextObject"))
+	static void ClearSoundMixClassOverride(UObject* WorldContextObject, class USoundMix* InSoundMixModifier, class USoundClass* InSoundClass, float FadeOutTime = 1.0f);
+
 	/** Push a sound mix modifier onto the audio system **/
 	UFUNCTION(BlueprintCallable, Category="Audio", meta=(WorldContext = "WorldContextObject"))
 	static void PushSoundMixModifier(UObject* WorldContextObject, class USoundMix* InSoundMixModifier);
diff --git a/Engine/Source/Runtime/Engine/Classes/Sound/SoundBase.h b/Engine/Source/Runtime/Engine/Classes/Sound/SoundBase.h
--- a/Engine/Source/Runtime/Engine/Classes/Sound/SoundBase.h
+++ b/Engine/Source/Runtime/Engine/Classes/Sound/SoundBase.h
@@ -104,7 +104,7 @@ public:
 	/**
 	  * Returns the subtitle priority
 	  */
-	virtual float GetSubtitlePriority() const { return 10000.f; };
+	virtual float GetSubtitlePriority() const { return DEFAULT_SUBTITLE_PRIORITY; };
 
 	/** Returns whether or not this sound is looping. */
 	bool IsLooping();
diff --git a/Engine/Source/Runtime/Engine/Private/ActiveSound.cpp b/Engine/Source/Runtime/Engine/Private/ActiveSound.cpp
--- a/Engine/Source/Runtime/Engine/Private/ActiveSound.cpp
+++ b/Engine/Source/Runtime/Engine/Private/ActiveSound.cpp
@@ -37,6 +37,7 @@ FActiveSound::FActiveSound()
 	, bIsMusic(false)
 	, bReverb(false)
 	, bCenterChannelOnly(false)
+	, bIsPreviewSound(false)
 	, bGotInteriorSettings(false)
 #if !(NO_LOGGING || UE_BUILD_SHIPPING || UE_BUILD_TEST)
 	, bWarnedAboutOrphanedLooping(false)
@@ -58,7 +59,7 @@ FActiveSound::FActiveSound()
 	, CurrentOcclusionVolumeAttenuation(1.0f)
 	, ConcurrencyVolumeScale(1.f)
 	, ConcurrencyDuckingVolumeScale(1.f)
-	, SubtitlePriority(0.f)
+	, SubtitlePriority(DEFAULT_SUBTITLE_PRIORITY)
 	, Priority(1.0f)
 	, FocusPriorityScale(1.0f)
 	, FocusDistanceScale(1.0f)
@@ -217,7 +218,18 @@ void FActiveSound::UpdateWaveInstances( TArray<FWaveInstance*> &InWaveInstances,
 	// (even after the Sound has started playing, and this line takes them all into account and gives us
 	// final value that is correct
 	UpdateAdjustVolumeMultiplier(DeltaTime);
-	ParseParams.VolumeMultiplier = VolumeMultiplier * Sound->GetVolumeMultiplier() * CurrentAdjustVolumeMultiplier * AudioDevice->TransientMasterVolume * ConcurrencyVolumeScale;
+
+	// If the sound is a preview sound, then ignore the transient master volume and application volume
+	float MasterVolume = AudioDevice->TransientMasterVolume; 
+	float ApplicationVolume = FApp::GetVolumeMultiplier();
+	if (bIsPreviewSound)
+	{
+		MasterVolume = 1.0f;
+		ApplicationVolume = 1.0f;
+	}
+
+	ParseParams.VolumeMultiplier = VolumeMultiplier * Sound->GetVolumeMultiplier() * CurrentAdjustVolumeMultiplier * MasterVolume * ApplicationVolume * ConcurrencyVolumeScale;
+
 	ParseParams.Priority = Priority;
 	ParseParams.Pitch *= PitchMultiplier * Sound->GetPitchMultiplier();
 	ParseParams.bEnableLowPassFilter = bEnableLowPassFilter;
diff --git a/Engine/Source/Runtime/Engine/Private/AudioDevice.cpp b/Engine/Source/Runtime/Engine/Private/AudioDevice.cpp
--- a/Engine/Source/Runtime/Engine/Private/AudioDevice.cpp
+++ b/Engine/Source/Runtime/Engine/Private/AudioDevice.cpp
@@ -1639,6 +1639,16 @@ static void UpdateClassAdjustorOverrideEntry(FSoundClassAdjuster& ClassAdjustor,
 	// Reset the flags on the override adjuster
 	ClassAdjusterOverride.bOverrideApplied = true;
 	ClassAdjusterOverride.bOverrideChanged = false;
+
+	// Check if we're clearing and check the terminating condition
+	if (ClassAdjusterOverride.bIsClearing)
+	{
+		// If our override dynamic parameter is done, then we've finished clearing
+		if (ClassAdjusterOverride.VolumeOverride.IsDone())
+		{
+			ClassAdjusterOverride.bIsCleared = true;
+		}
+	}
 }
 
 void FAudioDevice::ApplyClassAdjusters(USoundMix* SoundMix, float InterpValue, float DeltaTime)
@@ -1678,12 +1688,20 @@ void FAudioDevice::ApplyClassAdjusters(USoundMix* SoundMix, float InterpValue, f
 			}
 		}
 
-		for (auto& SoundMixOverrideEntry : *SoundMixOverrideMap)
+		TArray<USoundClass*> SoundClassesToRemove;
+		for (TPair<USoundClass*, FSoundMixClassOverride>& SoundMixOverrideEntry : *SoundMixOverrideMap)
 		{
 			// Get the sound class object of the override
 			FSoundMixClassOverride& ClassAdjusterOverride = SoundMixOverrideEntry.Value;
 			USoundClass* SoundClassObject = ClassAdjusterOverride.SoundClassAdjustor.SoundClassObject;
 
+			// If the override has successfully cleared, then just remove it and continue iterating
+			if (ClassAdjusterOverride.bIsCleared)
+			{
+				SoundClassesToRemove.Add(SoundClassObject);
+				continue;
+			}
+
 			// Look for it in the adjusters copy 
 			bool bSoundClassAdjustorExisted = false;
 			for (FSoundClassAdjuster& Entry : *SoundClassAdjusters)
@@ -1712,6 +1730,17 @@ void FAudioDevice::ApplyClassAdjusters(USoundMix* SoundMix, float InterpValue, f
 				SoundClassAdjusters->Add(NewEntry);
 			}
 		}
+
+		for (USoundClass* SoundClassToRemove : SoundClassesToRemove)
+		{
+			SoundMixOverrideMap->Remove(SoundClassToRemove);
+
+			// If there are no more overrides, remove the sound mix override entry
+			if (SoundMixOverrideMap->Num() == 0)
+			{
+				SoundMixClassEffectOverrides.Remove(SoundMix);
+			}
+		}
 	}
 
 	// Loop through the sound class adjusters, everything should be up-to-date
@@ -1992,6 +2021,7 @@ void FAudioDevice::SetSoundMixClassOverride(USoundMix* InSoundMix, USoundClass*
 
 		// Flag that we've changed so that the update will interpolate to new values
 		ClassOverride->bOverrideChanged = true;
+		ClassOverride->bIsClearing = false;
 		ClassOverride->FadeInTime = FadeInTime;
 	}
 	else
@@ -2003,13 +2033,72 @@ void FAudioDevice::SetSoundMixClassOverride(USoundMix* InSoundMix, USoundClass*
 		NewClassOverride.SoundClassAdjustor.PitchAdjuster = Pitch;
 		NewClassOverride.SoundClassAdjustor.bApplyToChildren = bApplyToChildren;
 		NewClassOverride.FadeInTime = FadeInTime;
-		NewClassOverride.bOverrideApplied = false;
-		NewClassOverride.bOverrideChanged = false;
 
 		SoundMixClassOverrideMap.Add(InSoundClass, NewClassOverride);
 	}
 }
 
+void FAudioDevice::ClearSoundMixClassOverride(USoundMix* InSoundMix, USoundClass* InSoundClass, float FadeOutTime)
+{
+	if (!InSoundMix || !InSoundClass)
+	{
+		return;
+	}
+
+	// Get the sound mix class override map for the sound mix. If this doesn't exist, then nobody overrode the sound mix
+	FSoundMixClassOverrideMap* SoundMixClassOverrideMap = SoundMixClassEffectOverrides.Find(InSoundMix);
+	if (!SoundMixClassOverrideMap)
+	{
+		return;
+	}
+
+	// Get the sound class override. If this doesn't exist, then the sound class wasn't previously overridden.
+	FSoundMixClassOverride* SoundClassOverride = SoundMixClassOverrideMap->Find(InSoundClass);
+	if (!SoundClassOverride)
+	{
+		return;
+	}
+
+	// If the override is currently applied, then we need to "fade out" the override
+	if (SoundClassOverride->bOverrideApplied)
+	{
+		// Get the new target values that sound mix would be if it weren't overridden. 
+		// If this was a pure add to the sound mix, then the target values will be 1.0f (i.e. not applied)
+		float VolumeAdjuster = 1.0f;
+		float PitchAdjuster = 1.0f;
+
+		// Loop through the sound mix class adjusters and set the volume adjuster to the value that would be in the sound mix
+		for (const FSoundClassAdjuster& Adjustor : InSoundMix->SoundClassEffects)
+		{
+			if (Adjustor.SoundClassObject == InSoundClass)
+			{
+				VolumeAdjuster = Adjustor.VolumeAdjuster;
+				PitchAdjuster = Adjustor.PitchAdjuster;
+				break;
+			}
+		}
+
+		SoundClassOverride->bIsClearing = true;
+		SoundClassOverride->bIsCleared = false;
+		SoundClassOverride->bOverrideChanged = true;
+		SoundClassOverride->FadeInTime = FadeOutTime;
+		SoundClassOverride->SoundClassAdjustor.VolumeAdjuster = VolumeAdjuster;
+		SoundClassOverride->SoundClassAdjustor.PitchAdjuster = PitchAdjuster;
+	}
+	else
+	{
+		// Otherwise, we just simply remove the sound class override in the sound class override map
+		SoundMixClassOverrideMap->Remove(InSoundClass);
+
+		// If there are no more overrides, remove the sound mix override entry
+		if (!SoundMixClassOverrideMap->Num())
+		{
+			SoundMixClassEffectOverrides.Remove(InSoundMix);
+		}
+	}
+
+}
+
 void FAudioDevice::PopSoundMixModifier(USoundMix* SoundMix, bool bIsPassive)
 {
 	if (SoundMix)
@@ -2387,7 +2476,11 @@ void FAudioDevice::StartSources(TArray<FWaveInstance*>& WaveInstances, int32 Fir
 						if (bSuccess)
 						{
 							check(Source->IsInitialized());
-							Source->Play();
+							// If the source didn't get paused while initializing, then play it
+							if (!Source->IsPaused())
+							{
+								Source->Play();
+							}
 							Source->Update();
 						}
 					}
diff --git a/Engine/Source/Runtime/Engine/Private/Components/AudioComponent.cpp b/Engine/Source/Runtime/Engine/Private/Components/AudioComponent.cpp
--- a/Engine/Source/Runtime/Engine/Private/Components/AudioComponent.cpp
+++ b/Engine/Source/Runtime/Engine/Private/Components/AudioComponent.cpp
@@ -24,6 +24,7 @@ UAudioComponent::UAudioComponent(const FObjectInitializer& ObjectInitializer)
 #endif
 	VolumeMultiplier = 1.f;
 	bOverridePriority = false;
+	bIsPreviewSound = false;
 	Priority = 1.f;
 	PitchMultiplier = 1.f;
 	VolumeModulationMin = 1.f;
@@ -216,7 +217,7 @@ void UAudioComponent::PlayInternal(const float StartTime, const float FadeInDura
 				NewActiveSound.bAlwaysPlay = bAlwaysPlay;
 				NewActiveSound.bReverb = bReverb;
 				NewActiveSound.bCenterChannelOnly = bCenterChannelOnly;
-
+				NewActiveSound.bIsPreviewSound = bIsPreviewSound;
 				NewActiveSound.bLocationDefined = !bPreviewComponent;
 				if (NewActiveSound.bLocationDefined)
 				{
diff --git a/Engine/Source/Runtime/Engine/Private/Components/CharacterMovementComponent.cpp b/Engine/Source/Runtime/Engine/Private/Components/CharacterMovementComponent.cpp
--- a/Engine/Source/Runtime/Engine/Private/Components/CharacterMovementComponent.cpp
+++ b/Engine/Source/Runtime/Engine/Private/Components/CharacterMovementComponent.cpp
@@ -2027,14 +2027,17 @@ void UCharacterMovementComponent::Crouch(bool bClientSimulation)
 	// Change collision size to crouching dimensions
 	const float ComponentScale = CharacterOwner->GetCapsuleComponent()->GetShapeScale();
 	const float OldUnscaledHalfHeight = CharacterOwner->GetCapsuleComponent()->GetUnscaledCapsuleHalfHeight();
-	CharacterOwner->GetCapsuleComponent()->SetCapsuleSize(CharacterOwner->GetCapsuleComponent()->GetUnscaledCapsuleRadius(), CrouchedHalfHeight);
-	float HalfHeightAdjust = (OldUnscaledHalfHeight - CrouchedHalfHeight);
+	const float OldUnscaledRadius = CharacterOwner->GetCapsuleComponent()->GetUnscaledCapsuleRadius();
+	// Height is not allowed to be smaller than radius.
+	const float ClampedCrouchedHalfHeight = FMath::Max3(0.f, OldUnscaledRadius, CrouchedHalfHeight);
+	CharacterOwner->GetCapsuleComponent()->SetCapsuleSize(OldUnscaledRadius, ClampedCrouchedHalfHeight);
+	float HalfHeightAdjust = (OldUnscaledHalfHeight - ClampedCrouchedHalfHeight);
 	float ScaledHalfHeightAdjust = HalfHeightAdjust * ComponentScale;
 
 	if( !bClientSimulation )
 	{
 		// Crouching to a larger height? (this is rare)
-		if (CrouchedHalfHeight > OldUnscaledHalfHeight)
+		if (ClampedCrouchedHalfHeight > OldUnscaledHalfHeight)
 		{
 			FCollisionQueryParams CapsuleParams(CharacterMovementComponentStatics::CrouchTraceName, false, CharacterOwner);
 			FCollisionResponseParams ResponseParam;
@@ -2045,7 +2048,7 @@ void UCharacterMovementComponent::Crouch(bool bClientSimulation)
 			// If encroached, cancel
 			if( bEncroached )
 			{
-				CharacterOwner->GetCapsuleComponent()->SetCapsuleSize(CharacterOwner->GetCapsuleComponent()->GetUnscaledCapsuleRadius(), OldUnscaledHalfHeight);
+				CharacterOwner->GetCapsuleComponent()->SetCapsuleSize(OldUnscaledRadius, OldUnscaledHalfHeight);
 				return;
 			}
 		}
@@ -2063,7 +2066,7 @@ void UCharacterMovementComponent::Crouch(bool bClientSimulation)
 
 	// OnStartCrouch takes the change from the Default size, not the current one (though they are usually the same).
 	ACharacter* DefaultCharacter = CharacterOwner->GetClass()->GetDefaultObject<ACharacter>();
-	HalfHeightAdjust = (DefaultCharacter->GetCapsuleComponent()->GetUnscaledCapsuleHalfHeight() - CrouchedHalfHeight);
+	HalfHeightAdjust = (DefaultCharacter->GetCapsuleComponent()->GetUnscaledCapsuleHalfHeight() - ClampedCrouchedHalfHeight);
 	ScaledHalfHeightAdjust = HalfHeightAdjust * ComponentScale;
 
 	AdjustProxyCapsuleSize();
diff --git a/Engine/Source/Runtime/Engine/Private/GameplayStatics.cpp b/Engine/Source/Runtime/Engine/Private/GameplayStatics.cpp
--- a/Engine/Source/Runtime/Engine/Private/GameplayStatics.cpp
+++ b/Engine/Source/Runtime/Engine/Private/GameplayStatics.cpp
@@ -984,7 +984,7 @@ UAudioComponent* UGameplayStatics::SpawnSoundAtLocation(UObject* WorldContextObj
 		AudioComponent->bAllowSpatialization	= bIsInGameWorld;
 		AudioComponent->bIsUISound				= !bIsInGameWorld;
 		AudioComponent->bAutoDestroy			= true;
-		AudioComponent->SubtitlePriority		= 10000.f; // Fixme: pass in? Do we want that exposed to blueprints though?
+		AudioComponent->SubtitlePriority		= DEFAULT_SUBTITLE_PRIORITY; // Fixme: pass in? Do we want that exposed to blueprints though?
 		AudioComponent->Play(StartTime);
 	}
 
@@ -1030,7 +1030,7 @@ class UAudioComponent* UGameplayStatics::SpawnSoundAttached(class USoundBase* So
 		AudioComponent->bAllowSpatialization	= bIsInGameWorld;
 		AudioComponent->bIsUISound				= !bIsInGameWorld;
 		AudioComponent->bAutoDestroy			= true;
-		AudioComponent->SubtitlePriority		= 10000.f; // Fixme: pass in? Do we want that exposed to blueprints though?
+		AudioComponent->SubtitlePriority		= DEFAULT_SUBTITLE_PRIORITY; // Fixme: pass in? Do we want that exposed to blueprints though?
 		AudioComponent->Play(StartTime);
 	}
 
@@ -1137,6 +1137,25 @@ void UGameplayStatics::SetSoundMixClassOverride(UObject* WorldContextObject, cla
 	}
 }
 
+void UGameplayStatics::ClearSoundMixClassOverride(UObject* WorldContextObject, class USoundMix* InSoundMixModifier, class USoundClass* InSoundClass, float FadeOutTime)
+{
+	if (!InSoundMixModifier || !GEngine || !GEngine->UseSound())
+	{
+		return;
+	}
+
+	UWorld* ThisWorld = GEngine->GetWorldFromContextObject(WorldContextObject);
+	if (!ThisWorld || !ThisWorld->bAllowAudioPlayback)
+	{
+		return;
+	}
+
+	if (FAudioDevice* AudioDevice = ThisWorld->GetAudioDevice())
+	{
+		AudioDevice->ClearSoundMixClassOverride(InSoundMixModifier, InSoundClass, FadeOutTime);
+	}
+}
+
 void UGameplayStatics::PopSoundMixModifier(UObject* WorldContextObject, USoundMix* InSoundMixModifier)
 {
 	if (InSoundMixModifier == nullptr || GEngine == nullptr || !GEngine->UseSound())
diff --git a/Engine/Source/Runtime/Engine/Private/SoundWave.cpp b/Engine/Source/Runtime/Engine/Private/SoundWave.cpp
--- a/Engine/Source/Runtime/Engine/Private/SoundWave.cpp
+++ b/Engine/Source/Runtime/Engine/Private/SoundWave.cpp
@@ -61,7 +61,7 @@ USoundWave::USoundWave(const FObjectInitializer& ObjectInitializer)
 	Volume = 1.0;
 	Pitch = 1.0;
 	CompressionQuality = 40;
-	SubtitlePriority = 10000.f;
+	SubtitlePriority = DEFAULT_SUBTITLE_PRIORITY;
 }
 
 SIZE_T USoundWave::GetResourceSize(EResourceSizeMode::Type Mode)
diff --git a/Engine/Source/Runtime/Engine/Private/SubtitleManager.cpp b/Engine/Source/Runtime/Engine/Private/SubtitleManager.cpp
--- a/Engine/Source/Runtime/Engine/Private/SubtitleManager.cpp
+++ b/Engine/Source/Runtime/Engine/Private/SubtitleManager.cpp
@@ -49,10 +49,9 @@ void FSubtitleManager::QueueSubtitles( PTRINT SubtitleID, float Priority, bool b
 		return;
 	}
 
-	// Warn when a subtitle is played without its priority having been set.
-	if( Priority == 0.0f )
+	// Return when a subtitle is played with its subtitle suppressed.
+	if (Priority == 0.0f)
 	{
-		UE_LOG(LogSubtitle, Warning, TEXT( "Received subtitle with no priority." ) );
 		return;
 	}
 
diff --git a/Engine/Source/Runtime/Engine/Private/VorbisAudioInfo.cpp b/Engine/Source/Runtime/Engine/Private/VorbisAudioInfo.cpp
--- a/Engine/Source/Runtime/Engine/Private/VorbisAudioInfo.cpp
+++ b/Engine/Source/Runtime/Engine/Private/VorbisAudioInfo.cpp
@@ -89,8 +89,10 @@ FVorbisAudioInfo::FVorbisAudioInfo( void )
 
 FVorbisAudioInfo::~FVorbisAudioInfo( void ) 
 { 
-	check( VFWrapper != NULL );
+	FScopeLock ScopeLock(&VorbisCriticalSection);
+	check(VFWrapper != nullptr);
 	delete VFWrapper;
+	VFWrapper = nullptr;
 }
 
 /** Emulate read from memory functionality */
@@ -171,7 +173,13 @@ bool FVorbisAudioInfo::ReadCompressedInfo( const uint8* InSrcBufferData, uint32
 {
 	SCOPE_CYCLE_COUNTER( STAT_VorbisPrepareDecompressionTime );
 
-	check( VFWrapper != NULL );
+	FScopeLock ScopeLock(&VorbisCriticalSection);
+
+	if (!VFWrapper)
+	{
+		return false;
+	}
+
 	ov_callbacks		Callbacks;
 
 	SrcBufferData = InSrcBufferData;
@@ -183,8 +191,6 @@ bool FVorbisAudioInfo::ReadCompressedInfo( const uint8* InSrcBufferData, uint32
 	Callbacks.close_func = OggClose;
 	Callbacks.tell_func = OggTell;
 
-	FScopeLock ScopeLock(&VorbisCriticalSection);
-
 	// Set up the read from memory variables
 	int Result = ov_open_callbacks(this, &VFWrapper->vf, NULL, 0, Callbacks);
 	if (Result < 0)
diff --git a/Engine/Source/Runtime/Engine/Public/ActiveSound.h b/Engine/Source/Runtime/Engine/Public/ActiveSound.h
--- a/Engine/Source/Runtime/Engine/Public/ActiveSound.h
+++ b/Engine/Source/Runtime/Engine/Public/ActiveSound.h
@@ -194,6 +194,9 @@ public:
 	/** Whether or not this sound class forces sounds to the center channel */
 	uint32 bCenterChannelOnly:1;
 
+	/** Whether or not this active sound is a preview sound */
+	uint32 bIsPreviewSound:1;
+
 	/** Whether we have queried for the interior settings at least once */
 	uint32 bGotInteriorSettings:1;
 
diff --git a/Engine/Source/Runtime/Engine/Public/Audio.h b/Engine/Source/Runtime/Engine/Public/Audio.h
--- a/Engine/Source/Runtime/Engine/Public/Audio.h
+++ b/Engine/Source/Runtime/Engine/Public/Audio.h
@@ -54,6 +54,8 @@ ENGINE_API DECLARE_LOG_CATEGORY_EXTERN(LogAudioDebug, Display, All);
 #define MIN_FILTER_BANDWIDTH			0.1f
 #define MAX_FILTER_BANDWIDTH			2.0f
 
+#define DEFAULT_SUBTITLE_PRIORITY		10000.0f
+
 /**
  * Audio stats
  */
@@ -410,6 +412,7 @@ public:
 		, Paused(false)
 		, bInitialized(true) // Note: this is defaulted to true since not all platforms need to deal with async initialization.
 		, bReverbApplied(false)
+		, bIsPreviewSound(false)
 		, StereoBleed(0.0f)
 		, LFEBleed(0.5f)
 		, LPFFrequency(MAX_FILTER_FREQUENCY)
@@ -555,13 +558,15 @@ protected:
 	class FSoundBuffer*		Buffer;
 
 	/** Cached status information whether we are playing or not. */
-	uint32				Playing:1;
+	FThreadSafeBool		Playing;
 	/** Cached status information whether we are paused or not. */
 	uint32				Paused:1;
 	/** Whether or not the sound source is ready to be initialized */
 	uint32				bInitialized:1;
 	/** Cached sound mode value used to detect when to switch outputs. */
 	uint32				bReverbApplied:1;
+	/** Whether or not the sound is a preview sound */
+	uint32				bIsPreviewSound:1;
 	/** The amount of stereo sounds to bleed to the rear speakers */
 	float				StereoBleed;
 	/** The amount of a sound to bleed to the LFE speaker */
@@ -643,10 +648,15 @@ public:
 class ENGINE_API FDynamicParameter
 {
 public:
-	FDynamicParameter(float Value);
+	explicit FDynamicParameter(float Value);
 
 	void Set(float Value, float InDuration);
 	void Update(float DeltaTime);
+	
+	bool IsDone() const 
+	{
+		return CurrTimeSec >= DurationSec;
+	}
 	float GetValue() const
 	{
 		return CurrValue;
diff --git a/Engine/Source/Runtime/Engine/Public/AudioDevice.h b/Engine/Source/Runtime/Engine/Public/AudioDevice.h
--- a/Engine/Source/Runtime/Engine/Public/AudioDevice.h
+++ b/Engine/Source/Runtime/Engine/Public/AudioDevice.h
@@ -170,6 +170,8 @@ struct FSoundMixClassOverride
 	float FadeInTime;
 	uint8 bOverrideApplied : 1;
 	uint8 bOverrideChanged : 1;
+	uint8 bIsClearing : 1;
+	uint8 bIsCleared : 1;
 
 	FSoundMixClassOverride()
 		: VolumeOverride(1.0f)
@@ -177,6 +179,8 @@ struct FSoundMixClassOverride
 		, FadeInTime(0.0f)
 		, bOverrideApplied(false)
 		, bOverrideChanged(false)
+		, bIsClearing(false)
+		, bIsCleared(false)
 	{
 	}
 };
@@ -575,6 +579,11 @@ public:
 	 */
 	void SetSoundMixClassOverride(USoundMix* InSoundMix, USoundClass* InSoundClass, float Volume, float Pitch, float FadeInTime, bool bApplyToChildren);
 
+	/**
+	* Clears a sound class override in the given sound mix.
+	*/
+	void ClearSoundMixClassOverride(USoundMix* InSoundMix, USoundClass* InSoundClass, float FadeOutTime);
+
 	/**
 	 * Pop a SoundMix from the Audio Device's list.
 	 *
@@ -962,9 +971,6 @@ public:
 	/** The Base SoundMix that should be applied by default */
 	USoundMix* DefaultBaseSoundMix;
 
-	/** Special sound mix set by user preferences */
-	USoundMix* UserPreferencesSoundMix;
-
 	/** Map of sound mixes currently affecting audio properties */
 	TMap<USoundMix*, FSoundMixState> SoundMixModifiers;
 
diff --git a/Engine/Source/Runtime/IOS/IOSAudio/Private/IOSAudioSource.cpp b/Engine/Source/Runtime/IOS/IOSAudio/Private/IOSAudioSource.cpp
--- a/Engine/Source/Runtime/IOS/IOSAudio/Private/IOSAudioSource.cpp
+++ b/Engine/Source/Runtime/IOS/IOSAudio/Private/IOSAudioSource.cpp
@@ -257,7 +257,6 @@ void FIOSAudioSoundSource::Update(void)
 	}
 
 	// Apply global multiplier to disable sound when not the foreground app
-	Volume *= FApp::GetVolumeMultiplier();
 	Volume *= AudioDevice->PlatformAudioHeadroom;
 	Volume = FMath::Clamp(Volume, 0.0f, 1.0f);
 
diff --git a/Engine/Source/Runtime/Mac/CoreAudio/Private/CoreAudioSource.cpp b/Engine/Source/Runtime/Mac/CoreAudio/Private/CoreAudioSource.cpp
--- a/Engine/Source/Runtime/Mac/CoreAudio/Private/CoreAudioSource.cpp
+++ b/Engine/Source/Runtime/Mac/CoreAudio/Private/CoreAudioSource.cpp
@@ -332,7 +332,6 @@ void FCoreAudioSoundSource::Update( void )
 		}
 		
 		// apply global multiplier (ie to disable sound when not the foreground app)
-		Volume *= FApp::GetVolumeMultiplier();
 		Volume = FMath::Clamp<float>( Volume, 0.0f, MAX_VOLUME );
 		
 		// Convert to dB
@@ -370,7 +369,6 @@ void FCoreAudioSoundSource::Update( void )
 	else
 	{
 		// apply global multiplier (ie to disable sound when not the foreground app)
-		Volume *= FApp::GetVolumeMultiplier();
 		Volume = FMath::Clamp<float>( Volume, 0.0f, MAX_VOLUME );
 		
 		if( AudioDevice->GetMixDebugState() == DEBUGSTATE_IsolateReverb )
diff --git a/Engine/Source/Runtime/Windows/XAudio2/Private/XAudio2Source.cpp b/Engine/Source/Runtime/Windows/XAudio2/Private/XAudio2Source.cpp
--- a/Engine/Source/Runtime/Windows/XAudio2/Private/XAudio2Source.cpp
+++ b/Engine/Source/Runtime/Windows/XAudio2/Private/XAudio2Source.cpp
@@ -115,11 +115,6 @@ void FXAudio2SoundSource::FreeResources( void )
 	// If we're a streaming buffer...
 	if( bResourcesNeedFreeing )
 	{
-		// ... free the buffers
-		FMemory::Free( ( void* )XAudio2Buffers[0].pAudioData );
-		FMemory::Free( ( void* )XAudio2Buffers[1].pAudioData );
-		FMemory::Free( ( void* )XAudio2Buffers[2].pAudioData );
-
 		// Buffers without a valid resource ID are transient and need to be deleted.
 		if( Buffer )
 		{
@@ -213,9 +208,18 @@ bool FXAudio2SoundSource::ReadMorePCMData( const int32 BufferIndex, EDataReadMod
 	}
 }
 
-/** 
- * Submit the relevant audio buffers to the system
- */
+uint8* FXAudio2SoundSource::GetRealtimeBufferData(const int32 InBufferIndex, const int32 InBufferSize)
+{
+	// Only supporting 3 realtime buffers
+	check(InBufferIndex < 3);
+
+	// Resize the array in case the new buffer size is bigger than previously allocated
+	RealtimeBufferData[InBufferIndex].Reset();
+	RealtimeBufferData[InBufferIndex].AddZeroed(InBufferSize);
+
+	return RealtimeBufferData[InBufferIndex].GetData();
+}
+
 void FXAudio2SoundSource::SubmitPCMRTBuffers( void )
 {
 	SCOPE_CYCLE_COUNTER( STAT_AudioSubmitBuffersTime );
@@ -226,13 +230,13 @@ void FXAudio2SoundSource::SubmitPCMRTBuffers( void )
 	CurrentBuffer = 0;
 
 	const uint32 BufferSize = MONO_PCM_BUFFER_SIZE * Buffer->NumChannels;
-	
+
 	// Set up buffer areas to decompress to
-	XAudio2Buffers[0].pAudioData = (uint8*)FMemory::Malloc(BufferSize);
+	XAudio2Buffers[0].pAudioData = GetRealtimeBufferData(0, BufferSize);
 	XAudio2Buffers[0].AudioBytes = BufferSize;
 	XAudio2Buffers[0].pContext = this;
 
-	XAudio2Buffers[1].pAudioData = (uint8*)FMemory::Malloc(BufferSize);
+	XAudio2Buffers[1].pAudioData = GetRealtimeBufferData(1, BufferSize);
 	XAudio2Buffers[1].AudioBytes = BufferSize;
 	XAudio2Buffers[1].pContext = this;
 
@@ -263,7 +267,7 @@ void FXAudio2SoundSource::SubmitPCMRTBuffers( void )
 	// Prepare the third buffer for the OnBufferEnd callback to write to in the OnBufferEnd callback
 	CurrentBuffer = 2;
 
-	XAudio2Buffers[2].pAudioData = (uint8*)FMemory::Malloc(BufferSize);
+	XAudio2Buffers[2].pAudioData = GetRealtimeBufferData(2, BufferSize);
 	XAudio2Buffers[2].AudioBytes = BufferSize;
 	XAudio2Buffers[2].pContext = this;
 
@@ -569,7 +573,7 @@ bool FXAudio2SoundSource::Init(FWaveInstance* InWaveInstance)
  */
 void FXAudio2SoundSource::GetChannelVolumes(float ChannelVolumes[CHANNEL_MATRIX_COUNT], float AttenuatedVolume)
 {
-	if (FApp::GetVolumeMultiplier() == 0.0f || AudioDevice->IsAudioDeviceMuted())
+	if (AudioDevice->IsAudioDeviceMuted())
 	{
 		for( int32 i = 0; i < CHANNELOUT_COUNT; i++ )
 		{
@@ -632,7 +636,7 @@ void FXAudio2SoundSource::GetChannelVolumes(float ChannelVolumes[CHANNEL_MATRIX_
 			UE_LOG(LogXAudio2, Warning, TEXT("FXAudio2SoundSource contains unreasonble value %f in channel %d: %s"), ChannelVolumes[i], i, *Describe_Internal(true, false));
 		}
 
-		ChannelVolumes[i] = FMath::Clamp<float>(ChannelVolumes[i] * FApp::GetVolumeMultiplier() * AudioDevice->PlatformAudioHeadroom, 0.0f, MAX_VOLUME);
+		ChannelVolumes[i] = FMath::Clamp<float>(ChannelVolumes[i] * AudioDevice->PlatformAudioHeadroom, 0.0f, MAX_VOLUME);
 	}
 }
 
@@ -1324,11 +1328,9 @@ void FXAudio2SoundSource::Update()
 
 void FXAudio2SoundSource::Play()
 {
-	check(bInitialized);
-
-	if( WaveInstance )
+	if (WaveInstance)
 	{
-		if( !Playing )
+		if (!Playing)
 		{
 			if( Buffer->NumChannels >= SPEAKER_COUNT )
 			{
@@ -1336,10 +1338,12 @@ void FXAudio2SoundSource::Play()
 			}
 		}
 
-		if( Source )
+		// It's possible if Pause and Play are called while a sound is async initializing. In this case
+		// we'll just not actually play the source here. Instead we'll call play when the sound finishes loading.
+		if (Source && bInitialized)
 		{
-			AudioDevice->ValidateAPICall( TEXT( "Start" ), 
-				Source->Start( 0 ) );
+			AudioDevice->ValidateAPICall(TEXT("Start"),
+										 Source->Start(0));
 		}
 
 		Paused = false;
@@ -1363,10 +1367,13 @@ void FXAudio2SoundSource::Stop()
 			}
 		}
 
-		if( Source )
+		if (Source && Playing)
 		{
-			AudioDevice->ValidateAPICall( TEXT( "Stop" ), 
-				Source->Stop( XAUDIO2_PLAY_TAILS ) );
+			AudioDevice->ValidateAPICall(TEXT("FlushSourceBuffers"),
+										 Source->FlushSourceBuffers());
+
+			AudioDevice->ValidateAPICall(TEXT("Stop"),
+										 Source->Stop(0));
 		}
 
 		Paused = false;
@@ -1386,11 +1393,14 @@ void FXAudio2SoundSource::Stop()
 
 void FXAudio2SoundSource::Pause()
 {
-	check(bInitialized);
 	if (WaveInstance)
 	{
 		if (Source)
 		{
+			// If a source is pasued while it's async loading for realtime decoding,
+			// we'll set the paused flag but our IXAudio2Source pointer won't be valid yet.
+			// We check if the sound is paused after initialization finishes.
+			check(bInitialized);
 			AudioDevice->ValidateAPICall(TEXT("Stop"),
 										 Source->Stop(0));
 		}
diff --git a/Engine/Source/Runtime/Windows/XAudio2/Private/XAudio2Support.h b/Engine/Source/Runtime/Windows/XAudio2/Private/XAudio2Support.h
--- a/Engine/Source/Runtime/Windows/XAudio2/Private/XAudio2Support.h
+++ b/Engine/Source/Runtime/Windows/XAudio2/Private/XAudio2Support.h
@@ -420,6 +420,9 @@ protected:
 	/** Decompress through XAudio2Buffer, or call USoundWave procedure to generate more PCM data. Returns true/false: did audio loop? */
 	bool ReadMorePCMData(const int32 BufferIndex, EDataReadMode DataReadMode);
 
+	/** Retrieves the realtime buffer data from the given buffer index */
+	uint8* GetRealtimeBufferData(const int32 InBufferIndex, const int32 InBufferSize);
+
 	/** Returns if the source is using the default 3d spatialization. */
 	bool IsUsingHrtfSpatializer();
 
@@ -480,9 +483,12 @@ protected:
 	/** Destination voices */
 	XAUDIO2_SEND_DESCRIPTOR Destinations[DEST_COUNT];
 
-	/** A pair of sound buffers to allow notification when a sound loops. */
+	/** Used to allow notification when a sound loops and to feed audio to realtime decoded sources. */
 	XAUDIO2_BUFFER XAudio2Buffers[3];
 
+	/** Raw real-time buffer data for use with realtime XAudio2Buffers sources */
+	TArray<uint8> RealtimeBufferData[3];
+	
 	/** Additional buffer info for XWMA sounds */
 	XAUDIO2_BUFFER_WMA XAudio2BufferXWMA[1];
 
