Merging //UE4/Release-4.11 to //UE4/Main (Up to CL#2897738)

diff --git a/Engine/Source/Programs/CrashReporter/CrashReportClient/Private/CrashDescription.h b/Engine/Source/Programs/CrashReporter/CrashReportClient/Private/CrashDescription.h
--- a/Engine/Source/Programs/CrashReporter/CrashReportClient/Private/CrashDescription.h
+++ b/Engine/Source/Programs/CrashReporter/CrashReportClient/Private/CrashDescription.h
@@ -338,6 +338,9 @@ struct FPrimaryCrashProperties
 	/** Whether this crash contains primary usable data. */
 	bool bHasPrimaryData;
 
+	/** Copy of CommandLine that isn't anonymized so it can be used to restart the process */
+	FString RestartCommandLine;
+
 	/**
 	 *	Whether the report comes from a non-fatal event such as an ensure
 	 */
diff --git a/Engine/Source/Programs/CrashReporter/CrashReportClient/Private/CrashReportClient.cpp b/Engine/Source/Programs/CrashReporter/CrashReportClient/Private/CrashReportClient.cpp
--- a/Engine/Source/Programs/CrashReporter/CrashReportClient/Private/CrashReportClient.cpp
+++ b/Engine/Source/Programs/CrashReporter/CrashReportClient/Private/CrashReportClient.cpp
@@ -2,13 +2,11 @@
 
 #include "CrashReportClientApp.h"
 #include "CrashReportUtil.h"
-
 #include "CrashDescription.h"
-
 #include "CrashReportClient.h"
 #include "UniquePtr.h"
-
 #include "TaskGraphInterfaces.h"
+#include "DesktopPlatformModule.h"
 
 #define LOCTEXT_NAMESPACE "CrashReportClient"
 
@@ -93,10 +91,48 @@ FReply FCrashReportClient::SubmitAndRestart()
 {
 	Submit();
 
+	// Check for processes that were started from the Launcher using -EpicPortal on the command line
+	bool bRunFromLauncher = FParse::Param(*FPrimaryCrashProperties::Get()->RestartCommandLine, TEXT("EPICPORTAL"));
 	const FString CrashedAppPath = ErrorReport.FindCrashedAppPath();
-	const FString CommandLineArguments = FPrimaryCrashProperties::Get()->CommandLine.AsString();
 
-	FPlatformProcess::CreateProc(*CrashedAppPath, *CommandLineArguments, true, false, false, NULL, 0, NULL, NULL);
+	bool bLauncherRestarted = false;
+	if (bRunFromLauncher)
+	{
+		// We'll restart Launcher-run processes by having the installed Launcher handle it
+		IDesktopPlatform* DesktopPlatform = FDesktopPlatformModule::Get();
+
+		if (DesktopPlatform != nullptr)
+		{
+			// Split the path so we can format it as a URI
+			TArray<FString> PathArray;
+			CrashedAppPath.Replace(TEXT("//"), TEXT("/")).ParseIntoArray(PathArray, TEXT("/"), false);	// WER saves this out on Windows with double slashes as the separator for some reason.
+			FString CrashedAppPathUri;
+
+			// Exclude the last item (the filename). The Launcher currently expects an installed application folder.
+			for (int32 ItemIndex = 0; ItemIndex < PathArray.Num() - 1; ItemIndex++)
+			{
+				FString& PathItem = PathArray[ItemIndex];
+				CrashedAppPathUri += FPlatformHttp::UrlEncode(PathItem);
+				CrashedAppPathUri += TEXT("/");
+			}
+			CrashedAppPathUri.RemoveAt(CrashedAppPathUri.Len() - 1);
+
+			// Re-run the application via the Launcher
+			FOpenLauncherOptions OpenOptions(FString::Printf(TEXT("apps/%s"), *CrashedAppPathUri));
+			OpenOptions.bSilent = true;
+			if (DesktopPlatform->OpenLauncher(OpenOptions))
+			{
+				bLauncherRestarted = true;
+			}
+		}
+	}
+
+	if (!bLauncherRestarted)
+	{
+		// Launcher didn't restart the process so start it ourselves
+		const FString CommandLineArguments = FPrimaryCrashProperties::Get()->RestartCommandLine;
+		FPlatformProcess::CreateProc(*CrashedAppPath, *CommandLineArguments, true, false, false, NULL, 0, NULL, NULL);
+	}
 
 	return FReply::Handled();
 }
diff --git a/Engine/Source/Programs/CrashReporter/CrashReportClient/Private/CrashReportClientApp.cpp b/Engine/Source/Programs/CrashReporter/CrashReportClient/Private/CrashReportClientApp.cpp
--- a/Engine/Source/Programs/CrashReporter/CrashReportClient/Private/CrashReportClientApp.cpp
+++ b/Engine/Source/Programs/CrashReporter/CrashReportClient/Private/CrashReportClientApp.cpp
@@ -1,6 +1,7 @@
 // Copyright 1998-2016 Epic Games, Inc. All Rights Reserved.
 
 #include "CrashReportClientApp.h"
+#include "EngineVersion.h"
 #include "CrashDescription.h"
 #include "CrashReportAnalytics.h"
 #include "QoSReporter.h"
@@ -248,4 +249,4 @@ void RunCrashReportClient(const TCHAR* CommandLine)
 	FTaskGraphInterface::Shutdown();
 
 	FEngineLoop::AppExit();
-}
\ No newline at end of file
+}
diff --git a/Engine/Source/Programs/IOS/iPhonePackager/CodeSigning.cs b/Engine/Source/Programs/IOS/iPhonePackager/CodeSigning.cs
--- a/Engine/Source/Programs/IOS/iPhonePackager/CodeSigning.cs
+++ b/Engine/Source/Programs/IOS/iPhonePackager/CodeSigning.cs
@@ -85,7 +85,7 @@ namespace iPhonePackager
 				string plistFile = Config.EngineBuildDirectory + "/UE4Game-Info.plist";
 				if (!string.IsNullOrEmpty(Config.ProjectFile))
 				{
-					plistFile = Path.GetDirectoryName(Config.ProjectFile) + "/Intermediate/IOS/" + Path.GetFileNameWithoutExtension(Config.ProjectFile) + "-Info.plist";
+					plistFile = Path.GetDirectoryName(Config.ProjectFile) + "/Intermediate/" + Config.OSString + "/" + Path.GetFileNameWithoutExtension(Config.ProjectFile) + "-Info.plist";
 
 					if (!File.Exists(plistFile))
 					{
diff --git a/Engine/Source/Programs/IOS/iPhonePackager/CompileTime.cs b/Engine/Source/Programs/IOS/iPhonePackager/CompileTime.cs
--- a/Engine/Source/Programs/IOS/iPhonePackager/CompileTime.cs
+++ b/Engine/Source/Programs/IOS/iPhonePackager/CompileTime.cs
@@ -71,11 +71,18 @@ namespace iPhonePackager
 		static private string GetBaseXcodeCommandline()
 		{
 			string CmdLine = XcodeDeveloperDir + "usr/bin/xcodebuild" +
-					" -project UE4_FromPC.xcodeproj" + 
+					" -project UE4_FromPC.xcodeproj" +
 					" -configuration \"" + Program.SchemeConfiguration + "\"" +
-					" -target '" + Program.SchemeName + "'" +
-					" -destination generic/platform=iOS" +
-					" -sdk " + ((Program.Architecture == "-simulator") ? "iphonesimulator" : "iphoneos");
+					" -target '" + Program.SchemeName + "'";
+			
+			if (Config.OSString == "IOS")
+			{
+				CmdLine += " -destination generic/platform=iOS" +" -sdk " + ((Program.Architecture == "-simulator") ? "iphonesimulator" : "iphoneos");
+			}
+			else
+			{
+				CmdLine += " -destination generic/platform=tvOS" + " -sdk " + ((Program.Architecture == "-simulator") ? "appletvsimulator" : "appletvos");
+			}
 
 			// sign with the Distribution identity when packaging for distribution
 			if (Config.bUseRPCUtil)
@@ -131,11 +138,11 @@ namespace iPhonePackager
 			Console.WriteLine("BranchPath = {0} --- GameBranchPath = {1}", BranchPath, GameBranchPath);
 
 			// generate the directories to recursively copy into later on
-			MacStagingRootDir = string.Format("{0}/{1}/IOS", iPhone_SigningDevRootMac, GameBranchPath);
+			MacStagingRootDir = string.Format("{0}/{1}/" + Config.OSString, iPhone_SigningDevRootMac, GameBranchPath);
 			MacStagingRootDir = MacStagingRootDir.Replace("//", "/");
-			MacBinariesDir = string.Format("{0}/{1}/IOS", iPhone_SigningDevRootMac, GameBranchPath);
+			MacBinariesDir = string.Format("{0}/{1}/" + Config.OSString, iPhone_SigningDevRootMac, GameBranchPath);
 			MacBinariesDir = MacBinariesDir.Replace("//", "/");
-			MacXcodeStagingDir = string.Format("{0}/{1}/IOS/XcodeSupportFiles", iPhone_SigningDevRootMac, GameBranchPath);
+			MacXcodeStagingDir = string.Format("{0}/{1}/" + Config.OSString + "/XcodeSupportFiles", iPhone_SigningDevRootMac, GameBranchPath);
 			MacXcodeStagingDir = MacXcodeStagingDir.Replace("//", "/");
 
 			MacMobileProvisionFilename = MachineName + "_UE4Temp.mobileprovision";
@@ -277,7 +284,7 @@ namespace iPhonePackager
 			File.WriteAllBytes(Path.Combine(Config.PCXcodeStagingDir, MacSigningIdentityFilename), Data);
 
 			// needs Mac line endings so it can be executed
-			string SrcPath = @"..\..\..\Build\IOS\XcodeSupportFiles\prepackage.sh";
+			string SrcPath = @"..\..\..\Build\" + Config.OSString + @"\XcodeSupportFiles\prepackage.sh";
 			string DestPath = Path.Combine(Config.PCXcodeStagingDir, @"prepackage.sh");
 			Program.Log(" ... '" + SrcPath + "' -> '" + DestPath + "'");
 			string SHContents = File.ReadAllText(SrcPath);
@@ -345,7 +352,7 @@ namespace iPhonePackager
 
 			case "prepackage":
 				Program.Log(" ... running prepackage script remotely ");
-				DisplayCommandLine = String.Format("sh prepackage.sh {0} IOS {1} {2}", Program.GameName, Program.GameConfiguration, Program.Architecture);
+				DisplayCommandLine = String.Format("sh prepackage.sh {0} " + Config.OSString + " {1} {2}", Program.GameName, Program.GameConfiguration, Program.Architecture);
 				CommandLine = "\"" + MacXcodeStagingDir + "\" " + DisplayCommandLine;
 				WorkingFolder = "\"" + MacXcodeStagingDir + "\"";
 				break;
diff --git a/Engine/Source/Programs/IOS/iPhonePackager/Config.cs b/Engine/Source/Programs/IOS/iPhonePackager/Config.cs
--- a/Engine/Source/Programs/IOS/iPhonePackager/Config.cs
+++ b/Engine/Source/Programs/IOS/iPhonePackager/Config.cs
@@ -68,6 +68,11 @@ namespace iPhonePackager
 		/// </summary>
 		public static string OverrideDevRoot = null;
 
+		/// <summary>
+		/// Optional value specifying that we are working with tvOS
+		/// </summary>
+		public static string OSString = "IOS";
+
 		/// <summary>
 		/// The local build directory (on PC)
 		/// </summary>
@@ -104,9 +109,9 @@ namespace iPhonePackager
 		{
 			get { 
 				if (Environment.OSVersion.Platform == PlatformID.MacOSX || Environment.OSVersion.Platform == PlatformID.Unix) {
-					return Path.GetFullPath (RootRelativePath + "Engine/Build/IOS"); 
+					return Path.GetFullPath (RootRelativePath + "Engine/Build/" + Config.OSString); 
 				} else {
-					return Path.GetFullPath (RootRelativePath + @"Engine\Build\IOS"); 
+					return Path.GetFullPath (RootRelativePath + @"Engine\Build\" + Config.OSString); 
 				}
 			}
 		}
@@ -118,7 +123,7 @@ namespace iPhonePackager
 		{
 			get
 			{
-				string IntermediateGameBuildDir = GameDirectory + @"\Intermediate\IOS";
+				string IntermediateGameBuildDir = GameDirectory + @"\Intermediate\" + Config.OSString;
 				// if the normal Build dir exists, return it, otherwise, use the program Resources directory
 				return Path.GetFullPath(Directory.Exists(IntermediateGameBuildDir) ?
 					IntermediateGameBuildDir :
@@ -138,7 +143,7 @@ namespace iPhonePackager
 		/// </summary>
 		public static string CookedDirectory
 		{
-			get { return Path.GetFullPath(GameDirectory + @"\Saved\Cooked\IOS"); }
+			get { return Path.GetFullPath(GameDirectory + @"\Saved\Cooked\" + Config.OSString); }
 		}
 
 		/// <summary>
@@ -195,7 +200,7 @@ namespace iPhonePackager
 				}
 				else
 				{
-					return SigningPrefix + Program.GameName + "-IOS-" + Program.GameConfiguration + Program.Architecture + ".ipa";
+					return SigningPrefix + Program.GameName + "-" + Config.OSString + "-" + Program.GameConfiguration + Program.Architecture + ".ipa";
 				}
 			}
 		}
@@ -235,7 +240,7 @@ namespace iPhonePackager
 			}
 			else
 			{
-				Filename = Path.Combine(Config.BinariesDirectory, FilePrefix + Program.GameName + "-IOS-" + Program.GameConfiguration + Program.Architecture + FileSuffix);
+				Filename = Path.Combine(Config.BinariesDirectory, FilePrefix + Program.GameName + "-" + Config.OSString + "-" + Program.GameConfiguration + Program.Architecture + FileSuffix);
 			}
 
 			return Filename;
@@ -252,7 +257,7 @@ namespace iPhonePackager
 			string GameName = Program.GameName;
 			if (!String.IsNullOrEmpty(ProjectFile))
 			{
-				BinariesDir = Path.Combine(Path.GetDirectoryName(ProjectFile), "Binaries", "IOS");
+				BinariesDir = Path.Combine(Path.GetDirectoryName(ProjectFile), "Binaries", Config.OSString);
 				GameName = Path.GetFileNameWithoutExtension(ProjectFile);
 			}
 			if (Program.GameConfiguration == "Development")
@@ -261,7 +266,7 @@ namespace iPhonePackager
 			}
 			else
 			{
-				Filename = Path.Combine(BinariesDir, FilePrefix + GameName + "-IOS-" + Program.GameConfiguration + Program.Architecture + FileSuffix);
+				Filename = Path.Combine(BinariesDir, FilePrefix + GameName + "-" + Config.OSString + "-" + Program.GameConfiguration + Program.Architecture + FileSuffix);
 			}
 
 			// ensure the directory exists
@@ -280,7 +285,7 @@ namespace iPhonePackager
 			}
 			else
 			{
-				return FileOperations.FindPrefixedFile(Config.BinariesDirectory, Program.GameName + "-IOS-" + Program.GameConfiguration + Program.Architecture + FileSuffix);
+				return FileOperations.FindPrefixedFile(Config.BinariesDirectory, Program.GameName + "-" + Config.OSString + "-" + Program.GameConfiguration + Program.Architecture + FileSuffix);
 			}
 		}
 
@@ -391,7 +396,7 @@ namespace iPhonePackager
 		/// <returns></returns>
 		public static string GetRootBackedUpDocumentsDirectory()
 		{
-			return Path.GetFullPath(Path.Combine(GameDirectory + @"\IOS_Backups"));
+			return Path.GetFullPath(Path.Combine(GameDirectory + @"\" + Config.OSString + "_Backups"));
 		}
 
 		public static bool Initialize(string InitialCurrentDirectory, string GamePath)
@@ -434,12 +439,12 @@ namespace iPhonePackager
 			// special case handling for anything inside Engine/Source, it will go to the Engine/Binaries directory, not the Game's binaries directory
 			if (OrigGamePath.Replace("\\", "/").Contains("Engine/Source"))
 			{
-				BinariesDirectory = Path.Combine(RootRelativePath, @"Engine\Binaries\IOS\");
+				BinariesDirectory = Path.Combine(RootRelativePath, @"Engine\Binaries\" + Config.OSString + @"\");
 			}
-			else if (!OrigGamePath.Contains(@"Binaries\IOS"))
+			else if (!OrigGamePath.Contains(@"Binaries\" + Config.OSString))
 			{
 				// no sense in adding Binaries\IOS when it's already there. This is a special case to handle packaging UnrealLaunchDaemon from the command line.
-				BinariesDirectory = Path.Combine(GameDirectory, @"Binaries\IOS\");
+				BinariesDirectory = Path.Combine(GameDirectory, @"Binaries\" + Config.OSString + @"\");
 			}
 			else
 			{
diff --git a/Engine/Source/Programs/IOS/iPhonePackager/CookTime.cs b/Engine/Source/Programs/IOS/iPhonePackager/CookTime.cs
--- a/Engine/Source/Programs/IOS/iPhonePackager/CookTime.cs
+++ b/Engine/Source/Programs/IOS/iPhonePackager/CookTime.cs
@@ -155,7 +155,7 @@ namespace iPhonePackager
 			}
 			else
 			{
-				NameDecoration = "-IOS-" + Program.GameConfiguration + Program.Architecture;
+				NameDecoration = "-" + Config.OSString + "-" + Program.GameConfiguration + Program.Architecture;
 			}
 
 			// Copy and un-decorate the binary name
diff --git a/Engine/Source/Programs/UnrealBuildTool/Android/AndroidToolChain.cs b/Engine/Source/Programs/UnrealBuildTool/Android/AndroidToolChain.cs
--- a/Engine/Source/Programs/UnrealBuildTool/Android/AndroidToolChain.cs
+++ b/Engine/Source/Programs/UnrealBuildTool/Android/AndroidToolChain.cs
@@ -33,9 +33,9 @@ namespace UnrealBuildTool
 
 		static private Dictionary<string, string[]> LibrariesToSkip = new Dictionary<string, string[]> {
 			{ "-armv7", new string[] { } }, 
-			{ "-arm64", new string[] { "nvToolsExt", "nvToolsExtStub", "oculus", "vrapi", "gpg", } }, 
-			{ "-x86",   new string[] { "nvToolsExt", "nvToolsExtStub", "oculus", "vrapi", } }, 
-			{ "-x64",   new string[] { "nvToolsExt", "nvToolsExtStub", "oculus", "vrapi", "gpg", } }, 
+			{ "-arm64", new string[] { "nvToolsExt", "nvToolsExtStub", "oculus", "vrapi", "ovrkernel", "systemutils", "openglloader", "gpg", } },
+			{ "-x86",   new string[] { "nvToolsExt", "nvToolsExtStub", "oculus", "vrapi", "ovrkernel", "systemutils", "openglloader", } }, 
+			{ "-x64",   new string[] { "nvToolsExt", "nvToolsExtStub", "oculus", "vrapi", "ovrkernel", "systemutils", "openglloader", "gpg", } }, 
 		};
 
 		static private Dictionary<string, string[]> ModulesToSkip = new Dictionary<string, string[]> {
diff --git a/Engine/Source/Programs/UnrealBuildTool/Configuration/BuildConfiguration.cs b/Engine/Source/Programs/UnrealBuildTool/Configuration/BuildConfiguration.cs
--- a/Engine/Source/Programs/UnrealBuildTool/Configuration/BuildConfiguration.cs
+++ b/Engine/Source/Programs/UnrealBuildTool/Configuration/BuildConfiguration.cs
@@ -667,9 +667,9 @@ namespace UnrealBuildTool
 			// set up some paths
 			BaseIntermediateFolder = "Intermediate/Build/";
 
-			// By default check and stop the build on EULA violation
-			bCheckLicenseViolations = false;
-			bBreakBuildOnLicenseViolation = true;
+			// By default check for EULA violation and warn
+			bCheckLicenseViolations = true;
+			bBreakBuildOnLicenseViolation = false;
 
 			// Enables support for fast include dependency scanning, as well as gathering data for 'UBT Makefiles', then quickly
 			// assembling builds in subsequent runs using data in those cached makefiles
diff --git a/Engine/Source/Programs/UnrealBuildTool/Configuration/UEBuildModule.cs b/Engine/Source/Programs/UnrealBuildTool/Configuration/UEBuildModule.cs
--- a/Engine/Source/Programs/UnrealBuildTool/Configuration/UEBuildModule.cs
+++ b/Engine/Source/Programs/UnrealBuildTool/Configuration/UEBuildModule.cs
@@ -1143,7 +1143,7 @@ namespace UnrealBuildTool
 					}
 
 					// A shared PCH was not already set up for us, so set one up.
-					if (ModulePCHEnvironment == null)
+					if (ModulePCHEnvironment == null && SourceFilesToBuild.CPPFiles.Count > 0)
 					{
 						FileItem PCHHeaderFile = ProcessedDependencies.UniquePCHHeaderFile;
 						string PCHModuleName = this.Name;
diff --git a/Engine/Source/Programs/UnrealBuildTool/Configuration/UEBuildTarget.cs b/Engine/Source/Programs/UnrealBuildTool/Configuration/UEBuildTarget.cs
--- a/Engine/Source/Programs/UnrealBuildTool/Configuration/UEBuildTarget.cs
+++ b/Engine/Source/Programs/UnrealBuildTool/Configuration/UEBuildTarget.cs
@@ -978,6 +978,7 @@ namespace UnrealBuildTool
 		/// Filename for the receipt for this target.
 		/// </summary>
 		private string ReceiptFileName;
+		public string BuildReceiptFileName { get { return ReceiptFileName; } }
 
 		/// <summary>
 		/// Version manifests to be written to each output folder
diff --git a/Engine/Source/Programs/UnrealBuildTool/IOS/IOSToolChain.cs b/Engine/Source/Programs/UnrealBuildTool/IOS/IOSToolChain.cs
--- a/Engine/Source/Programs/UnrealBuildTool/IOS/IOSToolChain.cs
+++ b/Engine/Source/Programs/UnrealBuildTool/IOS/IOSToolChain.cs
@@ -553,7 +553,9 @@ namespace UnrealBuildTool
 				string AllArgs = Arguments + FileArguments + CompileEnvironment.Config.AdditionalArguments;
 				string SourceText = System.IO.File.ReadAllText(SourceFile.AbsolutePath);
 				if (UEBuildConfiguration.bCompileForSize && (SourceFile.AbsolutePath.Contains("ElementBatcher.cpp") || SourceText.Contains("ElementBatcher.cpp") || SourceFile.AbsolutePath.Contains("AnimationRuntime.cpp") || SourceText.Contains("AnimationRuntime.cpp")
-					|| SourceFile.AbsolutePath.Contains("AnimEncoding.cpp") || SourceText.Contains("AnimEncoding.cpp") || SourceFile.AbsolutePath.Contains("TextRenderComponent.cpp") || SourceText.Contains("TextRenderComponent.cpp")))
+					|| SourceFile.AbsolutePath.Contains("AnimEncoding.cpp") || SourceText.Contains("AnimEncoding.cpp") || SourceFile.AbsolutePath.Contains("TextRenderComponent.cpp") || SourceText.Contains("TextRenderComponent.cpp")
+					|| SourceFile.AbsolutePath.Contains("SWidget.cpp") || SourceText.Contains("SWidget.cpp") || SourceFile.AbsolutePath.Contains("SCanvas.cpp") || SourceText.Contains("SCanvas.cpp")
+					|| SourceFile.AbsolutePath.Contains("ParticleSystemRender.cpp") || SourceText.Contains("ParticleSystemRender.cpp")))
 				{
 					Console.WriteLine("Forcing {0} to --O3!", SourceFile.AbsolutePath);
 
@@ -992,13 +994,13 @@ namespace UnrealBuildTool
 					// generate the dummy project so signing works
 					if (AppName == "UE4Game" || AppName == "UE4Client" || Utils.IsFileUnderDirectory(Target.ProjectDirectory + "/" + AppName + ".uproject", Path.GetFullPath("../..")))
 					{
-						UnrealBuildTool.GenerateProjectFiles(new XcodeProjectFileGenerator(Target.ProjectFile), new string[] { "-platforms=" + (CppPlatform == CPPTargetPlatform.IOS ? "IOS" : "TVOS"), "-NoIntellIsense", "-iosdeployonly", "-ignorejunk" });
+						UnrealBuildTool.GenerateProjectFiles(new XcodeProjectFileGenerator(Target.ProjectFile), new string[] { "-platforms=" + (CppPlatform == CPPTargetPlatform.IOS ? "IOS" : "TVOS"), "-NoIntellIsense", (CppPlatform == CPPTargetPlatform.IOS ? "-iosdeployonly" : "-tvosdeployonly"), "-ignorejunk" });
 						Project = Path.GetFullPath("../..") + "/UE4_" + (CppPlatform == CPPTargetPlatform.IOS ? "IOS" : "TVOS") + ".xcworkspace";
 						SchemeName = "UE4";
 					}
 					else
 					{
-						UnrealBuildTool.GenerateProjectFiles(new XcodeProjectFileGenerator(Target.ProjectFile), new string[] { "-platforms" + (CppPlatform == CPPTargetPlatform.IOS ? "IOS" : "TVOS"), "-NoIntellIsense", "-iosdeployonly", "-ignorejunk", "-project=\"" + Target.ProjectDirectory + "/" + AppName + ".uproject\"", "-game" });
+						UnrealBuildTool.GenerateProjectFiles(new XcodeProjectFileGenerator(Target.ProjectFile), new string[] { "-platforms" + (CppPlatform == CPPTargetPlatform.IOS ? "IOS" : "TVOS"), "-NoIntellIsense", (CppPlatform == CPPTargetPlatform.IOS ? "-iosdeployonly" : "-tvosdeployonly"), "-ignorejunk", "-project=\"" + Target.ProjectDirectory + "/" + AppName + ".uproject\"", "-game" });
 						Project = Target.ProjectDirectory + "/" + AppName + "_" + (CppPlatform == CPPTargetPlatform.IOS ? "IOS" : "TVOS") + ".xcworkspace";
 					}
 
@@ -1138,7 +1140,7 @@ namespace UnrealBuildTool
 					if (!bUseDangerouslyFastMode)
 					{
 						// generate the dummy project so signing works
-						UnrealBuildTool.GenerateProjectFiles(new XcodeProjectFileGenerator(Target.ProjectFile), new string[] { "-NoIntellisense", "-iosdeployonly", ((Target.ProjectFile != null) ? "-game" : "") });
+						UnrealBuildTool.GenerateProjectFiles(new XcodeProjectFileGenerator(Target.ProjectFile), new string[] { "-NoIntellisense", (CppPlatform == CPPTargetPlatform.IOS ? "-iosdeployonly" : "-tvosdeployonly"), ((Target.ProjectFile != null) ? "-game" : "") });
 					}
 
 					// now that 
@@ -1213,6 +1215,11 @@ namespace UnrealBuildTool
 						Arguments += " -architecture " + Architecture.Substring(1);
 					}
 
+					if (CppPlatform == CPPTargetPlatform.TVOS)
+					{
+						Arguments += " -tvos";
+					}
+
 					if (!bUseRPCUtil)
 					{
 						Arguments += " -usessh";
diff --git a/Engine/Source/Programs/UnrealBuildTool/Linux/UEBuildLinux.cs b/Engine/Source/Programs/UnrealBuildTool/Linux/UEBuildLinux.cs
--- a/Engine/Source/Programs/UnrealBuildTool/Linux/UEBuildLinux.cs
+++ b/Engine/Source/Programs/UnrealBuildTool/Linux/UEBuildLinux.cs
@@ -233,8 +233,8 @@ namespace UnrealBuildTool
 
 		public override bool CanUseXGE()
 		{
-			// [RCL] 2015-08-04 FIXME: modular (cross-)builds (e.g. editor, UT server) fail with XGE as FixDeps step apparently depends on artifacts (object files) which aren't listed among its prerequisites.
-			return false;
+			// [RCL] 2015-08-04 FIXME: we have seen XGE builds fail on Windows
+			return BuildHostPlatform.Current.Platform == UnrealTargetPlatform.Linux;
 		}
 
 		/// <summary>
diff --git a/Engine/Source/Programs/UnrealBuildTool/System/MakefileGenerator.cs b/Engine/Source/Programs/UnrealBuildTool/System/MakefileGenerator.cs
--- a/Engine/Source/Programs/UnrealBuildTool/System/MakefileGenerator.cs
+++ b/Engine/Source/Programs/UnrealBuildTool/System/MakefileGenerator.cs
@@ -125,7 +125,7 @@ namespace UnrealBuildTool
 
 			MakefileContent.Append("\n\n" + BuildCommand + ProjectBuildCommand + "\n" +
 				"all: StandardSet\n\n" +
-				"RequiredTools: CrashReportClient ShaderCompileWorker UnrealPak UnrealLightmass\n\n" +
+				"RequiredTools: CrashReportClient-Linux-Shipping ShaderCompileWorker UnrealPak UnrealLightmass\n\n" +
 				"StandardSet: RequiredTools UnrealFrontend UE4Editor\n\n" +
 				"DebugSet: RequiredTools UnrealFrontend-Linux-Debug UE4Editor-Linux-Debug\n\n"
 			);
diff --git a/Engine/Source/Programs/UnrealBuildTool/System/ModuleDescriptor.cs b/Engine/Source/Programs/UnrealBuildTool/System/ModuleDescriptor.cs
--- a/Engine/Source/Programs/UnrealBuildTool/System/ModuleDescriptor.cs
+++ b/Engine/Source/Programs/UnrealBuildTool/System/ModuleDescriptor.cs
@@ -29,6 +29,7 @@ namespace UnrealBuildTool
 		PostConfigInit,
 		PreLoadingScreen,
 		PostEngineInit,
+		None,
 	}
 
 	[DebuggerDisplay("Name={Name}")]
diff --git a/Engine/Source/Programs/UnrealBuildTool/System/ProjectFileGenerator.cs b/Engine/Source/Programs/UnrealBuildTool/System/ProjectFileGenerator.cs
--- a/Engine/Source/Programs/UnrealBuildTool/System/ProjectFileGenerator.cs
+++ b/Engine/Source/Programs/UnrealBuildTool/System/ProjectFileGenerator.cs
@@ -419,6 +419,9 @@ namespace UnrealBuildTool
 							}
 						}
 
+						// Engine Extras files
+						AddEngineExtrasFiles(EngineProject);
+
 						// Engine localization files
 						if( IncludeLocalizationFiles )
 						{
@@ -843,6 +846,11 @@ namespace UnrealBuildTool
 			}
 		}
 
+		/// Adds all engine extras files to the specified project
+		protected virtual void AddEngineExtrasFiles(ProjectFile EngineProject)
+		{
+		}
+
 		/// Adds UnrealHeaderTool config files to the specified project
 		private void AddUnrealHeaderToolConfigFiles(ProjectFile EngineProject)
 		{
diff --git a/Engine/Source/Programs/UnrealBuildTool/System/UnrealBuildTool.cs b/Engine/Source/Programs/UnrealBuildTool/System/UnrealBuildTool.cs
--- a/Engine/Source/Programs/UnrealBuildTool/System/UnrealBuildTool.cs
+++ b/Engine/Source/Programs/UnrealBuildTool/System/UnrealBuildTool.cs
@@ -1018,7 +1018,7 @@ namespace UnrealBuildTool
 
 					// @todo ubtmake: remove this when building with RPCUtility works
 					// @todo tvos merge: Check the change to this line, not clear why. Is TVOS needed here?
-					if (CheckPlatform == UnrealTargetPlatform.Mac || CheckPlatform == UnrealTargetPlatform.IOS)
+					if (CheckPlatform == UnrealTargetPlatform.Mac || CheckPlatform == UnrealTargetPlatform.IOS || CheckPlatform == UnrealTargetPlatform.TVOS)
 					{
 						BuildConfiguration.bUseUBTMakefiles = false;
 					}
@@ -2016,7 +2016,10 @@ namespace UnrealBuildTool
 								// clean up any stale modules
 								foreach (UEBuildTarget Target in Targets)
 								{
-									Target.CleanStaleModules();
+									if (Target.OnlyModules == null || Target.OnlyModules.Count == 0)
+									{
+										Target.CleanStaleModules();
+									}
 								}
 							}
 
@@ -2218,6 +2221,12 @@ namespace UnrealBuildTool
 					throw new BuildException("ShouldDoHotReload cannot handle multiple binaries returning from UEBuildTarget.MakeExecutablePaths");
 				}
 
+				string EditorProcessFilename = EditorProcessFilenames[0].CanonicalName;
+				if (TargetDesc.Platform == UnrealTargetPlatform.Mac && !EditorProcessFilename.Contains(".app/contents/macos/"))
+				{
+					EditorProcessFilename += ".app/contents/macos/" + Path.GetFileNameWithoutExtension(EditorProcessFilename);
+				}
+					
 				BuildHostPlatform.ProcessInfo[] Processes = BuildHostPlatform.Current.GetProcesses();
 				string EditorRunsDir = Path.Combine(UnrealBuildTool.EngineDirectory.FullName, "Intermediate", "EditorRuns");
 
@@ -2243,7 +2252,7 @@ namespace UnrealBuildTool
 					if (!bIsRunning)
 					{
 						// Otherwise check if the path matches.
-						bIsRunning = new FileReference(Proc.Filename).CanonicalName == EditorProcessFilenames[0].CanonicalName;
+						bIsRunning = new FileReference (Proc.Filename).CanonicalName == EditorProcessFilename;
 					}
 				}
 			}
diff --git a/Engine/Source/Programs/UnrealBuildTool/System/VCProjectFileGenerator.cs b/Engine/Source/Programs/UnrealBuildTool/System/VCProjectFileGenerator.cs
--- a/Engine/Source/Programs/UnrealBuildTool/System/VCProjectFileGenerator.cs
+++ b/Engine/Source/Programs/UnrealBuildTool/System/VCProjectFileGenerator.cs
@@ -231,6 +231,21 @@ namespace UnrealBuildTool
 			}
 		}
 
+		/// <summary>
+		/// Adds Extra files that are specific to Visual Studio projects
+		/// </summary>
+		/// <param name="EngineProject">Project to add files to</param>
+		protected override void AddEngineExtrasFiles(ProjectFile EngineProject)
+		{
+			base.AddEngineExtrasFiles(EngineProject);
+
+			// Add our UE4.natvis file
+			var NatvisFilePath = FileReference.Combine(UnrealBuildTool.EngineDirectory, "Extras", "VisualStudioDebugging", "UE4.natvis");
+			if (NatvisFilePath.Exists())
+			{
+				EngineProject.AddFileToProject(NatvisFilePath, UnrealBuildTool.EngineDirectory);
+			}
+		}
 
 		/// <summary>
 		/// Selects which platforms and build configurations we want in the project file
diff --git a/Engine/Source/Programs/UnrealBuildTool/System/XGE.cs b/Engine/Source/Programs/UnrealBuildTool/System/XGE.cs
--- a/Engine/Source/Programs/UnrealBuildTool/System/XGE.cs
+++ b/Engine/Source/Programs/UnrealBuildTool/System/XGE.cs
@@ -343,10 +343,11 @@ namespace UnrealBuildTool
 					ToolElement.SetAttribute("OutputPrefix", OutputPrefix);
 				}
 
-				// batch files (.bat, .cmd) need to be run via or cmd /c or shellexecute, 
+				// When running on Windows, differentiate between .exe and batch files.
+				// Those (.bat, .cmd) need to be run via cmd /c or shellexecute, 
 				// the latter which we can't use because we want to redirect input/output
 
-				bool bLaunchViaCmdExe = !Path.GetExtension(Action.CommandPath).ToLower().EndsWith("exe");
+				bool bLaunchViaCmdExe = (BuildHostPlatform.Current.Platform == UnrealTargetPlatform.Win64) && (!Path.GetExtension(Action.CommandPath).ToLower().EndsWith(".exe"));
 
 				string CommandPath = "";
 				string CommandArguments = "";
diff --git a/Engine/Source/Programs/UnrealBuildTool/ToolChain/RemoteToolChain.cs b/Engine/Source/Programs/UnrealBuildTool/ToolChain/RemoteToolChain.cs
--- a/Engine/Source/Programs/UnrealBuildTool/ToolChain/RemoteToolChain.cs
+++ b/Engine/Source/Programs/UnrealBuildTool/ToolChain/RemoteToolChain.cs
@@ -926,6 +926,7 @@ namespace UnrealBuildTool
 			SSHProcess.StartInfo.FileName = ResolvedSSHExe;
 			SSHProcess.StartInfo.Arguments = string.Format(
 				"{0} {1}@{2} \"{3}\"",
+//				"-o CheckHostIP=no {0} {1}@{2} \"{3}\"",
 				ResolvedSSHAuthentication,
 				RSyncUsername,
 				RemoteServerName,
diff --git a/Engine/Source/Programs/UnrealHeaderTool/Private/HeaderParser.cpp b/Engine/Source/Programs/UnrealHeaderTool/Private/HeaderParser.cpp
--- a/Engine/Source/Programs/UnrealHeaderTool/Private/HeaderParser.cpp
+++ b/Engine/Source/Programs/UnrealHeaderTool/Private/HeaderParser.cpp
@@ -4930,6 +4930,10 @@ void FHeaderParser::CompileClassDeclaration(FClasses& AllClasses)
 
 	auto PrologFinishLine = InputLine;
 
+	// Members of classes have a default private access level in c++
+	// Setting this directly should be ok as we don't support nested classes, so the outer scope access should not need restoring
+	CurrentAccessSpecifier = ACCESS_Private;
+
 	AddFormattedPrevCommentAsTooltipMetaData(MetaData);
 
 	// New style files have the class name / extends afterwards
diff --git a/Engine/Source/Runtime/AIModule/Classes/EnvironmentQuery/Generators/EnvQueryGenerator_BlueprintBase.h b/Engine/Source/Runtime/AIModule/Classes/EnvironmentQuery/Generators/EnvQueryGenerator_BlueprintBase.h
--- a/Engine/Source/Runtime/AIModule/Classes/EnvironmentQuery/Generators/EnvQueryGenerator_BlueprintBase.h
+++ b/Engine/Source/Runtime/AIModule/Classes/EnvironmentQuery/Generators/EnvQueryGenerator_BlueprintBase.h
@@ -22,6 +22,9 @@ class AIMODULE_API UEnvQueryGenerator_BlueprintBase : public UEnvQueryGenerator
 	UPROPERTY(EditAnywhere, Category = Generator)
 	TSubclassOf<UEnvQueryContext> Context;
 
+	/** @todo this should show up only in the generator's BP, but 
+	 *	due to the way EQS editor is generating widgets it's there as well
+	 *	It's a bug and we'll fix it */
 	UPROPERTY(EditDefaultsOnly, Category = Generator)
 	TSubclassOf<UEnvQueryItemType> GeneratedItemType;
 
diff --git a/Engine/Source/Runtime/AIModule/Classes/Perception/AIPerceptionComponent.h b/Engine/Source/Runtime/AIModule/Classes/Perception/AIPerceptionComponent.h
--- a/Engine/Source/Runtime/AIModule/Classes/Perception/AIPerceptionComponent.h
+++ b/Engine/Source/Runtime/AIModule/Classes/Perception/AIPerceptionComponent.h
@@ -153,6 +153,7 @@ protected:
 	UPROPERTY(Transient)
 	AAIController* AIOwner;
 
+	/** @todo this field is misnamed. It's a whitelist. */
 	FPerceptionChannelWhitelist PerceptionFilter;
 
 private:
@@ -219,6 +220,9 @@ public:
 	UFUNCTION(BlueprintCallable, Category="AI|Perception")
 	void RequestStimuliListenerUpdate();
 
+	/** Allows toggling senses on and off */
+	void UpdatePerceptionWhitelist(const FAISenseID Channel, const bool bNewValue);
+
 	void RegisterStimulus(AActor* Source, const FAIStimulus& Stimulus);
 	void ProcessStimuli();
 	/** Returns true if, as result of stimuli aging, this listener needs an update (like if some stimuli expired) */
@@ -264,8 +268,9 @@ public:
 	FActorPerceptionUpdatedDelegate OnTargetPerceptionUpdated;
 
 protected:
-
+	DEPRECATED(4.11, "Function has been renamed and made public. Please use UpdatePerceptionWhitelist instead")
 	void UpdatePerceptionFilter(FAISenseID Channel, bool bNewValue);
+
 	TActorPerceptionContainer& GetPerceptualData() { return PerceptualData; }
 
 	/** called to clean up on owner's end play or destruction */
diff --git a/Engine/Source/Runtime/AIModule/Private/EnvironmentQuery/Generators/EnvQueryGenerator_BlueprintBase.cpp b/Engine/Source/Runtime/AIModule/Private/EnvironmentQuery/Generators/EnvQueryGenerator_BlueprintBase.cpp
--- a/Engine/Source/Runtime/AIModule/Private/EnvironmentQuery/Generators/EnvQueryGenerator_BlueprintBase.cpp
+++ b/Engine/Source/Runtime/AIModule/Private/EnvironmentQuery/Generators/EnvQueryGenerator_BlueprintBase.cpp
@@ -7,7 +7,6 @@
 #include "EnvironmentQuery/Generators/EnvQueryGenerator_BlueprintBase.h"
 #include "EnvironmentQuery/Items/EnvQueryItemType_Actor.h"
 #include "AI/Navigation/RecastNavMesh.h"
-#include "VisualLogger/VisualLogger.h"
 
 #define LOCTEXT_NAMESPACE "EnvQueryGenerator"
 
@@ -43,15 +42,33 @@ void UEnvQueryGenerator_BlueprintBase::GenerateItems(FEnvQueryInstance& QueryIns
 void UEnvQueryGenerator_BlueprintBase::AddGeneratedVector(FVector Vector) const
 {
 	check(CachedQueryInstance);
-	ensure(ItemType->IsChildOf<UEnvQueryItemType_VectorBase>());
-	CachedQueryInstance->AddItemData<UEnvQueryItemType_Point>(Vector);
+	if (ensure(ItemType->IsChildOf<UEnvQueryItemType_ActorBase>() == false))
+	{
+		CachedQueryInstance->AddItemData<UEnvQueryItemType_Point>(Vector);
+	}
+	else
+	{
+		UE_LOG(LogEQS, Error, TEXT("Trying to generate a Vector item while generator %s is configured to produce Actor items")
+			, *GetName());
+	}
 }
 
 void UEnvQueryGenerator_BlueprintBase::AddGeneratedActor(AActor* Actor) const
 {
 	check(CachedQueryInstance);
-	ensure(ItemType->IsChildOf<UEnvQueryItemType_ActorBase>());
-	CachedQueryInstance->AddItemData<UEnvQueryItemType_Actor>(Actor);
+	if (ensure(ItemType->IsChildOf<UEnvQueryItemType_ActorBase>()))
+	{
+		CachedQueryInstance->AddItemData<UEnvQueryItemType_Actor>(Actor);
+	}
+	else
+	{
+		UE_LOG(LogEQS, Error, TEXT("Trying to generate an Actor item while generator %s is configured to produce Vector items. Will use Actor\'s location, but please update your BP code.")
+			, *GetName());
+		if (Actor)
+		{
+			CachedQueryInstance->AddItemData<UEnvQueryItemType_Point>(Actor->GetActorLocation());
+		}
+	}
 }
 
 UObject* UEnvQueryGenerator_BlueprintBase::GetQuerier() const
diff --git a/Engine/Source/Runtime/AIModule/Private/Perception/AIPerceptionComponent.cpp b/Engine/Source/Runtime/AIModule/Private/Perception/AIPerceptionComponent.cpp
--- a/Engine/Source/Runtime/AIModule/Private/Perception/AIPerceptionComponent.cpp
+++ b/Engine/Source/Runtime/AIModule/Private/Perception/AIPerceptionComponent.cpp
@@ -244,7 +244,7 @@ void UAIPerceptionComponent::BeginDestroy()
 	Super::BeginDestroy();
 }
 
-void UAIPerceptionComponent::UpdatePerceptionFilter(FAISenseID Channel, bool bNewValue)
+void UAIPerceptionComponent::UpdatePerceptionWhitelist(const FAISenseID Channel, const bool bNewValue)
 {
 	const bool bCurrentValue = PerceptionFilter.ShouldRespondToChannel(Channel);
 	if (bNewValue != bCurrentValue)
@@ -704,3 +704,11 @@ void UAIPerceptionComponent::DescribeSelfToVisLog(FVisualLogEntry* Snapshot) con
 
 }
 #endif // ENABLE_VISUAL_LOG
+
+//----------------------------------------------------------------------//
+// deprecated
+//----------------------------------------------------------------------//
+void UAIPerceptionComponent::UpdatePerceptionFilter(FAISenseID Channel, bool bNewValue)
+{
+	UpdatePerceptionWhitelist(Channel, bNewValue);
+}
\ No newline at end of file
diff --git a/Engine/Source/Runtime/AnimGraphRuntime/Private/AnimNodes/AnimNode_RandomPlayer.cpp b/Engine/Source/Runtime/AnimGraphRuntime/Private/AnimNodes/AnimNode_RandomPlayer.cpp
--- a/Engine/Source/Runtime/AnimGraphRuntime/Private/AnimNodes/AnimNode_RandomPlayer.cpp
+++ b/Engine/Source/Runtime/AnimGraphRuntime/Private/AnimNodes/AnimNode_RandomPlayer.cpp
@@ -76,6 +76,12 @@ void FAnimNode_RandomPlayer::Update(const FAnimationUpdateContext& Context)
 {
 	EvaluateGraphExposedInputs.Execute(Context);
 
+	if(Entries.Num() == 0)
+	{
+		// We don't have any entries, play data will be invalid - early out
+		return;
+	}
+
 	FRandomAnimPlayData* CurrentData = &PlayData[GetDataIndex(ERandomDataIndexType::Current)];
 	FRandomAnimPlayData* NextData = &PlayData[GetDataIndex(ERandomDataIndexType::Next)];
 
diff --git a/Engine/Source/Runtime/AnimGraphRuntime/Private/BoneControllers/AnimNode_Trail.cpp b/Engine/Source/Runtime/AnimGraphRuntime/Private/BoneControllers/AnimNode_Trail.cpp
--- a/Engine/Source/Runtime/AnimGraphRuntime/Private/BoneControllers/AnimNode_Trail.cpp
+++ b/Engine/Source/Runtime/AnimGraphRuntime/Private/BoneControllers/AnimNode_Trail.cpp
@@ -214,8 +214,8 @@ bool FAnimNode_Trail::IsValidToEvaluate(const USkeleton* Skeleton, const FBoneCo
 		{
 			if (ChainIndex == INDEX_NONE)
 			{
-				// we don't want this to spam. We want this to happen once. 
-				ensureMsgf(false, TEXT("AnimNode_Trail [%s]: contains invalid chain index. Verify if the node contains correct length of chain. "), *TrailBone.BoneName.ToString());
+				// unfortunately there is no easy way to communicate this back to user other than spamming here because this gets called every frame
+				// originally tried in AnimGraphNode, but that doesn't know hierarchy so I can't verify it there. Maybe should try with USkeleton asset there. @todo
 				return false;
 			}
 			else if (RequiredBones.Contains(ChainIndex) == false)
diff --git a/Engine/Source/Runtime/Apple/MetalRHI/Private/MetalCommandEncoder.cpp b/Engine/Source/Runtime/Apple/MetalRHI/Private/MetalCommandEncoder.cpp
--- a/Engine/Source/Runtime/Apple/MetalRHI/Private/MetalCommandEncoder.cpp
+++ b/Engine/Source/Runtime/Apple/MetalRHI/Private/MetalCommandEncoder.cpp
@@ -494,7 +494,7 @@ void FMetalCommandEncoder::PopDebugGroup(void)
 	
 #pragma mark - Public Render State Mutators -
 
-void FMetalCommandEncoder::SetRenderPassDescriptor(MTLRenderPassDescriptor* const RenderPass)
+void FMetalCommandEncoder::SetRenderPassDescriptor(MTLRenderPassDescriptor* const RenderPass, bool const bReset)
 {
 	check(IsRenderCommandEncoderActive() == false && IsComputeCommandEncoderActive() == false && IsBlitCommandEncoderActive() == false);
 	
@@ -508,11 +508,14 @@ void FMetalCommandEncoder::SetRenderPassDescriptor(MTLRenderPassDescriptor* cons
 		RenderPassDesc = RenderPass;
 	}
 	
-	for(uint32 i = 0; i < SF_NumFrequencies; i++)
+	if (bReset)
 	{
-		ShaderBuffers[i].Bound = 0;
-		ShaderTextures[i].Bound = 0;
-		ShaderSamplers[i].Bound = 0;
+		for(uint32 i = 0; i < SF_NumFrequencies; i++)
+		{
+			ShaderBuffers[i].Bound = 0;
+			ShaderTextures[i].Bound = 0;
+			ShaderSamplers[i].Bound = 0;
+		}
 	}
 }
 
diff --git a/Engine/Source/Runtime/Apple/MetalRHI/Private/MetalCommandEncoder.h b/Engine/Source/Runtime/Apple/MetalRHI/Private/MetalCommandEncoder.h
--- a/Engine/Source/Runtime/Apple/MetalRHI/Private/MetalCommandEncoder.h
+++ b/Engine/Source/Runtime/Apple/MetalRHI/Private/MetalCommandEncoder.h
@@ -109,8 +109,9 @@ public:
 	/**
 	 * Set the render pass descriptor - no encoder may be active when this function is called.
 	 * @param RenderPass The render pass descriptor to set. May be nil.
+	 * @param bReset Whether to reset existing state.
 	 */
-	void SetRenderPassDescriptor(MTLRenderPassDescriptor* const RenderPass);
+	void SetRenderPassDescriptor(MTLRenderPassDescriptor* const RenderPass, bool const bReset);
 	
 	/*
 	 * Sets the current render pipeline state object.
diff --git a/Engine/Source/Runtime/Apple/MetalRHI/Private/MetalContext.cpp b/Engine/Source/Runtime/Apple/MetalRHI/Private/MetalContext.cpp
--- a/Engine/Source/Runtime/Apple/MetalRHI/Private/MetalContext.cpp
+++ b/Engine/Source/Runtime/Apple/MetalRHI/Private/MetalContext.cpp
@@ -641,13 +641,13 @@ void FMetalContext::ResetRenderCommandEncoder()
 	
 	ConditionalSwitchToGraphics();
 	
-    if (IsFeatureLevelSupported( GMaxRHIShaderPlatform, ERHIFeatureLevel::SM4 ))
+	if (IsFeatureLevelSupported( GMaxRHIShaderPlatform, ERHIFeatureLevel::SM4 ))
     {
-        StateCache.SetRenderTargetsInfo(StateCache.GetRenderTargetsInfo(), QueryBuffer->GetCurrentQueryBuffer()->Buffer);
+        StateCache.SetRenderTargetsInfo(StateCache.GetRenderTargetsInfo(), QueryBuffer->GetCurrentQueryBuffer()->Buffer, false);
     }
     else
     {
-        StateCache.SetRenderTargetsInfo(StateCache.GetRenderTargetsInfo(), NULL);
+        StateCache.SetRenderTargetsInfo(StateCache.GetRenderTargetsInfo(), NULL, false);
     }
 	
 	if (CommandEncoder.IsRenderCommandEncoderActive())
@@ -719,16 +719,7 @@ void FMetalContext::PrepareToDraw(uint32 PrimitiveType)
 		FTexture2DRHIRef DepthStencil = RHICreateTexture2D(FBSize.width, FBSize.height, PF_DepthStencil, 1, 1, TexCreate_DepthStencilTargetable, TexInfo);
 		Info.DepthStencilRenderTarget.Texture = DepthStencil;
 		
-		TRefCountPtr<FMetalBlendState> BlendState = StateCache.GetBlendState();
-		TRefCountPtr<FMetalDepthStencilState> DepthState = StateCache.GetDepthStencilState();
-		TRefCountPtr<FMetalRasterizerState> RasterState = StateCache.GetRasterizerState();
-
-		StateCache.SetRenderTargetsInfo(Info, StateCache.GetVisibilityResultsBuffer());
-		
-		StateCache.SetBlendState(BlendState);
-		StateCache.SetDepthStencilState(DepthState);
-		StateCache.SetRasterizerState(RasterState);
-		StateCache.SetBoundShaderState(CurrentBoundShaderState);
+		StateCache.SetRenderTargetsInfo(Info, StateCache.GetVisibilityResultsBuffer(), false);
 		
 		bRestoreState = true;
 	}
@@ -767,11 +758,11 @@ void FMetalContext::SetRenderTargetsInfo(const FRHISetRenderTargetsInfo& RenderT
 //	}
     if (IsFeatureLevelSupported( GMaxRHIShaderPlatform, ERHIFeatureLevel::SM4 ))
     {
-        StateCache.SetRenderTargetsInfo(RenderTargetsInfo, QueryBuffer->GetCurrentQueryBuffer()->Buffer);
+        StateCache.SetRenderTargetsInfo(RenderTargetsInfo, QueryBuffer->GetCurrentQueryBuffer()->Buffer, true);
     }
     else
     {
-        StateCache.SetRenderTargetsInfo(RenderTargetsInfo, NULL);
+        StateCache.SetRenderTargetsInfo(RenderTargetsInfo, NULL, true);
     }
 }
 
diff --git a/Engine/Source/Runtime/Apple/MetalRHI/Private/MetalStateCache.cpp b/Engine/Source/Runtime/Apple/MetalRHI/Private/MetalStateCache.cpp
--- a/Engine/Source/Runtime/Apple/MetalRHI/Private/MetalStateCache.cpp
+++ b/Engine/Source/Runtime/Apple/MetalRHI/Private/MetalStateCache.cpp
@@ -233,7 +233,7 @@ void FMetalStateCache::SetComputeShader(FMetalComputeShader* InComputeShader)
 	}
 }
 
-void FMetalStateCache::SetRenderTargetsInfo(FRHISetRenderTargetsInfo const& InRenderTargets, id<MTLBuffer> const QueryBuffer)
+void FMetalStateCache::SetRenderTargetsInfo(FRHISetRenderTargetsInfo const& InRenderTargets, id<MTLBuffer> const QueryBuffer, bool const bReset)
 {
 	// see if our new Info matches our previous Info
 	if (NeedsToSetRenderTarget(InRenderTargets) || QueryBuffer != VisibilityResults)
@@ -604,18 +604,21 @@ void FMetalStateCache::SetRenderTargetsInfo(FRHISetRenderTargetsInfo const& InRe
 		}
 	
 		// Set render to the framebuffer
-		CommandEncoder.SetRenderPassDescriptor(RenderPass);
+		CommandEncoder.SetRenderPassDescriptor(RenderPass, bReset);
 		
 		// if (bNeedsClear)
 		{
 			CommandEncoder.BeginRenderCommandEncoding();
 		}
 		
-		// Reset any existing state as that must be fully reinitialised by the caller.
-		DepthStencilState.SafeRelease();
-		RasterizerState.SafeRelease();
-		BlendState.SafeRelease();
-		BoundShaderState.SafeRelease();
+		if (bReset)
+		{
+			// Reset any existing state as that must be fully reinitialised by the caller.
+			DepthStencilState.SafeRelease();
+			RasterizerState.SafeRelease();
+			BlendState.SafeRelease();
+			BoundShaderState.SafeRelease();
+		}
 	}
 }
 
diff --git a/Engine/Source/Runtime/Apple/MetalRHI/Private/MetalStateCache.h b/Engine/Source/Runtime/Apple/MetalRHI/Private/MetalStateCache.h
--- a/Engine/Source/Runtime/Apple/MetalRHI/Private/MetalStateCache.h
+++ b/Engine/Source/Runtime/Apple/MetalRHI/Private/MetalStateCache.h
@@ -23,7 +23,7 @@ public:
 	void SetRasterizerState(FMetalRasterizerState* InRasterizerState);
 	void SetBoundShaderState(FMetalBoundShaderState* BoundShaderState);
 	void SetComputeShader(FMetalComputeShader* InComputeShader);
-	void SetRenderTargetsInfo(FRHISetRenderTargetsInfo const& InRenderTargets, id<MTLBuffer> const QueryBuffer);
+	void SetRenderTargetsInfo(FRHISetRenderTargetsInfo const& InRenderTargets, id<MTLBuffer> const QueryBuffer, bool const bReset);
 	void SetHasValidRenderTarget(bool const bIsValid);
 	void SetViewport(const MTLViewport& InViewport);
 	void SetVertexBuffer(uint32 const Index, id<MTLBuffer> Buffer, uint32 const Stride, uint32 const Offset);
diff --git a/Engine/Source/Runtime/Apple/MetalRHI/Private/MetalTexture.cpp b/Engine/Source/Runtime/Apple/MetalRHI/Private/MetalTexture.cpp
--- a/Engine/Source/Runtime/Apple/MetalRHI/Private/MetalTexture.cpp
+++ b/Engine/Source/Runtime/Apple/MetalRHI/Private/MetalTexture.cpp
@@ -414,6 +414,17 @@ FMetalSurface::FMetalSurface(ERHIResourceType ResourceType, EPixelFormat Format,
 			case MTLPixelFormatBGRA8Unorm:
 				MTLFormat = MTLPixelFormatBGRA8Unorm_sRGB;
 				break;
+			case MTLPixelFormatR8Unorm:
+				// For now R8 sRGB expansion is 2D only, log other usage for later.
+				if(Type == RRT_Texture2D)
+				{
+					MTLFormat = MTLPixelFormatRGBA8Unorm_sRGB;
+				}
+				else
+				{
+					UE_LOG(LogMetal, Display, TEXT("Attempting to use unsupported MTLPixelFormatR8Unorm_sRGB on Mac with texture type: %d, no format expansion will be provided so rendering errors may occur."), Type);
+				}
+				break;
 			default:
 				break;
 		}
@@ -739,7 +750,7 @@ void* FMetalSurface::Lock(uint32 MipIndex, uint32 ArrayIndex, EResourceLockMode
 	if(!LockedMemory[MipIndex])
 	{
 #if METAL_API_1_1
-		NSUInteger ResMode = MTLResourceStorageModeShared | (GetMetalDeviceContext().SupportsFeature(EMetalFeaturesResourceOptions) ? MTLResourceCPUCacheModeWriteCombined : 0);
+		NSUInteger ResMode = MTLResourceStorageModeShared | (GetMetalDeviceContext().SupportsFeature(EMetalFeaturesResourceOptions) && !(PixelFormat == PF_G8 && (Flags & TexCreate_SRGB)) ? MTLResourceCPUCacheModeWriteCombined : 0);
 #else
 		NSUInteger ResMode = MTLStorageModeShared;
 #endif
@@ -773,11 +784,37 @@ void* FMetalSurface::Lock(uint32 MipIndex, uint32 ArrayIndex, EResourceLockMode
 			//kick the current command buffer.
 			GetMetalDeviceContext().SubmitCommandBufferAndWait();
 			
+			// Pack RGBA8_sRGB into R8_sRGB for Mac.
+			if (PixelFormat == PF_G8 && (Flags & TexCreate_SRGB) && Type == RRT_Texture2D)
+			{
+				TArray<uint8> Data;
+				uint8* ExpandedMem = (uint8*)[LockedMemory[MipIndex] contents];
+				Data.Append(ExpandedMem, MipBytes);
+				uint32 SrcStride = DestStride;
+				DestStride = FMath::Max<uint32>(SizeX >> MipIndex, 1);
+				for(uint y = 0; y < FMath::Max<uint32>(SizeY >> MipIndex, 1); y++)
+				{
+					uint8* RowDest = ExpandedMem;
+					for(uint x = 0; x < FMath::Max<uint32>(SizeX >> MipIndex, 1); x++)
+					{
+						*(RowDest++) = Data[(y * SrcStride) + (x * 4)];
+					}
+					ExpandedMem = (ExpandedMem + DestStride);
+				}
+			}
+			
             break;
         }
         case RLM_WriteOnly:
         {
             WriteLock |= 1 << MipIndex;
+#if PLATFORM_MAC
+			// Expand R8_sRGB into RGBA8_sRGB for Mac.
+			if (PixelFormat == PF_G8 && (Flags & TexCreate_SRGB) && Type == RRT_Texture2D)
+			{
+				DestStride = FMath::Max<uint32>(SizeX >> MipIndex, 1);
+			}
+#endif
             break;
         }
         default:
@@ -808,6 +845,27 @@ void FMetalSurface::Unlock(uint32 MipIndex, uint32 ArrayIndex)
             Region = MTLRegionMake3D(0, 0, 0, FMath::Max<uint32>(SizeX >> MipIndex, 1), FMath::Max<uint32>(SizeY >> MipIndex, 1), FMath::Max<uint32>(SizeZ >> MipIndex, 1));
         }
 #if PLATFORM_MAC
+		// Expand R8_sRGB into RGBA8_sRGB for Mac.
+		if (PixelFormat == PF_G8 && (Flags & TexCreate_SRGB) && Type == RRT_Texture2D)
+		{
+			TArray<uint8> Data;
+			uint8* ExpandedMem = (uint8*)[LockedMemory[MipIndex] contents];
+			Data.Append(ExpandedMem, BytesPerImage);
+			uint32 SrcStride = FMath::Max<uint32>(SizeX >> MipIndex, 1);
+			for(uint y = 0; y < FMath::Max<uint32>(SizeY >> MipIndex, 1); y++)
+			{
+				uint8* RowDest = ExpandedMem;
+				for(uint x = 0; x < FMath::Max<uint32>(SizeX >> MipIndex, 1); x++)
+				{
+					*(RowDest++) = Data[(y * SrcStride) + x];
+					*(RowDest++) = Data[(y * SrcStride) + x];
+					*(RowDest++) = Data[(y * SrcStride) + x];
+					*(RowDest++) = Data[(y * SrcStride) + x];
+				}
+				ExpandedMem = (ExpandedMem + Stride);
+			}
+		}
+
 		if(Texture.storageMode == MTLStorageModePrivate)
 		{
 			SCOPED_AUTORELEASE_POOL;
@@ -857,6 +915,13 @@ uint32 FMetalSurface::GetMipSize(uint32 MipIndex, uint32* Stride, bool bSingleLa
 		NumBlocksX = FMath::Max<uint32>(NumBlocksX, 2);
 		NumBlocksY = FMath::Max<uint32>(NumBlocksY, 2);
 	}
+#if PLATFORM_MAC
+	else if (PixelFormat == PF_G8 && (Flags & TexCreate_SRGB))
+	{
+		// RGBA_sRGB is the closest match - so expand the data.
+		NumBlocksX *= 4;
+	}
+#endif
 
 	const uint32 MipBytes = NumBlocksX * NumBlocksY * BlockBytes * MipSizeZ;
 
diff --git a/Engine/Source/Runtime/Core/Private/HAL/FileManagerGeneric.cpp b/Engine/Source/Runtime/Core/Private/HAL/FileManagerGeneric.cpp
--- a/Engine/Source/Runtime/Core/Private/HAL/FileManagerGeneric.cpp
+++ b/Engine/Source/Runtime/Core/Private/HAL/FileManagerGeneric.cpp
@@ -310,22 +310,30 @@ bool FFileManagerGeneric::Move( const TCHAR* Dest, const TCHAR* Src, bool Replac
 		{
 			return false;
 		}
-		// If the move failed, throw a warning but retry before we throw an error
-		UE_LOG( LogFileManager, Warning, TEXT( "MoveFile was unable to move '%s' to '%s', retrying in .5s..." ), Src, Dest );
 
-		// Wait just a little bit( i.e. a totally arbitrary amount )...
-		FPlatformProcess::Sleep( 0.5f );
+		int32 RetryCount = 10;
+		bool bSuccess = false;
+		while (RetryCount--)
+		{
+			// If the move failed, throw a warning but retry before we throw an error
+			UE_LOG(LogFileManager, Warning, TEXT("MoveFile was unable to move '%s' to '%s', retrying in .5s..."), Src, Dest);
 
-		// Try again
-		if( !GetLowLevel().MoveFile( Dest, Src ) )
+			// Wait just a little bit( i.e. a totally arbitrary amount )...
+			FPlatformProcess::Sleep(0.5f);
+
+			// Try again
+			bSuccess = GetLowLevel().MoveFile(Dest, Src);
+			if (bSuccess)
+			{
+				UE_LOG(LogFileManager, Warning, TEXT("MoveFile recovered during retry!"));
+				break;
+			}
+		}
+		if (!bSuccess)
 		{
 			UE_LOG( LogFileManager, Error, TEXT( "Error moving file '%s' to '%s'." ), Src, Dest );
 			return false;
 		}
-		else
-		{
-			UE_LOG( LogFileManager, Warning, TEXT( "MoveFile recovered during retry!" ) );
-		}
 	}
 	return true;
 }
@@ -605,8 +613,8 @@ FArchiveFileReaderGeneric::FArchiveFileReaderGeneric( IFileHandle* InHandle, con
 
 void FArchiveFileReaderGeneric::Seek( int64 InPos )
 {
-	check( InPos>=0 );
-	check( InPos<=Size );
+	checkf(InPos >= 0, TEXT("Attempted to seek to a negative location (%lld/%lld), file: %s. The file is most likely corrupt."), InPos, Size, *Filename);
+	checkf(InPos <= Size, TEXT("Attempted to seek past the end of file (%lld/%lld), file: %s. The file is most likely corrupt."), InPos, Size, *Filename);
 	if( !SeekLowLevel( InPos ) )
 	{
 		TCHAR ErrorBuffer[1024];
diff --git a/Engine/Source/Runtime/Core/Private/HAL/UnrealMemory.cpp b/Engine/Source/Runtime/Core/Private/HAL/UnrealMemory.cpp
--- a/Engine/Source/Runtime/Core/Private/HAL/UnrealMemory.cpp
+++ b/Engine/Source/Runtime/Core/Private/HAL/UnrealMemory.cpp
@@ -326,6 +326,21 @@ void FMemory::Free( void* Original )
 	return GMalloc->Free( Original );
 }
 
+void* FMemory::GPUMalloc(SIZE_T Count, uint32 Alignment /* = DEFAULT_ALIGNMENT */)
+{
+	return FPlatformMemory::GPUMalloc(Count, Alignment);
+}
+
+void* FMemory::GPURealloc(void* Original, SIZE_T Count, uint32 Alignment /* = DEFAULT_ALIGNMENT */)
+{
+	return FPlatformMemory::GPURealloc(Original, Count, Alignment);
+}
+
+void FMemory::GPUFree(void* Original)
+{
+	return FPlatformMemory::GPUFree(Original);
+}
+
 SIZE_T FMemory::GetAllocSize( void* Original )
 {
 	if( !GMalloc )
diff --git a/Engine/Source/Runtime/Core/Private/IOS/IOSPlatformFramePacer.cpp b/Engine/Source/Runtime/Core/Private/IOS/IOSPlatformFramePacer.cpp
--- a/Engine/Source/Runtime/Core/Private/IOS/IOSPlatformFramePacer.cpp
+++ b/Engine/Source/Runtime/Core/Private/IOS/IOSPlatformFramePacer.cpp
@@ -65,25 +65,31 @@ FIOSFramePacer* FIOSPlatformRHIFramePacer::FramePacer = nil;
 bool FIOSPlatformRHIFramePacer::IsEnabled()
 {
     static bool bIsRHIFramePacerEnabled = false;
-    
-    FString FrameRateLockAsEnum;
-    GConfig->GetString( TEXT("/Script/IOSRuntimeSettings.IOSRuntimeSettings"), TEXT("FrameRateLock"), FrameRateLockAsEnum, GEngineIni );
-    
-    uint32 FrameRateLock = 0;
-    if( FParse::Value(*FrameRateLockAsEnum, TEXT("PUFRL_"), FrameRateLock))
-    {
-        if( !bIsRHIFramePacerEnabled && FrameRateLock > 0 )
-        {
-            check( (IOSDisplayConstants::MaxRefreshRate % FrameRateLock) == 0 );
-            FrameInterval = IOSDisplayConstants::MaxRefreshRate / FrameRateLock;
-        
-            bIsRHIFramePacerEnabled = (FrameInterval > 0);
-        }
-    }
-    
-
-    
-    return bIsRHIFramePacerEnabled;
+	static bool bInitialized = false;
+
+	if (!bInitialized)
+	{
+		FString FrameRateLockAsEnum;
+		GConfig->GetString(TEXT("/Script/IOSRuntimeSettings.IOSRuntimeSettings"), TEXT("FrameRateLock"), FrameRateLockAsEnum, GEngineIni);
+
+		uint32 FrameRateLock = 60;
+		FParse::Value(*FrameRateLockAsEnum, TEXT("PUFRL_"), FrameRateLock);
+		if (FrameRateLock == 0)
+		{
+			FrameRateLock = 60;
+		}
+
+		if (!bIsRHIFramePacerEnabled)
+		{
+			check((IOSDisplayConstants::MaxRefreshRate % FrameRateLock) == 0);
+			FrameInterval = IOSDisplayConstants::MaxRefreshRate / FrameRateLock;
+
+			bIsRHIFramePacerEnabled = (FrameInterval > 0);
+		}
+		bInitialized = true;
+	}
+	
+	return bIsRHIFramePacerEnabled;
 }
 
 void FIOSPlatformRHIFramePacer::InitWithEvent(FEvent* TriggeredEvent)
diff --git a/Engine/Source/Runtime/Core/Private/IOS/IOSPlatformProcess.cpp b/Engine/Source/Runtime/Core/Private/IOS/IOSPlatformProcess.cpp
--- a/Engine/Source/Runtime/Core/Private/IOS/IOSPlatformProcess.cpp
+++ b/Engine/Source/Runtime/Core/Private/IOS/IOSPlatformProcess.cpp
@@ -41,10 +41,8 @@ void FIOSPlatformProcess::LaunchURL( const TCHAR* URL, const TCHAR* Parms, FStri
 {
 	UE_LOG(LogIOS, Log,  TEXT("LaunchURL %s %s"), URL, Parms?Parms:TEXT("") );
 	NSString* CFUrl = (NSString*)FPlatformString::TCHARToCFString( URL );
-//	CFUrl = CFURLCreateStringByAddingPercentEscapes(NULL, CFUrl, CFSTR("#?+"), NULL, kCFStringEncodingUTF8);
-    CFUrl = [CFUrl stringByAddingPercentEncodingWithAllowedCharacters: [NSCharacterSet characterSetWithCharactersInString:@"#?+"]];
+    CFUrl = [CFUrl stringByAddingPercentEncodingWithAllowedCharacters: [[NSCharacterSet characterSetWithCharactersInString:@"#?+"] invertedSet]];
 	[[UIApplication sharedApplication] openURL: [NSURL URLWithString:( NSString *)CFUrl]];
-	CFRelease( CFUrl );
 
 	if (Error != nullptr)
 	{
@@ -55,10 +53,8 @@ void FIOSPlatformProcess::LaunchURL( const TCHAR* URL, const TCHAR* Parms, FStri
 bool FIOSPlatformProcess::CanLaunchURL(const TCHAR* URL)
 {
 	NSString* CFUrl = (NSString*)FPlatformString::TCHARToCFString(URL);
-//	CFUrl = CFURLCreateStringByAddingPercentEscapes(NULL, CFUrl, CFSTR("#?+"), NULL, kCFStringEncodingUTF8);
-    CFUrl = [CFUrl stringByAddingPercentEncodingWithAllowedCharacters: [NSCharacterSet characterSetWithCharactersInString: @"#?+"]];
+    CFUrl = [CFUrl stringByAddingPercentEncodingWithAllowedCharacters: [[NSCharacterSet characterSetWithCharactersInString: @"#?+"] invertedSet]];
 	bool Result = [[UIApplication sharedApplication] canOpenURL:[NSURL URLWithString : (NSString *)CFUrl]];
-	CFRelease(CFUrl);
 
 	return Result;
 }
diff --git a/Engine/Source/Runtime/Core/Private/Internationalization/TextLocalizationManager.cpp b/Engine/Source/Runtime/Core/Private/Internationalization/TextLocalizationManager.cpp
--- a/Engine/Source/Runtime/Core/Private/Internationalization/TextLocalizationManager.cpp
+++ b/Engine/Source/Runtime/Core/Private/Internationalization/TextLocalizationManager.cpp
@@ -655,7 +655,7 @@ void FTextLocalizationManager::UpdateFromLocalizationResource(const FString& Loc
 {
 	TArray<FLocalizationEntryTracker> LocalizationEntryTrackers;
 
-	FLocalizationEntryTracker LocalizationEntryTracker = LocalizationEntryTrackers[LocalizationEntryTrackers.Add(FLocalizationEntryTracker())];
+	FLocalizationEntryTracker& LocalizationEntryTracker = LocalizationEntryTrackers[LocalizationEntryTrackers.Add(FLocalizationEntryTracker())];
 	LocalizationEntryTracker.LoadFromFile(LocalizationResourceFilePath);
 	LocalizationEntryTracker.DetectAndLogConflicts();
 
diff --git a/Engine/Source/Runtime/Core/Private/Linux/LinuxPlatformSplash.cpp b/Engine/Source/Runtime/Core/Private/Linux/LinuxPlatformSplash.cpp
--- a/Engine/Source/Runtime/Core/Private/Linux/LinuxPlatformSplash.cpp
+++ b/Engine/Source/Runtime/Core/Private/Linux/LinuxPlatformSplash.cpp
@@ -512,46 +512,50 @@ static int RenderString (GLuint tex_idx)
  *
  * @return Splash surface
  */
-static SDL_Surface* LinuxSplash_LoadImage(const TCHAR* ImagePath, IImageWrapperPtr & OutImageWrapper)
+static SDL_Surface* LinuxSplash_LoadImage(const TCHAR* InImagePath, IImageWrapperPtr & OutImageWrapper)
 {
 	TArray<uint8> RawFileData;
+	FString ImagePath(InImagePath);
 
-	// Load the image buffer first
-	if (!FFileHelper::LoadFileToArray(RawFileData, ImagePath))
+	// Load the image buffer first (unless it's BMP)
+	if (!ImagePath.EndsWith(TEXT("bmp"), ESearchCase::IgnoreCase) && FFileHelper::LoadFileToArray(RawFileData, *ImagePath))
 	{
-		// If for some reason the image cannot be loaded, use the default BMP function
-		return SDL_LoadBMP(TCHAR_TO_UTF8(ImagePath));
-	}
-
-	auto& ImageWrapperModule = FModuleManager::LoadModuleChecked<IImageWrapperModule>(FName("ImageWrapper"));
-	auto Format = ImageWrapperModule.DetectImageFormat(RawFileData.GetData(), RawFileData.Num());
-	OutImageWrapper = ImageWrapperModule.CreateImageWrapper(Format);
-
-	if (OutImageWrapper.IsValid() && OutImageWrapper->SetCompressed(RawFileData.GetData(), RawFileData.Num()))
-	{
-		const TArray<uint8>* RawData = nullptr;
+		auto& ImageWrapperModule = FModuleManager::LoadModuleChecked<IImageWrapperModule>(FName("ImageWrapper"));
+		auto Format = ImageWrapperModule.DetectImageFormat(RawFileData.GetData(), RawFileData.Num());
+		OutImageWrapper = ImageWrapperModule.CreateImageWrapper(Format);
 
-		if (OutImageWrapper->GetRaw(ERGBFormat::BGRA, 8, RawData))
+		if (OutImageWrapper.IsValid() && OutImageWrapper->SetCompressed(RawFileData.GetData(), RawFileData.Num()))
 		{
-			uint32 Bmask = 0x000000ff;
-			uint32 Gmask = 0x0000ff00;
-			uint32 Rmask = 0x00ff0000;
-			uint32 Amask = 0xff000000;
-
-			return SDL_CreateRGBSurfaceFrom((void*)RawData->GetData(),
-				OutImageWrapper->GetWidth(), OutImageWrapper->GetHeight(),
-				32, OutImageWrapper->GetWidth() * 4,
-				Rmask, Gmask, Bmask, Amask);
+			const TArray<uint8>* RawData = nullptr;
+
+			if (OutImageWrapper->GetRaw(ERGBFormat::BGRA, 8, RawData))
+			{
+				uint32 Bmask = 0x000000ff;
+				uint32 Gmask = 0x0000ff00;
+				uint32 Rmask = 0x00ff0000;
+				uint32 Amask = 0xff000000;
+
+				return SDL_CreateRGBSurfaceFrom((void*)RawData->GetData(),
+					OutImageWrapper->GetWidth(), OutImageWrapper->GetHeight(),
+					32, OutImageWrapper->GetWidth() * 4,
+					Rmask, Gmask, Bmask, Amask);
+			}
 		}
 	}
 
-	// If for some reason the image cannot be load, use the default BMP function
-	return SDL_LoadBMP(TCHAR_TO_UTF8(ImagePath));
+	// If for some reason the image cannot be loaded, use the default BMP function
+	return SDL_LoadBMP(TCHAR_TO_UTF8(*ImagePath));
 }
 
+void LinuxSplash_TearDownSplashResources();
+
 /** Helper function to init resources used by the splash thread */
 bool LinuxSplash_InitSplashResources()
 {
+	checkf(GSplashScreenImage == nullptr, TEXT("LinuxSplash_InitSplashResources() has been called multiple times."));
+	checkf(GSplashWindow == nullptr, TEXT("LinuxSplash_InitSplashResources() has been called multiple times."));
+	checkf(GSplashIconImage == nullptr, TEXT("LinuxSplash_InitSplashResources() has been called multiple times."));
+
 	if (!FPlatformMisc::PlatformInitMultimedia()) //	will not initialize more than once
 	{
 		UE_LOG(LogInit, Warning, TEXT("LinuxSplash_InitSplashResources() : PlatformInitMultimedia() failed, there will be no splash."));
@@ -578,12 +582,13 @@ bool LinuxSplash_InitSplashResources()
 	if (GSplashWindow == nullptr)
 	{
 		UE_LOG(LogHAL, Error, TEXT("LinuxSplash_InitSplashResources() : Splash screen window could not be created! SDL_Error: %s"), UTF8_TO_TCHAR(SDL_GetError()));
-		return -1;
+		LinuxSplash_TearDownSplashResources();
+		return false;
 	}
 
 	SDL_SetWindowTitle( GSplashWindow, TCHAR_TO_UTF8(*GSplashScreenAppName.ToString()) );
 
-	GSplashIconImage = LinuxSplash_LoadImage(*(FString(FPlatformProcess::BaseDir()) / GIconPath), GIconImageWrapper);
+	GSplashIconImage = LinuxSplash_LoadImage(*GIconPath, GIconImageWrapper);
 
 	if (GSplashIconImage != nullptr)
 	{
@@ -592,6 +597,7 @@ bool LinuxSplash_InitSplashResources()
 	else
 	{
 		UE_LOG(LogHAL, Warning, TEXT("LinuxSplash_InitSplashResources() : Splash icon could not be created! SDL_Error: %s"), UTF8_TO_TCHAR(SDL_GetError()));
+		LinuxSplash_TearDownSplashResources();
 		return false;
 	}
 
@@ -601,16 +607,25 @@ bool LinuxSplash_InitSplashResources()
 /** Helper function to tear down resources used by the splash thread */
 void LinuxSplash_TearDownSplashResources()
 {
-	SDL_DestroyWindow(GSplashWindow);
-	GSplashWindow = nullptr;
-
-	SDL_FreeSurface(GSplashIconImage);
-	GSplashIconImage = nullptr;
+	if (GSplashIconImage)
+	{
+		SDL_FreeSurface(GSplashIconImage);
+		GSplashIconImage = nullptr;
+	}
 	GIconImageWrapper.Reset();
 
-	SDL_FreeSurface(GSplashScreenImage);
-	GSplashScreenImage = nullptr;
-	GSplashImageWrapper.Reset();
+	if (GSplashWindow)
+	{
+		SDL_DestroyWindow(GSplashWindow);
+		GSplashWindow = nullptr;
+	}
+
+	if (GSplashScreenImage)
+	{
+		SDL_FreeSurface(GSplashScreenImage);
+		GSplashScreenImage = nullptr;
+		GSplashImageWrapper.Reset();
+	}
 
 	// do not deinit SDL here
 }
@@ -890,18 +905,75 @@ void FLinuxPlatformSplash::Show( )
 	// decide on which splash screen to show
 	const FText GameName = FText::FromString(FApp::GetGameName());
 
-	const TCHAR* SplashImage = GIsEditor ? ( GameName.IsEmpty() ? TEXT("EdSplashDefault") : TEXT("EdSplash") ) : ( GameName.IsEmpty() ? TEXT("SplashDefault") : TEXT("Splash") );
-	const TCHAR* IconImage = GIsEditor ? ( GameName.IsEmpty() ? TEXT("EdIconDefault") : TEXT("EdIcon") ) : ( GameName.IsEmpty() ? TEXT("IconDefault") : TEXT("Icon.bmp") );
+	bool IsCustom = false;
 
-	// make sure a splash was found
-	bool IsCustom;
-	
-	if (GetSplashPath(SplashImage, IconImage, GSplashPath, GIconPath, IsCustom) == false)
+	// first look for the splash, do not init anything if not found
 	{
-		UE_LOG(LogHAL, Warning, TEXT("Splash screen image not found."));
+		const TCHAR* EditorSplashes[] =
+		{
+			TEXT("EdSplash"),
+			TEXT("EdSplashDefault"),
+			nullptr
+		};
+
+		const TCHAR* GameSplashes[] =
+		{
+			TEXT("Splash"),
+			TEXT("SplashDefault"),
+			nullptr
+		};
+
+		const TCHAR** SplashesToTry = GIsEditor? EditorSplashes : GameSplashes;
+		bool bSplashFound = false;
+		for (const TCHAR* SplashImage = *SplashesToTry; SplashImage != nullptr && !bSplashFound; SplashImage = *(++SplashesToTry))
+		{
+			if (GetSplashPath(SplashImage, GSplashPath, IsCustom))
+			{
+				bSplashFound = true;
+			}
+		}
+
+		if (!bSplashFound)
+		{
+			UE_LOG(LogHAL, Warning, TEXT("Splash screen image not found."));
+			return;	// early out
+		}
+	}
+
+	// look for the icon separately, also avoid initialization if not found
+	{
+		const TCHAR* EditorIcons[] =
+		{
+			TEXT("EdIcon"),
+			TEXT("EdIconDefault"),
+			nullptr
+		};
+
+		const TCHAR* GameIcons[] =
+		{
+			TEXT("Icon"),
+			TEXT("IconDefault"),
+			nullptr
+		};
+
+		const TCHAR** IconsToTry = GIsEditor? EditorIcons : GameIcons;
+		bool bIconFound = false;
+		for (const TCHAR* IconImage = *IconsToTry; IconImage != nullptr && !bIconFound; IconImage = *(++IconsToTry))
+		{
+			bool bDummy;
+			if (GetSplashPath(IconImage, GIconPath, bDummy))
+			{
+				bIconFound = true;
+			}
+		}
+
+		if (!bIconFound)
+		{
+			UE_LOG(LogHAL, Warning, TEXT("Game icon not found."));
+			return;	// early out
+		}
 	}
 
-	
 	// Don't set the game name if the splash screen is custom.
 	if ( !IsCustom )
 	{
@@ -985,11 +1057,11 @@ void FLinuxPlatformSplash::Hide()
 			// busy loop!
 			FPlatformProcess::Sleep(0.01f);
 		}
-
-		// tear down resources that thread used
-		LinuxSplash_TearDownSplashResources();
 	}
 
+	// tear down resources that thread used (safe to call even if they were not initialized
+	LinuxSplash_TearDownSplashResources();
+
 	GSplashThread = nullptr;
 #endif
 }
diff --git a/Engine/Source/Runtime/Core/Private/Linux/LinuxWindow.cpp b/Engine/Source/Runtime/Core/Private/Linux/LinuxWindow.cpp
--- a/Engine/Source/Runtime/Core/Private/Linux/LinuxWindow.cpp
+++ b/Engine/Source/Runtime/Core/Private/Linux/LinuxWindow.cpp
@@ -209,49 +209,57 @@ void FLinuxWindow::Initialize( FLinuxApplication* const Application, const TShar
 	//	the size of the window you input is the sizeof the client.
 	HWnd = SDL_CreateWindow( TCHAR_TO_ANSI( *Definition->Title ), X, Y, ClientWidth, ClientHeight, WindowStyle  );
 
-	// Try to find an icon for the window
-	const FText GameName = FText::FromString(FApp::GetGameName());
-	FString Iconpath;
-	static SDL_Surface *IconImage = nullptr;
-
-	if (!GameName.IsEmpty())
+	if (Definition->AppearsInTaskbar)
 	{
-		if (GIsEditor)
-		{
-			Iconpath = FPaths::GameContentDir() / TEXT("Splash/EdIcon.bmp");
-		}
-		else
-		{
-			Iconpath = FPaths::GameContentDir() / TEXT("Splash/Icon.bmp");
-		}
+		// Try to find an icon for the window
+		const FText GameName = FText::FromString(FApp::GetGameName());
+		FString Iconpath;
+		SDL_Surface *IconImage = nullptr;
 
-		if (IFileManager::Get().FileSize(*Iconpath) != -1)
+		if (!GameName.IsEmpty())
 		{
-			IconImage = SDL_LoadBMP(TCHAR_TO_UTF8(*Iconpath));
-		}
-	}
+			if (GIsEditor)
+			{
+				Iconpath = FPaths::GameContentDir() / TEXT("Splash/EdIcon.bmp");
+			}
+			else
+			{
+				Iconpath = FPaths::GameContentDir() / TEXT("Splash/Icon.bmp");
+			}
 
-	if (IconImage == nullptr)
-	{
-		// no game specified or there's no custom icons for the game. use default icons.
-		if (GIsEditor)
-		{
-			Iconpath = FPaths::EngineContentDir() / TEXT("Splash/EdIconDefault.bmp");
+			Iconpath = FPaths::FPaths::ConvertRelativePathToFull(Iconpath);
+			if (IFileManager::Get().FileSize(*Iconpath) != -1)
+			{
+				IconImage = SDL_LoadBMP(TCHAR_TO_UTF8(*Iconpath));
+			}
 		}
-		else
+
+		if (IconImage == nullptr)
 		{
-			Iconpath = FPaths::EngineContentDir() / TEXT("Splash/IconDefault.bmp");
+			// no game specified or there's no custom icons for the game. use default icons.
+			if (GIsEditor)
+			{
+				Iconpath = FPaths::EngineContentDir() / TEXT("Splash/EdIconDefault.bmp");
+			}
+			else
+			{
+				Iconpath = FPaths::EngineContentDir() / TEXT("Splash/IconDefault.bmp");
+			}
+
+			Iconpath = FPaths::FPaths::ConvertRelativePathToFull(Iconpath);
+			if (IFileManager::Get().FileSize(*Iconpath) != -1)
+			{
+				IconImage = SDL_LoadBMP(TCHAR_TO_UTF8(*Iconpath));
+			}
 		}
 
-		if (IFileManager::Get().FileSize(*Iconpath) != -1)
+		if (IconImage != nullptr)
 		{
-			IconImage = SDL_LoadBMP(TCHAR_TO_UTF8(*Iconpath));
-		}
-	}
+			SDL_SetWindowIcon(HWnd, IconImage);
 
-	if (IconImage != nullptr)
-	{
-		SDL_SetWindowIcon(HWnd, IconImage);
+			SDL_FreeSurface(IconImage);
+			IconImage = nullptr;
+		}
 	}
 
 	SDL_SetWindowHitTest( HWnd, FLinuxWindow::HitTest, this );
diff --git a/Engine/Source/Runtime/Core/Private/Mac/MacApplication.cpp b/Engine/Source/Runtime/Core/Private/Mac/MacApplication.cpp
--- a/Engine/Source/Runtime/Core/Private/Mac/MacApplication.cpp
+++ b/Engine/Source/Runtime/Core/Private/Mac/MacApplication.cpp
@@ -49,6 +49,7 @@ FMacApplication::FMacApplication()
 ,	bSystemModalMode(false)
 ,	ModifierKeysFlags(0)
 ,	CurrentModifierFlags(0)
+,	bIsRightClickEmulationEnabled(true)
 ,	bEmulatingRightClick(false)
 ,	bIgnoreMouseMoveDelta(false)
 ,	bIsWorkspaceSessionActive(true)
@@ -359,7 +360,7 @@ void FMacApplication::DeferEvent(NSObject* Object)
 			case NSOtherMouseUp:
 				DeferredEvent.ButtonNumber = [Event buttonNumber];
 				DeferredEvent.ClickCount = [Event clickCount];
-				if (DeferredEvent.Type == NSLeftMouseDown && (DeferredEvent.ModifierFlags & NSControlKeyMask))
+				if (bIsRightClickEmulationEnabled && DeferredEvent.Type == NSLeftMouseDown && (DeferredEvent.ModifierFlags & NSControlKeyMask))
 				{
 					bEmulatingRightClick = true;
 					DeferredEvent.Type = NSRightMouseDown;
diff --git a/Engine/Source/Runtime/Core/Private/Mac/MacPlatformMisc.cpp b/Engine/Source/Runtime/Core/Private/Mac/MacPlatformMisc.cpp
--- a/Engine/Source/Runtime/Core/Private/Mac/MacPlatformMisc.cpp
+++ b/Engine/Source/Runtime/Core/Private/Mac/MacPlatformMisc.cpp
@@ -2133,6 +2133,29 @@ void FMacCrashContext::GenerateInfoInFolder(char const* const InfoFolder, bool b
 		{
 			write(LogDst, Data, Bytes);
 		}
+		
+		// Copy the system log to capture GPU restarts and other nasties not reported by our application
+		if ( !GMacAppInfo.bIsSandboxed && GMacAppInfo.SystemLogSize >= 0 && access("/var/log/system.log", R_OK|F_OK) == 0 )
+		{
+			char const* SysLogHeader = "\nAppending System Log:\n";
+			write(LogDst, SysLogHeader, strlen(SysLogHeader));
+			
+			int SysLogSrc = open("/var/log/system.log", O_RDONLY);
+			
+			// Attempt to capture only the system log from while our application was running but
+			if (lseek(SysLogSrc, GMacAppInfo.SystemLogSize, SEEK_SET) != GMacAppInfo.SystemLogSize)
+			{
+				close(SysLogSrc);
+				SysLogSrc = open("/var/log/system.log", O_RDONLY);
+			}
+			
+			while((Bytes = read(SysLogSrc, Data, PATH_MAX)) > 0)
+			{
+				write(LogDst, Data, Bytes);
+			}
+			close(SysLogSrc);
+		}
+		
 		close(LogDst);
 		close(LogSrc);
 		// best effort, so don't care about result: couldn't copy -> tough, no log
@@ -2153,30 +2176,6 @@ void FMacCrashContext::GenerateInfoInFolder(char const* const InfoFolder, bool b
 			close(VideoDst);
 			close(VideoSrc);
 		}
-		
-		// Copy the system log to capture GPU restarts and other nasties not reported by our application
-		if ( !GMacAppInfo.bIsSandboxed && GMacAppInfo.SystemLogSize >= 0 && access("/var/log/system.log", R_OK|F_OK) == 0 )
-		{
-			FCStringAnsi::Strncpy(FilePath, CrashInfoFolder, PATH_MAX);
-			FCStringAnsi::Strcat(FilePath, PATH_MAX, "/");
-			FCStringAnsi::Strcat(FilePath, PATH_MAX, "system.log");
-			int SysLogSrc = open("/var/log/system.log", O_RDONLY);
-			int SysLogDst = open(FilePath, O_CREAT|O_WRONLY, 0766);
-			
-			// Attempt to capture only the system log from while our application was running but
-			if (lseek(SysLogSrc, GMacAppInfo.SystemLogSize, SEEK_SET) != GMacAppInfo.SystemLogSize)
-			{
-				close(SysLogSrc);
-				SysLogSrc = open("/var/log/system.log", O_RDONLY);
-			}
-			
-			while((Bytes = read(SysLogSrc, Data, PATH_MAX)) > 0)
-			{
-				write(SysLogDst, Data, Bytes);
-			}
-			close(SysLogDst);
-			close(SysLogSrc);
-		}
 	}
 }
 
diff --git a/Engine/Source/Runtime/Core/Public/Containers/ContainerAllocationPolicies.h b/Engine/Source/Runtime/Core/Public/Containers/ContainerAllocationPolicies.h
--- a/Engine/Source/Runtime/Core/Public/Containers/ContainerAllocationPolicies.h
+++ b/Engine/Source/Runtime/Core/Public/Containers/ContainerAllocationPolicies.h
@@ -239,7 +239,11 @@ public:
 
 			if (Data)
 			{
+#if PLATFORM_HAS_UMA
+				FMemory::GPUFree(Data);
+#else
 				FMemory::Free(Data);
+#endif
 			}
 
 			Data       = Other.Data;
@@ -251,7 +255,11 @@ public:
 		{
 			if(Data)
 			{
+#if PLATFORM_HAS_UMA
+				// The RHI should have taken ownership of this memory, so we don't free it here
+#else
 				FMemory::Free(Data);
+#endif
 			}
 		}
 
@@ -267,10 +275,14 @@ public:
 			)
 		{
 			// Avoid calling FMemory::Realloc( nullptr, 0 ) as ANSI C mandates returning a valid pointer which is not what we want.
-			if( Data || NumElements )
+			if (Data || NumElements)
 			{
+#if PLATFORM_HAS_UMA
+				Data = (FScriptContainerElement*)FMemory::GPURealloc(Data, NumElements*NumBytesPerElement, Alignment);
+#else
 				//checkSlow(((uint64)NumElements*(uint64)ElementTypeInfo.GetSize() < (uint64)INT_MAX));
 				Data = (FScriptContainerElement*)FMemory::Realloc( Data, NumElements*NumBytesPerElement, Alignment );
+#endif
 			}
 		}
 		int32 CalculateSlack(
diff --git a/Engine/Source/Runtime/Core/Public/Containers/LockFreeVoidPointerListBase.h b/Engine/Source/Runtime/Core/Public/Containers/LockFreeVoidPointerListBase.h
--- a/Engine/Source/Runtime/Core/Public/Containers/LockFreeVoidPointerListBase.h
+++ b/Engine/Source/Runtime/Core/Public/Containers/LockFreeVoidPointerListBase.h
@@ -2128,7 +2128,7 @@ public:
 	 *	CAUTION: This methods safety depends on external assumptions. For example, if another thread could add to the list at any time, the return value is no better than a best guess.
 	 *	As typically used, the list is not being access concurrently when this is called.
 	 */
-	FORCEINLINE bool IsEmpty() const  
+	FORCEINLINE bool IsEmpty()  
 	{
 		return FLockFreePointerListFIFOBase<TPaddingForCacheContention>::IsEmpty();
 	}
diff --git a/Engine/Source/Runtime/Core/Public/GenericPlatform/GenericPlatformMemory.h b/Engine/Source/Runtime/Core/Public/GenericPlatform/GenericPlatformMemory.h
--- a/Engine/Source/Runtime/Core/Public/GenericPlatform/GenericPlatformMemory.h
+++ b/Engine/Source/Runtime/Core/Public/GenericPlatform/GenericPlatformMemory.h
@@ -251,6 +251,12 @@ struct CORE_API FGenericPlatformMemory
 	 */
 	static void BinnedFreeToOS( void* Ptr );
 
+	// These alloc/free memory that is mapped to the GPU
+	// Only for platforms with UMA (XB1/PS4/etc)
+	static void* GPUMalloc(SIZE_T Count, uint32 Alignment = 0) { return nullptr; };
+	static void* GPURealloc(void* Original, SIZE_T Count, uint32 Alignment = 0) { return nullptr; };
+	static void GPUFree(void* Original) { };
+
 	/** Dumps basic platform memory statistics into the specified output device. */
 	static void DumpStats( FOutputDevice& Ar );
 
diff --git a/Engine/Source/Runtime/Core/Public/HAL/UnrealMemory.h b/Engine/Source/Runtime/Core/Public/HAL/UnrealMemory.h
--- a/Engine/Source/Runtime/Core/Public/HAL/UnrealMemory.h
+++ b/Engine/Source/Runtime/Core/Public/HAL/UnrealMemory.h
@@ -129,6 +129,15 @@ struct CORE_API FMemory
 	static void* Realloc( void* Original, SIZE_T Count, uint32 Alignment=DEFAULT_ALIGNMENT );
 	static void Free( void* Original );
 
+	//
+	// Malloc for GPU mapped memory on UMA systems (XB1/PS4/etc)
+	// It is expected that the RHI on platforms that use these knows what to 
+	// do with the memory and avoids unnecessary copies into GPU resources, etc.
+	//
+	static void* GPUMalloc(SIZE_T Count, uint32 Alignment = DEFAULT_ALIGNMENT);
+	static void* GPURealloc(void* Original, SIZE_T Count, uint32 Alignment = DEFAULT_ALIGNMENT);
+	static void GPUFree(void* Original);
+
 	static SIZE_T GetAllocSize( void* Original );
 
 	/**
diff --git a/Engine/Source/Runtime/Core/Public/Mac/MacApplication.h b/Engine/Source/Runtime/Core/Public/Mac/MacApplication.h
--- a/Engine/Source/Runtime/Core/Public/Mac/MacApplication.h
+++ b/Engine/Source/Runtime/Core/Public/Mac/MacApplication.h
@@ -202,6 +202,8 @@ public:
 
 	void IgnoreMouseMoveDelta() { bIgnoreMouseMoveDelta = true; }
 
+	void SetIsRightClickEmulationEnabled(bool bEnabled) { bIsRightClickEmulationEnabled = bEnabled; }
+
 public:
 
 	static void UpdateScreensArray();
@@ -308,6 +310,7 @@ private:
 	/** The current set of Cocoa modifier flags, used to detect when Mission Control has been invoked & returned so that we can synthesis the modifier events it steals */
 	NSUInteger CurrentModifierFlags;
 
+	bool bIsRightClickEmulationEnabled;
 	bool bEmulatingRightClick;
 
 	bool bIgnoreMouseMoveDelta;
diff --git a/Engine/Source/Runtime/CoreUObject/Private/UObject/Enum.cpp b/Engine/Source/Runtime/CoreUObject/Private/UObject/Enum.cpp
--- a/Engine/Source/Runtime/CoreUObject/Private/UObject/Enum.cpp
+++ b/Engine/Source/Runtime/CoreUObject/Private/UObject/Enum.cpp
@@ -80,18 +80,21 @@ FString UEnum::GetBaseEnumNameOnDuplication() const
 
 void UEnum::RenameNamesAfterDuplication()
 {
-	// Get name of base enum, from which we're duplicating.
-	FString BaseEnumName = GetBaseEnumNameOnDuplication();
+	if (Names.Num() != 0)
+	{
+		// Get name of base enum, from which we're duplicating.
+		FString BaseEnumName = GetBaseEnumNameOnDuplication();
 
-	// Get name of duplicated enum.
-	FString ThisName = GetName();
+		// Get name of duplicated enum.
+		FString ThisName = GetName();
 
-	// Replace all usages of base class name to the duplicated one.
-	for (TPair<FName, uint8>& Kvp : Names)
-	{
-		FString NameString = Kvp.Key.ToString();
-		NameString.ReplaceInline(*BaseEnumName, *ThisName);
-		Kvp.Key = FName(*NameString);
+		// Replace all usages of base class name to the duplicated one.
+		for (TPair<FName, uint8>& Kvp : Names)
+		{
+			FString NameString = Kvp.Key.ToString();
+			NameString.ReplaceInline(*BaseEnumName, *ThisName);
+			Kvp.Key = FName(*NameString);
+		}
 	}
 }
 
diff --git a/Engine/Source/Runtime/CoreUObject/Private/UObject/GarbageCollection.cpp b/Engine/Source/Runtime/CoreUObject/Private/UObject/GarbageCollection.cpp
--- a/Engine/Source/Runtime/CoreUObject/Private/UObject/GarbageCollection.cpp
+++ b/Engine/Source/Runtime/CoreUObject/Private/UObject/GarbageCollection.cpp
@@ -349,12 +349,13 @@ public:
 			{
 				if (ReferencedMutableObjectItem->IsUnreachable() && ReferencedMutableObjectItem->ThisThreadAtomicallyClearedRFUnreachable())
 				{
+					ReferencedMutableObjectItem->ThisThreadAtomicallyClearedFlag(EInternalObjectFlags::NoStrongReference);
 					ObjectsToSerialize.Add(static_cast<UObject*>(ReferencedMutableObjectItem->Object));
 				}
 			}
 			else if (ReferencedMutableObjectItem->IsUnreachable())
 			{
-				ReferencedMutableObjectItem->ClearUnreachable();
+				ReferencedMutableObjectItem->ClearFlags(EInternalObjectFlags::NoStrongReference | EInternalObjectFlags::Unreachable);
 				ObjectsToSerialize.Add(static_cast<UObject*>(ReferencedMutableObjectItem->Object));
 			}
 		}
@@ -373,11 +374,11 @@ public:
 			// This condition should get collapsed by the compiler based on the template argument
 			if (bParallel)
 			{
-				ReferencedClusterRootObjectItem->ThisThreadAtomicallyClearedRFUnreachable();
+				ReferencedClusterRootObjectItem->ThisThreadAtomicallyClearedFlag(EInternalObjectFlags::NoStrongReference | EInternalObjectFlags::Unreachable);
 			}
 			else
 			{
-				ReferencedClusterRootObjectItem->ClearUnreachable();
+				ReferencedClusterRootObjectItem->ClearFlags(EInternalObjectFlags::NoStrongReference | EInternalObjectFlags::Unreachable);
 			}
 			FUObjectCluster* ReferencedCluster = GUObjectClusters.FindChecked(ReferncedClusterIndex);
 			MarkClusterMutableObjectsAsReachable<bParallel>(ReferencedCluster, ObjectsToSerialize);
@@ -482,13 +483,14 @@ public:
 			{
 				if (RootObjectItem->ThisThreadAtomicallyClearedRFUnreachable())
 				{
+					RootObjectItem->ThisThreadAtomicallyClearedFlag(EInternalObjectFlags::NoStrongReference);
 					// Make sure all referenced clusters are marked as reachable too
 					MarkReferencedClustersAsReachable<true>(OwnerIndex, ObjectsToSerialize);
 				}
 			}
 			else if (RootObjectItem->IsUnreachable())
 			{
-				RootObjectItem->ClearUnreachable();
+				RootObjectItem->ClearFlags(EInternalObjectFlags::Unreachable | EInternalObjectFlags::NoStrongReference);
 				// Make sure all referenced clusters are marked as reachable too
 				MarkReferencedClustersAsReachable<false>(OwnerIndex, ObjectsToSerialize);
 			}
diff --git a/Engine/Source/Runtime/CoreUObject/Private/UObject/Obj.cpp b/Engine/Source/Runtime/CoreUObject/Private/UObject/Obj.cpp
--- a/Engine/Source/Runtime/CoreUObject/Private/UObject/Obj.cpp
+++ b/Engine/Source/Runtime/CoreUObject/Private/UObject/Obj.cpp
@@ -1630,16 +1630,16 @@ void UObject::LoadConfig( UClass* ConfigClass/*=NULL*/, const TCHAR* InFilename/
 					Key = FString::Printf(TEXT("%s[%i]"), *Property->GetName(), i);
 				}
 
-				FString Value;
+					FString Value;
 				const bool bFoundValue = GConfig->GetString( *ClassSection, *Key, Value, *PropFileName );
-				if (bFoundValue)
-				{
-					if (Property->ImportText(*Value, Property->ContainerPtrToValuePtr<uint8>(this, i), PortFlags, this) == NULL)
+					if (bFoundValue)
 					{
-						// this should be an error as the properties from the .ini / .int file are not correctly being read in and probably are affecting things in subtle ways
-						UE_LOG(LogObj, Error, TEXT("LoadConfig (%s): import failed for %s in: %s"), *GetPathName(), *Property->GetName(), *Value);
+						if (Property->ImportText(*Value, Property->ContainerPtrToValuePtr<uint8>(this, i), PortFlags, this) == NULL)
+						{
+							// this should be an error as the properties from the .ini / .int file are not correctly being read in and probably are affecting things in subtle ways
+							UE_LOG(LogObj, Error, TEXT("LoadConfig (%s): import failed for %s in: %s"), *GetPathName(), *Property->GetName(), *Value);
+						}
 					}
-				}
 
 #if !UE_BUILD_SHIPPING
 				if (!bFoundValue && !FPlatformProperties::RequiresCookedData())
@@ -1671,11 +1671,11 @@ void UObject::LoadConfig( UClass* ConfigClass/*=NULL*/, const TCHAR* InFilename/
 				if ( List.Num() > 0 )
 				{
 					ArrayHelper.EmptyAndAddValues(List.Num());
-					for( int32 i=List.Num()-1,c=0; i>=0; i--,c++ )
-					{
-						Array->Inner->ImportText( *List[i], ArrayHelper.GetRawPtr(c), PortFlags, this );
+						for( int32 i=List.Num()-1,c=0; i>=0; i--,c++ )
+						{
+							Array->Inner->ImportText( *List[i], ArrayHelper.GetRawPtr(c), PortFlags, this );
+						}
 					}
-				}
 				else
 				{
 					int32 Index = 0;
@@ -1827,13 +1827,13 @@ void UObject::SaveConfig( uint64 Flags, const TCHAR* InFilename, FConfigCacheIni
 					FString CompleteKey = FString::Printf(TEXT("%s%s"), bIsADefaultIniWrite ? TEXT("+") : TEXT(""), *Key);
 
 					FScriptArrayHelper_InContainer ArrayHelper(Array, this);
-					for( int32 i=0; i<ArrayHelper.Num(); i++ )
-					{
-						FString	Buffer;
-						Array->Inner->ExportTextItem( Buffer, ArrayHelper.GetRawPtr(i), ArrayHelper.GetRawPtr(i), this, PortFlags );
-						Sec->Add(*CompleteKey, *Buffer);
+						for( int32 i=0; i<ArrayHelper.Num(); i++ )
+						{
+							FString	Buffer;
+							Array->Inner->ExportTextItem( Buffer, ArrayHelper.GetRawPtr(i), ArrayHelper.GetRawPtr(i), this, PortFlags );
+							Sec->Add(*CompleteKey, *Buffer);
+						}
 					}
-				}
 				else if( Property->Identical_InContainer(this, SuperClassDefaultObject) )
 				{
 					// If we are not writing it to config above, we should make sure that this property isn't stagnant in the cache.
@@ -1858,10 +1858,10 @@ void UObject::SaveConfig( uint64 Flags, const TCHAR* InFilename, FConfigCacheIni
 							Key = TempKey;
 						}
 
-						FString	Value;
-						Property->ExportText_InContainer( Index, Value, this, this, this, PortFlags );
-						Config->SetString( *Section, *Key, *Value, *PropFileName );
-					}
+							FString	Value;
+							Property->ExportText_InContainer( Index, Value, this, this, this, PortFlags );
+							Config->SetString( *Section, *Key, *Value, *PropFileName );
+						}
 					else if( Property->Identical_InContainer(this, SuperClassDefaultObject, Index) )
 					{
 						// If we are not writing it to config above, we should make sure that this property isn't stagnant in the cache.
diff --git a/Engine/Source/Runtime/CoreUObject/Private/UObject/SavePackage.cpp b/Engine/Source/Runtime/CoreUObject/Private/UObject/SavePackage.cpp
--- a/Engine/Source/Runtime/CoreUObject/Private/UObject/SavePackage.cpp
+++ b/Engine/Source/Runtime/CoreUObject/Private/UObject/SavePackage.cpp
@@ -3134,7 +3134,7 @@ ESavePackageResult UPackage::Save(UPackage* InOuter, UObject* Base, EObjectFlags
 
 		const bool bIsCooking = TargetPlatform != nullptr;
 #if WITH_EDITORONLY_DATA
-		if (bIsCooking)
+		if (bIsCooking && (!(SaveFlags & ESaveFlags::SAVE_KeepEditorOnlyCookedPackages)))
 		{
 			static struct FCanSkipEditorReferencedPackagesWhenCooking
 			{
diff --git a/Engine/Source/Runtime/CoreUObject/Private/UObject/ScriptCore.cpp b/Engine/Source/Runtime/CoreUObject/Private/UObject/ScriptCore.cpp
--- a/Engine/Source/Runtime/CoreUObject/Private/UObject/ScriptCore.cpp
+++ b/Engine/Source/Runtime/CoreUObject/Private/UObject/ScriptCore.cpp
@@ -1290,11 +1290,21 @@ void UObject::execLocalVariable(FFrame& Stack, RESULT_DECL)
 	checkSlow(Stack.Locals != NULL);
 
 	UProperty* VarProperty = Stack.ReadProperty();
-	Stack.MostRecentPropertyAddress = VarProperty->ContainerPtrToValuePtr<uint8>(Stack.Locals);
+	if (VarProperty == nullptr)
+	{
+		FBlueprintExceptionInfo ExceptionInfo(EBlueprintExceptionType::AccessViolation, LOCTEXT("MissingLocalVariable", "Attempted to access missing local variable. If this is a packaged/cooked build, are you attempting to use an editor-only property?"));
+		FBlueprintCoreDelegates::ThrowScriptException(this, Stack, ExceptionInfo);
 
-	if (RESULT_PARAM)
+		Stack.MostRecentPropertyAddress = nullptr;
+	}
+	else
 	{
-		VarProperty->CopyCompleteValueToScriptVM( RESULT_PARAM, Stack.MostRecentPropertyAddress );
+		Stack.MostRecentPropertyAddress = VarProperty->ContainerPtrToValuePtr<uint8>(Stack.Locals);
+
+		if (RESULT_PARAM)
+		{
+			VarProperty->CopyCompleteValueToScriptVM(RESULT_PARAM, Stack.MostRecentPropertyAddress);
+		}
 	}
 }
 IMPLEMENT_VM_FUNCTION( EX_LocalVariable, execLocalVariable );
diff --git a/Engine/Source/Runtime/CoreUObject/Public/UObject/ObjectBase.h b/Engine/Source/Runtime/CoreUObject/Public/UObject/ObjectBase.h
--- a/Engine/Source/Runtime/CoreUObject/Public/UObject/ObjectBase.h
+++ b/Engine/Source/Runtime/CoreUObject/Public/UObject/ObjectBase.h
@@ -79,6 +79,7 @@ enum ESaveFlags
 	SAVE_Async			= 0x00000010,	// Save to a memory writer, then actually write to disk async
 	SAVE_Unversioned	= 0x00000020,	// Save all versions as zero. Upon load this is changed to the current version. This is only reasonable to use with full cooked builds for distribution.
 	SAVE_CutdownPackage	= 0x00000040,	// Saving cutdown packages in a temp location WITHOUT renaming the package.
+	SAVE_KeepEditorOnlyCookedPackages = 0x00000080,  // keep packages which are marked as editor only even though we are cooking
 };
 
 //
diff --git a/Engine/Source/Runtime/CoreUObject/Public/UObject/UObjectAnnotation.h b/Engine/Source/Runtime/CoreUObject/Public/UObject/UObjectAnnotation.h
--- a/Engine/Source/Runtime/CoreUObject/Public/UObject/UObjectAnnotation.h
+++ b/Engine/Source/Runtime/CoreUObject/Public/UObject/UObjectAnnotation.h
@@ -308,16 +308,14 @@ public:
 	void RemoveAnnotation(const UObjectBase *Object)
 	{
 		TAnnotation Annotation = this->GetAndRemoveAnnotation(Object);
-		if (Annotation.IsDefault())
+		for (auto It = InverseAnnotationMap.CreateIterator(); It; ++It)
 		{
-			// should not exist in the mapping
-			checkSlow(!InverseAnnotationMap.Find(Annotation));
-		}
-		else
-		{
-			int32 NumRemoved = InverseAnnotationMap.Remove(Annotation);
-			checkSlow(NumRemoved == 1);
+			if (It->Value == Object)
+			{
+				It.RemoveCurrent();
+			}
 		}
+		checkSlow(!InverseAnnotationMap.Find(Annotation));
 	}
 	/**
 	 * Removes all annotation from the annotation list. 
diff --git a/Engine/Source/Runtime/Engine/Classes/AI/Navigation/NavLinkProxy.h b/Engine/Source/Runtime/Engine/Classes/AI/Navigation/NavLinkProxy.h
--- a/Engine/Source/Runtime/Engine/Classes/AI/Navigation/NavLinkProxy.h
+++ b/Engine/Source/Runtime/Engine/Classes/AI/Navigation/NavLinkProxy.h
@@ -68,6 +68,7 @@ public:
 #if WITH_EDITOR
 	virtual void PostEditChangeProperty(FPropertyChangedEvent& PropertyChangedEvent) override;
 #endif // WITH_EDITOR
+	virtual void PostInitializeComponents() override;
 
 #if ENABLE_VISUAL_LOG
 	virtual void BeginPlay() override;
diff --git a/Engine/Source/Runtime/Engine/Classes/Animation/AnimInstance.h b/Engine/Source/Runtime/Engine/Classes/Animation/AnimInstance.h
--- a/Engine/Source/Runtime/Engine/Classes/Animation/AnimInstance.h
+++ b/Engine/Source/Runtime/Engine/Classes/Animation/AnimInstance.h
@@ -1016,7 +1016,7 @@ public:
 	bool ParallelCanEvaluate(const USkeletalMesh* InSkeletalMesh) const;
 
 	/** Perform evaluation. Can be called from worker threads. */
-	void ParallelEvaluateAnimation(bool bForceRefPose, const USkeletalMesh* InSkeletalMesh, TArray<FTransform>& OutLocalAtoms, TArray<FActiveVertexAnim>& OutVertexAnims, FBlendedHeapCurve& OutCurve);
+	void ParallelEvaluateAnimation(bool bForceRefPose, const USkeletalMesh* InSkeletalMesh, TArray<FTransform>& OutLocalAtoms, FBlendedHeapCurve& OutCurve);
 
 	void PostEvaluateAnimation();
 	void UninitializeAnimation();
@@ -1105,6 +1105,9 @@ public:
 	/** Update all internal curves from Blended Curve */
 	void UpdateCurves(const FBlendedHeapCurve& InCurves);
 
+	/** Refresh currently existing curves */
+	void RefreshCurves(USkeletalMeshComponent* Component);
+
 	/** Check whether we have active morph target curves */
 	bool HasMorphTargetCurves() const;
 
@@ -1182,12 +1185,6 @@ public:
 	/** Add curve float data using a curve Uid, the name of the curve will be resolved from the skeleton **/
 	void AddCurveValue(const USkeleton::AnimCurveUID Uid, float Value, int32 CurveTypeFlags);
 
-	/** Pass-thru function to proxy - only call on the game thread */
-	void UpdateMorphTargetCurves(const TMap<FName, float>& InMorphTargetCurves);
-
-	/** Pass-thru function to proxy - only call on the game thread */
-	void UpdateComponentsMaterialParameters(UPrimitiveComponent* Component);
-
 	/** Given a machine and state index, record a state weight for this frame */
 	void RecordStateWeight(const int32& InMachineClassIndex, const int32& InStateIndex, const float& InStateWeight);
 
@@ -1211,8 +1208,6 @@ public:
 	/** Get current accumulated root motion, removing it from the AnimInstance in the process */
 	FRootMotionMovementParams ConsumeExtractedRootMotion(float Alpha);
 
-	/** Wrapper around UpdateActiveVertexAnims that can use vertex anims internal to the proxy */
-	TArray<FActiveVertexAnim> UpdateActiveVertexAnims(const USkeletalMesh* SkeletalMesh);
 private:
 	/** Active Root Motion Montage Instance, if any. */
 	struct FAnimMontageInstance * RootMotionMontageInstance;
diff --git a/Engine/Source/Runtime/Engine/Classes/Components/SkeletalMeshComponent.h b/Engine/Source/Runtime/Engine/Classes/Components/SkeletalMeshComponent.h
--- a/Engine/Source/Runtime/Engine/Classes/Components/SkeletalMeshComponent.h
+++ b/Engine/Source/Runtime/Engine/Classes/Components/SkeletalMeshComponent.h
@@ -18,8 +18,8 @@ class UAnimInstance;
 struct FEngineShowFlags;
 struct FConvexVolume;
 
-
-
+DECLARE_MULTICAST_DELEGATE(FOnSkelMeshPhysicsCreatedMultiCast);
+typedef FOnSkelMeshPhysicsCreatedMultiCast::FDelegate FOnSkelMeshPhysicsCreated;
 
 namespace physx
 { 
@@ -108,7 +108,6 @@ struct FAnimationEvaluationContext
 	// Double buffer evaluation data
 	TArray<FTransform> SpaceBases;
 	TArray<FTransform> LocalAtoms;
-	TArray<FActiveVertexAnim> VertexAnims;
 	FVector RootBoneTranslation;
 
 	// Double buffer curve data
@@ -142,8 +141,6 @@ struct FAnimationEvaluationContext
 		SpaceBases.Append(Other.SpaceBases);
 		LocalAtoms.Reset();
 		LocalAtoms.Append(Other.LocalAtoms);
-		VertexAnims.Reset();
-		VertexAnims.Append(Other.VertexAnims);
 		RootBoneTranslation = Other.RootBoneTranslation;
 		Curve.InitFrom(Other.Curve);
 		bDoInterpolation = Other.bDoInterpolation;
@@ -1035,7 +1032,7 @@ public:
 	* @param	OutVertexAnims			Active vertex animations
 	* @param	OutRootBoneTranslation	Calculated root bone translation
 	*/
-	void PerformAnimationEvaluation(const USkeletalMesh* InSkeletalMesh, UAnimInstance* InAnimInstance, TArray<FTransform>& OutSpaceBases, TArray<FTransform>& OutLocalAtoms, TArray<FActiveVertexAnim>& OutVertexAnims, FVector& OutRootBoneTranslation, FBlendedHeapCurve& OutCurve) const;
+	void PerformAnimationEvaluation(const USkeletalMesh* InSkeletalMesh, UAnimInstance* InAnimInstance, TArray<FTransform>& OutSpaceBases, TArray<FTransform>& OutLocalAtoms, FVector& OutRootBoneTranslation, FBlendedHeapCurve& OutCurve) const;
 	void PostAnimEvaluation( FAnimationEvaluationContext& EvaluationContext );
 
 	/**
@@ -1082,9 +1079,6 @@ public:
 	UFUNCTION(BlueprintCallable, Category="Components|SkeletalMesh")
 	void SetAllBodiesSimulatePhysics(bool bNewSimulate);
 
-	/** Update Material Parameters based on AnimInstance */
-	void UpdateMaterialParameters();
-
 	/** This is global set up for setting physics blend weight
 	 * This does multiple things automatically
 	 * If PhysicsBlendWeight == 1.f, it will enable Simulation, and if PhysicsBlendWeight == 0.f, it will disable Simulation. 
@@ -1331,7 +1325,7 @@ private:
 	void PostPhysicsTickComponent(FSkeletalMeshComponentPostPhysicsTickFunction& ThisTickFunction);
 
 	/** Evaluate Anim System **/
-	void EvaluateAnimation(const USkeletalMesh* InSkeletalMesh, UAnimInstance* InAnimInstance, TArray<FTransform>& OutLocalAtoms, TArray<struct FActiveVertexAnim>& OutVertexAnims, FVector& OutRootBoneTranslation, FBlendedHeapCurve& OutCurve) const;
+	void EvaluateAnimation(const USkeletalMesh* InSkeletalMesh, UAnimInstance* InAnimInstance, TArray<FTransform>& OutLocalAtoms, FVector& OutRootBoneTranslation, FBlendedHeapCurve& OutCurve) const;
 
 	/**
 	* Take the SourceAtoms array (translation vector, rotation quaternion and scale vector) and update the array of component-space bone transformation matrices (DestSpaceBases).
@@ -1382,6 +1376,9 @@ public:
 	/** Returns array containing cachec animation curve mapping UIDs (which are copied over from USkeleton) */
 	TArray<FSmartNameMapping::UID> const * GetCachedAnimCurveMappingNameUids();
 
+	/** Apply animation curves to this component */
+	void ApplyAnimationCurvesToComponent(const TMap<FName, float>* InMaterialParameterCurves, const TMap<FName, float>* InAnimationMorphCurves);
+	
 private:
 	// Returns whether we need to run the Pre Cloth Tick or not
 	bool ShouldRunPostPhysicsTick() const;
@@ -1478,7 +1475,14 @@ public:
 	FBlendedHeapCurve& GetEditableAnimationCurves() { return CurvesArray[CurrentEditableSpaceBases]; }
 	const FBlendedHeapCurve& GetEditableAnimationCurves() const { return CurvesArray[CurrentEditableSpaceBases]; }
 #endif 
-};
 
+public:
+	/** Register/Unregister for physics state creation callback */
+	FDelegateHandle RegisterOnPhysicsCreatedDelegate(const FOnSkelMeshPhysicsCreated& Delegate);
+	void UnregisterOnPhysicsCreatedDelegate(const FDelegateHandle& DelegateHandle);
 
+private:
 
+	/** Multicaster fired when this component creates physics state (in case external objects rely on physics state)*/
+	FOnSkelMeshPhysicsCreatedMultiCast OnSkelMeshPhysicsCreated;
+};
diff --git a/Engine/Source/Runtime/Engine/Classes/Engine/EngineBaseTypes.h b/Engine/Source/Runtime/Engine/Classes/Engine/EngineBaseTypes.h
--- a/Engine/Source/Runtime/Engine/Classes/Engine/EngineBaseTypes.h
+++ b/Engine/Source/Runtime/Engine/Classes/Engine/EngineBaseTypes.h
@@ -24,6 +24,21 @@ enum EInputEvent
 	IE_MAX                  =5,
 };
 
+UENUM()
+enum class EMouseCaptureMode : uint8
+{
+	/** Do not capture the mouse at all */
+	NoCapture,
+	/** Capture the mouse permanently when the viewport is clicked, and consume the initial mouse down that caused the capture so it isn't processed by player input */
+	CapturePermanently,
+	/** Capture the mouse permanently when the viewport is clicked, and allow player input to process the mouse down that caused the capture */
+	CapturePermanently_IncludingInitialMouseDown,
+	/** Capture the mouse during a mouse down, releases on mouse up */
+	CaptureDuringMouseDown,
+	/** Capture only when the right mouse button is down, not any of the other mouse buttons */
+	CaptureDuringRightMouseDown,
+};
+
 /** Type of tick we wish to perform on the level */
 enum ELevelTick
 {
diff --git a/Engine/Source/Runtime/Engine/Classes/Engine/GameInstance.h b/Engine/Source/Runtime/Engine/Classes/Engine/GameInstance.h
--- a/Engine/Source/Runtime/Engine/Classes/Engine/GameInstance.h
+++ b/Engine/Source/Runtime/Engine/Classes/Engine/GameInstance.h
@@ -7,6 +7,7 @@
 
 class ULocalPlayer;
 class FUniqueNetId;
+struct FLatentActionManager;
 
 // 
 // 	EWelcomeScreen, 	//initial screen.  Used for platforms where we may not have a signed in user yet.
@@ -204,6 +205,8 @@ public:
 
 	inline FTimerManager& GetTimerManager() const { return *TimerManager; }
 
+	inline FLatentActionManager& GetLatentActionManager() const { return *LatentActionManager;  }
+
 	/**
 	 * Start recording a replay with the given custom name and friendly name.
 	 *
@@ -234,6 +237,7 @@ public:
 	{}
 
 	FTimerManager* TimerManager;
+	FLatentActionManager* LatentActionManager;
 
 	/** @return online session management object associated with this game instance */
 	class UOnlineSession* GetOnlineSession() const { return OnlineSession; }
diff --git a/Engine/Source/Runtime/Engine/Classes/Engine/GameViewportClient.h b/Engine/Source/Runtime/Engine/Classes/Engine/GameViewportClient.h
--- a/Engine/Source/Runtime/Engine/Classes/Engine/GameViewportClient.h
+++ b/Engine/Source/Runtime/Engine/Classes/Engine/GameViewportClient.h
@@ -638,7 +638,7 @@ public:
 	/**
 	 * Set the mouse capture behavior when the viewport is clicked
 	 */
-	void SetCaptureMouseOnClick(EMouseCaptureMode::Type Mode)
+	void SetCaptureMouseOnClick(EMouseCaptureMode Mode)
 	{
 		MouseCaptureMode = Mode;
 	}
@@ -646,7 +646,7 @@ public:
 	/**
 	 * Gets the mouse capture behavior when the viewport is clicked
 	 */
-	virtual EMouseCaptureMode::Type CaptureMouseOnClick() override
+	virtual EMouseCaptureMode CaptureMouseOnClick() override
 	{
 		return MouseCaptureMode;
 	}
@@ -796,7 +796,7 @@ private:
 	bool bIgnoreInput;
 
 	/** Mouse capture behavior when the viewport is clicked */
-	EMouseCaptureMode::Type MouseCaptureMode;
+	EMouseCaptureMode MouseCaptureMode;
 
 	/** Whether or not the cursor is hidden when the viewport captures the mouse */
 	bool bHideCursorDuringCapture;
@@ -809,6 +809,9 @@ private:
 
 	/** Whether or not this audio device is in audio-focus */
 	bool bHasAudioFocus;
+
+	/** Whether or not this viewport is active*/
+	bool bActive;
 };
 
 
