Copying //UE4/Dev-Framework to Dev-Main (//UE4/Dev-Main) @ 2910422

==========================
MAJOR FEATURES + CHANGES
==========================

Change 2860430 on 2016/02/09 by [REDACTED_PERSON]

	UE-26216 EQ Settings are incorrect, unable to produce a LPF setting

	- Fix is to expose the xaudio2 eq effect parameters actually used
	- Added support for real-time previewing of EQ parameter changes (interpolates from current to new values)

	Platform code:
	- CoreAudio's EQ AU effects have identical settings as well, translated the params to the mac code
	- PS4 code was disabled before -- this change is a first-pass implementation for PS4, but its still not enabled

Change 2867932 on 2016/02/15 by [REDACTED_PERSON]

	Some adjustments and cleanups to recent dev-framework features

	- Removing change to ObjectVersion in favor of a TaggedPropertyRedirect in BaseEngine.ini for EQSettings variables
	- Adding bIsTraceDelegateBound bool for flagging need to add delegate binding for async traces
	- Reverse iteratoing through PendingSoundsToDelete for deferred delete for async occlusion checks and using RemoveAtAtSwap

Change 2894543 on 2016/03/04 by [REDACTED_PERSON]

	When doing stricmp only convert to lower if characters are not already equal

Change 2896383 on 2016/03/06 by [REDACTED_PERSON]

	Fixed duplicating EQS queries copying query name as well #UE4

Change 2897688 on 2016/03/07 by [REDACTED_PERSON]

	Fix particle emitter post load not calling post load on its modules.

Change 2897806 on 2016/03/07 by [REDACTED_PERSON]

	#ue4 - Protect against MTD sometimes returning negative distance when it should be positive.

Change 2897815 on 2016/03/07 by [REDACTED_PERSON]

	#ue4 - Avoid multiple iterations with axis flipping in FindTeleportSpot unless we are spawning something. Otherwise the teleport spot could be placed on the other side of a wall. Also avoid using a partial MTD result in EncroachingBlockingGeometry() when one of the MTD tests fails.

Change 2897819 on 2016/03/07 by [REDACTED_PERSON]

	#ue4 - FInterpTo() now treats negative values as instant speed as well as zero, consistent with other *InterpTo() functions. Previously it only treated zero as instant, and negative speeds returned the "Current" value.

Change 2898891 on 2016/03/08 by [REDACTED_PERSON]

	Don't dirty level when launching standalone from editor
	Rework spawning temporary spawn here player start when launching standalone from editor to use a transaction and undo it rather than spawn and destroy
	Add a mechanism to squelch notification messages
	Add the ability for a undo to say that the transaction cannot be undone

Change 2898902 on 2016/03/08 by [REDACTED_PERSON]

	Fix budget mode not showing up

Change 2898968 on 2016/03/08 by [REDACTED_PERSON]

	Added OnComponentTransformChanged delegate to editor. (UE-27305) #rb marc.audy

Change 2899511 on 2016/03/08 by [REDACTED_PERSON]

	Changed AllocateNameEntry to be template based instead of branching off of whether the string is ansi or wide in multiple places
	Removed the load time hashing of FNameEntry as they are loaded which improves FLinkerLoad::SerializeNameMap() perf

	#test: everything since this is FName, PS4

Change 2899525 on 2016/03/08 by [REDACTED_PERSON]

	Add an IsGamepadAttached function to Platform and Slate Application.
	Implemented for XInputInterface and PS4/WinDualShock plugin

Change 2899672 on 2016/03/08 by [REDACTED_PERSON]

	Expose setting vertical alignment for text render component to blueprints

Change 2899693 on 2016/03/08 by [REDACTED_PERSON]

	- Increase name bucket count to 65k
	- Increase name pool alocator to 256kb

Change 2901383 on 2016/03/09 by Jurre.deBaare

	Adding Alembic TPS file

Change 2901520 on 2016/03/09 by [REDACTED_PERSON]

	#ue4 - Make "checkCode(Code)" macro work without requiring a semicolon in the code block.

	Example:
	    checkCode(ensure(SomeTest););  // previous requirement:
	    checkCode(ensure(SomeTest)); // now supported

Change 2901532 on 2016/03/09 by [REDACTED_PERSON]

	#ue4 - Add DiagnosticCheckNaN() version taking a string to FRotator and FQuat.

diff --git a/Engine/Source/Editor/BehaviorTreeEditor/Private/BehaviorTreeGraph.cpp b/Engine/Source/Editor/BehaviorTreeEditor/Private/BehaviorTreeGraph.cpp
--- a/Engine/Source/Editor/BehaviorTreeEditor/Private/BehaviorTreeGraph.cpp
+++ b/Engine/Source/Editor/BehaviorTreeEditor/Private/BehaviorTreeGraph.cpp
@@ -1066,6 +1066,7 @@ void UBehaviorTreeGraph::RebuildExecutionOrder()
 
 void UBehaviorTreeGraph::RebuildChildOrder(UEdGraphNode* ParentNode)
 {
+	bool bUpdateExecutionOrder = false;
 	if (ParentNode)
 	{
 		for (int32 PinIdx = 0; PinIdx < ParentNode->Pins.Num(); PinIdx++)
@@ -1073,10 +1074,18 @@ void UBehaviorTreeGraph::RebuildChildOrder(UEdGraphNode* ParentNode)
 			UEdGraphPin* Pin = ParentNode->Pins[PinIdx];
 			if (Pin->Direction == EGPD_Output)
 			{
+				TArray<UEdGraphPin*> PrevOrder(Pin->LinkedTo);
 				Pin->LinkedTo.Sort(FCompareNodeXLocation());
+
+				bUpdateExecutionOrder = bUpdateExecutionOrder || (PrevOrder != Pin->LinkedTo);
 			}
 		}
 	}
+
+	if (bUpdateExecutionOrder)
+	{
+		UpdateAsset(KeepRebuildCounter);
+	}
 }
 
 namespace BTAutoArrangeHelpers
diff --git a/Engine/Source/Editor/BehaviorTreeEditor/Private/SGraphNode_BehaviorTree.cpp b/Engine/Source/Editor/BehaviorTreeEditor/Private/SGraphNode_BehaviorTree.cpp
--- a/Engine/Source/Editor/BehaviorTreeEditor/Private/SGraphNode_BehaviorTree.cpp
+++ b/Engine/Source/Editor/BehaviorTreeEditor/Private/SGraphNode_BehaviorTree.cpp
@@ -20,6 +20,11 @@
 
 #define LOCTEXT_NAMESPACE "BehaviorTreeEditor"
 
+namespace
+{
+	static const bool bShowExecutionIndexInEditorMode = true;
+}
+
 /////////////////////////////////////////////////////
 // SBehaviorTreePin
 
@@ -883,7 +888,7 @@ EVisibility SGraphNode_BehaviorTree::GetIndexVisibility() const
 	}
 
 	// Visible if we are in PIE or if we have siblings
-	const bool bCanShowIndex = (GEditor->bIsSimulatingInEditor || GEditor->PlayWorld != NULL) || (MyParentOutputPin && MyParentOutputPin->LinkedTo.Num() > 1);
+	const bool bCanShowIndex = (bShowExecutionIndexInEditorMode || GEditor->bIsSimulatingInEditor || GEditor->PlayWorld != NULL) || (MyParentOutputPin && MyParentOutputPin->LinkedTo.Num() > 1);
 
 	// LOD this out once things get too small
 	TSharedPtr<SGraphPanel> MyOwnerPanel = GetOwnerPanel();
@@ -902,7 +907,7 @@ FText SGraphNode_BehaviorTree::GetIndexText() const
 
 	int32 Index = 0;
 
-	if (GEditor->bIsSimulatingInEditor || GEditor->PlayWorld != NULL)
+	if (bShowExecutionIndexInEditorMode || GEditor->bIsSimulatingInEditor || GEditor->PlayWorld != NULL)
 	{
 		// special case: range of execution indices in composite decorator node
 		UBehaviorTreeGraphNode_CompositeDecorator* CompDecorator = Cast<UBehaviorTreeGraphNode_CompositeDecorator>(GraphNode);
@@ -935,7 +940,7 @@ FText SGraphNode_BehaviorTree::GetIndexText() const
 
 FText SGraphNode_BehaviorTree::GetIndexTooltipText() const
 {
-	if (GEditor->bIsSimulatingInEditor || GEditor->PlayWorld != NULL)
+	if (bShowExecutionIndexInEditorMode || GEditor->bIsSimulatingInEditor || GEditor->PlayWorld != NULL)
 	{
 		return LOCTEXT("ExecutionIndexTooltip", "Execution index: this shows the order in which nodes are executed.");
 	}
diff --git a/Engine/Source/Editor/LandscapeEditor/Private/ActorFactoryLandscape.cpp b/Engine/Source/Editor/LandscapeEditor/Private/ActorFactoryLandscape.cpp
--- a/Engine/Source/Editor/LandscapeEditor/Private/ActorFactoryLandscape.cpp
+++ b/Engine/Source/Editor/LandscapeEditor/Private/ActorFactoryLandscape.cpp
@@ -13,14 +13,14 @@ UActorFactoryLandscape::UActorFactoryLandscape(const FObjectInitializer& ObjectI
 	NewActorClass = ALandscapeProxy::StaticClass();
 }
 
-AActor* UActorFactoryLandscape::SpawnActor(UObject* Asset, ULevel* InLevel, const FVector& Location, const FRotator& Rotation, EObjectFlags ObjectFlags, const FName& Name)
+AActor* UActorFactoryLandscape::SpawnActor(UObject* Asset, ULevel* InLevel, const FTransform& Transform, EObjectFlags ObjectFlags, const FName Name)
 {
 	GLevelEditorModeTools().ActivateMode(FBuiltinEditorModes::EM_Landscape);
 
 	FEdModeLandscape* EdMode = (FEdModeLandscape*)GLevelEditorModeTools().GetActiveMode(FBuiltinEditorModes::EM_Landscape);
 
-	EdMode->UISettings->NewLandscape_Location = Location;
-	EdMode->UISettings->NewLandscape_Rotation = Rotation;
+	EdMode->UISettings->NewLandscape_Location = Transform.GetLocation();
+	EdMode->UISettings->NewLandscape_Rotation = Transform.GetRotation().Rotator();
 
 	EdMode->SetCurrentTool("NewLandscape");
 
@@ -28,7 +28,7 @@ AActor* UActorFactoryLandscape::SpawnActor(UObject* Asset, ULevel* InLevel, cons
 	SpawnInfo.OverrideLevel = InLevel;
 	SpawnInfo.ObjectFlags = ObjectFlags;
 	SpawnInfo.Name = Name;
-	return InLevel->OwningWorld->SpawnActor(ALandscapePlaceholder::StaticClass(), &Location, &Rotation, SpawnInfo);
+	return InLevel->OwningWorld->SpawnActor(ALandscapePlaceholder::StaticClass(), &Transform, SpawnInfo);
 }
 
 ALandscapePlaceholder::ALandscapePlaceholder(const FObjectInitializer& ObjectInitializer)
diff --git a/Engine/Source/Editor/LandscapeEditor/Private/Classes/ActorFactoryLandscape.h b/Engine/Source/Editor/LandscapeEditor/Private/Classes/ActorFactoryLandscape.h
--- a/Engine/Source/Editor/LandscapeEditor/Private/Classes/ActorFactoryLandscape.h
+++ b/Engine/Source/Editor/LandscapeEditor/Private/Classes/ActorFactoryLandscape.h
@@ -10,10 +10,6 @@ class UActorFactoryLandscape : public UActorFactory
 	GENERATED_UCLASS_BODY()
 
 	//~ Begin UActorFactory Interface
-	//virtual bool CanCreateActorFrom( const FAssetData& AssetData, FText& OutErrorMsg ) override;
-	virtual AActor* SpawnActor( UObject* Asset, ULevel* InLevel, const FVector& Location, const FRotator& Rotation, EObjectFlags ObjectFlags, const FName& Name ) override;
-	//virtual void PostSpawnActor( UObject* Asset, AActor* NewActor) override;
-	//virtual void PostCreateBlueprint( UObject* Asset, AActor* CDO ) override;
-	//virtual UObject* GetAssetFromActorInstance(AActor* ActorInstance) override;
+	virtual AActor* SpawnActor( UObject* Asset, ULevel* InLevel, const FTransform& Transform, EObjectFlags ObjectFlags, const FName Name ) override;
 	//~ End UActorFactory Interface
 };
diff --git a/Engine/Source/Editor/UnrealEd/Classes/ActorFactories/ActorFactory.h b/Engine/Source/Editor/UnrealEd/Classes/ActorFactories/ActorFactory.h
--- a/Engine/Source/Editor/UnrealEd/Classes/ActorFactories/ActorFactory.h
+++ b/Engine/Source/Editor/UnrealEd/Classes/ActorFactories/ActorFactory.h
@@ -66,7 +66,7 @@ class UNREALED_API UActorFactory : public UObject
 protected:
 
 	virtual bool PreSpawnActor( UObject* Asset, FTransform& InOutLocation);
-	virtual AActor* SpawnActor( UObject* Asset, ULevel* InLevel, const FVector& Location, const FRotator& Rotation, EObjectFlags ObjectFlags, const FName& Name );
+	virtual AActor* SpawnActor( UObject* Asset, ULevel* InLevel, const FTransform& Transform, EObjectFlags ObjectFlags, const FName Name );
 
 	/** Subclasses may implement this to modify the actor after it has been spawned 
 	    IMPORTANT: If you override this, you should usually also override PostCreateBlueprint()! */
diff --git a/Engine/Source/Editor/UnrealEd/Classes/ActorFactories/ActorFactoryCharacter.h b/Engine/Source/Editor/UnrealEd/Classes/ActorFactories/ActorFactoryCharacter.h
--- a/Engine/Source/Editor/UnrealEd/Classes/ActorFactories/ActorFactoryCharacter.h
+++ b/Engine/Source/Editor/UnrealEd/Classes/ActorFactories/ActorFactoryCharacter.h
@@ -11,7 +11,4 @@ class UActorFactoryCharacter : public UActorFactory
 
 public:
 	virtual bool CanCreateActorFrom( const FAssetData& AssetData, FText& OutErrorMsg ) override;
-
-protected:
-	virtual AActor* SpawnActor( UObject* Asset, ULevel* InLevel, const FVector& Location, const FRotator& Rotation, EObjectFlags ObjectFlags, const FName& Name ) override;
 };
diff --git a/Engine/Source/Editor/UnrealEd/Classes/ActorFactories/ActorFactoryClass.h b/Engine/Source/Editor/UnrealEd/Classes/ActorFactories/ActorFactoryClass.h
--- a/Engine/Source/Editor/UnrealEd/Classes/ActorFactories/ActorFactoryClass.h
+++ b/Engine/Source/Editor/UnrealEd/Classes/ActorFactories/ActorFactoryClass.h
@@ -17,5 +17,5 @@ class UActorFactoryClass : public UActorFactory
 	//~ End UActorFactory Interface	
 
 protected:
-	virtual AActor* SpawnActor( UObject* Asset, ULevel* InLevel, const FVector& Location, const FRotator& Rotation, EObjectFlags ObjectFlags, const FName& Name ) override;
+	virtual AActor* SpawnActor( UObject* Asset, ULevel* InLevel, const FTransform& Transform, EObjectFlags ObjectFlags, const FName Name ) override;
 };
diff --git a/Engine/Source/Editor/UnrealEd/Classes/ActorFactories/ActorFactoryEmptyActor.h b/Engine/Source/Editor/UnrealEd/Classes/ActorFactories/ActorFactoryEmptyActor.h
--- a/Engine/Source/Editor/UnrealEd/Classes/ActorFactories/ActorFactoryEmptyActor.h
+++ b/Engine/Source/Editor/UnrealEd/Classes/ActorFactories/ActorFactoryEmptyActor.h
@@ -13,6 +13,6 @@ public:
 	virtual bool CanCreateActorFrom( const FAssetData& AssetData, FText& OutErrorMsg ) override;
 
 protected:
-	virtual AActor* SpawnActor( UObject* Asset, ULevel* InLevel, const FVector& Location, const FRotator& Rotation, EObjectFlags ObjectFlags, const FName& Name ) override;
+	virtual AActor* SpawnActor( UObject* Asset, ULevel* InLevel, const FTransform& Transform, EObjectFlags ObjectFlags, const FName Name ) override;
 
 };
diff --git a/Engine/Source/Editor/UnrealEd/Classes/ActorFactories/ActorFactoryPawn.h b/Engine/Source/Editor/UnrealEd/Classes/ActorFactories/ActorFactoryPawn.h
--- a/Engine/Source/Editor/UnrealEd/Classes/ActorFactories/ActorFactoryPawn.h
+++ b/Engine/Source/Editor/UnrealEd/Classes/ActorFactories/ActorFactoryPawn.h
@@ -11,7 +11,4 @@ class UActorFactoryPawn : public UActorFactoryEmptyActor
 
 public:
 	virtual bool CanCreateActorFrom( const FAssetData& AssetData, FText& OutErrorMsg ) override;
-
-protected:
-	virtual AActor* SpawnActor( UObject* Asset, ULevel* InLevel, const FVector& Location, const FRotator& Rotation, EObjectFlags ObjectFlags, const FName& Name ) override;
 };
diff --git a/Engine/Source/Editor/UnrealEd/Classes/Editor/EditorEngine.h b/Engine/Source/Editor/UnrealEd/Classes/Editor/EditorEngine.h
--- a/Engine/Source/Editor/UnrealEd/Classes/Editor/EditorEngine.h
+++ b/Engine/Source/Editor/UnrealEd/Classes/Editor/EditorEngine.h
@@ -437,10 +437,6 @@ public:
 	UPROPERTY(globalconfig)
 	uint32 bAllowMultiplePIEWorlds:1;
 
-	/** The PlayerStart class used when spawning the player at the current camera location. */
-	UPROPERTY()
-	TSubclassOf<class ANavigationObjectBase>  PlayFromHerePlayerStartClass;
-
 	/** True if there is a pending end play map queued */
 	UPROPERTY()
 	uint32 bRequestEndPlayMapQueued:1;
@@ -461,6 +457,14 @@ public:
 	UPROPERTY()
 	uint32 bIsSimulatingInEditor:1;
 
+	/** True if we should not display notifications about undo/redo */
+	UPROPERTY()
+	uint32 bSquelchTransactionNotification:1;
+
+	/** The PlayerStart class used when spawning the player at the current camera location. */
+	UPROPERTY()
+	TSubclassOf<class ANavigationObjectBase>  PlayFromHerePlayerStartClass;
+
 	/** When Simulating In Editor, a pointer to the original (non-simulating) editor world */
 	UPROPERTY()
 	class UWorld* EditorWorld;
@@ -887,7 +891,7 @@ public:
 	int32 EndTransaction();
 	void ResetTransaction(const FText& Reason);
 	void CancelTransaction(int32 Index);
-	bool UndoTransaction();
+	bool UndoTransaction(bool bCanRedo = true);
 	bool RedoTransaction();
 	bool IsTransactionActive();
 	FText GetTransactionName() const;
diff --git a/Engine/Source/Editor/UnrealEd/Classes/Editor/TransBuffer.h b/Engine/Source/Editor/UnrealEd/Classes/Editor/TransBuffer.h
--- a/Engine/Source/Editor/UnrealEd/Classes/Editor/TransBuffer.h
+++ b/Engine/Source/Editor/UnrealEd/Classes/Editor/TransBuffer.h
@@ -108,7 +108,7 @@ public:
 	virtual void SetUndoBarrier() override;
 	virtual void RemoveUndoBarrier() override;
 	virtual void ClearUndoBarriers() override;
-	virtual bool Undo() override;
+	virtual bool Undo(bool bCanRedo = true) override;
 	virtual bool Redo() override;
 	virtual bool EnableObjectSerialization() override;
 	virtual bool DisableObjectSerialization() override;
diff --git a/Engine/Source/Editor/UnrealEd/Classes/Editor/Transactor.h b/Engine/Source/Editor/UnrealEd/Classes/Editor/Transactor.h
--- a/Engine/Source/Editor/UnrealEd/Classes/Editor/Transactor.h
+++ b/Engine/Source/Editor/UnrealEd/Classes/Editor/Transactor.h
@@ -552,9 +552,11 @@ class UTransactor : public UObject
 	/**
 	 * Executes an undo transaction, undoing all actions contained by that transaction.
 	 * 
+	 * @param	bCanRedo	If false indicates that the undone transaction (and any transactions that came after it) cannot be redone
+	 * 
 	 * @return				true if the transaction was successfully undone
 	 */
-	virtual bool Undo() PURE_VIRTUAL(UTransactor::Undo,return false;);
+	virtual bool Undo(bool bCanRedo = true) PURE_VIRTUAL(UTransactor::Undo,return false;);
 
 	/**
 	 * Executes an redo transaction, redoing all actions contained by that transaction.
diff --git a/Engine/Source/Editor/UnrealEd/Private/EditorEngine.cpp b/Engine/Source/Editor/UnrealEd/Private/EditorEngine.cpp
--- a/Engine/Source/Editor/UnrealEd/Private/EditorEngine.cpp
+++ b/Engine/Source/Editor/UnrealEd/Private/EditorEngine.cpp
@@ -4590,7 +4590,7 @@ void UEditorEngine::ReplaceActors(UActorFactory* Factory, const FAssetData& Asse
 		const FTransform OldTransform = OldActor->ActorToWorld();
 
 		// create the actor
-		NewActor = Factory->CreateActor( Asset, Level, OldTransform, RF_Transactional, OldActorName );
+		NewActor = Factory->CreateActor( Asset, Level, OldTransform );
 		// For blueprints, try to copy over properties
 		if (Factory->IsA(UActorFactoryBlueprint::StaticClass()))
 		{
@@ -4605,8 +4605,10 @@ void UEditorEngine::ReplaceActors(UActorFactory* Factory, const FAssetData& Asse
 			}
 		}
 
-		if ( NewActor != NULL )
+		if (NewActor)
 		{
+			NewActor->Rename(*OldActorName.ToString());
+
 			// The new actor might not have a root component
 			USceneComponent* const NewActorRootComponent = NewActor->GetRootComponent();
 			if(NewActorRootComponent)
diff --git a/Engine/Source/Editor/UnrealEd/Private/EditorServer.cpp b/Engine/Source/Editor/UnrealEd/Private/EditorServer.cpp
--- a/Engine/Source/Editor/UnrealEd/Private/EditorServer.cpp
+++ b/Engine/Source/Editor/UnrealEd/Private/EditorServer.cpp
@@ -1088,7 +1088,10 @@ void UEditorEngine::HandleTransactorRedo( FUndoSessionContext SessionContext, bo
 
 	BroadcastPostRedo(SessionContext.Context, SessionContext.PrimaryObject, Succeeded);
 	InvalidateAllViewportsAndHitProxies();
-	ShowUndoRedoNotification(FText::Format(NSLOCTEXT("UnrealEd", "RedoMessageFormat", "Redo: {0}"), SessionContext.Title), Succeeded);
+	if (!bSquelchTransactionNotification)
+	{
+		ShowUndoRedoNotification(FText::Format(NSLOCTEXT("UnrealEd", "RedoMessageFormat", "Redo: {0}"), SessionContext.Title), Succeeded);
+	}
 }
 
 void UEditorEngine::HandleTransactorUndo( FUndoSessionContext SessionContext, bool Succeeded )
@@ -1098,7 +1101,10 @@ void UEditorEngine::HandleTransactorUndo( FUndoSessionContext SessionContext, bo
 
 	BroadcastPostUndo(SessionContext.Context, SessionContext.PrimaryObject, Succeeded);
 	InvalidateAllViewportsAndHitProxies();
-	ShowUndoRedoNotification(FText::Format(NSLOCTEXT("UnrealEd", "UndoMessageFormat", "Undo: {0}"), SessionContext.Title), Succeeded);
+	if (!bSquelchTransactionNotification)
+	{
+		ShowUndoRedoNotification(FText::Format(NSLOCTEXT("UnrealEd", "UndoMessageFormat", "Undo: {0}"), SessionContext.Title), Succeeded);
+	}
 }
 
 bool UEditorEngine::AreEditorAnalyticsEnabled() const 
@@ -1261,7 +1267,7 @@ void UEditorEngine::PostUndo(bool bSuccess)
 	FBlueprintCompileReinstancer::BatchReplaceInstancesOfClass(OldToNewClassMapToReinstance);
 }
 
-bool UEditorEngine::UndoTransaction()
+bool UEditorEngine::UndoTransaction(bool bCanRedo)
 {
 	// make sure we're in a valid state to perform this
 	if (GIsSavingPackage || IsGarbageCollecting())
@@ -1269,7 +1275,7 @@ bool UEditorEngine::UndoTransaction()
 		return false;
 	}
 
-	return Trans->Undo();
+	return Trans->Undo(bCanRedo);
 }
 
 bool UEditorEngine::RedoTransaction()
diff --git a/Engine/Source/Editor/UnrealEd/Private/EditorTransaction.cpp b/Engine/Source/Editor/UnrealEd/Private/EditorTransaction.cpp
--- a/Engine/Source/Editor/UnrealEd/Private/EditorTransaction.cpp
+++ b/Engine/Source/Editor/UnrealEd/Private/EditorTransaction.cpp
@@ -781,7 +781,7 @@ void UTransBuffer::ClearUndoBarriers()
 }
 
 
-bool UTransBuffer::Undo()
+bool UTransBuffer::Undo(bool bCanRedo)
 {
 	CheckState();
 
@@ -801,6 +801,12 @@ bool UTransBuffer::Undo()
 		BeforeRedoUndoDelegate.Broadcast(Transaction.GetContext());
 		Transaction.Apply();
 		UndoDelegate.Broadcast(Transaction.GetContext(), true);
+
+		if (!bCanRedo)
+		{
+			UndoBuffer.RemoveAt(UndoBuffer.Num() - UndoCount, UndoCount);
+			UndoCount = 0;
+		}
 	}
 	GIsTransacting = false;
 
diff --git a/Engine/Source/Editor/UnrealEd/Private/Factories/ActorFactory.cpp b/Engine/Source/Editor/UnrealEd/Private/Factories/ActorFactory.cpp
--- a/Engine/Source/Editor/UnrealEd/Private/Factories/ActorFactory.cpp
+++ b/Engine/Source/Editor/UnrealEd/Private/Factories/ActorFactory.cpp
@@ -200,9 +200,7 @@ AActor* UActorFactory::CreateActor( UObject* Asset, ULevel* InLevel, FTransform
 
 	if ( PreSpawnActor(Asset, SpawnTransform) )
 	{
-		const auto Location = SpawnTransform.GetLocation();
-		const auto Rotation = SpawnTransform.GetRotation().Rotator();
-		NewActor = SpawnActor(Asset, InLevel, Location, Rotation, ObjectFlags, Name);
+		NewActor = SpawnActor(Asset, InLevel, SpawnTransform, ObjectFlags, Name);
 
 		if ( NewActor )
 		{
@@ -230,7 +228,7 @@ bool UActorFactory::PreSpawnActor( UObject* Asset, FTransform& InOutLocation)
 	return true;
 }
 
-AActor* UActorFactory::SpawnActor( UObject* Asset, ULevel* InLevel, const FVector& Location, const FRotator& Rotation, EObjectFlags ObjectFlags, const FName& Name )
+AActor* UActorFactory::SpawnActor( UObject* Asset, ULevel* InLevel, const FTransform& Transform, EObjectFlags ObjectFlags, const FName Name )
 {
 	AActor* DefaultActor = GetDefaultActor( FAssetData( Asset ) );
 	if ( DefaultActor )
@@ -239,7 +237,7 @@ AActor* UActorFactory::SpawnActor( UObject* Asset, ULevel* InLevel, const FVecto
 		SpawnInfo.OverrideLevel = InLevel;
 		SpawnInfo.ObjectFlags = ObjectFlags;
 		SpawnInfo.Name = Name;
-		return InLevel->OwningWorld->SpawnActor( DefaultActor->GetClass(), &Location, &Rotation, SpawnInfo );
+		return InLevel->OwningWorld->SpawnActor( DefaultActor->GetClass(), &Transform, SpawnInfo );
 	}
 
 	return NULL;
@@ -1193,17 +1191,17 @@ bool UActorFactoryEmptyActor::CanCreateActorFrom( const FAssetData& AssetData, F
 	return AssetData.ObjectPath == FName(*AActor::StaticClass()->GetPathName());
 }
 
-AActor* UActorFactoryEmptyActor::SpawnActor( UObject* Asset, ULevel* InLevel, const FVector& Location, const FRotator& Rotation, EObjectFlags ObjectFlags, const FName& Name )
+AActor* UActorFactoryEmptyActor::SpawnActor( UObject* Asset, ULevel* InLevel, const FTransform& Transform, EObjectFlags ObjectFlags, const FName Name )
 {
 	AActor* NewActor = nullptr;
 	{
 		// Spawn a temporary actor for dragging around
-		NewActor = Super::SpawnActor(Asset, InLevel, Location, Rotation, ObjectFlags, Name);
+		NewActor = Super::SpawnActor(Asset, InLevel, Transform, ObjectFlags, Name);
 
 		USceneComponent* RootComponent = NewObject<USceneComponent>(NewActor, USceneComponent::GetDefaultSceneRootVariableName(), RF_Transactional);
 		RootComponent->Mobility = EComponentMobility::Movable;
 		RootComponent->bVisualizeComponent = true;
-		RootComponent->SetWorldLocationAndRotation(Location, Rotation);
+		RootComponent->SetWorldTransform(Transform);
 
 		NewActor->SetRootComponent(RootComponent);
 		NewActor->AddInstanceComponent(RootComponent);
@@ -1231,14 +1229,6 @@ bool UActorFactoryCharacter::CanCreateActorFrom(const FAssetData& AssetData, FTe
 	return AssetData.ObjectPath == FName(*ACharacter::StaticClass()->GetPathName());
 }
 
-AActor* UActorFactoryCharacter::SpawnActor(UObject* Asset, ULevel* InLevel, const FVector& Location, const FRotator& Rotation, EObjectFlags ObjectFlags, const FName& Name)
-{
-	AActor* NewActor = Super::SpawnActor( Asset, InLevel, Location, Rotation, ObjectFlags, Name );
-
-	return NewActor;
-}
-
-
 /*-----------------------------------------------------------------------------
 UActorFactoryPawn
 -----------------------------------------------------------------------------*/
@@ -1254,12 +1244,6 @@ bool UActorFactoryPawn::CanCreateActorFrom(const FAssetData& AssetData, FText& O
 	return AssetData.ObjectPath == FName(*APawn::StaticClass()->GetPathName());
 }
 
-AActor* UActorFactoryPawn::SpawnActor(UObject* Asset, ULevel* InLevel, const FVector& Location, const FRotator& Rotation, EObjectFlags ObjectFlags, const FName& Name)
-{
-	AActor* NewActor = Super::SpawnActor(Asset, InLevel, Location, Rotation, ObjectFlags, Name);
-	return NewActor;
-}
-
 /*-----------------------------------------------------------------------------
 UActorFactoryAmbientSound
 -----------------------------------------------------------------------------*/
@@ -1375,7 +1359,7 @@ bool UActorFactoryClass::PreSpawnActor( UObject* Asset, FTransform& InOutLocatio
 	return false;
 }
 
-AActor* UActorFactoryClass::SpawnActor( UObject* Asset, ULevel* InLevel, const FVector& Location, const FRotator& Rotation, EObjectFlags ObjectFlags, const FName& Name )
+AActor* UActorFactoryClass::SpawnActor( UObject* Asset, ULevel* InLevel, const FTransform& Transform, EObjectFlags ObjectFlags, const FName Name )
 {
 	UClass* ActualClass = Cast<UClass>(Asset);
 
@@ -1385,7 +1369,7 @@ AActor* UActorFactoryClass::SpawnActor( UObject* Asset, ULevel* InLevel, const F
 		SpawnInfo.OverrideLevel = InLevel;
 		SpawnInfo.ObjectFlags = ObjectFlags;
 		SpawnInfo.Name = Name;
-		return InLevel->OwningWorld->SpawnActor( ActualClass, &Location, &Rotation, SpawnInfo );
+		return InLevel->OwningWorld->SpawnActor( ActualClass, &Transform, SpawnInfo );
 	}
 
 	return NULL;
@@ -1969,9 +1953,9 @@ bool UActorFactoryMovieScene::CanCreateActorFrom( const FAssetData& AssetData, F
 	return true;
 }
 
-AActor* UActorFactoryMovieScene::SpawnActor( UObject* Asset, ULevel* InLevel, const FVector& Location, const FRotator& Rotation, EObjectFlags ObjectFlags, const FName& Name )
+AActor* UActorFactoryMovieScene::SpawnActor( UObject* Asset, ULevel* InLevel, const FTransform& Transform, EObjectFlags ObjectFlags, const FName Name )
 {
-	ALevelSequenceActor* NewActor = Cast<ALevelSequenceActor>(Super::SpawnActor(Asset, InLevel, Location, Rotation, ObjectFlags, Name));
+	ALevelSequenceActor* NewActor = Cast<ALevelSequenceActor>(Super::SpawnActor(Asset, InLevel, Transform, ObjectFlags, Name));
 
 	if (NewActor)
 	{
diff --git a/Engine/Source/Editor/UnrealEd/Private/Factories/ActorFactoryMovieScene.h b/Engine/Source/Editor/UnrealEd/Private/Factories/ActorFactoryMovieScene.h
--- a/Engine/Source/Editor/UnrealEd/Private/Factories/ActorFactoryMovieScene.h
+++ b/Engine/Source/Editor/UnrealEd/Private/Factories/ActorFactoryMovieScene.h
@@ -11,7 +11,7 @@ class UActorFactoryMovieScene : public UActorFactory
 	GENERATED_UCLASS_BODY()
 
 	// Begin UActorFactory Interface
-	virtual AActor* SpawnActor( UObject* Asset, ULevel* InLevel, const FVector& Location, const FRotator& Rotation, EObjectFlags ObjectFlags, const FName& Name ) override;
+	virtual AActor* SpawnActor( UObject* Asset, ULevel* InLevel, const FTransform& Transform, EObjectFlags ObjectFlags, const FName Name ) override;
 	virtual bool CanCreateActorFrom( const FAssetData& AssetData, FText& OutErrorMsg ) override;
 	virtual UObject* GetAssetFromActorInstance(AActor* ActorInstance) override;
 	// End UActorFactory Interface
diff --git a/Engine/Source/Editor/UnrealEd/Private/PlayLevel.cpp b/Engine/Source/Editor/UnrealEd/Private/PlayLevel.cpp
--- a/Engine/Source/Editor/UnrealEd/Private/PlayLevel.cpp
+++ b/Engine/Source/Editor/UnrealEd/Private/PlayLevel.cpp
@@ -1153,7 +1153,7 @@ private:
 
 
 void UEditorEngine::SaveWorldForPlay(TArray<FString>& SavedMapNames)
-	{
+{
 	UWorld* World = GWorld;
 
 	// check if PersistentLevel has any external references
@@ -1171,41 +1171,50 @@ void UEditorEngine::SaveWorldForPlay(TArray<FString>& SavedMapNames)
 	// Temporarily rename streaming levels for pie saving
 	FScopedRenameStreamingLevels ScopedRenameStreamingLevels( World, PlayOnConsolePackageName, Prefix );
 	
-	const FString WorldPackageName = World->GetOutermost()->GetName();
-	
 	// spawn a play-from-here player start or a temporary player start
 	AActor* PlayerStart = NULL;
 	bool bCreatedPlayerStart = false;
 
-	SpawnPlayFromHereStart( World, PlayerStart, PlayWorldLocation, PlayWorldRotation );
+	// Gross hack to avoid marking the level dirty while this temporary spawn is done while saving we flag the editor as loading package
+	bool bTempIsEditorLoadingPackage = GIsEditorLoadingPackage;
+	GIsEditorLoadingPackage = true;
 
-	if (PlayerStart != NULL)
-	{
-		bCreatedPlayerStart = true;
-	}
-	else
 	{
-		PlayerStart = CheckForPlayerStart();
-	
-		if( PlayerStart == NULL )
-		{
-			FActorSpawnParameters SpawnInfo;
-			SpawnInfo.SpawnCollisionHandlingOverride = ESpawnActorCollisionHandlingMethod::AlwaysSpawn;
-			PlayerStart = World->SpawnActor<AActor>( APlayerStart::StaticClass(), FVector::ZeroVector, FRotator::ZeroRotator, SpawnInfo );
+		// Do the spawning in a transaction so that we can undo it below
+		FScopedTransaction SpawnPlayFromHereStartTransaction(LOCTEXT("SpawnPlayFromHereStart","Spawn Play From Here Start"));
 
+		SpawnPlayFromHereStart( World, PlayerStart, PlayWorldLocation, PlayWorldRotation );
+
+		if (PlayerStart != NULL)
+		{
 			bCreatedPlayerStart = true;
 		}
+		else
+		{
+			PlayerStart = CheckForPlayerStart();
+	
+			if( PlayerStart == NULL )
+			{
+				FActorSpawnParameters SpawnInfo;
+				SpawnInfo.SpawnCollisionHandlingOverride = ESpawnActorCollisionHandlingMethod::AlwaysSpawn;
+				PlayerStart = World->SpawnActor<AActor>( APlayerStart::StaticClass(), FVector::ZeroVector, FRotator::ZeroRotator, SpawnInfo );
+
+				bCreatedPlayerStart = true;
+			}
+		}
 	}
 	
 	// save out all open map packages
 	TArray<FString> SavedWorldFileNames;
 	bool bSavedWorld = SavePlayWorldPackages(World, *Prefix, /*out*/ SavedWorldFileNames);
 
-	// Remove the player start we added if we made one
-	if (bCreatedPlayerStart)
-	{
-		World->DestroyActor(PlayerStart);
-	}
+	// Undo the (potentially) spawning transaction and don't display a notification banner about it
+	bool bTempSquelched = GEditor->bSquelchTransactionNotification;
+	GEditor->bSquelchTransactionNotification = true;
+	GEditor->UndoTransaction(false);
+	GEditor->bSquelchTransactionNotification = bTempSquelched;
+
+	GIsEditorLoadingPackage = bTempIsEditorLoadingPackage;
 	
 	if (bSavedWorld)
 	{
diff --git a/Engine/Source/Programs/UnrealHeaderTool/Private/UHTMakefile/UHTMakefile.cpp b/Engine/Source/Programs/UnrealHeaderTool/Private/UHTMakefile/UHTMakefile.cpp
--- a/Engine/Source/Programs/UnrealHeaderTool/Private/UHTMakefile/UHTMakefile.cpp
+++ b/Engine/Source/Programs/UnrealHeaderTool/Private/UHTMakefile/UHTMakefile.cpp
@@ -2221,10 +2221,10 @@ void FUHTMakefile::LoadAndSetupReferencedNames(FArchive& Ar)
 	Ar << NameCount;
 	for (int32 i = 0; i < NameCount; i++)
 	{
-		FNameEntry NameEntry(ENAME_LinkerConstructor);
+		FNameEntrySerialized NameEntry(ENAME_LinkerConstructor);
 		Ar << NameEntry;
 
-		FName Name = NameEntry.IsWide() ? FName(ENAME_LinkerConstructor, NameEntry.GetWideName()) : FName(ENAME_LinkerConstructor, NameEntry.GetAnsiName());
+		FName Name(NameEntry);
 		NameIndices.Add(Name, NameIndices.Num());
 		ReferencedNames.Add(Name);
 	}
diff --git a/Engine/Source/Runtime/AIModule/Classes/EnvironmentQuery/EnvQuery.h b/Engine/Source/Runtime/AIModule/Classes/EnvironmentQuery/EnvQuery.h
--- a/Engine/Source/Runtime/AIModule/Classes/EnvironmentQuery/EnvQuery.h
+++ b/Engine/Source/Runtime/AIModule/Classes/EnvironmentQuery/EnvQuery.h
@@ -42,6 +42,9 @@ public:
 
 	/** QueryName patching up */
 	virtual void PostLoad() override;
+#if WITH_EDITOR
+	virtual void PostDuplicate(bool bDuplicateForPIE) override;
+#endif // WITH_EDITOR
 
 	FName GetQueryName() const { return QueryName; }
 
diff --git a/Engine/Source/Runtime/AIModule/Private/EnvironmentQuery/EnvQuery.cpp b/Engine/Source/Runtime/AIModule/Private/EnvironmentQuery/EnvQuery.cpp
--- a/Engine/Source/Runtime/AIModule/Private/EnvironmentQuery/EnvQuery.cpp
+++ b/Engine/Source/Runtime/AIModule/Private/EnvironmentQuery/EnvQuery.cpp
@@ -138,4 +138,16 @@ void UEnvQuery::PostLoad()
 	{
 		QueryName = GetFName();
 	}
-}
\ No newline at end of file
+}
+
+#if WITH_EDITOR
+void UEnvQuery::PostDuplicate(bool bDuplicateForPIE)
+{
+	if (bDuplicateForPIE == false)
+	{
+		QueryName = GetFName();
+	}
+
+	Super::PostDuplicate(bDuplicateForPIE);
+}
+#endif // WITH_EDITOR
diff --git a/Engine/Source/Runtime/AssetRegistry/Private/AssetDataGatherer.cpp b/Engine/Source/Runtime/AssetRegistry/Private/AssetDataGatherer.cpp
--- a/Engine/Source/Runtime/AssetRegistry/Private/AssetDataGatherer.cpp
+++ b/Engine/Source/Runtime/AssetRegistry/Private/AssetDataGatherer.cpp
@@ -4,7 +4,7 @@
 
 namespace AssetDataGathererConstants
 {
-	static const int32 CacheSerializationVersion = 6;
+	static const int32 CacheSerializationVersion = 7;
 	static const int32 MaxFilesToDiscoverBeforeFlush = 2500;
 	static const int32 MaxFilesToGatherBeforeFlush = 250;
 	static const int32 MaxFilesToProcessBeforeCacheWrite = 50000;
diff --git a/Engine/Source/Runtime/AssetRegistry/Private/NameTableArchive.cpp b/Engine/Source/Runtime/AssetRegistry/Private/NameTableArchive.cpp
--- a/Engine/Source/Runtime/AssetRegistry/Private/NameTableArchive.cpp
+++ b/Engine/Source/Runtime/AssetRegistry/Private/NameTableArchive.cpp
@@ -68,7 +68,7 @@ bool FNameTableArchiveReader::SerializeNameMap()
 		for ( int32 NameMapIdx = 0; NameMapIdx < NameCount; ++NameMapIdx )
 		{
 			// Read the name entry from the file.
-			FNameEntry NameEntry(ENAME_LinkerConstructor);
+			FNameEntrySerialized NameEntry(ENAME_LinkerConstructor);
 			*this << NameEntry;
 
 			if (IsError())
@@ -76,11 +76,7 @@ bool FNameTableArchiveReader::SerializeNameMap()
 				return false;
 			}
 
-			NameMap.Add( 
-				NameEntry.IsWide() ? 
-					FName(ENAME_LinkerConstructor, NameEntry.GetWideName()) : 
-					FName(ENAME_LinkerConstructor, NameEntry.GetAnsiName())
-				);
+			NameMap.Add(FName(NameEntry));
 		}
 
 		Seek( OriginalOffset );
diff --git a/Engine/Source/Runtime/AssetRegistry/Private/PackageReader.cpp b/Engine/Source/Runtime/AssetRegistry/Private/PackageReader.cpp
--- a/Engine/Source/Runtime/AssetRegistry/Private/PackageReader.cpp
+++ b/Engine/Source/Runtime/AssetRegistry/Private/PackageReader.cpp
@@ -342,14 +342,10 @@ void FPackageReader::SerializeNameMap()
 		for ( int32 NameMapIdx = 0; NameMapIdx < PackageFileSummary.NameCount; ++NameMapIdx )
 		{
 			// Read the name entry from the file.
-			FNameEntry NameEntry(ENAME_LinkerConstructor);
+			FNameEntrySerialized NameEntry(ENAME_LinkerConstructor);
 			*this << NameEntry;
 
-			NameMap.Add( 
-				NameEntry.IsWide() ? 
-					FName(ENAME_LinkerConstructor, NameEntry.GetWideName()) : 
-					FName(ENAME_LinkerConstructor, NameEntry.GetAnsiName())
-				);
+			NameMap.Add(FName(NameEntry));
 		}
 	}
 }
diff --git a/Engine/Source/Runtime/Core/Private/Math/UnrealMath.cpp b/Engine/Source/Runtime/Core/Private/Math/UnrealMath.cpp
--- a/Engine/Source/Runtime/Core/Private/Math/UnrealMath.cpp
+++ b/Engine/Source/Runtime/Core/Private/Math/UnrealMath.cpp
@@ -2352,7 +2352,7 @@ CORE_API FRotator FMath::RInterpTo( const FRotator& Current, const FRotator& Tar
 CORE_API float FMath::FInterpTo( float Current, float Target, float DeltaTime, float InterpSpeed )
 {
 	// If no interp speed, jump to target value
-	if( InterpSpeed == 0.f )
+	if( InterpSpeed <= 0.f )
 	{
 		return Target;
 	}
diff --git a/Engine/Source/Runtime/Core/Private/UObject/UnrealNames.cpp b/Engine/Source/Runtime/Core/Private/UObject/UnrealNames.cpp
--- a/Engine/Source/Runtime/Core/Private/UObject/UnrealNames.cpp
+++ b/Engine/Source/Runtime/Core/Private/UObject/UnrealNames.cpp
@@ -16,7 +16,11 @@ DEFINE_LOG_CATEGORY_STATIC(LogUnrealNames, Log, All);
 	}
 #endif
 
-FNameEntry* AllocateNameEntry( const void* Name, NAME_INDEX Index, FNameEntry* HashNext, bool bIsPureAnsi );
+// Get the size of a FNameEntry without the union buffer included
+static const int32 NameEntryWithoutUnionSize = sizeof(FNameEntry) - (NAME_SIZE * sizeof(TCHAR));
+
+template<typename TCharType>
+FNameEntry* AllocateNameEntry( const void* Name, NAME_INDEX Index, FNameEntry* HashNext);
 
 /**
 * Helper function that can be used inside the debuggers watch window. E.g. "DebugFName(Class->Name.Index)". 
@@ -61,6 +65,17 @@ const TCHAR* DebugFName(FName& Name)
 	return TempName;
 }
 
+template <typename TCharType>
+static uint16 GetRawCasePreservingHash(const TCharType* Source)
+{
+	return FCrc::StrCrc32(Source) & 0xFFFF;
+
+}
+template <typename TCharType>
+static uint16 GetRawNonCasePreservingHash(const TCharType* Source)
+{
+	return FCrc::Strihash_DEPRECATED(Source) & 0xFFFF;
+}
 
 /*-----------------------------------------------------------------------------
 	FNameEntry
@@ -172,13 +187,29 @@ int32 FNameEntry::GetSize( const TCHAR* Name )
  */
 int32 FNameEntry::GetSize( int32 Length, bool bIsPureAnsi )
 {
-	// Calculate base size without union array.
-	int32 Size = sizeof(FNameEntry) - NAME_SIZE * sizeof(TCHAR);
-	// Add size required for string.
-	Size += (Length+1) * (bIsPureAnsi ? sizeof(ANSICHAR) : sizeof(TCHAR));
+	// Add size required for string to the base size used by the FNameEntry
+	int32 Size = NameEntryWithoutUnionSize + (Length + 1) * (bIsPureAnsi ? sizeof(ANSICHAR) : sizeof(TCHAR));
 	return Size;
 }
 
+FNameEntrySerialized::FNameEntrySerialized(const FNameEntry& NameEntry)
+{
+	if (NameEntry.IsWide())
+	{
+		PreSetIsWideForSerialization(true);
+		FCString::Strcpy(WideName, NAME_SIZE, NameEntry.GetWideName());
+		NonCasePreservingHash = GetRawNonCasePreservingHash(NameEntry.GetWideName());
+		CasePreservingHash = GetRawCasePreservingHash(NameEntry.GetWideName());
+	}
+	else
+	{
+		PreSetIsWideForSerialization(false);
+		FCStringAnsi::Strcpy(AnsiName, NAME_SIZE, NameEntry.GetAnsiName());
+		NonCasePreservingHash = GetRawNonCasePreservingHash(NameEntry.GetAnsiName());
+		CasePreservingHash = GetRawCasePreservingHash(NameEntry.GetAnsiName());
+	}
+}
+
 /*-----------------------------------------------------------------------------
 	FName statics.
 -----------------------------------------------------------------------------*/
@@ -399,23 +430,33 @@ FName::FName( const TCHAR* Name, int32 InNumber, EFindName FindType )
 	Init(Name, InNumber, FindType);
 }
 
-
-/**
- * Constructor used by FLinkerLoad when loading its name table; Creates an FName with an instance
- * number of 0 that does not attempt to split the FName into string and number portions.
- */
-FName::FName( ELinkerNameTableConstructor, const WIDECHAR* Name )
+FName::FName(const FNameEntrySerialized& LoadedEntry)
 {
-	Init(Name, NAME_NO_NUMBER_INTERNAL, FNAME_Add, false);
+	if (LoadedEntry.bWereHashesLoaded)
+{
+		// Since the name table can change sizes we need to mask the raw hash to the current size so we don't access out of bounds
+		const uint16 NonCasePreservingHash = LoadedEntry.NonCasePreservingHash & (ARRAY_COUNT(NameHash) - 1);
+		const uint16 CasePreservingHash = LoadedEntry.CasePreservingHash & (ARRAY_COUNT(NameHash) - 1);
+		if (LoadedEntry.IsWide())
+		{
+			Init(LoadedEntry.GetWideName(), NAME_NO_NUMBER_INTERNAL, FNAME_Add, NonCasePreservingHash, CasePreservingHash);
 }
-
-/**
- * Constructor used by FLinkerLoad when loading its name table; Creates an FName with an instance
- * number of 0 that does not attempt to split the FName into string and number portions.
- */
-FName::FName( ELinkerNameTableConstructor, const ANSICHAR* Name )
+		else
+		{
+			Init(LoadedEntry.GetAnsiName(), NAME_NO_NUMBER_INTERNAL, FNAME_Add, NonCasePreservingHash, CasePreservingHash);
+		}
+	}
+	else
+	{
+		if (LoadedEntry.IsWide())
+		{
+			Init(LoadedEntry.GetWideName(), NAME_NO_NUMBER_INTERNAL, FNAME_Add, false);
+		}
+		else
 {
-	Init(Name, NAME_NO_NUMBER_INTERNAL, FNAME_Add, false);
+			Init(LoadedEntry.GetAnsiName(), NAME_NO_NUMBER_INTERNAL, FNAME_Add, false);
+		}
+	}
 }
 
 FName::FName( EName HardcodedIndex, const TCHAR* Name )
@@ -463,27 +504,51 @@ int32 FName::Compare( const FName& Other ) const
 	}
 }
 
+template <typename TCharType>
+int32 FName::GetCasePreservingHash(const TCharType* Source)
+{
+	return GetRawCasePreservingHash(Source) & (ARRAY_COUNT(NameHash) - 1);
+}
+
+template <typename TCharType>
+int32 FName::GetNonCasePreservingHash(const TCharType* Source)
+{
+	return GetRawNonCasePreservingHash(Source) & (ARRAY_COUNT(NameHash) - 1);
+}
+
 void FName::Init(const WIDECHAR* InName, int32 InNumber, EFindName FindType, bool bSplitName, int32 HardcodeIndex)
 {
-	InitInternal<WIDECHAR>(InName, InNumber, FindType, bSplitName, HardcodeIndex);
+	const bool bIsPureAnsi = TCString<WIDECHAR>::IsPureAnsi(InName);
+	// Switch to ANSI if possible to save memory
+	if (bIsPureAnsi)
+	{
+		InitInternal_HashSplit<ANSICHAR>(StringCast<ANSICHAR>(InName).Get(), InNumber, FindType, bSplitName, HardcodeIndex);
+	}
+	else
+	{
+		InitInternal_HashSplit<WIDECHAR>(InName, InNumber, FindType, bSplitName, HardcodeIndex);
+	}
 }
 
-void FName::Init(const ANSICHAR* InName, int32 InNumber, EFindName FindType, bool bSplitName, int32 HardcodeIndex)
+void FName::Init(const WIDECHAR* InName, int32 InNumber, EFindName FindType, const uint16 NonCasePreservingHash, const uint16 CasePreservingHash)
 {
-	InitInternal<ANSICHAR>(InName, InNumber, FindType, bSplitName, HardcodeIndex);
+	// Since this comes from the FLinkerLoad we know that it is not pure ansi
+	InitInternal<WIDECHAR>(InName, InNumber, FindType, -1, NonCasePreservingHash, CasePreservingHash);
 }
 
-template <typename TCharType>
-void FName::InitInternal(const TCharType* InName, int32 InNumber, const EFindName FindType, const bool bSplitName, const int32 HardcodeIndex)
+void FName::Init(const ANSICHAR* InName, int32 InNumber, EFindName FindType, bool bSplitName, int32 HardcodeIndex)
 {
-	check(TCString<TCharType>::Strlen(InName)<=NAME_SIZE);
+	InitInternal_HashSplit<ANSICHAR>(InName, InNumber, FindType, bSplitName, HardcodeIndex);
+}
 
-	// initialize the name subsystem if necessary
-	if (!GetIsInitialized())
+void FName::Init(const ANSICHAR* InName, int32 InNumber, EFindName FindType, const uint16 NonCasePreservingHash, const uint16 CasePreservingHash)
 	{
-		StaticInit();
+	InitInternal<ANSICHAR>(InName, InNumber, FindType, -1, NonCasePreservingHash, CasePreservingHash);
 	}
 
+template <typename TCharType>
+void FName::InitInternal_HashSplit(const TCharType* InName, int32 InNumber, const EFindName FindType, bool bSplitName, const int32 HardcodeIndex)
+{
 	TCharType TempBuffer[NAME_SIZE];
 	int32 TempNumber;
 	// if we were passed in a number, we can't split again, other wise, a_1_2_3_4 would change every time
@@ -496,6 +561,26 @@ void FName::InitInternal(const TCharType* InName, int32 InNumber, const EFindNam
 			InNumber = NAME_EXTERNAL_TO_INTERNAL(TempNumber);
 		}
 	}
+	// Hash value of string after splitting
+	const uint16 NonCasePreservingHash = GetNonCasePreservingHash(InName);
+#if WITH_CASE_PRESERVING_NAME
+	const uint16 CasePreservingHash = GetCasePreservingHash(InName);
+#else
+	const uint16 CasePreservingHash = 0;
+#endif
+	InitInternal<TCharType>(InName, InNumber, FindType, HardcodeIndex, NonCasePreservingHash, CasePreservingHash);
+}
+
+template <typename TCharType>
+void FName::InitInternal(const TCharType* InName, int32 InNumber, const EFindName FindType, const int32 HardcodeIndex, const uint16 NonCasePreservingHash, const uint16 CasePreservingHash)
+{
+	check(TCString<TCharType>::Strlen(InName)<=NAME_SIZE);
+
+	// initialize the name subsystem if necessary
+	if (!GetIsInitialized())
+	{
+		StaticInit();
+	}
 
 	check(InName);
 
@@ -521,11 +606,11 @@ void FName::InitInternal(const TCharType* InName, int32 InNumber, const EFindNam
 	const bool bIsPureAnsi = TCString<TCharType>::IsPureAnsi(InName);
 	if(bIsPureAnsi)
 	{
-		bWasFoundOrAdded = InitInternal_FindOrAdd<ANSICHAR>(StringCast<ANSICHAR>(InName).Get(), FindType, HardcodeIndex, OutComparisonIndex, OutDisplayIndex);
+		bWasFoundOrAdded = InitInternal_FindOrAdd<ANSICHAR>(StringCast<ANSICHAR>(InName).Get(), FindType, HardcodeIndex, NonCasePreservingHash, CasePreservingHash, OutComparisonIndex, OutDisplayIndex);
 	}
 	else
 	{
-		bWasFoundOrAdded = InitInternal_FindOrAdd<WIDECHAR>(StringCast<WIDECHAR>(InName).Get(), FindType, HardcodeIndex, OutComparisonIndex, OutDisplayIndex);
+		bWasFoundOrAdded = InitInternal_FindOrAdd<WIDECHAR>(StringCast<WIDECHAR>(InName).Get(), FindType, HardcodeIndex, NonCasePreservingHash, CasePreservingHash, OutComparisonIndex, OutDisplayIndex);
 	}
 
 	if(bWasFoundOrAdded)
@@ -546,6 +631,16 @@ void FName::InitInternal(const TCharType* InName, int32 InNumber, const EFindNam
 	}
 }
 
+template <typename TCharType> void IncrementNameCount();
+template <> void IncrementNameCount<ANSICHAR>()
+{
+	FName::NumAnsiNames++;
+}
+template <> void IncrementNameCount<WIDECHAR>()
+{
+	FName::NumWideNames++;
+}
+
 template <typename TCharType>
 struct FNameInitHelper
 {
@@ -567,6 +662,25 @@ struct FNameInitHelper<ANSICHAR>
 		// the remainder of the array of NAME_SIZE with 0xfd.  So, we have to pass in the length of the dynamically allocated array instead.
 		FCStringAnsi::Strcpy(const_cast<ANSICHAR*>(DestNameEntry->GetAnsiName()), DestNameEntry->GetNameLength()+1, SrcName);
 	}
+
+	static void SetNameString(FNameEntry* const DestNameEntry, const ANSICHAR* SrcName, const int32 NameLength)
+	{
+		// Can't rely on the template override for static arrays since the safe crt version of strcpy will fill in
+		// the remainder of the array of NAME_SIZE with 0xfd.  So, we have to pass in the length of the dynamically allocated array instead.
+		FCStringAnsi::Strcpy(const_cast<ANSICHAR*>(DestNameEntry->GetAnsiName()), NameLength + 1, SrcName);
+	}
+
+	static int32 GetIndexShiftValue()
+	{
+		return 0;
+	}
+
+	static int32 GetSize(int32 Length)
+	{
+		// Add size required for string to the base size used by the FNameEntry
+		int32 Size = NameEntryWithoutUnionSize + (Length + 1) * sizeof(ANSICHAR);
+		return Size;
+	}
 };
 
 template <>
@@ -585,12 +699,31 @@ struct FNameInitHelper<WIDECHAR>
 		// the remainder of the array of NAME_SIZE with 0xfd.  So, we have to pass in the length of the dynamically allocated array instead.
 		FCStringWide::Strcpy(const_cast<WIDECHAR*>(DestNameEntry->GetWideName()), DestNameEntry->GetNameLength()+1, SrcName);
 	}
+
+	static void SetNameString(FNameEntry* const DestNameEntry, const WIDECHAR* SrcName, const int32 NameLength)
+	{
+		// Can't rely on the template override for static arrays since the safe crt version of strcpy will fill in
+		// the remainder of the array of NAME_SIZE with 0xfd.  So, we have to pass in the length of the dynamically allocated array instead.
+		FCStringWide::Strcpy(const_cast<WIDECHAR*>(DestNameEntry->GetWideName()), NameLength + 1, SrcName);
+	}
+
+	static int32 GetIndexShiftValue()
+	{
+		return 1;
+	}
+
+	static int32 GetSize(int32 Length)
+	{
+		// Add size required for string to the base size used by the FNameEntry
+		int32 Size = NameEntryWithoutUnionSize + (Length + 1) * sizeof(TCHAR);
+		return Size;
+	}
 };
 
 template <typename TCharType>
-bool FName::InitInternal_FindOrAdd(const TCharType* InName, const EFindName FindType, const int32 HardcodeIndex, int32& OutComparisonIndex, int32& OutDisplayIndex)
+bool FName::InitInternal_FindOrAdd(const TCharType* InName, const EFindName FindType, const int32 HardcodeIndex, const uint16 NonCasePreservingHash, const uint16 CasePreservingHash, int32& OutComparisonIndex, int32& OutDisplayIndex)
 {
-	const bool bWasFoundOrAdded = InitInternal_FindOrAddNameEntry<TCharType>(InName, FindType, ENameCase::IgnoreCase, OutComparisonIndex);
+	const bool bWasFoundOrAdded = InitInternal_FindOrAddNameEntry<TCharType>(InName, FindType, ENameCase::IgnoreCase, NonCasePreservingHash, OutComparisonIndex);
 	
 #if WITH_CASE_PRESERVING_NAME
 	if(bWasFoundOrAdded && HardcodeIndex < 0)
@@ -601,7 +734,7 @@ bool FName::InitInternal_FindOrAdd(const TCharType* InName, const EFindName Find
 		// If the string we got back doesn't match the case of the string we provided, also add a case variant version for display purposes
 		if(TCString<TCharType>::Strcmp(InName, FNameInitHelper<TCharType>::GetNameString(NameEntry)) != 0)
 		{
-			if(!InitInternal_FindOrAddNameEntry<TCharType>(InName, FindType, ENameCase::CaseSensitive, OutDisplayIndex))
+			if(!InitInternal_FindOrAddNameEntry<TCharType>(InName, FindType, ENameCase::CaseSensitive, CasePreservingHash, OutDisplayIndex))
 			{
 				// We don't consider failing to find/add the case variant a full failure
 				OutDisplayIndex = OutComparisonIndex;
@@ -622,12 +755,9 @@ bool FName::InitInternal_FindOrAdd(const TCharType* InName, const EFindName Find
 }
 
 template <typename TCharType>
-bool FName::InitInternal_FindOrAddNameEntry(const TCharType* InName, const EFindName FindType, const ENameCase ComparisonMode, int32& OutIndex)
+bool FName::InitInternal_FindOrAddNameEntry(const TCharType* InName, const EFindName FindType, const ENameCase ComparisonMode, const uint16 iHash, int32& OutIndex)
 {
 	CallNameCreationHook();
-	// Hash value of string
-	const int32 iHash = ( (ComparisonMode == ENameCase::IgnoreCase) ? FCrc::Strihash_DEPRECATED( InName ) : FCrc::StrCrc32( InName ) ) & (ARRAY_COUNT(NameHash)-1);
-
 	if (OutIndex < 0)
 	{
 		// Try to find the name in the hash.
@@ -693,7 +823,7 @@ bool FName::InitInternal_FindOrAddNameEntry(const TCharType* InName, const EFind
 	{
 		check(OutIndex < Names.Num());
 	}
-	FNameEntry* NewEntry = AllocateNameEntry( InName, OutIndex, OldHash, FNameInitHelper<TCharType>::IsAnsi );
+	FNameEntry* NewEntry = AllocateNameEntry<TCharType>( InName, OutIndex, OldHash);
 	if (FPlatformAtomics::InterlockedCompareExchangePointer((void**)&Names[OutIndex], NewEntry, NULL) != NULL) // we use an atomic operation to check for unexpected concurrency, verify alignment, etc
 	{
 		UE_LOG(LogUnrealNames, Fatal, TEXT("Hardcoded name '%s' at index %i was duplicated (or unexpected concurrency). Existing entry is '%s'."), *NewEntry->GetPlainNameString(), NewEntry->GetIndex(), *Names[OutIndex]->GetPlainNameString() );
@@ -1008,9 +1138,65 @@ FArchive& operator<<( FArchive& Ar, FNameEntry& E )
 		}
 	}
 	else
+	{
+		// Convert to our serialized type
+		FNameEntrySerialized EntrySerialized(E);
+		Ar << EntrySerialized;
+	}
+
+	return Ar;
+}
+
+FArchive& operator<<(FArchive& Ar, FNameEntrySerialized& E)
+{
+	if (Ar.IsLoading())
+	{
+		// for optimization reasons, we want to keep pure Ansi strings as Ansi for initializing the name entry
+		// (and later the FName) to stop copying in and out of TCHARs
+		int32 StringLen;
+		Ar << StringLen;
+
+		// negative stringlen means it's a wide string
+		if (StringLen < 0)
+		{
+			StringLen = -StringLen;
+
+			// mark the name will be wide
+			E.PreSetIsWideForSerialization(true);
+
+			// get the pointer to the wide array 
+			WIDECHAR* WideName = const_cast<WIDECHAR*>(E.GetWideName());
+
+			// read in the UCS2CHAR string and byteswap it, etc
+			auto Sink = StringMemoryPassthru<UCS2CHAR>(WideName, StringLen, StringLen);
+			Ar.Serialize(Sink.Get(), StringLen * sizeof(UCS2CHAR));
+			Sink.Apply();
+
+			INTEL_ORDER_TCHARARRAY(WideName)
+		}
+		else
+		{
+			// mark the name will be ansi
+			E.PreSetIsWideForSerialization(false);
+
+			// ansi strings can go right into the AnsiBuffer
+			ANSICHAR* AnsiName = const_cast<ANSICHAR*>(E.GetAnsiName());
+			Ar.Serialize(AnsiName, StringLen);
+		}
+		if (Ar.UE4Ver() >= VER_UE4_NAME_HASHES_SERIALIZED)
+		{
+			// Read the save time calculated hashes to save load time perf
+			Ar << E.NonCasePreservingHash;
+			Ar << E.CasePreservingHash;
+			E.bWereHashesLoaded = true;
+		}
+	}
+	else
 	{
 		FString Str = E.GetPlainNameString();
 		Ar << Str;
+		Ar << E.NonCasePreservingHash;
+		Ar << E.CasePreservingHash;
 	}
 
 	return Ar;
@@ -1069,7 +1255,7 @@ public:
 	FORCEINLINE int32 PoolSize()
 	{
 		// Allocate in 64k chunks as it's ideal for page size.
-		return 64 * 1024;
+		return 256 * 1024;
 	}
 
 	/**
@@ -1104,26 +1290,17 @@ private:
 /** Global allocator for name entries. */
 FNameEntryPoolAllocator GNameEntryPoolAllocator;
 
-FNameEntry* AllocateNameEntry( const void* Name, NAME_INDEX Index, FNameEntry* HashNext, bool bIsPureAnsi )
+template<typename TCharType>
+FNameEntry* AllocateNameEntry( const void* Name, NAME_INDEX Index, FNameEntry* HashNext)
 {
-	const SIZE_T NameLen  = bIsPureAnsi ? FCStringAnsi::Strlen((ANSICHAR*)Name) : FCString::Strlen((TCHAR*)Name);
-	int32 NameEntrySize	  = FNameEntry::GetSize( NameLen, bIsPureAnsi );
+	const SIZE_T NameLen  = TCString<TCharType>::Strlen((TCharType*)Name);
+	int32 NameEntrySize	  = FNameInitHelper<TCharType>::GetSize( NameLen );
 	FNameEntry* NameEntry = GNameEntryPoolAllocator.Allocate( NameEntrySize );
 	FName::NameEntryMemorySize += NameEntrySize;
-	NameEntry->Index      = (Index << NAME_INDEX_SHIFT) | (bIsPureAnsi ? 0 : 1);
+	NameEntry->Index      = (Index << NAME_INDEX_SHIFT) | (FNameInitHelper<TCharType>::GetIndexShiftValue());
 	NameEntry->HashNext   = HashNext;
-	// Can't rely on the template override for static arrays since the safe crt version of strcpy will fill in
-	// the remainder of the array of NAME_SIZE with 0xfd.  So, we have to pass in the length of the dynamically allocated array instead.
-	if( bIsPureAnsi )
-	{
-		FCStringAnsi::Strcpy( const_cast<ANSICHAR*>(NameEntry->GetAnsiName()), NameLen + 1, (ANSICHAR*) Name );
-		FName::NumAnsiNames++;
-	}
-	else
-	{
-		FCStringWide::Strcpy( const_cast<WIDECHAR*>(NameEntry->GetWideName()), NameLen + 1, (WIDECHAR*) Name );
-		FName::NumWideNames++;
-	}
+	FNameInitHelper<TCharType>::SetNameString(NameEntry, (TCharType*)Name, NameLen);
+	IncrementNameCount<TCharType>();
 	return NameEntry;
 }
 
diff --git a/Engine/Source/Runtime/Core/Private/Windows/WindowsApplication.cpp b/Engine/Source/Runtime/Core/Private/Windows/WindowsApplication.cpp
--- a/Engine/Source/Runtime/Core/Private/Windows/WindowsApplication.cpp
+++ b/Engine/Source/Runtime/Core/Private/Windows/WindowsApplication.cpp
@@ -24,6 +24,7 @@
 #include <devguid.h>
 #include <dwmapi.h>
 #include <cfgmgr32.h>
+#include <windowsx.h>
 
 
 // This might not be defined by Windows when maintaining backwards-compatibility to pre-Vista builds
@@ -267,6 +268,24 @@ void FWindowsApplication::SetMessageHandler( const TSharedRef< FGenericApplicati
 
 }
 
+bool FWindowsApplication::IsGamepadAttached() const
+{
+	if (XInput->IsGamepadAttached())
+	{
+		return true;
+	}
+
+	for( auto DeviceIt = ExternalInputDevices.CreateConstIterator(); DeviceIt; ++DeviceIt )
+	{
+		if ((*DeviceIt)->IsGamepadAttached())
+		{
+			return true;
+		}
+	}
+
+	return false;
+}
+
 FModifierKeysState FWindowsApplication::GetModifierKeys() const
 {
 	return CachedModifierKeyState;
@@ -1394,91 +1413,91 @@ int32 FWindowsApplication::ProcessDeferredMessage( const FDeferredWindowsMessage
 			// Mouse Button Down
 		case WM_LBUTTONDBLCLK:
 		case WM_LBUTTONDOWN:
-			{
-				if( msg == WM_LBUTTONDOWN )
-				{
-					MessageHandler->OnMouseDown( CurrentNativeEventWindowPtr, EMouseButtons::Left );
-				}
-				else
-				{
-					MessageHandler->OnMouseDoubleClick( CurrentNativeEventWindowPtr, EMouseButtons::Left );
-				}
-				return 0;
-			}
-			break;
-
 		case WM_MBUTTONDBLCLK:
 		case WM_MBUTTONDOWN:
-			{
-				if( msg == WM_MBUTTONDOWN )
-				{
-					MessageHandler->OnMouseDown( CurrentNativeEventWindowPtr, EMouseButtons::Middle );
-				}
-				else
-				{
-					MessageHandler->OnMouseDoubleClick( CurrentNativeEventWindowPtr, EMouseButtons::Middle );
-				}
-				return 0;
-			}
-			break;
-
 		case WM_RBUTTONDBLCLK:
 		case WM_RBUTTONDOWN:
+		case WM_XBUTTONDBLCLK:
+		case WM_XBUTTONDOWN:
+		case WM_LBUTTONUP:
+		case WM_MBUTTONUP:
+		case WM_RBUTTONUP:
+		case WM_XBUTTONUP:
 			{
-				if( msg == WM_RBUTTONDOWN )
+				POINT CursorPoint;
+				CursorPoint.x = GET_X_LPARAM(lParam);
+				CursorPoint.y = GET_Y_LPARAM(lParam); 
+
+				ClientToScreen(hwnd, &CursorPoint);
+
+				const FVector2D CursorPos(CursorPoint.x, CursorPoint.y);
+
+				EMouseButtons::Type MouseButton = EMouseButtons::Invalid;
+				bool bDoubleClick = false;
+				bool bMouseUp = false;
+				switch(msg)
 				{
-					MessageHandler->OnMouseDown( CurrentNativeEventWindowPtr, EMouseButtons::Right );
+				case WM_LBUTTONDBLCLK:
+					bDoubleClick = true;
+					MouseButton = EMouseButtons::Left;
+					break;
+				case WM_LBUTTONUP:
+					bMouseUp = true;
+					MouseButton = EMouseButtons::Left;
+					break;
+				case WM_LBUTTONDOWN:
+					MouseButton = EMouseButtons::Left;
+					break;
+				case WM_MBUTTONDBLCLK:
+					bDoubleClick = true;
+					MouseButton = EMouseButtons::Middle;
+					break;
+				case WM_MBUTTONUP:
+					bMouseUp = true;
+					MouseButton = EMouseButtons::Middle;
+					break;
+				case WM_MBUTTONDOWN:
+					MouseButton = EMouseButtons::Middle;
+					break;
+				case WM_RBUTTONDBLCLK:
+					bDoubleClick = true;
+					MouseButton = EMouseButtons::Right;
+					break;
+				case WM_RBUTTONUP:
+					bMouseUp = true;
+					MouseButton = EMouseButtons::Right;
+					break;
+				case WM_RBUTTONDOWN:
+					MouseButton = EMouseButtons::Right;
+					break;
+				case WM_XBUTTONDBLCLK:
+					bDoubleClick = true;
+					MouseButton = ( HIWORD(wParam) & XBUTTON1 ) ? EMouseButtons::Thumb01  : EMouseButtons::Thumb02;
+					break;
+				case WM_XBUTTONUP:
+					bMouseUp = true;
+					MouseButton = ( HIWORD(wParam) & XBUTTON1 ) ? EMouseButtons::Thumb01  : EMouseButtons::Thumb02;
+					break;
+				case WM_XBUTTONDOWN:
+					MouseButton = ( HIWORD(wParam) & XBUTTON1 ) ? EMouseButtons::Thumb01  : EMouseButtons::Thumb02;
+					break;
+				default:
+					check(0);
 				}
-				else
+
+				if (bMouseUp)
 				{
-					MessageHandler->OnMouseDoubleClick( CurrentNativeEventWindowPtr, EMouseButtons::Right );
+					return MessageHandler->OnMouseUp( MouseButton, CursorPos ) ? 0 : 1;
 				}
-				return 0;
-			}
-			break;
-
-		case WM_XBUTTONDBLCLK:
-		case WM_XBUTTONDOWN:
-			{
-				EMouseButtons::Type MouseButton = ( HIWORD(wParam) & XBUTTON1 ) ? EMouseButtons::Thumb01  : EMouseButtons::Thumb02;
-
-				BOOL Result = false;
-				if( msg == WM_XBUTTONDOWN )
+				else if (bDoubleClick)
 				{
-					Result = MessageHandler->OnMouseDown( CurrentNativeEventWindowPtr, MouseButton );
+					MessageHandler->OnMouseDoubleClick( CurrentNativeEventWindowPtr, MouseButton, CursorPos );
 				}
 				else
 				{
-					Result = MessageHandler->OnMouseDoubleClick( CurrentNativeEventWindowPtr, MouseButton );
+					MessageHandler->OnMouseDown( CurrentNativeEventWindowPtr, MouseButton, CursorPos );
 				}
-
-				return Result ? 0 : 1;
-			}
-			break;
-
-			// Mouse Button Up
-		case WM_LBUTTONUP:
-			{
-				return MessageHandler->OnMouseUp( EMouseButtons::Left ) ? 0 : 1;
-			}
-			break;
-
-		case WM_MBUTTONUP:
-			{
-				return MessageHandler->OnMouseUp( EMouseButtons::Middle ) ? 0 : 1;
-			}
-			break;
-
-		case WM_RBUTTONUP:
-			{
-				return MessageHandler->OnMouseUp( EMouseButtons::Right ) ? 0 : 1;
-			}
-			break;
-
-		case WM_XBUTTONUP:
-			{
-				EMouseButtons::Type MouseButton = ( HIWORD(wParam) & XBUTTON1 ) ? EMouseButtons::Thumb01  : EMouseButtons::Thumb02;
-				return MessageHandler->OnMouseUp( MouseButton ) ? 0 : 1;
+				return 0;
 			}
 			break;
 
@@ -1518,7 +1537,13 @@ int32 FWindowsApplication::ProcessDeferredMessage( const FDeferredWindowsMessage
 				const float SpinFactor = 1 / 120.0f;
 				const SHORT WheelDelta = GET_WHEEL_DELTA_WPARAM(wParam);
 
-				const BOOL Result = MessageHandler->OnMouseWheel( static_cast<float>( WheelDelta ) * SpinFactor );
+				POINT CursorPoint;
+				CursorPoint.x = GET_X_LPARAM(lParam);
+				CursorPoint.y = GET_Y_LPARAM(lParam); 
+
+				const FVector2D CursorPos(CursorPoint.x, CursorPoint.y);
+
+				const BOOL Result = MessageHandler->OnMouseWheel( static_cast<float>( WheelDelta ) * SpinFactor, CursorPos );
 				return Result ? 0 : 1;
 			}
 			break;
diff --git a/Engine/Source/Runtime/Core/Private/Windows/XInputInterface.cpp b/Engine/Source/Runtime/Core/Private/Windows/XInputInterface.cpp
--- a/Engine/Source/Runtime/Core/Private/Windows/XInputInterface.cpp
+++ b/Engine/Source/Runtime/Core/Private/Windows/XInputInterface.cpp
@@ -29,6 +29,7 @@ XInputInterface::XInputInterface( const TSharedRef< FGenericApplicationMessageHa
 		ControllerState.ControllerId = ControllerIndex;
 	}
 
+	bIsGamepadAttached = false;
 	bNeedsControllerStateUpdate = true;
 	InitialButtonRepeatDelay = 0.2f;
 	ButtonRepeatDelay = 0.1f;
@@ -96,143 +97,162 @@ float ShortToNormalizedFloat(SHORT AxisVal)
 
 void XInputInterface::SendControllerEvents()
 {
+	bool bWereConnected[MAX_NUM_XINPUT_CONTROLLERS];
+	XINPUT_STATE XInputStates[MAX_NUM_XINPUT_CONTROLLERS];
+
+	bIsGamepadAttached = false;
 	for ( int32 ControllerIndex=0; ControllerIndex < MAX_NUM_XINPUT_CONTROLLERS; ++ControllerIndex )
 	{
 		FControllerState& ControllerState = ControllerStates[ControllerIndex];
+
+		bWereConnected[ControllerIndex] = ControllerState.bIsConnected;
+
 		if( ControllerState.bIsConnected || bNeedsControllerStateUpdate )
 		{
-			XINPUT_STATE XInputState;
+			XINPUT_STATE& XInputState = XInputStates[ControllerIndex];
 			FMemory::Memzero( &XInputState, sizeof(XINPUT_STATE) );
 
-			bool bWasConnected = ControllerState.bIsConnected;
-
 			ControllerState.bIsConnected = ( XInputGetState( ControllerIndex, &XInputState ) == ERROR_SUCCESS ) ? true : false;
-			
-			// If the controller is connected send events or if the controller was connected send a final event with default states so that 
-			// the game doesn't think that controller buttons are still held down
-			if( ControllerState.bIsConnected || bWasConnected )
+
+			if (ControllerState.bIsConnected)
 			{
-				// If the controller is connected now but was not before, refresh the information
-				if (!bWasConnected && ControllerState.bIsConnected)
-				{
-					FCoreDelegates::OnControllerConnectionChange.Broadcast(true, -1, ControllerIndex);
-				}
-				else if (bWasConnected && !ControllerState.bIsConnected)
-				{
-					FCoreDelegates::OnControllerConnectionChange.Broadcast(false, -1, ControllerIndex);
-				}
+				bIsGamepadAttached = true;
+			}
+		}
+	}
+		
+	for ( int32 ControllerIndex=0; ControllerIndex < MAX_NUM_XINPUT_CONTROLLERS; ++ControllerIndex )
+	{
+		FControllerState& ControllerState = ControllerStates[ControllerIndex];
+
+		const bool bWasConnected = bWereConnected[ControllerIndex];
+
+		// If the controller is connected send events or if the controller was connected send a final event with default states so that 
+		// the game doesn't think that controller buttons are still held down
+		if( ControllerState.bIsConnected || bWasConnected )
+		{
+			const XINPUT_STATE& XInputState = XInputStates[ControllerIndex];
+
+			// If the controller is connected now but was not before, refresh the information
+			if (!bWasConnected && ControllerState.bIsConnected)
+			{
+				FCoreDelegates::OnControllerConnectionChange.Broadcast(true, -1, ControllerIndex);
+			}
+			else if (bWasConnected && !ControllerState.bIsConnected)
+			{
+				FCoreDelegates::OnControllerConnectionChange.Broadcast(false, -1, ControllerIndex);
+			}
 			
-				bool CurrentStates[MAX_NUM_CONTROLLER_BUTTONS] = {0};
+			bool CurrentStates[MAX_NUM_CONTROLLER_BUTTONS] = {0};
 		
-				// Get the current state of all buttons
-				CurrentStates[X360ToXboxControllerMapping[0]] = !!(XInputState.Gamepad.wButtons & XINPUT_GAMEPAD_A);
-				CurrentStates[X360ToXboxControllerMapping[1]] = !!(XInputState.Gamepad.wButtons & XINPUT_GAMEPAD_B);
-				CurrentStates[X360ToXboxControllerMapping[2]] = !!(XInputState.Gamepad.wButtons & XINPUT_GAMEPAD_X);
-				CurrentStates[X360ToXboxControllerMapping[3]] = !!(XInputState.Gamepad.wButtons & XINPUT_GAMEPAD_Y);
-				CurrentStates[X360ToXboxControllerMapping[4]] = !!(XInputState.Gamepad.wButtons & XINPUT_GAMEPAD_LEFT_SHOULDER);
-				CurrentStates[X360ToXboxControllerMapping[5]] = !!(XInputState.Gamepad.wButtons & XINPUT_GAMEPAD_RIGHT_SHOULDER);
-				CurrentStates[X360ToXboxControllerMapping[6]] = !!(XInputState.Gamepad.wButtons & XINPUT_GAMEPAD_BACK);
-				CurrentStates[X360ToXboxControllerMapping[7]] = !!(XInputState.Gamepad.wButtons & XINPUT_GAMEPAD_START);
-				CurrentStates[X360ToXboxControllerMapping[8]] = !!(XInputState.Gamepad.wButtons & XINPUT_GAMEPAD_LEFT_THUMB);
-				CurrentStates[X360ToXboxControllerMapping[9]] = !!(XInputState.Gamepad.wButtons & XINPUT_GAMEPAD_RIGHT_THUMB);
-				CurrentStates[X360ToXboxControllerMapping[10]] = !!(XInputState.Gamepad.bLeftTrigger > XINPUT_GAMEPAD_TRIGGER_THRESHOLD);
-				CurrentStates[X360ToXboxControllerMapping[11]] = !!(XInputState.Gamepad.bRightTrigger > XINPUT_GAMEPAD_TRIGGER_THRESHOLD);
-				CurrentStates[X360ToXboxControllerMapping[12]] = !!(XInputState.Gamepad.wButtons & XINPUT_GAMEPAD_DPAD_UP);
-				CurrentStates[X360ToXboxControllerMapping[13]] = !!(XInputState.Gamepad.wButtons & XINPUT_GAMEPAD_DPAD_DOWN);
-				CurrentStates[X360ToXboxControllerMapping[14]] = !!(XInputState.Gamepad.wButtons & XINPUT_GAMEPAD_DPAD_LEFT);
-				CurrentStates[X360ToXboxControllerMapping[15]] = !!(XInputState.Gamepad.wButtons & XINPUT_GAMEPAD_DPAD_RIGHT);
-				CurrentStates[X360ToXboxControllerMapping[16]] = !!(XInputState.Gamepad.sThumbLY > XINPUT_GAMEPAD_LEFT_THUMB_DEADZONE);
-				CurrentStates[X360ToXboxControllerMapping[17]] = !!(XInputState.Gamepad.sThumbLY < -XINPUT_GAMEPAD_LEFT_THUMB_DEADZONE);
-				CurrentStates[X360ToXboxControllerMapping[18]] = !!(XInputState.Gamepad.sThumbLX < -XINPUT_GAMEPAD_LEFT_THUMB_DEADZONE);
-				CurrentStates[X360ToXboxControllerMapping[19]] = !!(XInputState.Gamepad.sThumbLX > XINPUT_GAMEPAD_LEFT_THUMB_DEADZONE);
-				CurrentStates[X360ToXboxControllerMapping[20]] = !!(XInputState.Gamepad.sThumbRY > XINPUT_GAMEPAD_RIGHT_THUMB_DEADZONE);
-				CurrentStates[X360ToXboxControllerMapping[21]] = !!(XInputState.Gamepad.sThumbRY < -XINPUT_GAMEPAD_RIGHT_THUMB_DEADZONE);
-				CurrentStates[X360ToXboxControllerMapping[22]] = !!(XInputState.Gamepad.sThumbRX < -XINPUT_GAMEPAD_RIGHT_THUMB_DEADZONE);
-				CurrentStates[X360ToXboxControllerMapping[23]] = !!(XInputState.Gamepad.sThumbRX > XINPUT_GAMEPAD_RIGHT_THUMB_DEADZONE);
-
-				// Check Analog state
-
-				if( ControllerState.LeftXAnalog != XInputState.Gamepad.sThumbLX )
-				{
-					MessageHandler->OnControllerAnalog( FGamepadKeyNames::LeftAnalogX, ControllerState.ControllerId, ShortToNormalizedFloat(XInputState.Gamepad.sThumbLX) );
-					ControllerState.LeftXAnalog = XInputState.Gamepad.sThumbLX;
-				}
+			// Get the current state of all buttons
+			CurrentStates[X360ToXboxControllerMapping[0]] = !!(XInputState.Gamepad.wButtons & XINPUT_GAMEPAD_A);
+			CurrentStates[X360ToXboxControllerMapping[1]] = !!(XInputState.Gamepad.wButtons & XINPUT_GAMEPAD_B);
+			CurrentStates[X360ToXboxControllerMapping[2]] = !!(XInputState.Gamepad.wButtons & XINPUT_GAMEPAD_X);
+			CurrentStates[X360ToXboxControllerMapping[3]] = !!(XInputState.Gamepad.wButtons & XINPUT_GAMEPAD_Y);
+			CurrentStates[X360ToXboxControllerMapping[4]] = !!(XInputState.Gamepad.wButtons & XINPUT_GAMEPAD_LEFT_SHOULDER);
+			CurrentStates[X360ToXboxControllerMapping[5]] = !!(XInputState.Gamepad.wButtons & XINPUT_GAMEPAD_RIGHT_SHOULDER);
+			CurrentStates[X360ToXboxControllerMapping[6]] = !!(XInputState.Gamepad.wButtons & XINPUT_GAMEPAD_BACK);
+			CurrentStates[X360ToXboxControllerMapping[7]] = !!(XInputState.Gamepad.wButtons & XINPUT_GAMEPAD_START);
+			CurrentStates[X360ToXboxControllerMapping[8]] = !!(XInputState.Gamepad.wButtons & XINPUT_GAMEPAD_LEFT_THUMB);
+			CurrentStates[X360ToXboxControllerMapping[9]] = !!(XInputState.Gamepad.wButtons & XINPUT_GAMEPAD_RIGHT_THUMB);
+			CurrentStates[X360ToXboxControllerMapping[10]] = !!(XInputState.Gamepad.bLeftTrigger > XINPUT_GAMEPAD_TRIGGER_THRESHOLD);
+			CurrentStates[X360ToXboxControllerMapping[11]] = !!(XInputState.Gamepad.bRightTrigger > XINPUT_GAMEPAD_TRIGGER_THRESHOLD);
+			CurrentStates[X360ToXboxControllerMapping[12]] = !!(XInputState.Gamepad.wButtons & XINPUT_GAMEPAD_DPAD_UP);
+			CurrentStates[X360ToXboxControllerMapping[13]] = !!(XInputState.Gamepad.wButtons & XINPUT_GAMEPAD_DPAD_DOWN);
+			CurrentStates[X360ToXboxControllerMapping[14]] = !!(XInputState.Gamepad.wButtons & XINPUT_GAMEPAD_DPAD_LEFT);
+			CurrentStates[X360ToXboxControllerMapping[15]] = !!(XInputState.Gamepad.wButtons & XINPUT_GAMEPAD_DPAD_RIGHT);
+			CurrentStates[X360ToXboxControllerMapping[16]] = !!(XInputState.Gamepad.sThumbLY > XINPUT_GAMEPAD_LEFT_THUMB_DEADZONE);
+			CurrentStates[X360ToXboxControllerMapping[17]] = !!(XInputState.Gamepad.sThumbLY < -XINPUT_GAMEPAD_LEFT_THUMB_DEADZONE);
+			CurrentStates[X360ToXboxControllerMapping[18]] = !!(XInputState.Gamepad.sThumbLX < -XINPUT_GAMEPAD_LEFT_THUMB_DEADZONE);
+			CurrentStates[X360ToXboxControllerMapping[19]] = !!(XInputState.Gamepad.sThumbLX > XINPUT_GAMEPAD_LEFT_THUMB_DEADZONE);
+			CurrentStates[X360ToXboxControllerMapping[20]] = !!(XInputState.Gamepad.sThumbRY > XINPUT_GAMEPAD_RIGHT_THUMB_DEADZONE);
+			CurrentStates[X360ToXboxControllerMapping[21]] = !!(XInputState.Gamepad.sThumbRY < -XINPUT_GAMEPAD_RIGHT_THUMB_DEADZONE);
+			CurrentStates[X360ToXboxControllerMapping[22]] = !!(XInputState.Gamepad.sThumbRX < -XINPUT_GAMEPAD_RIGHT_THUMB_DEADZONE);
+			CurrentStates[X360ToXboxControllerMapping[23]] = !!(XInputState.Gamepad.sThumbRX > XINPUT_GAMEPAD_RIGHT_THUMB_DEADZONE);
+
+			// Check Analog state
+
+			if( ControllerState.LeftXAnalog != XInputState.Gamepad.sThumbLX )
+			{
+				MessageHandler->OnControllerAnalog( FGamepadKeyNames::LeftAnalogX, ControllerState.ControllerId, ShortToNormalizedFloat(XInputState.Gamepad.sThumbLX) );
+				ControllerState.LeftXAnalog = XInputState.Gamepad.sThumbLX;
+			}
 
-				if( ControllerState.LeftYAnalog != XInputState.Gamepad.sThumbLY )
-				{
-					MessageHandler->OnControllerAnalog( FGamepadKeyNames::LeftAnalogY, ControllerState.ControllerId, ShortToNormalizedFloat(XInputState.Gamepad.sThumbLY) );
-					ControllerState.LeftYAnalog = XInputState.Gamepad.sThumbLY;
-				}
+			if( ControllerState.LeftYAnalog != XInputState.Gamepad.sThumbLY )
+			{
+				MessageHandler->OnControllerAnalog( FGamepadKeyNames::LeftAnalogY, ControllerState.ControllerId, ShortToNormalizedFloat(XInputState.Gamepad.sThumbLY) );
+				ControllerState.LeftYAnalog = XInputState.Gamepad.sThumbLY;
+			}
 
-				if( ControllerState.RightXAnalog != XInputState.Gamepad.sThumbRX )
-				{
-					MessageHandler->OnControllerAnalog( FGamepadKeyNames::RightAnalogX, ControllerState.ControllerId, ShortToNormalizedFloat(XInputState.Gamepad.sThumbRX) );
-					ControllerState.RightXAnalog = XInputState.Gamepad.sThumbRX;
-				}
+			if( ControllerState.RightXAnalog != XInputState.Gamepad.sThumbRX )
+			{
+				MessageHandler->OnControllerAnalog( FGamepadKeyNames::RightAnalogX, ControllerState.ControllerId, ShortToNormalizedFloat(XInputState.Gamepad.sThumbRX) );
+				ControllerState.RightXAnalog = XInputState.Gamepad.sThumbRX;
+			}
 
-				if( ControllerState.RightYAnalog != XInputState.Gamepad.sThumbRY )
-				{
-					MessageHandler->OnControllerAnalog( FGamepadKeyNames::RightAnalogY, ControllerState.ControllerId, ShortToNormalizedFloat(XInputState.Gamepad.sThumbRY) );
-					ControllerState.RightYAnalog = XInputState.Gamepad.sThumbRY;
-				}
+			if( ControllerState.RightYAnalog != XInputState.Gamepad.sThumbRY )
+			{
+				MessageHandler->OnControllerAnalog( FGamepadKeyNames::RightAnalogY, ControllerState.ControllerId, ShortToNormalizedFloat(XInputState.Gamepad.sThumbRY) );
+				ControllerState.RightYAnalog = XInputState.Gamepad.sThumbRY;
+			}
 
-				if( ControllerState.LeftTriggerAnalog != XInputState.Gamepad.bLeftTrigger )
-				{
-					MessageHandler->OnControllerAnalog( FGamepadKeyNames::LeftTriggerAnalog, ControllerState.ControllerId, XInputState.Gamepad.bLeftTrigger / 255.f );
-					ControllerState.LeftTriggerAnalog = XInputState.Gamepad.bLeftTrigger;
-				}
+			if( ControllerState.LeftTriggerAnalog != XInputState.Gamepad.bLeftTrigger )
+			{
+				MessageHandler->OnControllerAnalog( FGamepadKeyNames::LeftTriggerAnalog, ControllerState.ControllerId, XInputState.Gamepad.bLeftTrigger / 255.f );
+				ControllerState.LeftTriggerAnalog = XInputState.Gamepad.bLeftTrigger;
+			}
 
-				if( ControllerState.RightTriggerAnalog != XInputState.Gamepad.bRightTrigger )
-				{
-					MessageHandler->OnControllerAnalog( FGamepadKeyNames::RightTriggerAnalog, ControllerState.ControllerId, XInputState.Gamepad.bRightTrigger / 255.f );
-					ControllerState.RightTriggerAnalog = XInputState.Gamepad.bRightTrigger;
-				}
+			if( ControllerState.RightTriggerAnalog != XInputState.Gamepad.bRightTrigger )
+			{
+				MessageHandler->OnControllerAnalog( FGamepadKeyNames::RightTriggerAnalog, ControllerState.ControllerId, XInputState.Gamepad.bRightTrigger / 255.f );
+				ControllerState.RightTriggerAnalog = XInputState.Gamepad.bRightTrigger;
+			}
 
-				const double CurrentTime = FPlatformTime::Seconds();
+			const double CurrentTime = FPlatformTime::Seconds();
 
-				// For each button check against the previous state and send the correct message if any
-				for (int32 ButtonIndex = 0; ButtonIndex < MAX_NUM_CONTROLLER_BUTTONS; ++ButtonIndex)
+			// For each button check against the previous state and send the correct message if any
+			for (int32 ButtonIndex = 0; ButtonIndex < MAX_NUM_CONTROLLER_BUTTONS; ++ButtonIndex)
+			{
+				if (CurrentStates[ButtonIndex] != ControllerState.ButtonStates[ButtonIndex])
 				{
-					if (CurrentStates[ButtonIndex] != ControllerState.ButtonStates[ButtonIndex])
+					if( CurrentStates[ButtonIndex] )
 					{
-						if( CurrentStates[ButtonIndex] )
-						{
-							MessageHandler->OnControllerButtonPressed( Buttons[ButtonIndex], ControllerState.ControllerId, false );
-						}
-						else
-						{
-							MessageHandler->OnControllerButtonReleased( Buttons[ButtonIndex], ControllerState.ControllerId, false );
-						}
-
-						if ( CurrentStates[ButtonIndex] != 0 )
-						{
-							// this button was pressed - set the button's NextRepeatTime to the InitialButtonRepeatDelay
-							ControllerState.NextRepeatTime[ButtonIndex] = CurrentTime + InitialButtonRepeatDelay;
-						}
+						MessageHandler->OnControllerButtonPressed( Buttons[ButtonIndex], ControllerState.ControllerId, false );
 					}
-					else if ( CurrentStates[ButtonIndex] != 0 && ControllerState.NextRepeatTime[ButtonIndex] <= CurrentTime )
+					else
 					{
-						MessageHandler->OnControllerButtonPressed( Buttons[ButtonIndex], ControllerState.ControllerId, true );
+						MessageHandler->OnControllerButtonReleased( Buttons[ButtonIndex], ControllerState.ControllerId, false );
+					}
 
-						// set the button's NextRepeatTime to the ButtonRepeatDelay
-						ControllerState.NextRepeatTime[ButtonIndex] = CurrentTime + ButtonRepeatDelay;
+					if ( CurrentStates[ButtonIndex] != 0 )
+					{
+						// this button was pressed - set the button's NextRepeatTime to the InitialButtonRepeatDelay
+						ControllerState.NextRepeatTime[ButtonIndex] = CurrentTime + InitialButtonRepeatDelay;
 					}
+				}
+				else if ( CurrentStates[ButtonIndex] != 0 && ControllerState.NextRepeatTime[ButtonIndex] <= CurrentTime )
+				{
+					MessageHandler->OnControllerButtonPressed( Buttons[ButtonIndex], ControllerState.ControllerId, true );
 
-					// Update the state for next time
-					ControllerState.ButtonStates[ButtonIndex] = CurrentStates[ButtonIndex];
-				}	
+					// set the button's NextRepeatTime to the ButtonRepeatDelay
+					ControllerState.NextRepeatTime[ButtonIndex] = CurrentTime + ButtonRepeatDelay;
+				}
+
+				// Update the state for next time
+				ControllerState.ButtonStates[ButtonIndex] = CurrentStates[ButtonIndex];
+			}	
 
- 				// apply force feedback
- 				XINPUT_VIBRATION VibrationState;
+ 			// apply force feedback
+ 			XINPUT_VIBRATION VibrationState;
  
-				const float LargeValue = (ControllerState.ForceFeedback.LeftLarge > ControllerState.ForceFeedback.RightLarge ? ControllerState.ForceFeedback.LeftLarge : ControllerState.ForceFeedback.RightLarge);
-				const float SmallValue = (ControllerState.ForceFeedback.LeftSmall > ControllerState.ForceFeedback.RightSmall ? ControllerState.ForceFeedback.LeftSmall : ControllerState.ForceFeedback.RightSmall);
+			const float LargeValue = (ControllerState.ForceFeedback.LeftLarge > ControllerState.ForceFeedback.RightLarge ? ControllerState.ForceFeedback.LeftLarge : ControllerState.ForceFeedback.RightLarge);
+			const float SmallValue = (ControllerState.ForceFeedback.LeftSmall > ControllerState.ForceFeedback.RightSmall ? ControllerState.ForceFeedback.LeftSmall : ControllerState.ForceFeedback.RightSmall);
 
-				VibrationState.wLeftMotorSpeed = ( ::WORD ) ( LargeValue * 65535.0f );
- 				VibrationState.wRightMotorSpeed = ( ::WORD ) ( SmallValue * 65535.0f );
+			VibrationState.wLeftMotorSpeed = ( ::WORD ) ( LargeValue * 65535.0f );
+ 			VibrationState.wRightMotorSpeed = ( ::WORD ) ( SmallValue * 65535.0f );
  
-				XInputSetState( ( ::DWORD ) ControllerState.ControllerId, &VibrationState );			
-			}
+			XInputSetState( ( ::DWORD ) ControllerState.ControllerId, &VibrationState );			
 		}
 	}
 
diff --git a/Engine/Source/Runtime/Core/Private/Windows/XInputInterface.h b/Engine/Source/Runtime/Core/Private/Windows/XInputInterface.h
--- a/Engine/Source/Runtime/Core/Private/Windows/XInputInterface.h
+++ b/Engine/Source/Runtime/Core/Private/Windows/XInputInterface.h
@@ -49,6 +49,8 @@ public:
 	*/
 	void SetChannelValues( const int32 ControllerId, const FForceFeedbackValues& Values );
 
+	bool IsGamepadAttached() const { return bIsGamepadAttached; }
+
 private:
 
 	XInputInterface( const TSharedRef< FGenericApplicationMessageHandler >& MessageHandler );
@@ -93,6 +95,8 @@ private:
 	/** If we've been notified by the system that the controller state may have changed */
 	bool bNeedsControllerStateUpdate;
 
+	bool bIsGamepadAttached;
+
 	/** In the engine, all controllers map to xbox controllers for consistency */
 	uint8	X360ToXboxControllerMapping[MAX_NUM_CONTROLLER_BUTTONS];
 
diff --git a/Engine/Source/Runtime/Core/Public/GenericPlatform/GenericApplication.h b/Engine/Source/Runtime/Core/Public/GenericPlatform/GenericApplication.h
--- a/Engine/Source/Runtime/Core/Public/GenericPlatform/GenericApplication.h
+++ b/Engine/Source/Runtime/Core/Public/GenericPlatform/GenericApplication.h
@@ -437,6 +437,8 @@ public:
 
 	virtual bool IsMouseAttached() const { return true; }
 
+	virtual bool IsGamepadAttached() const { return false; }
+
 	virtual void RegisterConsoleCommandListener(const FOnConsoleCommandListener& InListener) {}
 
 	virtual void AddPendingConsoleCommand(const FString& InCommand) {}
diff --git a/Engine/Source/Runtime/Core/Public/GenericPlatform/GenericApplicationMessageHandler.h b/Engine/Source/Runtime/Core/Public/GenericPlatform/GenericApplicationMessageHandler.h
--- a/Engine/Source/Runtime/Core/Public/GenericPlatform/GenericApplicationMessageHandler.h
+++ b/Engine/Source/Runtime/Core/Public/GenericPlatform/GenericApplicationMessageHandler.h
@@ -261,21 +261,41 @@ public:
 		return false;
 	}
 
+	virtual bool OnMouseDown( const TSharedPtr< FGenericWindow >& Window, const EMouseButtons::Type Button, const FVector2D CursorPos )
+	{
+		return false;
+	}
+
 	virtual bool OnMouseUp( const EMouseButtons::Type Button )
 	{
 		return false;
 	}
 
+	virtual bool OnMouseUp( const EMouseButtons::Type Button, const FVector2D CursorPos )
+	{
+		return false;
+	}
+
 	virtual bool OnMouseDoubleClick( const TSharedPtr< FGenericWindow >& Window, const EMouseButtons::Type Button )
 	{
 		return false;
 	}
 
+	virtual bool OnMouseDoubleClick( const TSharedPtr< FGenericWindow >& Window, const EMouseButtons::Type Button, const FVector2D CursorPos )
+	{
+		return false;
+	}
+
 	virtual bool OnMouseWheel( const float Delta )
 	{
 		return false;
 	}
 
+	virtual bool OnMouseWheel( const float Delta, const FVector2D CursorPos )
+	{
+		return false;
+	}
+
 	virtual bool OnMouseMove()
 	{
 		return false;
diff --git a/Engine/Source/Runtime/Core/Public/GenericPlatform/GenericPlatformStricmp.h b/Engine/Source/Runtime/Core/Public/GenericPlatform/GenericPlatformStricmp.h
--- a/Engine/Source/Runtime/Core/Public/GenericPlatform/GenericPlatformStricmp.h
+++ b/Engine/Source/Runtime/Core/Public/GenericPlatform/GenericPlatformStricmp.h
@@ -43,12 +43,15 @@ private:
 		// Walk the strings, comparing them case insensitively.
 		for (; *String1 || *String2; String1++, String2++)
 		{
-			CompatibleCharType1 Char1 = TChar<CompatibleCharType1>::ToLower(*String1);
-			CompatibleCharType2 Char2 = TChar<CompatibleCharType2>::ToLower(*String2);
-
-			if (Char1 != Char2)
+			if(*String1 != *String2)
 			{
-				return Char1 - Char2;
+				CompatibleCharType1 Char1 = TChar<CompatibleCharType1>::ToLower(*String1);
+				CompatibleCharType2 Char2 = TChar<CompatibleCharType2>::ToLower(*String2);
+
+				if (Char1 != Char2)
+				{
+					return Char1 - Char2;
+				}
 			}
 		}
 		return 0;
diff --git a/Engine/Source/Runtime/Core/Public/GenericPlatform/StandardPlatformString.h b/Engine/Source/Runtime/Core/Public/GenericPlatform/StandardPlatformString.h
--- a/Engine/Source/Runtime/Core/Public/GenericPlatform/StandardPlatformString.h
+++ b/Engine/Source/Runtime/Core/Public/GenericPlatform/StandardPlatformString.h
@@ -45,11 +45,15 @@ public:
 		// walk the strings, comparing them case insensitively, up to a max size
 		for (; (*String1 || *String2) && Count > 0; String1++, String2++, Count--)
 		{
-			CharType Char1 = TChar<CharType>::ToUpper(*String1), Char2 = TChar<CharType>::ToUpper(*String2);
-			if (Char1 != Char2)
+			if(*String1 != *String2)
 			{
-				return Char1 - Char2;
+				CharType Char1 = TChar<CharType>::ToUpper(*String1), Char2 = TChar<CharType>::ToUpper(*String2);
+				if (Char1 != Char2)
+				{
+					return Char1 - Char2;
+				}
 			}
+			
 		}
 		return 0;
 	}
diff --git a/Engine/Source/Runtime/Core/Public/Math/Quat.h b/Engine/Source/Runtime/Core/Public/Math/Quat.h
--- a/Engine/Source/Runtime/Core/Public/Math/Quat.h
+++ b/Engine/Source/Runtime/Core/Public/Math/Quat.h
@@ -410,8 +410,18 @@ public:
 			*const_cast<FQuat*>(this) = FQuat::Identity;
 		}
 	}
+
+	FORCEINLINE void DiagnosticCheckNaN(const TCHAR* Message) const
+	{
+		if (ContainsNaN())
+		{
+			logOrEnsureNanError(TEXT("%s: FQuat contains NaN: %s"), Message, *ToString());
+			*const_cast<FQuat*>(this) = FQuat::Identity;
+		}
+	}
 #else
 	FORCEINLINE void DiagnosticCheckNaN() const {}
+	FORCEINLINE void DiagnosticCheckNaN(const TCHAR* Message) const {}
 #endif
 
 public:
diff --git a/Engine/Source/Runtime/Core/Public/Math/Rotator.h b/Engine/Source/Runtime/Core/Public/Math/Rotator.h
--- a/Engine/Source/Runtime/Core/Public/Math/Rotator.h
+++ b/Engine/Source/Runtime/Core/Public/Math/Rotator.h
@@ -36,8 +36,18 @@ public:
 			*const_cast<FRotator*>(this) = ZeroRotator;
 		}
 	}
+
+	FORCEINLINE void DiagnosticCheckNaN(const TCHAR* Message) const
+	{
+		if (ContainsNaN())
+		{
+			logOrEnsureNanError(TEXT("%s: FRotator contains NaN: %s"), Message, *ToString());
+			*const_cast<FRotator*>(this) = ZeroRotator;
+		}
+	}
 #else
 	FORCEINLINE void DiagnosticCheckNaN() const {}
+	FORCEINLINE void DiagnosticCheckNaN(const TCHAR* Message) const {}
 #endif
 
 	/**
diff --git a/Engine/Source/Runtime/Core/Public/Misc/AssertionMacros.h b/Engine/Source/Runtime/Core/Public/Misc/AssertionMacros.h
--- a/Engine/Source/Runtime/Core/Public/Misc/AssertionMacros.h
+++ b/Engine/Source/Runtime/Core/Public/Misc/AssertionMacros.h
@@ -23,7 +23,7 @@
 	#define _DebugBreakAndPromptForRemote()
 #endif // !UE_BUILD_SHIPPING
 #if DO_CHECK
-	#define checkCode( Code )		do { Code } while ( false );
+	#define checkCode( Code )		do { Code; } while ( false );
 	#define verify(expr)			{ if(!(expr)) { FDebug::LogAssertFailedMessage( #expr, __FILE__, __LINE__ ); _DebugBreakAndPromptForRemote(); FDebug::AssertFailed( #expr, __FILE__, __LINE__ ); CA_ASSUME(expr); } }
 	#define check(expr)				{ if(!(expr)) { FDebug::LogAssertFailedMessage( #expr, __FILE__, __LINE__ ); _DebugBreakAndPromptForRemote(); FDebug::AssertFailed( #expr, __FILE__, __LINE__ ); CA_ASSUME(expr); } }
 	
diff --git a/Engine/Source/Runtime/Core/Public/UObject/FrameworkObjectVersion.h b/Engine/Source/Runtime/Core/Public/UObject/FrameworkObjectVersion.h
--- a/Engine/Source/Runtime/Core/Public/UObject/FrameworkObjectVersion.h
+++ b/Engine/Source/Runtime/Core/Public/UObject/FrameworkObjectVersion.h
@@ -16,6 +16,9 @@ struct CORE_API FFrameworkObjectVersion
 		// missing graphs that appear read only when edited
 		AnimBlueprintSubgraphFix,
 
+		// Static and skeletal mesh sockets now use the specified scale
+		MeshSocketScaleUtilization,
+
 		// -----<new versions can be added above this line>-------------------------------------------------
 		VersionPlusOne,
 		LatestVersion = VersionPlusOne - 1
diff --git a/Engine/Source/Runtime/Core/Public/UObject/NameTypes.h b/Engine/Source/Runtime/Core/Public/UObject/NameTypes.h
--- a/Engine/Source/Runtime/Core/Public/UObject/NameTypes.h
+++ b/Engine/Source/Runtime/Core/Public/UObject/NameTypes.h
@@ -73,7 +73,7 @@ namespace FNameDefs
 {
 #if !WITH_EDITORONLY_DATA
 	// Use a modest bucket count on consoles
-	static const uint32 NameHashBucketCount = 4096;
+	static const uint32 NameHashBucketCount = 65536;
 #else
 	// On PC platform we use a large number of name hash buckets to accommodate the editor's
 	// use of FNames to store asset path and content tags
@@ -124,30 +124,31 @@ public:
 	/** Pointer to the next entry in this hash bin's linked list. */
 	FNameEntry*		HashNext;
 
-private:
+protected:
 	/** Name, variable-sized - note that AllocateNameEntry only allocates memory as needed. */
 	union
 	{
 		ANSICHAR	AnsiName[NAME_SIZE];
 		WIDECHAR	WideName[NAME_SIZE];
 	};
-
-	// DO NOT ADD VARIABLES BELOW UNION!
-
-public:
-	/** Default constructor doesn't do anything. AllocateNameEntry is responsible for work. */
-	FNameEntry()
-	{}
+	// DO NOT ADD VARIABLES BELOW!
 
 	/** 
 	 * Constructor called from the linker name table serialization function. Initializes the index
 	 * to a value that indicates widechar as that's what the linker is going to serialize.
+	 *
+	 * Only callable from the serialization version of this class
 	 */
 	FNameEntry( enum ELinkerNameTableConstructor )
 	{
 		Index = NAME_WIDE_MASK;
 	}
 
+public:
+	/** Default constructor doesn't do anything. AllocateNameEntry is responsible for work. */
+	FNameEntry()
+	{}
+
 	/** 
 	 * Sets whether or not the NameEntry will have a wide string, or an ansi string
 	 *
@@ -247,7 +248,34 @@ public:
 	}
 
 	// Friend for access to Flags.
-	friend FNameEntry* AllocateNameEntry( const void* Name, NAME_INDEX Index, FNameEntry* HashNext, bool bIsPureAnsi );
+	template<typename TCharType>
+	friend FNameEntry* AllocateNameEntry( const void* Name, NAME_INDEX Index, FNameEntry* HashNext);
+};
+
+/**
+ *  This struct is only used during loading/saving and is not part of the runtime costs
+ */
+struct FNameEntrySerialized :
+	public FNameEntry
+{
+	uint16 NonCasePreservingHash;
+	uint16 CasePreservingHash;
+	bool bWereHashesLoaded;
+
+	FNameEntrySerialized(const FNameEntry& NameEntry);
+	FNameEntrySerialized(enum ELinkerNameTableConstructor) :
+		FNameEntry(ENAME_LinkerConstructor),
+		NonCasePreservingHash(0),
+		CasePreservingHash(0),
+		bWereHashesLoaded(false)
+	{
+	}
+
+	friend CORE_API FArchive& operator<<(FArchive& Ar, FNameEntrySerialized& E);
+	friend CORE_API FArchive& operator<<(FArchive& Ar, FNameEntrySerialized* E)
+	{
+		return Ar << *E;
+	}
 };
 
 /**
@@ -788,15 +816,10 @@ public:
 
 	/**
 	 * Constructor used by FLinkerLoad when loading its name table; Creates an FName with an instance
-	 * number of 0 that does not attempt to split the FName into string and number portions.
+	 * number of 0 that does not attempt to split the FName into string and number portions. Also,
+	 * this version skips calculating the hashes of the names if possible
 	 */
-	FName( ELinkerNameTableConstructor, const WIDECHAR* Name );
-
-	/**
-	 * Constructor used by FLinkerLoad when loading its name table; Creates an FName with an instance
-	 * number of 0 that does not attempt to split the FName into string and number portions.
-	 */
-	FName( ELinkerNameTableConstructor, const ANSICHAR* Name );
+	FName(const FNameEntrySerialized& LoadedEntry);
 
 	/**
 	 * Create an FName with a hardcoded string index.
@@ -862,6 +885,11 @@ public:
 		return !operator==(Other);
 	}
 
+	template <typename TCharType>
+	static int32 GetCasePreservingHash(const TCharType* Source);
+	template <typename TCharType>
+	static int32 GetNonCasePreservingHash(const TCharType* Source);
+
 	static void StaticInit();
 	static void DisplayHash( class FOutputDevice& Ar );
 	static FString SafeString( int32 InDisplayIndex, int32 InstanceNumber=NAME_NO_NUMBER_INTERNAL )
@@ -986,7 +1014,10 @@ private:
 	friend const TCHAR* DebugFName(int32);
 	friend const TCHAR* DebugFName(int32, int32);
 	friend const TCHAR* DebugFName(FName&);
-	friend FNameEntry* AllocateNameEntry( const void* Name, NAME_INDEX Index, FNameEntry* HashNext, bool bIsPureAnsi );
+	template<typename TCharType>
+	friend FNameEntry* AllocateNameEntry( const void* Name, NAME_INDEX Index, FNameEntry* HashNext);
+	/** Used to increment the correct counter based upon TCharType */
+	template <typename TCharType> friend void IncrementNameCount();
 
 	/**
 	 * Initialization from a wide string
@@ -998,6 +1029,10 @@ private:
 	 * @param HardcodeIndex If >= 0, this represents a hardcoded FName and so automatically gets this index
 	 */
 	void Init(const WIDECHAR* InName, int32 InNumber, EFindName FindType, bool bSplitName=true, int32 HardcodeIndex = -1);
+	/**
+	 * Version that takes the hashes as part of the constructor (serialized previously)
+	 */
+	void Init(const WIDECHAR* InName, int32 InNumber, EFindName FindType, const uint16 NonCasePreservingHash, const uint16 CasePreservingHash);
 
 	/**
 	 * Initialization from an ANSI string
@@ -1009,15 +1044,25 @@ private:
 	 * @param HardcodeIndex If >= 0, this represents a hardcoded FName and so automatically gets this index
 	 */
 	void Init(const ANSICHAR* InName, int32 InNumber, EFindName FindType, bool bSplitName=true, int32 HardcodeIndex = -1);
+	/**
+	 * Version that takes the hashes as part of the constructor (serialized previously). Skips InitInternal_HashSplit for loading perf reasons
+	 */
+	void Init(const ANSICHAR* InName, int32 InNumber, EFindName FindType, const uint16 NonCasePreservingHash, const uint16 CasePreservingHash);
+
+	template <typename TCharType>
+	void InitInternal(const TCharType* InName, int32 InNumber, const EFindName FindType, const int32 HardcodeIndex, const uint16 NonCasePreservingHash, const uint16 CasePreservingHash);
 
+	/**
+	 * Version of InitInternal that calculates the hash after splitting the string. Used by runtime FName construction
+	 */
 	template <typename TCharType>
-	void InitInternal(const TCharType* InName, int32 InNumber, const EFindName FindType, const bool bSplitName, const int32 HardcodeIndex);
+	void InitInternal_HashSplit(const TCharType* InName, int32 InNumber, const EFindName FindType, bool bSplitName, const int32 HardcodeIndex);
 
 	template <typename TCharType>
-	static bool InitInternal_FindOrAdd(const TCharType* InName, const EFindName FindType, const int32 HardcodeIndex, int32& OutComparisonIndex, int32& OutDisplayIndex);
+	static bool InitInternal_FindOrAdd(const TCharType* InName, const EFindName FindType, const int32 HardcodeIndex, const uint16 NonCasePreservingHash, const uint16 CasePreservingHash, int32& OutComparisonIndex, int32& OutDisplayIndex);
 
 	template <typename TCharType>
-	static bool InitInternal_FindOrAddNameEntry(const TCharType* InName, const EFindName FindType, const ENameCase ComparisonMode, int32& OutIndex);
+	static bool InitInternal_FindOrAddNameEntry(const TCharType* InName, const EFindName FindType, const ENameCase ComparisonMode, const uint16 iHash, int32& OutIndex);
 
 	template <typename TCharType>
 	static bool SplitNameWithCheckImpl(const TCharType* OldName, TCharType* NewName, int32 NewNameLen, int32& NewNumber);
diff --git a/Engine/Source/Runtime/Core/Public/UObject/ObjectVersion.h b/Engine/Source/Runtime/Core/Public/UObject/ObjectVersion.h
--- a/Engine/Source/Runtime/Core/Public/UObject/ObjectVersion.h
+++ b/Engine/Source/Runtime/Core/Public/UObject/ObjectVersion.h
@@ -602,6 +602,8 @@ enum EUnrealEngineObjectUE4Version
 	VER_UE4_STREAMABLE_TEXTURE_AABB,
 	// FPropertyTag stores GUID of property
 	VER_UE4_PROPERTY_GUID_IN_PROPERTY_TAG,
+	// Name table hashes are calculated and saved out rather than at load time
+	VER_UE4_NAME_HASHES_SERIALIZED,
 	// -----<new versions can be added before this line>-------------------------------------------------
 	// - this needs to be the last line (see note below)
 	VER_UE4_AUTOMATIC_VERSION_PLUS_ONE,
diff --git a/Engine/Source/Runtime/Core/Public/Windows/WindowsApplication.h b/Engine/Source/Runtime/Core/Public/Windows/WindowsApplication.h
--- a/Engine/Source/Runtime/Core/Public/Windows/WindowsApplication.h
+++ b/Engine/Source/Runtime/Core/Public/Windows/WindowsApplication.h
@@ -340,6 +340,7 @@ public:
 	virtual void SetHighPrecisionMouseMode( const bool Enable, const TSharedPtr< FGenericWindow >& InWindow ) override;
 	virtual bool IsUsingHighPrecisionMouseMode() const override { return bUsingHighPrecisionMouseInput; }
 	virtual bool IsMouseAttached() const override { return bIsMouseAttached; }
+	virtual bool IsGamepadAttached() const override;
 	virtual FModifierKeysState GetModifierKeys() const override;
 	virtual bool IsCursorDirectlyOverSlateWindow() const override;
 	virtual FPlatformRect GetWorkArea( const FPlatformRect& CurrentWindow ) const override;
diff --git a/Engine/Source/Runtime/CoreUObject/Private/UObject/LinkerLoad.cpp b/Engine/Source/Runtime/CoreUObject/Private/UObject/LinkerLoad.cpp
--- a/Engine/Source/Runtime/CoreUObject/Private/UObject/LinkerLoad.cpp
+++ b/Engine/Source/Runtime/CoreUObject/Private/UObject/LinkerLoad.cpp
@@ -1246,19 +1246,12 @@ FLinkerLoad::ELinkerStatus FLinkerLoad::SerializeNameMap()
 	while( bFinishedPrecaching && NameMapIndex < Summary.NameCount && !IsTimeLimitExceeded(TEXT("serializing name map"),100) )
 	{
 		// Read the name entry from the file.
-		FNameEntry NameEntry(ENAME_LinkerConstructor);
+		FNameEntrySerialized NameEntry(ENAME_LinkerConstructor);
 		*this << NameEntry;
 
-		// Add it to the name table. We disregard the context flags as we don't support flags on names for final release builds.
+		// Add it to the name table with no splitting and no hash calculations
+		NameMap.Add(FName(NameEntry));
 
-		// now, we make sure we DO NOT split the name here because it will have been written out
-		// split, and we don't want to keep splitting A_3_4_9 every time
-
-		NameMap.Add( 
-			NameEntry.IsWide() ? 
-				FName(ENAME_LinkerConstructor, NameEntry.GetWideName()) : 
-				FName(ENAME_LinkerConstructor, NameEntry.GetAnsiName())
-			);
 		NameMapIndex++;
 	}
 
diff --git a/Engine/Source/Runtime/Engine/Classes/Components/ChildActorComponent.h b/Engine/Source/Runtime/Engine/Classes/Components/ChildActorComponent.h
--- a/Engine/Source/Runtime/Engine/Classes/Components/ChildActorComponent.h
+++ b/Engine/Source/Runtime/Engine/Classes/Components/ChildActorComponent.h
@@ -73,6 +73,7 @@ public:
 	virtual void OnComponentCreated() override;
 	virtual void OnComponentDestroyed(bool bDestroyingHierarchy) override;
 	virtual void OnRegister() override;
+	virtual void OnUnregister() override;
 	virtual FActorComponentInstanceData* GetComponentInstanceData() const override;
 	//~ End ActorComponent Interface.
 
diff --git a/Engine/Source/Runtime/Engine/Classes/Components/PrimitiveComponent.h b/Engine/Source/Runtime/Engine/Classes/Components/PrimitiveComponent.h
--- a/Engine/Source/Runtime/Engine/Classes/Components/PrimitiveComponent.h
+++ b/Engine/Source/Runtime/Engine/Classes/Components/PrimitiveComponent.h
@@ -111,8 +111,8 @@ DECLARE_DYNAMIC_MULTICAST_DELEGATE_OneParam(FComponentSleepSignature, FName, Bon
 
 DECLARE_DYNAMIC_MULTICAST_DELEGATE_OneParam( FComponentBeginCursorOverSignature, UPrimitiveComponent*, TouchedComponent );
 DECLARE_DYNAMIC_MULTICAST_DELEGATE_OneParam( FComponentEndCursorOverSignature, UPrimitiveComponent*, TouchedComponent );
-DECLARE_DYNAMIC_MULTICAST_DELEGATE_OneParam( FComponentOnClickedSignature, UPrimitiveComponent*, TouchedComponent );
-DECLARE_DYNAMIC_MULTICAST_DELEGATE_OneParam( FComponentOnReleasedSignature, UPrimitiveComponent*, TouchedComponent );
+DECLARE_DYNAMIC_MULTICAST_DELEGATE_TwoParams( FComponentOnClickedSignature, UPrimitiveComponent*, TouchedComponent , FKey, ButtonPressed);
+DECLARE_DYNAMIC_MULTICAST_DELEGATE_TwoParams( FComponentOnReleasedSignature, UPrimitiveComponent*, TouchedComponent, FKey, ButtonReleased);
 DECLARE_DYNAMIC_MULTICAST_DELEGATE_TwoParams( FComponentOnInputTouchBeginSignature, ETouchIndex::Type, FingerIndex, UPrimitiveComponent*, TouchedComponent );
 DECLARE_DYNAMIC_MULTICAST_DELEGATE_TwoParams( FComponentOnInputTouchEndSignature, ETouchIndex::Type, FingerIndex, UPrimitiveComponent*, TouchedComponent );
 DECLARE_DYNAMIC_MULTICAST_DELEGATE_TwoParams( FComponentBeginTouchOverSignature, ETouchIndex::Type, FingerIndex, UPrimitiveComponent*, TouchedComponent );
@@ -1880,8 +1880,8 @@ public:
 
 	static void DispatchMouseOverEvents(UPrimitiveComponent* CurrentComponent, UPrimitiveComponent* NewComponent);
 	static void DispatchTouchOverEvents(ETouchIndex::Type FingerIndex, UPrimitiveComponent* CurrentComponent, UPrimitiveComponent* NewComponent);
-	void DispatchOnClicked();
-	void DispatchOnReleased();
+	void DispatchOnClicked(FKey ButtonClicked = EKeys::LeftMouseButton);
+	void DispatchOnReleased(FKey ButtonReleased = EKeys::LeftMouseButton);
 	void DispatchOnInputTouchBegin(const ETouchIndex::Type Key);
 	void DispatchOnInputTouchEnd(const ETouchIndex::Type Key);
 };
diff --git a/Engine/Source/Runtime/Engine/Classes/Components/TextRenderComponent.h b/Engine/Source/Runtime/Engine/Classes/Components/TextRenderComponent.h
--- a/Engine/Source/Runtime/Engine/Classes/Components/TextRenderComponent.h
+++ b/Engine/Source/Runtime/Engine/Classes/Components/TextRenderComponent.h
@@ -118,6 +118,10 @@ class ENGINE_API UTextRenderComponent : public UPrimitiveComponent
 	UFUNCTION(BlueprintCallable, Category="Rendering|Components|TextRender")
 	void SetHorizontalAlignment(EHorizTextAligment Value);
 
+	/** Change the vertical alignment and signal the primitives to be rebuilt */
+	UFUNCTION(BlueprintCallable, Category="Rendering|Components|TextRender")
+	void SetVerticalAlignment(EVerticalTextAligment Value);
+
 	/** Change the text render color and signal the primitives to be rebuilt */
 	UFUNCTION(BlueprintCallable, Category="Rendering|Components|TextRender")
 	void SetTextRenderColor(FColor Value);
diff --git a/Engine/Source/Runtime/Engine/Classes/Engine/Engine.h b/Engine/Source/Runtime/Engine/Classes/Engine/Engine.h
--- a/Engine/Source/Runtime/Engine/Classes/Engine/Engine.h
+++ b/Engine/Source/Runtime/Engine/Classes/Engine/Engine.h
@@ -1570,6 +1570,13 @@ public:
 	/** Called by internal engine systems after an actor has been moved to notify other subsystems */
 	void BroadcastOnActorMoved( AActor* Actor ) { OnActorMovedEvent.Broadcast( Actor ); }
 
+	/** Editor-only event triggered when any component transform is changed */
+	DECLARE_EVENT_TwoParams(UEngine, FOnComponentTransformChangedEvent, USceneComponent*, ETeleportType);
+	FOnComponentTransformChangedEvent& OnComponentTransformChanged() { return OnComponentTransformChangedEvent; }
+
+	/** Called by SceneComponent PropagateTransformUpdate to nofify of any component transform change */
+	void BroadcastOnComponentTransformChanged(USceneComponent* InComponent, ETeleportType InTeleport) { OnComponentTransformChangedEvent.Broadcast(InComponent, InTeleport); }
+
 	/** Editor-only event triggered when actors are being requested to be renamed */
 	DECLARE_EVENT_OneParam( UEngine, FLevelActorRequestRenameEvent, const AActor* );
 	FLevelActorRequestRenameEvent& OnLevelActorRequestRename() { return LevelActorRequestRenameEvent; }
@@ -2330,6 +2337,9 @@ private:
 
 	/** Broadcasts after an actor has been moved, rotated or scaled */
 	FOnActorMovedEvent		OnActorMovedEvent;
+
+	/** Broadcasts after a component has been moved, rotated or scaled */
+	FOnComponentTransformChangedEvent OnComponentTransformChangedEvent;
 	
 	/** Delegate broadcast after UEditorEngine::Tick has been called (or UGameEngine::Tick in standalone) */
 	FPostEditorTick PostEditorTickEvent;
diff --git a/Engine/Source/Runtime/Engine/Classes/Engine/SkeletalMeshSocket.h b/Engine/Source/Runtime/Engine/Classes/Engine/SkeletalMeshSocket.h
--- a/Engine/Source/Runtime/Engine/Classes/Engine/SkeletalMeshSocket.h
+++ b/Engine/Source/Runtime/Engine/Classes/Engine/SkeletalMeshSocket.h
@@ -85,6 +85,8 @@ class USkeletalMeshSocket : public UObject
 	 */
 	ENGINE_API bool AttachActor (class AActor* Actor, class USkeletalMeshComponent* SkelComp) const;
 
+	virtual void Serialize(FArchive& Ar) override;
+
 #if WITH_EDITOR
 	/** Broadcasts a notification whenever the socket property has changed. */
 	DECLARE_EVENT_TwoParams(USkeletalMeshSocket, FSocketChangedEvent, const class USkeletalMeshSocket*, const class UProperty*);
diff --git a/Engine/Source/Runtime/Engine/Classes/Engine/StaticMeshSocket.h b/Engine/Source/Runtime/Engine/Classes/Engine/StaticMeshSocket.h
--- a/Engine/Source/Runtime/Engine/Classes/Engine/StaticMeshSocket.h
+++ b/Engine/Source/Runtime/Engine/Classes/Engine/StaticMeshSocket.h
@@ -51,6 +51,8 @@ class UStaticMeshSocket : public UObject
 	 */
 	ENGINE_API bool AttachActor(AActor* Actor, class UStaticMeshComponent* MeshComp) const;
 
+	virtual void Serialize(FArchive& Ar) override;
+
 public:
 #if WITH_EDITOR
 	/** Broadcasts a notification whenever the socket property has changed. */
diff --git a/Engine/Source/Runtime/Engine/Classes/GameFramework/Actor.h b/Engine/Source/Runtime/Engine/Classes/GameFramework/Actor.h
--- a/Engine/Source/Runtime/Engine/Classes/GameFramework/Actor.h
+++ b/Engine/Source/Runtime/Engine/Classes/GameFramework/Actor.h
@@ -29,8 +29,8 @@ DECLARE_DYNAMIC_MULTICAST_DELEGATE_FourParams( FActorHitSignature, AActor*, Self
 
 DECLARE_DYNAMIC_MULTICAST_DELEGATE( FActorBeginCursorOverSignature );
 DECLARE_DYNAMIC_MULTICAST_DELEGATE( FActorEndCursorOverSignature );
-DECLARE_DYNAMIC_MULTICAST_DELEGATE( FActorOnClickedSignature );
-DECLARE_DYNAMIC_MULTICAST_DELEGATE( FActorOnReleasedSignature );
+DECLARE_DYNAMIC_MULTICAST_DELEGATE_OneParam( FActorOnClickedSignature, FKey, ButtonPressed );
+DECLARE_DYNAMIC_MULTICAST_DELEGATE_OneParam( FActorOnReleasedSignature, FKey, ButtonReleased );
 DECLARE_DYNAMIC_MULTICAST_DELEGATE_OneParam( FActorOnInputTouchBeginSignature, ETouchIndex::Type, FingerIndex );
 DECLARE_DYNAMIC_MULTICAST_DELEGATE_OneParam( FActorOnInputTouchEndSignature, ETouchIndex::Type, FingerIndex );
 DECLARE_DYNAMIC_MULTICAST_DELEGATE_OneParam( FActorBeginTouchOverSignature, ETouchIndex::Type, FingerIndex );
@@ -1256,16 +1256,16 @@ public:
 	void ReceiveActorEndCursorOver();
 
 	/** Event when this actor is clicked by the mouse when using the clickable interface. */
-	virtual void NotifyActorOnClicked();
+	virtual void NotifyActorOnClicked(FKey ButtonPressed = EKeys::LeftMouseButton);
 	/** Event when this actor is clicked by the mouse when using the clickable interface. */
 	UFUNCTION(BlueprintImplementableEvent, meta=(DisplayName = "ActorOnClicked"), Category="Mouse Input")
-	void ReceiveActorOnClicked();
+	void ReceiveActorOnClicked(FKey ButtonPressed = EKeys::LeftMouseButton);
 
 	/** Event when this actor is under the mouse when left mouse button is released while using the clickable interface. */
-	virtual void NotifyActorOnReleased();
+	virtual void NotifyActorOnReleased(FKey ButtonReleased = EKeys::LeftMouseButton);
 	/** Event when this actor is under the mouse when left mouse button is released while using the clickable interface. */
 	UFUNCTION(BlueprintImplementableEvent, meta=(DisplayName = "ActorOnReleased"), Category="Mouse Input")
-	void ReceiveActorOnReleased();
+	void ReceiveActorOnReleased(FKey ButtonReleased = EKeys::LeftMouseButton);
 
 	/** Event when this actor is touched when click events are enabled. */
 	virtual void NotifyActorOnInputTouchBegin(const ETouchIndex::Type FingerIndex);
diff --git a/Engine/Source/Runtime/Engine/Classes/GameFramework/PlayerController.h b/Engine/Source/Runtime/Engine/Classes/GameFramework/PlayerController.h
--- a/Engine/Source/Runtime/Engine/Classes/GameFramework/PlayerController.h
+++ b/Engine/Source/Runtime/Engine/Classes/GameFramework/PlayerController.h
@@ -330,6 +330,9 @@ class ENGINE_API APlayerController : public AController
 	UPROPERTY(EditAnywhere, BlueprintReadWrite, Category="Game|Feedback")
 	uint32 bForceFeedbackEnabled:1;
 
+	UPROPERTY(EditAnywhere, BlueprintReadWrite, Category=MouseInterface, meta=(EditCondition="bEnableClickEvents"))
+	TArray<FKey> ClickEventKeys;
+
 	UPROPERTY(EditAnywhere, BlueprintReadOnly, Category=MouseInterface)
 	TEnumAsByte<EMouseCursor::Type> DefaultMouseCursor;
 
diff --git a/Engine/Source/Runtime/Engine/Private/Actor.cpp b/Engine/Source/Runtime/Engine/Private/Actor.cpp
--- a/Engine/Source/Runtime/Engine/Private/Actor.cpp
+++ b/Engine/Source/Runtime/Engine/Private/Actor.cpp
@@ -1200,16 +1200,16 @@ void AActor::NotifyActorEndCursorOver()
 	ReceiveActorEndCursorOver();
 }
 
-void AActor::NotifyActorOnClicked()
+void AActor::NotifyActorOnClicked(FKey ButtonPressed)
 {
 	// call BP handler
-	ReceiveActorOnClicked();
+	ReceiveActorOnClicked(ButtonPressed);
 }
 
-void AActor::NotifyActorOnReleased()
+void AActor::NotifyActorOnReleased(FKey ButtonReleased)
 {
 	// call BP handler
-	ReceiveActorOnReleased();
+	ReceiveActorOnReleased(ButtonReleased);
 }
 
 void AActor::NotifyActorOnInputTouchBegin(const ETouchIndex::Type FingerIndex)
@@ -2595,7 +2595,7 @@ void AActor::FinishSpawning(const FTransform& UserTransform, bool bIsDefaultTran
 	{
 		bHasFinishedSpawning = true;
 
-		FTransform FinalRootComponentTransform = UserTransform;
+		FTransform FinalRootComponentTransform = (RootComponent ? RootComponent->ComponentToWorld : UserTransform);
 
 		// see if we need to adjust the transform (i.e. in deferred cases where the caller passes in a different transform here 
 		// than was passed in during the original SpawnActor call)
diff --git a/Engine/Source/Runtime/Engine/Private/Components/CharacterMovementComponent.cpp b/Engine/Source/Runtime/Engine/Private/Components/CharacterMovementComponent.cpp
--- a/Engine/Source/Runtime/Engine/Private/Components/CharacterMovementComponent.cpp
+++ b/Engine/Source/Runtime/Engine/Private/Components/CharacterMovementComponent.cpp
@@ -482,21 +482,19 @@ bool UCharacterMovementComponent::HasValidData() const
 	if (bIsValid)
 	{
 		// NaN-checking updates
-		if (FMath::IsNaN(GravityScale) || !FMath::IsFinite(GravityScale))
-		{
-			logOrEnsureNanError(TEXT("UCharacterMovementComponent::HasValidData detected NaN/INF in GravityScale! (%f)"), GravityScale);
-		}
 		if (Velocity.ContainsNaN())
 		{
-			logOrEnsureNanError(TEXT("UCharacterMovementComponent::HasValidData detected NaN/INF in Velocity! (%s)"), *Velocity.ToString());
+			logOrEnsureNanError(TEXT("UCharacterMovementComponent::HasValidData() detected NaN/INF for (%s) in Velocity:\n%s"), *GetPathNameSafe(this), *Velocity.ToString());
+			UCharacterMovementComponent* MutableThis = const_cast<UCharacterMovementComponent*>(this);
+			MutableThis->Velocity = FVector::ZeroVector;
 		}
 		if (!UpdatedComponent->GetComponentTransform().IsValid())
 		{
-			logOrEnsureNanError(TEXT("UCharacterMovementComponent::HasValidData detected NaN/INF in UpdatedComponent ComponentTransform!"));
+			logOrEnsureNanError(TEXT("UCharacterMovementComponent::HasValidData() detected NaN/INF for (%s) in UpdatedComponent->ComponentTransform:\n%s"), *GetPathNameSafe(this), *UpdatedComponent->GetComponentTransform().ToHumanReadableString());
 		}
 		if (UpdatedComponent->GetComponentRotation().ContainsNaN())
 		{
-			logOrEnsureNanError(TEXT("UCharacterMovementComponent::HasValidData detected NaN/INF in UpdatedComponent->GetComponentRotation()! %s"), *UpdatedComponent->GetComponentRotation().ToString());
+			logOrEnsureNanError(TEXT("UCharacterMovementComponent::HasValidData() detected NaN/INF for (%s) in UpdatedComponent->GetComponentRotation():\n%s"), *GetPathNameSafe(this), *UpdatedComponent->GetComponentRotation().ToString());
 		}
 	}
 #endif
@@ -1856,7 +1854,7 @@ void UCharacterMovementComponent::PerformMovement(float DeltaSeconds)
 		}
 
 		// NaN tracking
-		checkf(!Velocity.ContainsNaN(), TEXT("UCharacterMovementComponent::PerformMovement: Velocity contains NaN (%s: %s)\n%s"), *GetPathNameSafe(this), *GetPathNameSafe(GetOuter()), *Velocity.ToString());
+		checkCode(ensureMsgf(!Velocity.ContainsNaN(), TEXT("UCharacterMovementComponent::PerformMovement: Velocity contains NaN (%s)\n%s"), *GetPathNameSafe(this), *Velocity.ToString()));
 
 		// Clear jump input now, to allow movement events to trigger it for next update.
 		CharacterOwner->ClearJumpInput();
@@ -4007,7 +4005,7 @@ void UCharacterMovementComponent::PhysWalking(float deltaTime, int32 Iterations)
 		return;
 	}
 
-	checkf(!Velocity.ContainsNaN(), TEXT("PhysWalking: Velocity contains NaN before Iteration (%s: %s)\n%s"), *GetPathNameSafe(this), *GetPathNameSafe(GetOuter()), *Velocity.ToString());
+	checkCode(ensureMsgf(!Velocity.ContainsNaN(), TEXT("PhysWalking: Velocity contains NaN before Iteration (%s)\n%s"), *GetPathNameSafe(this), *Velocity.ToString()));
 	
 	bJustTeleported = false;
 	bool bCheckedFall = false;
@@ -4039,11 +4037,11 @@ void UCharacterMovementComponent::PhysWalking(float deltaTime, int32 Iterations)
 		if( !HasAnimRootMotion() && !CurrentRootMotion.HasOverrideVelocity() )
 		{
 			CalcVelocity(timeTick, GroundFriction, false, BrakingDecelerationWalking);
+			checkCode(ensureMsgf(!Velocity.ContainsNaN(), TEXT("PhysWalking: Velocity contains NaN after CalcVelocity (%s)\n%s"), *GetPathNameSafe(this), *Velocity.ToString()));
 		}
-		checkf(!Velocity.ContainsNaN(), TEXT("PhysWalking: Velocity contains NaN after CalcVelocity (%s: %s)\n%s"), *GetPathNameSafe(this), *GetPathNameSafe(GetOuter()), *Velocity.ToString());
-
+		
 		ApplyRootMotionToVelocity(timeTick);
-		checkf(!Velocity.ContainsNaN(), TEXT("PhysWalking: Velocity contains NaN after Root Motion application (%s: %s)\n%s"), *GetPathNameSafe(this), *GetPathNameSafe(GetOuter()), *Velocity.ToString());
+		checkCode(ensureMsgf(!Velocity.ContainsNaN(), TEXT("PhysWalking: Velocity contains NaN after Root Motion application (%s)\n%s"), *GetPathNameSafe(this), *Velocity.ToString()));
 
 		if( IsFalling() )
 		{
@@ -4229,14 +4227,14 @@ void UCharacterMovementComponent::PhysNavWalking(float deltaTime, int32 Iteratio
 
 	// Ensure velocity is horizontal.
 	MaintainHorizontalGroundVelocity();
-	checkf(!Velocity.ContainsNaN(), TEXT("PhysNavWalking: Velocity contains NaN before CalcVelocity (%s: %s)\n%s"), *GetPathNameSafe(this), *GetPathNameSafe(GetOuter()), *Velocity.ToString());
+	checkCode(ensureMsgf(!Velocity.ContainsNaN(), TEXT("PhysNavWalking: Velocity contains NaN before CalcVelocity (%s)\n%s"), *GetPathNameSafe(this), *Velocity.ToString()));
 
 	//bound acceleration
 	Acceleration.Z = 0.f;
 	if (!HasAnimRootMotion() && !CurrentRootMotion.HasOverrideVelocity())
 	{
 		CalcVelocity(deltaTime, GroundFriction, false, BrakingDecelerationWalking);
-		checkf(!Velocity.ContainsNaN(), TEXT("PhysNavWalking: Velocity contains NaN after CalcVelocity (%s: %s)\n%s"), *GetPathNameSafe(this), *GetPathNameSafe(GetOuter()), *Velocity.ToString());
+		checkCode(ensureMsgf(!Velocity.ContainsNaN(), TEXT("PhysNavWalking: Velocity contains NaN after CalcVelocity (%s)\n%s"), *GetPathNameSafe(this), *Velocity.ToString()));
 	}
 
 	ApplyRootMotionToVelocity(deltaTime);
@@ -4799,7 +4797,7 @@ FVector CharAngularVelocity(FRotator const& OldRot, FRotator const& NewRot, floa
 		DeltaQRot.ToAxisAndAngle(Axis, Angle);
 
 		RetAngVel = Axis * Angle * InvDeltaTime;
-		check(!RetAngVel.ContainsNaN());
+		RetAngVel.DiagnosticCheckNaN();
 	}
 
 	return RetAngVel;
@@ -4846,25 +4844,12 @@ void UCharacterMovementComponent::PhysicsRotation(float DeltaTime)
 	}
 
 	FRotator CurrentRotation = UpdatedComponent->GetComponentRotation(); // Normalized
+	CurrentRotation.DiagnosticCheckNaN(TEXT("CharacterMovementComponent::PhysicsRotation(): CurrentRotation"));
 
-#if ENABLE_NAN_DIAGNOSTIC
-	if (CurrentRotation.ContainsNaN())
-	{
-		logOrEnsureNanError(TEXT("CharacterMovementComponent::PhysicsRotation found NaN in CurrentRotation"));
-		CurrentRotation = FRotator::ZeroRotator;
-	}
-#endif
 	FRotator DeltaRot = GetDeltaRotation(DeltaTime);
+	DeltaRot.DiagnosticCheckNaN(TEXT("CharacterMovementComponent::PhysicsRotation(): GetDeltaRotation"));
 
-#if ENABLE_NAN_DIAGNOSTIC
-	if (DeltaRot.ContainsNaN())
-	{
-		logOrEnsureNanError(TEXT("CharacterMovementComponent::PhysicsRotation found NaN from GetDeltaRotation"));
-		DeltaRot = FRotator::ZeroRotator;
-	}
-#endif
 	FRotator DesiredRotation = CurrentRotation;
-
 	if (bOrientRotationToMovement)
 	{
 		DesiredRotation = ComputeOrientToMovementRotation(CurrentRotation, DeltaTime, DeltaRot);
@@ -4914,13 +4899,7 @@ void UCharacterMovementComponent::PhysicsRotation(float DeltaTime)
 		}
 
 		// Set the new rotation.
-#if ENABLE_NAN_DIAGNOSTIC
-		if (DesiredRotation.ContainsNaN())
-		{
-			logOrEnsureNanError(TEXT("CharacterMovementComponent::PhysicsRotation found NaN in DesiredRotation"));
-			DesiredRotation = FRotator::ZeroRotator;
-		}
-#endif
+		DesiredRotation.DiagnosticCheckNaN(TEXT("CharacterMovementComponent::PhysicsRotation(): DesiredRotation"));
 		MoveUpdatedComponent( FVector::ZeroVector, DesiredRotation, true );
 	}
 }
diff --git a/Engine/Source/Runtime/Engine/Private/Components/ChildActorComponent.cpp b/Engine/Source/Runtime/Engine/Private/Components/ChildActorComponent.cpp
--- a/Engine/Source/Runtime/Engine/Private/Components/ChildActorComponent.cpp
+++ b/Engine/Source/Runtime/Engine/Private/Components/ChildActorComponent.cpp
@@ -24,11 +24,16 @@ void UChildActorComponent::OnRegister()
 		else
 		{
 			ChildActorName = ChildActor->GetFName();
-			// attach new actor to this component
-			// we can't attach in CreateChildActor since it has intermediate Mobility set up
-			// causing spam with inconsistent mobility set up
-			// so moving Attach to happen in Register
-			ChildActor->AttachRootComponentTo(this, NAME_None, EAttachLocation::SnapToTargetIncludingScale);
+			
+			USceneComponent* ChildRoot = ChildActor->GetRootComponent();
+			if (ChildRoot && ChildRoot->GetAttachParent() != this)
+			{
+				// attach new actor to this component
+				// we can't attach in CreateChildActor since it has intermediate Mobility set up
+				// causing spam with inconsistent mobility set up
+				// so moving Attach to happen in Register
+				ChildRoot->AttachTo(this, NAME_None, EAttachLocation::SnapToTarget);
+			}
 		}
 	}
 	else if (ChildActorClass)
@@ -89,6 +94,13 @@ void UChildActorComponent::OnComponentDestroyed(bool bDestroyingHierarchy)
 {
 	Super::OnComponentDestroyed(bDestroyingHierarchy);
 
+	check(ChildActor == nullptr);
+}
+
+void UChildActorComponent::OnUnregister()
+{
+	Super::OnUnregister();
+
 	const UWorld* const MyWorld = GetWorld();
 	DestroyChildActor(MyWorld && !MyWorld->IsGameWorld());
 }
@@ -309,7 +321,7 @@ void UChildActorComponent::CreateChildActor()
 					const FComponentInstanceDataCache* ComponentInstanceData = (CachedInstanceData ? CachedInstanceData->ComponentInstanceData : nullptr);
 					ChildActor->FinishSpawning(ComponentToWorld, false, ComponentInstanceData);
 
-					ChildActor->AttachRootComponentTo(this, NAME_None, EAttachLocation::SnapToTargetIncludingScale);
+					ChildActor->AttachRootComponentTo(this, NAME_None, EAttachLocation::SnapToTarget);
 				}
 			}
 		}
@@ -326,10 +338,10 @@ void UChildActorComponent::CreateChildActor()
 void UChildActorComponent::DestroyChildActor(const bool bRequiresRename)
 {
 	// If we own an Actor, kill it now
-	if(ChildActor != nullptr && !GExitPurge)
+	if (ChildActor != nullptr && !GExitPurge)
 	{
 		// if still alive, destroy, otherwise just clear the pointer
-		if(!ChildActor->IsPendingKill())
+		if (!ChildActor->IsPendingKill())
 		{
 #if WITH_EDITOR
 			if (CachedInstanceData)
@@ -343,7 +355,7 @@ void UChildActorComponent::DestroyChildActor(const bool bRequiresRename)
 
 			UWorld* World = ChildActor->GetWorld();
 			// World may be nullptr during shutdown
-			if(World != nullptr)
+			if (World != nullptr)
 			{
 				UClass* ChildClass = ChildActor->GetClass();
 
@@ -361,7 +373,7 @@ void UChildActorComponent::DestroyChildActor(const bool bRequiresRename)
 				World->DestroyActor(ChildActor);
 			}
 		}
-
-		ChildActor = nullptr;
 	}
+
+	ChildActor = nullptr;
 }
diff --git a/Engine/Source/Runtime/Engine/Private/Components/PrimitiveComponent.cpp b/Engine/Source/Runtime/Engine/Private/Components/PrimitiveComponent.cpp
--- a/Engine/Source/Runtime/Engine/Private/Components/PrimitiveComponent.cpp
+++ b/Engine/Source/Runtime/Engine/Private/Components/PrimitiveComponent.cpp
@@ -2777,37 +2777,37 @@ void UPrimitiveComponent::DispatchTouchOverEvents(ETouchIndex::Type FingerIndex,
 	}
 }
 
-void UPrimitiveComponent::DispatchOnClicked()
+void UPrimitiveComponent::DispatchOnClicked(FKey ButtonPressed)
 {
 	if (IsActorValidToNotify(GetOwner()))
 	{
-		GetOwner()->NotifyActorOnClicked();
+		GetOwner()->NotifyActorOnClicked(ButtonPressed);
 		if (IsActorValidToNotify(GetOwner()))
 		{
-			GetOwner()->OnClicked.Broadcast();
+			GetOwner()->OnClicked.Broadcast(ButtonPressed);
 		}
 	}
 
 	if (!IsPendingKill())
 	{
-		OnClicked.Broadcast(this);
+		OnClicked.Broadcast(this, ButtonPressed);
 	}
 }
 
-void UPrimitiveComponent::DispatchOnReleased()
+void UPrimitiveComponent::DispatchOnReleased(FKey ButtonReleased)
 {
 	if (IsActorValidToNotify(GetOwner()))
 	{
-		GetOwner()->NotifyActorOnReleased();
+		GetOwner()->NotifyActorOnReleased(ButtonReleased);
 		if (IsActorValidToNotify(GetOwner()))
 		{
-			GetOwner()->OnReleased.Broadcast();
+			GetOwner()->OnReleased.Broadcast(ButtonReleased);
 		}
 	}
 
 	if (!IsPendingKill())
 	{
-		OnReleased.Broadcast(this);
+		OnReleased.Broadcast(this, ButtonReleased);
 	}
 }
 
diff --git a/Engine/Source/Runtime/Engine/Private/Components/SceneComponent.cpp b/Engine/Source/Runtime/Engine/Private/Components/SceneComponent.cpp
--- a/Engine/Source/Runtime/Engine/Private/Components/SceneComponent.cpp
+++ b/Engine/Source/Runtime/Engine/Private/Components/SceneComponent.cpp
@@ -497,6 +497,14 @@ void USceneComponent::PropagateTransformUpdate(bool bTransformChanged, bool bSki
 			}
 		}
 
+#if WITH_EDITOR
+		// Notify the editor of transformation update
+		if (!IsTemplate())
+		{
+			GEngine->BroadcastOnComponentTransformChanged(this, Teleport);
+		}
+#endif // WITH_EDITOR
+
 		// Refresh navigation
 		if (bNavigationRelevant && bRegistered)
 		{
diff --git a/Engine/Source/Runtime/Engine/Private/Components/TextRenderComponent.cpp b/Engine/Source/Runtime/Engine/Private/Components/TextRenderComponent.cpp
--- a/Engine/Source/Runtime/Engine/Private/Components/TextRenderComponent.cpp
+++ b/Engine/Source/Runtime/Engine/Private/Components/TextRenderComponent.cpp
@@ -941,6 +941,12 @@ void UTextRenderComponent::SetHorizontalAlignment(EHorizTextAligment Value)
 	MarkRenderStateDirty();	
 }
 
+void UTextRenderComponent::SetVerticalAlignment(EVerticalTextAligment Value)
+{
+	VerticalAlignment = Value;
+	MarkRenderStateDirty();
+}
+
 void UTextRenderComponent::SetTextRenderColor(FColor Value)
 {
 	TextRenderColor = Value;
diff --git a/Engine/Source/Runtime/Engine/Private/Distributions.cpp b/Engine/Source/Runtime/Engine/Private/Distributions.cpp
--- a/Engine/Source/Runtime/Engine/Private/Distributions.cpp
+++ b/Engine/Source/Runtime/Engine/Private/Distributions.cpp
@@ -1222,7 +1222,7 @@ void FRawDistributionVector::GetOutRange(float& MinOut, float& MaxOut)
 
 void FRawDistributionVector::GetRange(FVector& MinOut, FVector& MaxOut)
 {
-	if (!HasLookupTable() && Distribution)
+	if (/*!HasLookupTable() &&*/ Distribution)
 	{
 		check(Distribution);
 		Distribution->GetRange(MinOut, MaxOut);
diff --git a/Engine/Source/Runtime/Engine/Private/LevelActor.cpp b/Engine/Source/Runtime/Engine/Private/LevelActor.cpp
--- a/Engine/Source/Runtime/Engine/Private/LevelActor.cpp
+++ b/Engine/Source/Runtime/Engine/Private/LevelActor.cpp
@@ -718,10 +718,12 @@ bool UWorld::FindTeleportSpot(AActor* TestActor, FVector& TestLocation, FRotator
 	// now try just XY
 	if (!FMath::IsNearlyZero(Adjust.X) || !FMath::IsNearlyZero(Adjust.Y))
 	{
-		for (int i = 0; i < 8; ++i)
+		const FVector OriginalTestLocation = TestLocation;
+		// If initially spawning allow testing a few permutations (though this needs improvement).
+		// During play only test the first adjustment, permuting axes could put the location on other sides of geometry.
+		const int32 Iterations = (TestActor->HasActorBegunPlay() ? 1 : 8);
+		for (int i = 0; i < Iterations; ++i)
 		{
-			const FVector OriginalTestLocation = TestLocation;
-
 			TestLocation.X += (i < 4 ? Adjust.X : Adjust.Y) * (i % 2 == 0 ? 1 : -1);
 			TestLocation.Y += (i < 4 ? Adjust.Y : Adjust.X) * (i % 4 < 2 ? 1 : -1);
 			if (!EncroachingBlockingGeometry(TestActor, TestLocation, TestRotation, &Adjust))
@@ -864,11 +866,14 @@ static bool ComponentEncroachesBlockingGeometry_WithAdjustment(UWorld const* Wor
 					else
 					{
 						UE_LOG(LogPhysics, Log, TEXT("OverlapTest says we are overlapping, yet MTD says we're not. Something is wrong"));
+						// It's not safe to use a partial result, that could push us out to an invalid location (like the other side of a wall).
+						OutProposedAdjustment = FVector::ZeroVector;
+						return true;
 					}
 
 					// #hack: sometimes for boxes, physx returns a 0 MTD even though it reports a contact (returns true)
 					// to get around this, let's go ahead and test again with the epsilon-shrunken collision shape to see if we're really in 
-					// the clear.  if so, we'll say we have no contact (despite what OverlapMultiByChannel said -- it uses a different algorithm)
+					// the clear.
 					if (bSuccess && FMath::IsNearlyZero(MTDResult.Distance))
 					{
 						FCollisionShape const ShrunkenCollisionShape = PrimComp->GetCollisionShape(-Epsilon);
@@ -879,8 +884,10 @@ static bool ComponentEncroachesBlockingGeometry_WithAdjustment(UWorld const* Wor
 						}
 						else
 						{
-							// let's call this "no contact" and be done
-							return false;
+							UE_LOG(LogPhysics, Log, TEXT("OverlapTest says we are overlapping, yet MTD says we're not (with smaller shape). Something is wrong"));
+							// It's not safe to use a partial result, that could push us out to an invalid location.
+							OutProposedAdjustment = FVector::ZeroVector;
+							return true;
 						}
 					}
 				}
diff --git a/Engine/Source/Runtime/Engine/Private/LevelBounds.cpp b/Engine/Source/Runtime/Engine/Private/LevelBounds.cpp
--- a/Engine/Source/Runtime/Engine/Private/LevelBounds.cpp
+++ b/Engine/Source/Runtime/Engine/Private/LevelBounds.cpp
@@ -39,6 +39,7 @@ void ALevelBounds::PostLoad()
 
 FBox ALevelBounds::GetComponentsBoundingBox(bool bNonColliding) const
 {
+	checkf(RootComponent != nullptr, TEXT("LevelBounds actor with null root component: %s"), *GetPathNameSafe(this));
 	FVector BoundsCenter = RootComponent->GetComponentLocation();
 	FVector BoundsExtent = RootComponent->ComponentToWorld.GetScale3D() * 0.5f;
 	return FBox(BoundsCenter - BoundsExtent, 
diff --git a/Engine/Source/Runtime/Engine/Private/Particles/ParticleComponents.cpp b/Engine/Source/Runtime/Engine/Private/Particles/ParticleComponents.cpp
--- a/Engine/Source/Runtime/Engine/Private/Particles/ParticleComponents.cpp
+++ b/Engine/Source/Runtime/Engine/Private/Particles/ParticleComponents.cpp
@@ -279,6 +279,12 @@ void UParticleLODLevel::PostLoad()
 	{
 		RequiredModule->ConditionalPostLoad();
 	}
+
+	for(UParticleModule* ParticleModule : Modules)
+	{
+		ParticleModule->ConditionalPostLoad();
+	}
+
 	// shouldn't ever set another UObjects serialized variable in post load
 	// this causes determinisitc cooking issues due to load order being different
 	/*if (RequiredModule)
diff --git a/Engine/Source/Runtime/Engine/Private/PhysicsEngine/BodyInstance.cpp b/Engine/Source/Runtime/Engine/Private/PhysicsEngine/BodyInstance.cpp
--- a/Engine/Source/Runtime/Engine/Private/PhysicsEngine/BodyInstance.cpp
+++ b/Engine/Source/Runtime/Engine/Private/PhysicsEngine/BodyInstance.cpp
@@ -4459,7 +4459,7 @@ bool FBodyInstance::OverlapPhysX_AssumesLocked(const PxGeometry& PGeom, const Px
 				if(OutMTD)
 				{
 					OutMTD->Direction = P2UVector(POutDirection);
-					OutMTD->Distance = OutDistance;
+					OutMTD->Distance = FMath::Abs(OutDistance);
 				}
 					
 				return true;
diff --git a/Engine/Source/Runtime/Engine/Private/PlayerController.cpp b/Engine/Source/Runtime/Engine/Private/PlayerController.cpp
--- a/Engine/Source/Runtime/Engine/Private/PlayerController.cpp
+++ b/Engine/Source/Runtime/Engine/Private/PlayerController.cpp
@@ -89,6 +89,8 @@ APlayerController::APlayerController(const FObjectInitializer& ObjectInitializer
 	bIsPlayerController = true;
 	bIsLocalPlayerController = false;
 
+	ClickEventKeys.Add(EKeys::LeftMouseButton);
+
 	if (RootComponent)
 	{
 		// We want to drive rotation with ControlRotation regardless of attachment state.
@@ -2066,9 +2068,7 @@ bool APlayerController::InputKey(FKey Key, EInputEvent EventType, float AmountDe
 	if (PlayerInput)
 	{
 		bResult = PlayerInput->InputKey(Key, EventType, AmountDepressed, bGamepad);
-
-		// TODO: Allow click key(s?) to be defined
-		if (bEnableClickEvents && Key == EKeys::LeftMouseButton)
+		if (bEnableClickEvents && (ClickEventKeys.Contains(Key) || ClickEventKeys.Contains(EKeys::AnyKey)))
 		{
 			FVector2D MousePosition;
 			UGameViewportClient* ViewportClient = CastChecked<ULocalPlayer>(Player)->ViewportClient;
@@ -2102,11 +2102,11 @@ bool APlayerController::InputKey(FKey Key, EInputEvent EventType, float AmountDe
 					{
 					case IE_Pressed:
 					case IE_DoubleClick:
-						ClickedPrimitive->DispatchOnClicked();
+						ClickedPrimitive->DispatchOnClicked(Key);
 						break;
 
 					case IE_Released:
-						ClickedPrimitive->DispatchOnReleased();
+						ClickedPrimitive->DispatchOnReleased(Key);
 						break;
 
 					case IE_Axis:
diff --git a/Engine/Source/Runtime/Engine/Private/SkeletalMesh.cpp b/Engine/Source/Runtime/Engine/Private/SkeletalMesh.cpp
--- a/Engine/Source/Runtime/Engine/Private/SkeletalMesh.cpp
+++ b/Engine/Source/Runtime/Engine/Private/SkeletalMesh.cpp
@@ -43,6 +43,8 @@
 #include "EditorFramework/AssetImportData.h"
 #include "Engine/SkeletalMeshSocket.h"
 #include "Components/BrushComponent.h"
+#include "FrameworkObjectVersion.h"
+
 #define LOCTEXT_NAMESPACE "SkeltalMesh"
 
 DEFINE_LOG_CATEGORY(LogSkeletalMesh);
@@ -3801,7 +3803,7 @@ bool USkeletalMeshSocket::GetSocketMatrix(FMatrix& OutMatrix, const class USkele
 	if(BoneIndex != INDEX_NONE)
 	{
 		FMatrix BoneMatrix = SkelComp->GetBoneMatrix(BoneIndex);
-		FRotationTranslationMatrix RelSocketMatrix( RelativeRotation, RelativeLocation );
+		FScaleRotationTranslationMatrix RelSocketMatrix( RelativeScale, RelativeRotation, RelativeLocation );
 		OutMatrix = RelSocketMatrix * BoneMatrix;
 		return true;
 	}
@@ -3822,7 +3824,7 @@ FTransform USkeletalMeshSocket::GetSocketTransform(const class USkeletalMeshComp
 	if(BoneIndex != INDEX_NONE)
 	{
 		FTransform BoneTM = SkelComp->GetBoneTransform(BoneIndex);
-		FTransform RelSocketTM( RelativeRotation, RelativeLocation );
+		FTransform RelSocketTM( RelativeRotation, RelativeLocation, RelativeScale );
 		OutTM = RelSocketTM * BoneTM;
 	}
 
@@ -3835,7 +3837,7 @@ bool USkeletalMeshSocket::GetSocketMatrixWithOffset(FMatrix& OutMatrix, class US
 	if(BoneIndex != INDEX_NONE)
 	{
 		FMatrix BoneMatrix = SkelComp->GetBoneMatrix(BoneIndex);
-		FRotationTranslationMatrix RelSocketMatrix(RelativeRotation, RelativeLocation);
+		FScaleRotationTranslationMatrix RelSocketMatrix(RelativeScale, RelativeRotation, RelativeLocation);
 		FRotationTranslationMatrix RelOffsetMatrix(InRotation, InOffset);
 		OutMatrix = RelOffsetMatrix * RelSocketMatrix * BoneMatrix;
 		return true;
@@ -3851,7 +3853,7 @@ bool USkeletalMeshSocket::GetSocketPositionWithOffset(FVector& OutPosition, clas
 	if(BoneIndex != INDEX_NONE)
 	{
 		FMatrix BoneMatrix = SkelComp->GetBoneMatrix(BoneIndex);
-		FRotationTranslationMatrix RelSocketMatrix(RelativeRotation, RelativeLocation);
+		FScaleRotationTranslationMatrix RelSocketMatrix(RelativeScale, RelativeRotation, RelativeLocation);
 		FRotationTranslationMatrix RelOffsetMatrix(InRotation, InOffset);
 		FMatrix SocketMatrix = RelOffsetMatrix * RelSocketMatrix * BoneMatrix;
 		OutPosition = SocketMatrix.GetOrigin();
@@ -3925,6 +3927,21 @@ void USkeletalMeshSocket::CopyFrom(const class USkeletalMeshSocket* OtherSocket)
 
 #endif
 
+void USkeletalMeshSocket::Serialize(FArchive& Ar)
+{
+	Super::Serialize(Ar);
+
+	Ar.UsingCustomVersion(FFrameworkObjectVersion::GUID);
+
+	if(Ar.CustomVer(FFrameworkObjectVersion::GUID) < FFrameworkObjectVersion::MeshSocketScaleUtilization)
+	{
+		// Set the relative scale to 1.0. As it was not used before this should allow existing data
+		// to work as expected.
+		RelativeScale = FVector(1.0f, 1.0f, 1.0f);
+	}
+}
+
+
 /*-----------------------------------------------------------------------------
 	ASkeletalMeshActor
 -----------------------------------------------------------------------------*/
diff --git a/Engine/Source/Runtime/Engine/Private/StaticMesh.cpp b/Engine/Source/Runtime/Engine/Private/StaticMesh.cpp
--- a/Engine/Source/Runtime/Engine/Private/StaticMesh.cpp
+++ b/Engine/Source/Runtime/Engine/Private/StaticMesh.cpp
@@ -2983,14 +2983,14 @@ UStaticMeshSocket::UStaticMeshSocket(const FObjectInitializer& ObjectInitializer
 bool UStaticMeshSocket::GetSocketMatrix(FMatrix& OutMatrix, UStaticMeshComponent const* MeshComp) const
 {
 	check( MeshComp );
-	OutMatrix = FRotationTranslationMatrix( RelativeRotation, RelativeLocation ) * MeshComp->ComponentToWorld.ToMatrixWithScale();
+	OutMatrix = FScaleRotationTranslationMatrix( RelativeScale, RelativeRotation, RelativeLocation ) * MeshComp->ComponentToWorld.ToMatrixWithScale();
 	return true;
 }
 
 bool UStaticMeshSocket::GetSocketTransform(FTransform& OutTransform, class UStaticMeshComponent const* MeshComp) const
 {
 	check( MeshComp );
-	OutTransform = FTransform(RelativeRotation, RelativeLocation) * MeshComp->ComponentToWorld;
+	OutTransform = FTransform(RelativeRotation, RelativeLocation, RelativeScale) * MeshComp->ComponentToWorld;
 	return true;
 }
 
@@ -3035,3 +3035,17 @@ void UStaticMeshSocket::PostEditChangeProperty( FPropertyChangedEvent& PropertyC
 	}
 }
 #endif
+
+void UStaticMeshSocket::Serialize(FArchive& Ar)
+{
+	Super::Serialize(Ar);
+
+	Ar.UsingCustomVersion(FFrameworkObjectVersion::GUID);
+
+	if(Ar.CustomVer(FFrameworkObjectVersion::GUID) < FFrameworkObjectVersion::MeshSocketScaleUtilization)
+	{
+		// Set the relative scale to 1.0. As it was not used before this should allow existing data
+		// to work as expected.
+		RelativeScale = FVector(1.0f, 1.0f, 1.0f);
+	}
+}
\ No newline at end of file
diff --git a/Engine/Source/Runtime/InputDevice/Public/IInputDevice.h b/Engine/Source/Runtime/InputDevice/Public/IInputDevice.h
--- a/Engine/Source/Runtime/InputDevice/Public/IInputDevice.h
+++ b/Engine/Source/Runtime/InputDevice/Public/IInputDevice.h
@@ -33,5 +33,10 @@ public:
 	{
 		return nullptr;
 	}
+
+	virtual bool IsGamepadAttached() const
+	{
+		return false;
+	}
 };
 
diff --git a/Engine/Source/Runtime/Slate/Private/Framework/Application/SlateApplication.cpp b/Engine/Source/Runtime/Slate/Private/Framework/Application/SlateApplication.cpp
--- a/Engine/Source/Runtime/Slate/Private/Framework/Application/SlateApplication.cpp
+++ b/Engine/Source/Runtime/Slate/Private/Framework/Application/SlateApplication.cpp
@@ -4322,8 +4322,12 @@ bool FSlateApplication::IsFakingTouchEvents() const
 	return bIsFakingTouch || bIsGameFakingTouch;
 }
 
+bool FSlateApplication::OnMouseDown(const TSharedPtr< FGenericWindow >& PlatformWindow, const EMouseButtons::Type Button)
+{
+	return OnMouseDown(PlatformWindow, Button, GetCursorPos());
+}
 
-bool FSlateApplication::OnMouseDown( const TSharedPtr< FGenericWindow >& PlatformWindow, const EMouseButtons::Type Button )
+bool FSlateApplication::OnMouseDown( const TSharedPtr< FGenericWindow >& PlatformWindow, const EMouseButtons::Type Button, const FVector2D CursorPos )
 {
 	// convert to touch event if we are faking it	
 	if (bIsFakingTouch || bIsGameFakingTouch)
@@ -4336,7 +4340,7 @@ bool FSlateApplication::OnMouseDown( const TSharedPtr< FGenericWindow >& Platfor
 
 	FPointerEvent MouseEvent(
 		CursorPointerIndex,
-		GetCursorPos(),
+		CursorPos,
 		GetLastCursorPos(),
 		PressedMouseButtons,
 		Key,
@@ -4853,6 +4857,11 @@ bool FSlateApplication::RoutePointerMoveEvent(const FWidgetPath& WidgetsUnderPoi
 }
 
 bool FSlateApplication::OnMouseDoubleClick( const TSharedPtr< FGenericWindow >& PlatformWindow, const EMouseButtons::Type Button )
+{
+	return OnMouseDoubleClick(PlatformWindow, Button, GetCursorPos());
+}
+
+bool FSlateApplication::OnMouseDoubleClick( const TSharedPtr< FGenericWindow >& PlatformWindow, const EMouseButtons::Type Button, const FVector2D CursorPos )
 {
 	if (bIsFakingTouch || bIsGameFakingTouch)
 	{
@@ -4864,7 +4873,7 @@ bool FSlateApplication::OnMouseDoubleClick( const TSharedPtr< FGenericWindow >&
 
 	FPointerEvent MouseEvent(
 		CursorPointerIndex,
-		GetCursorPos(),
+		CursorPos,
 		GetLastCursorPos(),
 		PressedMouseButtons,
 		Key,
@@ -4911,6 +4920,11 @@ bool FSlateApplication::ProcessMouseButtonDoubleClickEvent( const TSharedPtr< FG
 }
 
 bool FSlateApplication::OnMouseUp( const EMouseButtons::Type Button )
+{
+	return OnMouseUp(Button, GetCursorPos());
+}
+
+bool FSlateApplication::OnMouseUp( const EMouseButtons::Type Button, const FVector2D CursorPos )
 {
 	// convert to touch event if we are faking it	
 	if (bIsFakingTouch || bIsGameFakingTouch)
@@ -4923,7 +4937,7 @@ bool FSlateApplication::OnMouseUp( const EMouseButtons::Type Button )
 
 	FPointerEvent MouseEvent(
 		CursorPointerIndex,
-		GetCursorPos(),
+		CursorPos,
 		GetLastCursorPos(),
 		PressedMouseButtons,
 		Key,
@@ -4970,12 +4984,15 @@ bool FSlateApplication::ProcessMouseButtonUpEvent( FPointerEvent& MouseEvent )
 
 bool FSlateApplication::OnMouseWheel( const float Delta )
 {
-	const FVector2D CurrentCursorPosition = GetCursorPos();
+	return OnMouseWheel(Delta, GetCursorPos());
+}
 
+bool FSlateApplication::OnMouseWheel( const float Delta, const FVector2D CursorPos )
+{
 	FPointerEvent MouseWheelEvent(
 		CursorPointerIndex,
-		CurrentCursorPosition,
-		CurrentCursorPosition,
+		CursorPos,
+		CursorPos,
 		PressedMouseButtons,
 		EKeys::Invalid,
 		Delta,
diff --git a/Engine/Source/Runtime/Slate/Public/Framework/Application/SlateApplication.h b/Engine/Source/Runtime/Slate/Public/Framework/Application/SlateApplication.h
--- a/Engine/Source/Runtime/Slate/Public/Framework/Application/SlateApplication.h
+++ b/Engine/Source/Runtime/Slate/Public/Framework/Application/SlateApplication.h
@@ -595,6 +595,11 @@ public:
 	 */
 	bool IsMouseAttached() const { return PlatformApplication.IsValid() ? PlatformApplication->IsMouseAttached() : false; }
 
+	/**
+	 * @return True if there is a gamepad attached
+	 */
+	bool IsGamepadAttached() const { return PlatformApplication.IsValid() ? PlatformApplication->IsGamepadAttached() : false; }
+
 	/**
 	 * Sets the widget reflector.
 	 *
@@ -1144,9 +1149,13 @@ public:
 	virtual bool OnKeyDown( const int32 KeyCode, const uint32 CharacterCode, const bool IsRepeat ) override;
 	virtual bool OnKeyUp( const int32 KeyCode, const uint32 CharacterCode, const bool IsRepeat ) override;
 	virtual bool OnMouseDown( const TSharedPtr< FGenericWindow >& PlatformWindow, const EMouseButtons::Type Button ) override;
+	virtual bool OnMouseDown( const TSharedPtr< FGenericWindow >& PlatformWindow, const EMouseButtons::Type Button, const FVector2D CursorPos ) override;
 	virtual bool OnMouseUp( const EMouseButtons::Type Button ) override;
+	virtual bool OnMouseUp( const EMouseButtons::Type Button, const FVector2D CursorPos ) override;
 	virtual bool OnMouseDoubleClick( const TSharedPtr< FGenericWindow >& PlatformWindow, const EMouseButtons::Type Button ) override;
+	virtual bool OnMouseDoubleClick( const TSharedPtr< FGenericWindow >& PlatformWindow, const EMouseButtons::Type Button, const FVector2D CursorPos ) override;
 	virtual bool OnMouseWheel( const float Delta ) override;
+	virtual bool OnMouseWheel( const float Delta, const FVector2D CursorPos ) override;
 	virtual bool OnMouseMove() override;
 	virtual bool OnRawMouseMove( const int32 X, const int32 Y ) override;
 	virtual bool OnCursorSet() override;
