Copying //UE4/Dev-Rendering to Dev-Main (//UE4/Dev-Main)

==========================
MAJOR FEATURES + CHANGES
==========================

Change 2890427 on 2016/03/02 by [REDACTED_PERSON]

	Fixing global vector fields not working at all with fixed bounds particle systems.

Change 2890689 on 2016/03/02 by [REDACTED_PERSON]

	Noise buffers missing in (Slate) UI materials

Change 2892556 on 2016/03/03 by [REDACTED_PERSON]

	UE4 - Added the hipri and background threads to SetThreadAffinity

Change 2898686 on 2016/03/08 by [REDACTED_PERSON]

	PS4 - Reduced malloc spam on submission queues.

Change 2899191 on 2016/03/08 by [REDACTED_PERSON]

	Adding HasCompleted as a possible condition for auto-deactivate, so we don't restrict this purely to burst only emitters

	#hira UE-18899

Change 2899587 on 2016/03/08 by [REDACTED_PERSON]

	Added missing default value for FGBufferData.SelectiveOutputMask

Change 2901654 on 2016/03/09 by [REDACTED_PERSON]

	Opacity mask materials were producing incorrect shadows.  The pixel shader computation of world coordinates was incorrect.  Instead we can simply use the interpolated position from the vertex shader.

Change 2903496 on 2016/03/10 by [REDACTED_PERSON]

	PS4 - Added Cycles64 and optimized the PS4 GetCycles to use a shift instead of a divide.

Change 2903508 on 2016/03/10 by [REDACTED_PERSON]

	UE4 - Tweaked out MallocBinned2. Removed several recyclers in the engine since MB2 is faster. Reworked the container slack policies to use the allocator granularity.

Change 2905727 on 2016/03/11 by [REDACTED_PERSON]

	Merging as an Edit.

	//UE4/Private-GDC16-Character/Engine/Source/Runtime/Engine/Private/ShaderCompiler/ShaderCompiler.cpp

	to //UE4/Dev-Rendering/Engine/Source/Runtime/Engine/Private/ShaderCompiler/ShaderCompiler.cpp

Change 2910327 on 2016/03/15 by [REDACTED_PERSON]

	Re-added support for vs2012 project generation

diff --git a/Engine/Shaders/DeferredDecal.usf b/Engine/Shaders/DeferredDecal.usf
--- a/Engine/Shaders/DeferredDecal.usf
+++ b/Engine/Shaders/DeferredDecal.usf
@@ -409,6 +409,7 @@ void MainPS(
 	Data.PrecomputedShadowFactors = 0;
 	Data.GBufferAO = 1;
 	Data.ShadingModelID = SHADINGMODELID_DEFAULT_LIT;
+	Data.SelectiveOutputMask = 0;
 	Data.PerObjectGBufferData = 1;
 
 	#if DECAL_RENDERSTAGE == 0
diff --git a/Engine/Shaders/ShadowDepthPixelShader.usf b/Engine/Shaders/ShadowDepthPixelShader.usf
--- a/Engine/Shaders/ShadowDepthPixelShader.usf
+++ b/Engine/Shaders/ShadowDepthPixelShader.usf
@@ -61,9 +61,13 @@ void Main(
 	FMaterialPixelParameters MaterialParameters = GetMaterialPixelParameters(Inputs.FactoryInterpolants, SvPosition);
 	FPixelMaterialInputs PixelMaterialInputs;
 	#if INTERPOLATE_POSITION
-		{
-			float4 ScreenPosition = SvPositionToScreenPosition(SvPosition);
-			CalcMaterialParametersEx(MaterialParameters, PixelMaterialInputs, SvPosition, ScreenPosition, 1, SvPositionToResolvedTranslatedWorld(SvPosition), SvPosition.xyz);
+	    {
+			float4 ScreenPosition  = SvPositionToScreenPosition(SvPosition);  
+			
+			float3 TranslatedWorld = Inputs.PixelPosition.xyz;
+
+			CalcMaterialParametersEx(MaterialParameters, PixelMaterialInputs, SvPosition, ScreenPosition, 1, TranslatedWorld, TranslatedWorld);
+
 		}
 	#else
 		CalcMaterialParameters(MaterialParameters, PixelMaterialInputs, SvPosition, 1);
diff --git a/Engine/Source/Developer/ShaderCompilerCommon/Private/HlslUtils.h b/Engine/Source/Developer/ShaderCompilerCommon/Private/HlslUtils.h
--- a/Engine/Source/Developer/ShaderCompilerCommon/Private/HlslUtils.h
+++ b/Engine/Source/Developer/ShaderCompilerCommon/Private/HlslUtils.h
@@ -145,9 +145,17 @@ namespace CrossCompiler
 					}
 				}
 			}
-			int32 CalculateSlack(int32 NumElements, int32 NumAllocatedElements, int32 NumBytesPerElement) const
+			int32 CalculateSlackReserve(int32 NumElements, int32 NumBytesPerElement) const
 			{
-				return DefaultCalculateSlack(NumElements, NumAllocatedElements, NumBytesPerElement);
+				return DefaultCalculateSlackReserve(NumElements, NumBytesPerElement, false);
+			}
+			int32 CalculateSlackShrink(int32 NumElements, int32 NumAllocatedElements, int32 NumBytesPerElement) const
+			{
+				return DefaultCalculateSlackShrink(NumElements, NumAllocatedElements, NumBytesPerElement, false);
+			}
+			int32 CalculateSlackGrow(int32 NumElements, int32 NumAllocatedElements, int32 NumBytesPerElement) const
+			{
+				return DefaultCalculateSlackGrow(NumElements, NumAllocatedElements, NumBytesPerElement, false);
 			}
 
 			int32 GetAllocatedSize(int32 NumAllocatedElements, int32 NumBytesPerElement) const
@@ -155,6 +163,11 @@ namespace CrossCompiler
 				return NumAllocatedElements * NumBytesPerElement;
 			}
 
+			bool HasAllocation()
+			{
+				return !!Data;
+			}
+
 			FLinearAllocator* LinearAllocator;
 
 		private:
diff --git a/Engine/Source/Programs/UnrealBuildTool/System/VCProject.cs b/Engine/Source/Programs/UnrealBuildTool/System/VCProject.cs
--- a/Engine/Source/Programs/UnrealBuildTool/System/VCProject.cs
+++ b/Engine/Source/Programs/UnrealBuildTool/System/VCProject.cs
@@ -1099,6 +1099,10 @@ namespace UnrealBuildTool
 
 					// @todo UWP: For the MS toolchains, if an override was set for project generation, push that into the build strings to override the build toolchain as well
 					string BuildToolOverride = "";
+                    if (UnrealBuildTool.CommandLineContains("-2012"))
+                    {
+                        BuildToolOverride = " -2012";
+                    }
 					if (UnrealBuildTool.CommandLineContains("-2013"))
 					{
 						BuildToolOverride = " -2013";
diff --git a/Engine/Source/Programs/UnrealBuildTool/System/VCProjectFileGenerator.cs b/Engine/Source/Programs/UnrealBuildTool/System/VCProjectFileGenerator.cs
--- a/Engine/Source/Programs/UnrealBuildTool/System/VCProjectFileGenerator.cs
+++ b/Engine/Source/Programs/UnrealBuildTool/System/VCProjectFileGenerator.cs
@@ -120,7 +120,8 @@ namespace UnrealBuildTool
 
 
 		public enum VCProjectFileFormat
-		{
+        {
+            VisualStudio2012,
 			VisualStudio2013,
 			VisualStudio2015,
 		}
@@ -131,7 +132,9 @@ namespace UnrealBuildTool
 			get
 			{
 				switch (ProjectFileFormat)
-				{
+                {
+                    case VCProjectFileFormat.VisualStudio2012:
+                        return "4.0";
 					case VCProjectFileFormat.VisualStudio2013:
 						return "12.0";
 					case VCProjectFileFormat.VisualStudio2015:
@@ -146,13 +149,15 @@ namespace UnrealBuildTool
 		{
 			get
 			{
-				switch (WindowsPlatform.Compiler)
-				{
-					case WindowsCompiler.VisualStudio2013:
-						return "v120";
-					case WindowsCompiler.VisualStudio2015:
-						return "v140";
-				}
+                switch (ProjectFileFormat)
+                {
+                    case VCProjectFileFormat.VisualStudio2012:
+                        return "v110";
+                    case VCProjectFileFormat.VisualStudio2013:
+                        return "v120";
+                    case VCProjectFileFormat.VisualStudio2015:
+                        return "v140";
+                }
 				return string.Empty;
 			}
 		}
@@ -164,7 +169,7 @@ namespace UnrealBuildTool
 		static VCProjectFileGenerator()
 		{
 			switch (WindowsPlatform.Compiler)
-			{
+            {
 				case WindowsCompiler.VisualStudio2013:
 					ProjectFileFormat = VCProjectFileFormat.VisualStudio2013;
 					break;
@@ -196,7 +201,7 @@ namespace UnrealBuildTool
 			// Default the project file format to whatever UBT is using
 			{
 				switch (WindowsPlatform.Compiler)
-				{
+                {
 					case WindowsCompiler.VisualStudio2013:
 						ProjectFileFormat = VCProjectFileFormat.VisualStudio2013;
 						break;
@@ -209,7 +214,11 @@ namespace UnrealBuildTool
 			foreach (string CurArgument in Arguments)
 			{
 				switch (CurArgument.ToUpperInvariant())
-				{
+                {
+                    case "-2012":
+                        ProjectFileFormat = VCProjectFileFormat.VisualStudio2012;
+                        break;
+
 					case "-2013":
 						ProjectFileFormat = VCProjectFileFormat.VisualStudio2013;
 						break;
@@ -260,7 +269,7 @@ namespace UnrealBuildTool
 			// Certain platforms override the project file format because their debugger add-ins may not yet support the latest
 			// version of Visual Studio.  This is their chance to override that.
 			// ...but only if the user didn't override this via the command-line.
-			if (!UnrealBuildTool.CommandLineContains("-2015") && !UnrealBuildTool.CommandLineContains("-2013"))
+            if (!UnrealBuildTool.CommandLineContains("-2015") && !UnrealBuildTool.CommandLineContains("-2013") && !UnrealBuildTool.CommandLineContains("-2012"))
 			{
 				foreach (UnrealTargetPlatform SupportedPlatform in SupportedPlatforms)
 				{
@@ -372,7 +381,15 @@ namespace UnrealBuildTool
 					"Microsoft Visual Studio Solution File, Format Version 12.00" + ProjectFileGenerator.NewLine +
 					"# Visual Studio 2013" + ProjectFileGenerator.NewLine +
 					VersionTag + ProjectFileGenerator.NewLine);
-			}
+            }
+            else if (ProjectFileFormat == VCProjectFileFormat.VisualStudio2012)
+            {
+                VCSolutionFileContent.Append(
+                    ProjectFileGenerator.NewLine +
+                    "Microsoft Visual Studio Solution File, Format Version 12.00" + ProjectFileGenerator.NewLine +
+                    "# Visual Studio 2012" + ProjectFileGenerator.NewLine +
+                    VersionTag + ProjectFileGenerator.NewLine);
+            }
 			else
 			{
 				throw new BuildException("Unexpected ProjectFileFormat");
@@ -883,7 +900,10 @@ namespace UnrealBuildTool
 				// Figure out the filename for the SUO file. VS will automatically import the options from earlier versions if necessary.
 				string SolutionOptionsExtension = "vUnknown.suo";
 				switch (ProjectFileFormat)
-				{
+                {
+                    case VCProjectFileFormat.VisualStudio2012:
+                        SolutionOptionsExtension = "v11.suo";
+                        break;
 					case VCProjectFileFormat.VisualStudio2013:
 						SolutionOptionsExtension = "v12.suo";
 						break;
diff --git a/Engine/Source/Runtime/Apple/MetalRHI/Private/MetalContext.cpp b/Engine/Source/Runtime/Apple/MetalRHI/Private/MetalContext.cpp
--- a/Engine/Source/Runtime/Apple/MetalRHI/Private/MetalContext.cpp
+++ b/Engine/Source/Runtime/Apple/MetalRHI/Private/MetalContext.cpp
@@ -1222,9 +1222,7 @@ static TLockFreeFixedSizeAllocator<sizeof(FMetalCommandContextContainer), PLATFO
 
 void* FMetalCommandContextContainer::operator new(size_t Size)
 {
-	// doesn't support derived classes with a different size
-	check(Size == sizeof(FMetalCommandContextContainer));
-	return FMetalCommandContextContainerAllocator.Allocate();
+	return FMemory::Malloc(Size);
 }
 
 /**
@@ -1232,7 +1230,7 @@ void* FMetalCommandContextContainer::operator new(size_t Size)
  */
 void FMetalCommandContextContainer::operator delete(void *RawMemory)
 {
-	FMetalCommandContextContainerAllocator.Free(RawMemory);
+	FMemory::Free(RawMemory);
 }
 
 IRHICommandContextContainer* FMetalDynamicRHI::RHIGetCommandContextContainer()
diff --git a/Engine/Source/Runtime/Core/Private/Async/TaskGraph.cpp b/Engine/Source/Runtime/Core/Private/Async/TaskGraph.cpp
--- a/Engine/Source/Runtime/Core/Private/Async/TaskGraph.cpp
+++ b/Engine/Source/Runtime/Core/Private/Async/TaskGraph.cpp
@@ -616,7 +616,9 @@ public:
 	 */
 	virtual uint32 Run() override
 	{
+		check(OwnerWorker); // make sure we are started up
 		ProcessTasksUntilQuit(0);
+		FMemory::ClearAndDisableTLSCachesOnCurrentThread();
 		return 0;
 	}
 
@@ -1134,6 +1136,10 @@ public:
 	/** Used for named threads to start processing tasks until the thread is idle and RequestQuit has been called. **/
 	virtual void ProcessTasksUntilQuit(int32 QueueIndex) override
 	{
+		if (PriorityIndex != (ENamedThreads::BackgroundThreadPriority >> ENamedThreads::ThreadPriorityShift))
+		{
+			FMemory::SetupTLSCachesOnCurrentThread();
+		}
 		check(!QueueIndex);
 		Queue.QuitWhenIdle.Reset();
 		while (Queue.QuitWhenIdle.GetValue() == 0)
@@ -1536,17 +1542,7 @@ public:
 				}
 				else
 				{
-					for (int32 Index = 0; Index < GetNumWorkerThreads(); Index++)
-					{
-						for (int32 Priority = 0; Priority < ENamedThreads::NumThreadPriorities; Priority++)
-						{
-							if (Priority == (ENamedThreads::NormalThreadPriority >> ENamedThreads::ThreadPriorityShift) ||
-								(Priority == (ENamedThreads::HighThreadPriority >> ENamedThreads::ThreadPriorityShift) && bCreatedHiPriorityThreads))
-							{
-								StartTaskThread(Priority, Index);
-							}
-						}
-					}
+					StartAllTaskThreads(true);
 				}
 			}
 		}
@@ -1823,7 +1819,22 @@ public:
 		ENamedThreads::Type ThreadToWake = ENamedThreads::Type(IndexToStart + Priority * NumTaskThreadsPerSet + NumNamedThreads);
 		((FTaskThreadAnyThread&)Thread(ThreadToWake)).FTaskThreadAnyThread::WakeUp();
 	}
-
+	void StartAllTaskThreads(bool bDoBackgroundThreads)
+	{
+		for (int32 Index = 0; Index < GetNumWorkerThreads(); Index++)
+		{
+			for (int32 Priority = 0; Priority < ENamedThreads::NumThreadPriorities; Priority++)
+			{
+				if (Priority == (ENamedThreads::NormalThreadPriority >> ENamedThreads::ThreadPriorityShift) ||
+					(Priority == (ENamedThreads::HighThreadPriority >> ENamedThreads::ThreadPriorityShift) && bCreatedHiPriorityThreads) ||
+					(Priority == (ENamedThreads::BackgroundThreadPriority >> ENamedThreads::ThreadPriorityShift) && bCreatedBackgroundPriorityThreads && bDoBackgroundThreads)
+					)
+				{
+					StartTaskThread(Priority, Index);
+				}
+			}
+		}
+	}
 
 	void StartTaskThreadFastMode(int32 Priority, int32 MyIndex)
 	{
@@ -2633,6 +2644,119 @@ FGraphEvent::~FGraphEvent()
 	CheckDontCompleteUntilIsEmpty(); // We should not have any wait untils outstanding
 }
 
+DECLARE_CYCLE_STAT(TEXT("FBroadcastTask"), STAT_FBroadcastTask, STATGROUP_TaskGraphTasks);
+
+class FBroadcastTask
+{
+public:
+	FBroadcastTask(TFunction<void(ENamedThreads::Type CurrentThread)>& InFunction, ENamedThreads::Type InDesiredThread, FThreadSafeCounter* InStallForTaskThread)
+		: Function(InFunction)
+		, DesiredThread(InDesiredThread)
+		, StallForTaskThread(InStallForTaskThread)
+	{
+	}
+	ENamedThreads::Type GetDesiredThread()
+	{
+		return DesiredThread;
+	}
+
+	FORCEINLINE TStatId GetStatId() const
+	{
+		return GET_STATID(STAT_FBroadcastTask);
+	}
+
+	static FORCEINLINE ESubsequentsMode::Type GetSubsequentsMode() { return ESubsequentsMode::TrackSubsequents; }
+	void FORCEINLINE DoTask(ENamedThreads::Type CurrentThread, const FGraphEventRef& MyCompletionGraphEvent)
+	{
+		Function(CurrentThread);
+		if (StallForTaskThread)
+		{
+			StallForTaskThread->Decrement();
+			// we wait for the others to finish here so that we do all task threads
+			while (StallForTaskThread->GetValue())
+			{
+				FPlatformProcess::SleepNoStats(.0002f);
+			}
+		}
+	}
+private:
+	TFunction<void(ENamedThreads::Type CurrentThread)> Function;
+	const ENamedThreads::Type DesiredThread;
+	FThreadSafeCounter* StallForTaskThread;
+};
+
+
+void FTaskGraphInterface::BroadcastSlow_OnlyUseForSpecialPurposes(bool bDoBackgroundThreads, TFunction<void(ENamedThreads::Type CurrentThread)>& Callback)
+{
+	QUICK_SCOPE_CYCLE_COUNTER(STAT_FTaskGraphInterface_BroadcastSlow_OnlyUseForSpecialPurposes);
+	check(FPlatformTLS::GetCurrentThreadId() == GGameThreadId);
+	if (!TaskGraphImplementationSingleton)
+	{
+		// we aren't going yet
+		Callback(ENamedThreads::GameThread);
+		return;
+	}
+
+	// these are static because we don't actually wait for the tasks to clear
+	static FThreadSafeCounter StallForTaskThread;
+
+	int32 Workers = FTaskGraphInterface::Get().GetNumWorkerThreads();
+	StallForTaskThread.Reset();
+	StallForTaskThread.Add(Workers * (1 + (bDoBackgroundThreads && ENamedThreads::bHasBackgroundThreads) + !!(ENamedThreads::bHasHighPriorityThreads)));
+
+	{
+
+		for (int32 Index = 0; Index < Workers; Index++)
+		{
+			TGraphTask<FBroadcastTask>::CreateTask().ConstructAndDispatchWhenReady(Callback, ENamedThreads::AnyNormalThreadHiPriTask, &StallForTaskThread);
+		}
+
+	}
+	if (ENamedThreads::bHasHighPriorityThreads)
+	{
+		for (int32 Index = 0; Index < Workers; Index++)
+		{
+			TGraphTask<FBroadcastTask>::CreateTask().ConstructAndDispatchWhenReady(Callback, ENamedThreads::AnyHiPriThreadHiPriTask, &StallForTaskThread);
+		}
+	}
+	if (bDoBackgroundThreads && ENamedThreads::bHasBackgroundThreads)
+	{
+		for (int32 Index = 0; Index < Workers; Index++)
+		{
+			TGraphTask<FBroadcastTask>::CreateTask().ConstructAndDispatchWhenReady(Callback, ENamedThreads::AnyBackgroundHiPriTask, &StallForTaskThread);
+		}
+	}
+	FGraphEventArray Tasks;
+	STAT(Tasks.Add(TGraphTask<FBroadcastTask>::CreateTask().ConstructAndDispatchWhenReady(Callback, ENamedThreads::SetTaskPriority(ENamedThreads::StatsThread, ENamedThreads::HighTaskPriority), nullptr)););
+	if (GRHIThread)
+	{
+		Tasks.Add(TGraphTask<FBroadcastTask>::CreateTask().ConstructAndDispatchWhenReady(Callback, ENamedThreads::SetTaskPriority(ENamedThreads::RHIThread, ENamedThreads::HighTaskPriority), nullptr));
+	}
+	if (ENamedThreads::RenderThread != ENamedThreads::GameThread)
+	{
+		Tasks.Add(TGraphTask<FBroadcastTask>::CreateTask().ConstructAndDispatchWhenReady(Callback, ENamedThreads::SetTaskPriority(ENamedThreads::RenderThread, ENamedThreads::HighTaskPriority), nullptr));
+	}
+	Tasks.Add(TGraphTask<FBroadcastTask>::CreateTask().ConstructAndDispatchWhenReady(Callback, ENamedThreads::GameThread_Local, nullptr));
+	FTaskGraphInterface::Get().WaitUntilTasksComplete(Tasks, ENamedThreads::GameThread_Local);
+	bool bStartedThreads = false;
+	double StartTime = FPlatformTime::Seconds();
+	while (StallForTaskThread.GetValue())
+	{
+		FPlatformProcess::SleepNoStats(.0001f);
+		// this is probably not needed, but task are not generally tested to never miss a start
+		if (!bStartedThreads && FPlatformTime::Seconds() - StartTime > 0.5)
+		{
+			bStartedThreads = true;
+			TaskGraphImplementationSingleton->StartAllTaskThreads(bDoBackgroundThreads);
+		}
+		else if (FPlatformTime::Seconds() - StartTime > 3.0)
+		{
+			StartTime = FPlatformTime::Seconds();
+			UE_LOG(LogTaskGraph, Error, TEXT("Broadcast failed after three seconds"));
+		}
+	}
+}
+
 
 
 // Benchmark
diff --git a/Engine/Source/Runtime/Core/Private/Containers/ContainerAllocationPolicies.cpp b/Engine/Source/Runtime/Core/Private/Containers/ContainerAllocationPolicies.cpp
--- a/Engine/Source/Runtime/Core/Private/Containers/ContainerAllocationPolicies.cpp
+++ b/Engine/Source/Runtime/Core/Private/Containers/ContainerAllocationPolicies.cpp
@@ -71,61 +71,13 @@ static class FSlackTrackExec: private FSelfRegisteringExec
 	}
 } SlackTrackExec;
 
-#endif
-
-
-
-int32 DefaultCalculateSlack(int32 NumElements,int32 NumAllocatedElements,SIZE_T BytesPerElement)
+void TrackSlack(int32 NumElements, int32 NumAllocatedElements, SIZE_T BytesPerElement, int32 Retval)
 {
-	int32 Retval;
-
-	if(NumElements < NumAllocatedElements)
-	{
-		// If the container has too much slack, shrink it to exactly fit the number of elements.
-		const uint32 CurrentSlackElements = NumAllocatedElements-NumElements;
-		const SIZE_T CurrentSlackBytes = (NumAllocatedElements-NumElements)*BytesPerElement;
-		const bool bTooManySlackBytes = CurrentSlackBytes >= 16384;
-		const bool bTooManySlackElements = 3*NumElements < 2*NumAllocatedElements;
-		if(	(bTooManySlackBytes || bTooManySlackElements) && (CurrentSlackElements > 64 || !NumElements) ) //  hard coded 64 :-(
-		{
-			Retval = NumElements;
-		}
-		else
-		{
-			Retval = NumAllocatedElements;
-		}
-	}
-	else if(NumElements > 0)
-	{
-		const int32 FirstAllocation = 4;
-		if (!NumAllocatedElements && NumElements <= FirstAllocation )
-		{
-			// 17 is too large for an initial allocation. Many arrays never have more one or two elements.
-			Retval = FirstAllocation;
-		}
-		else
-		{
-			// Allocate slack for the array proportional to its size.
-			check(NumElements < MAX_int32);
-			Retval = NumElements + 3*NumElements/8 + 16;
-			// NumElements and MaxElements are stored in 32 bit signed integers so we must be careful not to overflow here.
-			if (NumElements > Retval)
-			{
-				Retval = MAX_int32;
-			}
-		}
-	}
-	else
-	{
-		Retval = 0;
-	}
-
-#if TRACK_ARRAY_SLACK 
 	if( !GSlackTracker )
 	{
 		GSlackTracker = new FStackTracker( SlackTrackerUpdateFn, SlackTrackerReportFn );
 	}
-#define SLACK_TRACE_TO_SKIP 4
+	#define SLACK_TRACE_TO_SKIP 4
 	FSlackTrackData* const LCData = static_cast<FSlackTrackData*>(FMemory::Malloc(sizeof(FSlackTrackData)));
 	FMemory::Memset(LCData, 0, sizeof(FSlackTrackData));
 
@@ -134,9 +86,6 @@ int32 DefaultCalculateSlack(int32 NumElements,int32 NumAllocatedElements,SIZE_T
 	LCData->CurrentSlackNum = NumAllocatedElements-NumElements;
 
 	GSlackTracker->CaptureStackTrace(SLACK_TRACE_TO_SKIP, static_cast<void*>(LCData));
-#endif
-
-
-	return Retval;
 }
+#endif
 
diff --git a/Engine/Source/Runtime/Core/Private/GenericPlatform/GenericPlatformMallocCrash.cpp b/Engine/Source/Runtime/Core/Private/GenericPlatform/GenericPlatformMallocCrash.cpp
--- a/Engine/Source/Runtime/Core/Private/GenericPlatform/GenericPlatformMallocCrash.cpp
+++ b/Engine/Source/Runtime/Core/Private/GenericPlatform/GenericPlatformMallocCrash.cpp
@@ -215,6 +215,10 @@ void FGenericPlatformMallocCrash::SetAsGMalloc()
 {
 	InternalLock.Lock();
 	GMalloc = this;
+	if (PLATFORM_USES_FIXED_GMalloc_CLASS && GFixedMallocLocationPtr)
+	{
+		*GFixedMallocLocationPtr = nullptr; // this disables any fast-path inline allocators
+	}
 	CrashedThreadId = FPlatformTLS::GetCurrentThreadId();
 }
 
diff --git a/Engine/Source/Runtime/Core/Private/GenericPlatform/GenericPlatformTime.cpp b/Engine/Source/Runtime/Core/Private/GenericPlatform/GenericPlatformTime.cpp
--- a/Engine/Source/Runtime/Core/Private/GenericPlatform/GenericPlatformTime.cpp
+++ b/Engine/Source/Runtime/Core/Private/GenericPlatform/GenericPlatformTime.cpp
@@ -5,6 +5,13 @@
 DEFINE_LOG_CATEGORY_STATIC(LogGenericPlatformTime, Log, All);
 
 double FGenericPlatformTime::SecondsPerCycle = 0.0;	
+double FGenericPlatformTime::SecondsPerCycle64 = 0.0;
+
+double FGenericPlatformTime::GetSecondsPerCycle64()
+{
+	check(SecondsPerCycle64 != 0.0);
+	return SecondsPerCycle64;
+}
 
 #if PLATFORM_HAS_BSD_TIME 
 
@@ -15,6 +22,7 @@ double FGenericPlatformTime::InitTiming()
 {
 	// we use gettimeofday() instead of rdtsc, so it's 1000000 "cycles" per second on this faked CPU.
 	SecondsPerCycle = 1.0f / 1000000.0f;
+	SecondsPerCycle64 = 1.0 / 1000000.0;
 	return FPlatformTime::Seconds();
 }
 
diff --git a/Engine/Source/Runtime/Core/Private/HAL/Allocators/AnsiAllocator.h b/Engine/Source/Runtime/Core/Private/HAL/Allocators/AnsiAllocator.h
--- a/Engine/Source/Runtime/Core/Private/HAL/Allocators/AnsiAllocator.h
+++ b/Engine/Source/Runtime/Core/Private/HAL/Allocators/AnsiAllocator.h
@@ -67,9 +67,17 @@ public:
 				Data = (FScriptContainerElement*)::realloc(Data, NumElements*NumBytesPerElement);
 			}
 		}
-		int32 CalculateSlack(int32 NumElements, int32 NumAllocatedElements, SIZE_T NumBytesPerElement) const
+		int32 CalculateSlackReserve(int32 NumElements, int32 NumBytesPerElement) const
 		{
-			return DefaultCalculateSlack(NumElements, NumAllocatedElements, NumBytesPerElement);
+			return DefaultCalculateSlackReserve(NumElements, NumBytesPerElement, false);
+		}
+		int32 CalculateSlackShrink(int32 NumElements, int32 NumAllocatedElements, int32 NumBytesPerElement) const
+		{
+			return DefaultCalculateSlackShrink(NumElements, NumAllocatedElements, NumBytesPerElement, false);
+		}
+		int32 CalculateSlackGrow(int32 NumElements, int32 NumAllocatedElements, int32 NumBytesPerElement) const
+		{
+			return DefaultCalculateSlackGrow(NumElements, NumAllocatedElements, NumBytesPerElement, false);
 		}
 
 		SIZE_T GetAllocatedSize(int32 NumAllocatedElements, SIZE_T NumBytesPerElement) const
@@ -77,6 +85,12 @@ public:
 			return NumAllocatedElements * NumBytesPerElement;
 		}
 
+		bool HasAllocation()
+		{
+			return !!Data;
+		}
+
+
 	private:
 		ForAnyElementType(const ForAnyElementType&);
 		ForAnyElementType& operator=(const ForAnyElementType&);
diff --git a/Engine/Source/Runtime/Core/Private/HAL/Allocators/CachedOSPageAllocator.cpp b/Engine/Source/Runtime/Core/Private/HAL/Allocators/CachedOSPageAllocator.cpp
--- a/Engine/Source/Runtime/Core/Private/HAL/Allocators/CachedOSPageAllocator.cpp
+++ b/Engine/Source/Runtime/Core/Private/HAL/Allocators/CachedOSPageAllocator.cpp
@@ -94,3 +94,19 @@ void FCachedOSPageAllocator::FreeImpl(void* Ptr, SIZE_T Size, uint32 NumCacheBlo
 	CachedTotal += Size;
 	++FreedPageBlocksNum;
 }
+
+void FCachedOSPageAllocator::FreeAllImpl(FFreePageBlock* First, uint32& FreedPageBlocksNum, uint32& CachedTotal)
+{
+	while (FreedPageBlocksNum)
+	{
+		//Remove the oldest one
+		void* FreePtr = First->Ptr;
+		CachedTotal -= First->ByteSize;
+		FreedPageBlocksNum--;
+		if (FreedPageBlocksNum)
+		{
+			FMemory::Memmove(First, First + 1, sizeof(FFreePageBlock)* FreedPageBlocksNum);
+		}
+		FPlatformMemory::BinnedFreeToOS(FreePtr);
+	}
+}
diff --git a/Engine/Source/Runtime/Core/Private/HAL/MallocBinned2.cpp b/Engine/Source/Runtime/Core/Private/HAL/MallocBinned2.cpp
--- a/Engine/Source/Runtime/Core/Private/HAL/MallocBinned2.cpp
+++ b/Engine/Source/Runtime/Core/Private/HAL/MallocBinned2.cpp
@@ -10,101 +10,66 @@
 #include "MemoryMisc.h"
 #include "HAL/PlatformAtomics.h"
 
-#define BINNED2_LARGE_ALLOC         65536 // Alignment of OS-allocated pointer - pool-allocated pointers will have a non-aligned pointer
-#define BINNED2_MINIMUM_ALIGNMENT   16    // Alignment of blocks
-#define BINNED2_MAX_SMALL_POOL_SIZE 32768 // Maximum block size in GMallocBinned2SmallBlockSizes
+#if BINNED2_ALLOW_RUNTIME_TWEAKING
 
-
-static int32 GMallocBinned2PerThreadCaches = 1;
+int32 GMallocBinned2PerThreadCaches = DEFAULT_GMallocBinned2PerThreadCaches;
 static FAutoConsoleVariableRef GMallocBinned2PerThreadCachesCVar(
 	TEXT("MallocBinned2.PerThreadCaches"),
 	GMallocBinned2PerThreadCaches,
 	TEXT("Enables per-thread caches of small (<= 32768 byte) allocations from FMallocBinned2")
-);
+	);
 
-static int32 GMallocBinned2BundleSize = BINNED2_LARGE_ALLOC * 2;
+int32 GMallocBinned2BundleSize = DEFAULT_GMallocBinned2BundleSize;
 static FAutoConsoleVariableRef GMallocBinned2BundleSizeCVar(
 	TEXT("MallocBinned2.BundleSize"),
 	GMallocBinned2BundleSize,
-	TEXT("Size in bytes of per-block bundles used in the recycling process")
-);
+	TEXT("Max size in bytes of per-block bundles used in the recycling process")
+	);
+
+int32 GMallocBinned2BundleCount = DEFAULT_GMallocBinned2BundleCount;
+static FAutoConsoleVariableRef GMallocBinned2BundleCountCVar(
+	TEXT("MallocBinned2.BundleCount"),
+	GMallocBinned2BundleCount,
+	TEXT("Max count in blocks per-block bundles used in the recycling process")
+	);
 
-static int32 GMallocBinned2MaxBundlesBeforeRecycle = 4;
+int32 GMallocBinned2MaxBundlesBeforeRecycle = BINNED2_MAX_GMallocBinned2MaxBundlesBeforeRecycle;
 static FAutoConsoleVariableRef GMallocBinned2MaxBundlesBeforeRecycleCVar(
 	TEXT("MallocBinned2.BundleRecycleCount"),
 	GMallocBinned2MaxBundlesBeforeRecycle,
-	TEXT("Number of freed bundles in the global recycler before it returns them to the system, per-block size")
-);
+	TEXT("Number of freed bundles in the global recycler before it returns them to the system, per-block size. Limited by BINNED2_MAX_GMallocBinned2MaxBundlesBeforeRecycle (currently 4)")
+	);
+
+int32 GMallocBinned2AllocExtra = DEFAULT_GMallocBinned2AllocExtra;
+static FAutoConsoleVariableRef GMallocBinned2AllocExtraCVar(
+	TEXT("MallocBinned2.AllocExtra"),
+	GMallocBinned2AllocExtra,
+	TEXT("When we do acquire the lock, how many blocks cached in TLS caches. In no case will we grab more than a page.")
+	);
+
+#endif
 
 // Block sizes are based around getting the maximum amount of allocations per pool, with as little alignment waste as possible.
 // Block sizes should be close to even divisors of the system page size, and well distributed.
 // They must be 16-byte aligned as well.
-const uint32 GMallocBinned2SmallBlockSizes[] =
+static uint16 SmallBlockSizes[] =
 {
-	16,		32,		48,		64,		80,		96,		112,	128,
-	160,	192,	224,	256,	288,	320,	384,	448,
-	512,	576,	640,	704,	768,	896,	1024,	1168,
-	1360,	1632,	2048,	2336,	2720,	3264,	4096,	4672,
-	5456,	6544,	8192,	9360,	10912,	13104,	16384,	21840,
-	32768
+	16, 32, 48, 64, 80, 96, 112, 128,
+	160, 192, 224, 256, 288, 320, 384, 448,
+	512, 576, 640, 704, 768, 896, 1024 - 16, 1168,
+	1360, 1632, 2048 - 16, 2336, 2720, 3264, 4096 - 16, 4368,
+	4672, 5040, 5456, 5952, 6544 - 16, 7280, 8192 - 16, 9360,
+	10912, 13104, 16384 - 16, 21840, 32768 - 16
 };
 
+MS_ALIGN(PLATFORM_CACHE_LINE_SIZE) static uint8 UnusedAlignPadding[PLATFORM_CACHE_LINE_SIZE] GCC_ALIGN(PLATFORM_CACHE_LINE_SIZE) = { 0 };
+uint16 FMallocBinned2::SmallBlockSizesReversed[BINNED2_SMALL_POOL_COUNT] = { 0 };
+uint32 FMallocBinned2::Binned2TlsSlot = 0;
+uint32 FMallocBinned2::OsAllocationGranularity = 0;
+uint32 FMallocBinned2::PageSize = 0;
+FMallocBinned2* FMallocBinned2::MallocBinned2 = nullptr;
 // Mapping of sizes to small table indices
-uint32 GMallocBinned2MemSizeToIndex[BINNED2_MAX_SMALL_POOL_SIZE];
-
-/** Information about a piece of free memory. */
-struct FMallocBinned2::FFreeBlock
-{
-	FFreeBlock(uint32 InPageSize, uint32 InBlockSize)
-		: NextFreeBlock       (nullptr)
-		, NumFreeBlocks       (InPageSize / InBlockSize - 1)
-		, BlockSize           (InBlockSize)
-		, bShortBlockAllocated(false)
-	{
-	}
-
-	uint32 GetNumFreeRegularBlocks() const
-	{
-		return NumFreeBlocks;
-	}
-
-	uint32 GetShortBlockSize() const
-	{
-		return BlockSize - BINNED2_MINIMUM_ALIGNMENT;
-	}
-
-	bool HasFreeShortBlock() const
-	{
-		return !bShortBlockAllocated;
-	}
-
-	void* AllocateRegularBlock()
-	{
-		--NumFreeBlocks;
-		void* Result = (uint8*)this + (NumFreeBlocks + IsAligned(this, BINNED2_LARGE_ALLOC)) * BlockSize;
-
-		return Result;
-	}
-
-	void* AllocateShortBlock()
-	{
-		// Make sure this isn't a fake block
-		checkSlow(IsAligned(this, BINNED2_LARGE_ALLOC));
-
-		if (bShortBlockAllocated)
-		{
-			return nullptr;
-		}
-
-		bShortBlockAllocated = true;
-		return (uint8*)this + BINNED2_MINIMUM_ALIGNMENT;
-	}
-
- public:	void*  NextFreeBlock;             // Next free block in another pool
- public:	uint32 NumFreeBlocks;             // Number of consecutive free blocks here, at least 1.
- public:	uint32 BlockSize            : 31; // Size of the blocks that this list points to
- public:	uint32 bShortBlockAllocated : 1;  // Whether the short block is allocated or not
-};
+uint8 FMallocBinned2::MemSizeToIndex[1 + (BINNED2_MAX_SMALL_POOL_SIZE >> BINNED2_MINIMUM_ALIGNMENT_SHIFT)] = { 0 };
 
 FMallocBinned2::FPoolList::FPoolList()
 	: Front(nullptr)
@@ -130,16 +95,6 @@ public:
 		return FirstFreeBlock && FirstFreeBlock->GetNumFreeRegularBlocks() != 0;
 	}
 
-	bool HasFreeShortBlock() const
-	{
-		return FirstFreeBlock && FirstFreeBlock->HasFreeShortBlock();
-	}
-
-	uint32 GetShortBlockSize() const
-	{
-		return FirstFreeBlock ? FirstFreeBlock->GetShortBlockSize() : 0;
-	}
-
 	void* AllocateRegularBlock()
 	{
 		checkSlow(HasFreeRegularBlock());
@@ -149,15 +104,6 @@ public:
 		return Result;
 	}
 
-	void* AllocateShortBlock()
-	{
-		checkSlow(HasFreeShortBlock());
-		++Taken;
-		void* Result = FirstFreeBlock->AllocateShortBlock();
-		ExhaustPoolIfNecessary();
-		return Result;
-	}
-
 	uint32 GetOSRequestedBytes() const
 	{
 		return AllocSize;
@@ -204,7 +150,7 @@ private:
 		{
 			FirstFreeBlock = (FFreeBlock*)FirstFreeBlock->NextFreeBlock;
 		}
-		checkSlow(!FirstFreeBlock || FirstFreeBlock->GetNumFreeRegularBlocks() != 0 || FirstFreeBlock->HasFreeShortBlock());
+//		checkSlow(!FirstFreeBlock || FirstFreeBlock->GetNumFreeRegularBlocks() != 0 || FirstFreeBlock->HasFreeShortBlock());
 	}
 };
 
@@ -247,16 +193,10 @@ struct FMallocBinned2::PoolHashBucket
 	}
 };
 
-struct FMallocBinned2::Private
-{
-	FORCEINLINE static bool IsOSAllocation(const void* Ptr)
-	{
-		return IsAligned(Ptr, BINNED2_LARGE_ALLOC);
-	}
 
-	static_assert(ARRAY_COUNT(GMallocBinned2SmallBlockSizes) == SMALL_POOL_COUNT, "Small block size array size must match SMALL_POOL_COUNT");
-	static_assert(sizeof(FFreeBlock) <= BINNED2_MINIMUM_ALIGNMENT,                "Free block struct must be small enough to fit into a block.");
 
+struct FMallocBinned2::Private
+{
 	// Implementation. 
 	static CA_NO_RETURN void OutOfMemory(uint64 Size, uint32 Alignment=0)
 	{
@@ -314,9 +254,7 @@ struct FMallocBinned2::Private
 		// Create a new hash bucket entry
 		if (!Allocator.HashBucketFreeList)
 		{
-			const uint32 PageSize = Allocator.PageSize;
-
-			Allocator.HashBucketFreeList = (PoolHashBucket*)FPlatformMemory::BinnedAllocFromOS(PageSize);
+			Allocator.HashBucketFreeList = (PoolHashBucket*)FPlatformMemory::BinnedAllocFromOS(FMallocBinned2::PageSize);
 
 			for (UPTRINT i = 0, n = PageSize / sizeof(PoolHashBucket); i < n; ++i)
 			{
@@ -349,8 +287,6 @@ struct FMallocBinned2::Private
 
 	static FPoolInfo* FindPoolInfo(FMallocBinned2& Allocator, void* InPtr)
 	{
-		const UPTRINT PageSize = (UPTRINT)Allocator.PageSize;
-
 		UPTRINT BucketIndex;
 		uint32  PoolIndex;
 		Allocator.PtrToPoolMapping.GetHashBucketAndPoolIndices(InPtr, BucketIndex, PoolIndex);
@@ -371,264 +307,59 @@ struct FMallocBinned2::Private
 		return nullptr;
 	}
 
-	static FFreeBlock* GetPoolHeaderFromPointer(void* Ptr)
-	{
-		return (FFreeBlock*)AlignDown(Ptr, BINNED2_LARGE_ALLOC);
-	}
-
-	struct FBundleNode
-	{
-		FBundleNode* NextNodeInCurrentBundle;
-		FBundleNode* NextBundle;
-	};
-
-	struct FBundle
-	{
-		FBundle()
-		{
-			Reset();
-		}
-
-		void Reset()
-		{
-			Head        = nullptr;
-			ApproxBytes = 0;
-		}
-
-		void PushHead(FBundleNode* Node, uint32 NodeSize)
-		{
-			Node->NextNodeInCurrentBundle = Head;
-			Head         = Node;
-			ApproxBytes += NodeSize;
-		}
-
-		FBundleNode* PopHead(uint32 NodeSize)
-		{
-			check(Head);
-			FBundleNode* Result = Head;
-
-			ApproxBytes = FMath::Max<int32>(ApproxBytes - NodeSize, 0);
-			Head        = Head->NextNodeInCurrentBundle;
-
-			return Result;
-		}
-
-		FBundleNode* Head;
-		uint32       ApproxBytes;
-	};
-
-	static_assert(sizeof(FBundleNode) <= BINNED2_MINIMUM_ALIGNMENT, "Bundle nodes must fit into the smallest block size");
-
 	struct FGlobalRecycler
 	{
-		FGlobalRecycler()
-		{
-			DefaultConstructItems<FBundleNode*>(CurrentBundle, SMALL_POOL_COUNT);
-		}
-
-		void PushBundle(uint32 InPoolIndex, FBundleNode* InBundle, FBundleNode*& OutBundlesToRecycle)
-		{
-			FBundleNode** CurrentHeadPtr = &CurrentBundle[InPoolIndex];
-			FBundleNode*  CurrentHead;
-			do
-			{
-				CurrentHead          = *CurrentHeadPtr;
-				InBundle->NextBundle = CurrentHead;
-			}
-			while (FPlatformAtomics::InterlockedCompareExchangePointer((void**)CurrentHeadPtr, InBundle, CurrentHead) != CurrentHead);
-
-			if (BundleCount[InPoolIndex].Increment() >= GMallocBinned2MaxBundlesBeforeRecycle)
-			{
-				OutBundlesToRecycle = PopBundle(InPoolIndex);
-				OutBundlesToRecycle->NextBundle = nullptr;
-			}
-			else
-			{
-				OutBundlesToRecycle = nullptr;
-			}
-		}
-
-		FBundleNode* PopBundle(uint32 InPoolIndex)
-		{
-			FBundleNode** CurrentHeadPtr = &CurrentBundle[InPoolIndex];
-			FBundleNode*  CurrentHead;
-			FBundleNode*  NextHead;
-			do 
-			{
-				CurrentHead = *CurrentHeadPtr;
-				if (!CurrentHead)
-				{
-					return nullptr;
-				}
-
-				NextHead = CurrentHead->NextBundle;
-			}
-			while (FPlatformAtomics::InterlockedCompareExchangePointer((void**)CurrentHeadPtr, NextHead, CurrentHead) != CurrentHead);
-
-			BundleCount[InPoolIndex].Decrement();
-
-			return CurrentHead;
-		}
-
-	private:
-		FBundleNode*       CurrentBundle[SMALL_POOL_COUNT];
-		FThreadSafeCounter BundleCount  [SMALL_POOL_COUNT];
-	};
-
-	static FGlobalRecycler GGlobalRecycler;
-
-	struct FFreeBlockList
-	{
-		void PushToFront(void* InPtr, uint32 InPoolIndex, FBundleNode*& OutBundlesToRecycle)
+		bool PushBundle(uint32 InPoolIndex, FBundleNode* InBundle)
 		{
-			checkSlow(InPtr);
-
-			OutBundlesToRecycle = nullptr;
-			if (PartialBundle.ApproxBytes >= (uint32)GMallocBinned2BundleSize)
+			uint32 NumCachedBundles = FMath::Min<uint32>(GMallocBinned2MaxBundlesBeforeRecycle, BINNED2_MAX_GMallocBinned2MaxBundlesBeforeRecycle);
+			for (uint32 Slot = 0; Slot < NumCachedBundles; Slot++)
 			{
-				if (FullBundle.Head)
+				if (!Bundles[InPoolIndex].FreeBundles[Slot])
 				{
-					GGlobalRecycler.PushBundle(InPoolIndex, FullBundle.Head, OutBundlesToRecycle);
+					if (!FPlatformAtomics::InterlockedCompareExchangePointer((void**)&Bundles[InPoolIndex].FreeBundles[Slot], InBundle, nullptr))
+					{
+						return true;
+					}
 				}
-
-				FullBundle = PartialBundle;
-				PartialBundle.Reset();
 			}
-
-			PartialBundle.PushHead((FBundleNode*)InPtr, GMallocBinned2SmallBlockSizes[InPoolIndex]);
+			return false;
 		}
 
-		void* PopFromFront(uint32 InPoolIndex)
+		FBundleNode* PopBundle(uint32 InPoolIndex)
 		{
-			if (!PartialBundle.Head)
+			uint32 NumCachedBundles = FMath::Min<uint32>(GMallocBinned2MaxBundlesBeforeRecycle, BINNED2_MAX_GMallocBinned2MaxBundlesBeforeRecycle);
+			for (uint32 Slot = 0; Slot < NumCachedBundles; Slot++)
 			{
-				FBundle NewPartialBundle = FullBundle;
-				if (!NewPartialBundle.Head)
+				FBundleNode* Result = Bundles[InPoolIndex].FreeBundles[Slot];
+				if (Result)
 				{
-					NewPartialBundle.Head = GGlobalRecycler.PopBundle(InPoolIndex);
-					if (!NewPartialBundle.Head)
+					if (FPlatformAtomics::InterlockedCompareExchangePointer((void**)&Bundles[InPoolIndex].FreeBundles[Slot], nullptr, Result) == Result)
 					{
-						return nullptr;
+						return Result;
 					}
-
-					NewPartialBundle.ApproxBytes = (uint32)GMallocBinned2BundleSize;
 				}
-
-				PartialBundle = NewPartialBundle;
-				FullBundle.Reset();
-			}
-
-			void* Result = PartialBundle.PopHead(GMallocBinned2SmallBlockSizes[InPoolIndex]);
-			return Result;
-		}
-
-		FBundleNode* PopBundles(uint32 InPoolIndex)
-		{
-			FBundleNode* Partial = PartialBundle.Head;
-			if (Partial)
-			{
-				PartialBundle.Reset();
-				Partial->NextBundle = nullptr;
-			}
-
-			FBundleNode* Full = FullBundle.Head;
-			if (Full)
-			{
-				FullBundle.Reset();
-				Full->NextBundle = nullptr;
-			}
-
-			FBundleNode* Result = Partial;
-			if (Result)
-			{
-				Result->NextBundle = Full;
-			}
-			else
-			{
-				Result = Full;
 			}
-
-			return Result;
-		}
-
-		bool IsEmpty() const
-		{
-			return !PartialBundle.Head && !FullBundle.Head;
+			return nullptr;
 		}
 
 	private:
-		FBundle PartialBundle;
-		FBundle FullBundle;
-	};
-
-	// We've reinvented TThreadSingleton here because we don't want it to be re-entering our allocation
-	// functions!
-	template <typename T>
-	struct TThreadSingleton2
-	{
-		FORCEINLINE static T& Get(uint32 PageSize)
+		struct FPaddedBundlePointer
 		{
-			static uint32 TlsSlot = 0;
-
-			if (TlsSlot == 0)
-			{
-				const uint32 ThisTlsSlot = FPlatformTLS::AllocTlsSlot();
-				const uint32 PrevTlsSlot = FPlatformAtomics::InterlockedCompareExchange((int32*)&TlsSlot, (int32)ThisTlsSlot, 0);
-				if (PrevTlsSlot != 0)
-				{
-					FPlatformTLS::FreeTlsSlot(ThisTlsSlot);
-				}
-			}
-
-			T* ThreadSingleton = (T*)FPlatformTLS::GetTlsValue(TlsSlot);
-			if (!ThreadSingleton)
+			FBundleNode* FreeBundles[BINNED2_MAX_GMallocBinned2MaxBundlesBeforeRecycle];
+#define BUNDLE_PADDING (PLATFORM_CACHE_LINE_SIZE - sizeof(FBundleNode*) * BINNED2_MAX_GMallocBinned2MaxBundlesBeforeRecycle)
+#if (4 + (4 * PLATFORM_64BITS)) * BINNED2_MAX_GMallocBinned2MaxBundlesBeforeRecycle < PLATFORM_CACHE_LINE_SIZE
+			uint8 Padding[BUNDLE_PADDING];
+#endif
+			FPaddedBundlePointer()
 			{
-				ThreadSingleton = new (FPlatformMemory::BinnedAllocFromOS(Align(sizeof(T), PageSize))) T();
-				FPlatformTLS::SetTlsValue(TlsSlot, ThreadSingleton);
+				DefaultConstructItems<FBundleNode*>(FreeBundles, BINNED2_MAX_GMallocBinned2MaxBundlesBeforeRecycle);
 			}
-
-			return *ThreadSingleton;
-		}
+		};
+		static_assert(sizeof(FPaddedBundlePointer) == PLATFORM_CACHE_LINE_SIZE, "FPaddedBundlePointer should be the same size as a cache line");
+		MS_ALIGN(PLATFORM_CACHE_LINE_SIZE) FPaddedBundlePointer Bundles[BINNED2_SMALL_POOL_COUNT] GCC_ALIGN(PLATFORM_CACHE_LINE_SIZE);
 	};
 
-	struct FPerThreadFreeBlockLists : TThreadSingleton2<FPerThreadFreeBlockLists>
-	{
-		FPerThreadFreeBlockLists()
-		{
-		}
-
-		void* Malloc(uint32 InPoolIndex)
-		{
-			checkSlow(InPoolIndex >= 0 && InPoolIndex < SMALL_POOL_COUNT);
-
-			FFreeBlockList& List = FreeLists[InPoolIndex];
-			if (List.IsEmpty())
-			{
-				return nullptr;
-			}
-
-			void* Result = List.PopFromFront(InPoolIndex);
-			return Result;
-		}
-
-		void Free(void* InPtr, uint32 InPoolIndex, FBundleNode*& OutBundlesToRecycle)
-		{
-			checkSlow(InPoolIndex >= 0 && InPoolIndex < SMALL_POOL_COUNT);
-
-			FreeLists[InPoolIndex].PushToFront(InPtr, InPoolIndex, OutBundlesToRecycle);
-		}
-
-		FBundleNode* PopBundles(uint32 InPoolIndex)
-		{
-			checkSlow(InPoolIndex >= 0 && InPoolIndex < SMALL_POOL_COUNT);
-
-			FBundleNode* Result = FreeLists[InPoolIndex].PopBundles(InPoolIndex);
-			return Result;
-		}
-
-	private:
-		FFreeBlockList FreeLists[SMALL_POOL_COUNT];
-	};
+	static FGlobalRecycler GGlobalRecycler;
 
 	static void FreeBundles(FMallocBinned2& Allocator, FBundleNode* BundlesToRecycle, uint32 InBlockSize, uint32 InPoolIndex)
 	{
@@ -656,7 +387,6 @@ struct FMallocBinned2::Private
 				Free->NumFreeBlocks = 1;
 				Free->NextFreeBlock = NodePool->FirstFreeBlock;
 				Free->BlockSize     = InBlockSize;
-				Free->bShortBlockAllocated = false;
 				NodePool->FirstFreeBlock   = Free;
 
 				// Free this pool.
@@ -703,17 +433,17 @@ void FMallocBinned2::FPoolList::LinkToFront(FPoolInfo* Pool)
 	Pool->Link(Front);
 }
 
-FMallocBinned2::FPoolInfo& FMallocBinned2::FPoolList::PushNewPoolToFront(FMallocBinned2& Allocator, uint32 InBlockSize)
+FMallocBinned2::FPoolInfo& FMallocBinned2::FPoolList::PushNewPoolToFront(FMallocBinned2& Allocator, uint32 InBlockSize, uint32 InPoolIndex)
 {
 	const uint32 LocalPageSize = Allocator.PageSize;
 
 	// Allocate memory.
-	FFreeBlock* Free = new (Allocator.CachedOSPageAllocator.Allocate(LocalPageSize)) FFreeBlock(LocalPageSize, InBlockSize);
+	FFreeBlock* Free = new (Allocator.CachedOSPageAllocator.Allocate(LocalPageSize)) FFreeBlock(LocalPageSize, InBlockSize, InPoolIndex);
 	if (!Free)
 	{
 		Private::OutOfMemory(LocalPageSize);
 	}
-
+	check(IsAligned(Free, LocalPageSize));
 	// Create pool
 	FPoolInfo* Result = Private::GetOrCreatePoolInfo(Allocator, Free);
 	Result->Link(Front);
@@ -723,46 +453,73 @@ FMallocBinned2::FPoolInfo& FMallocBinned2::FPoolList::PushNewPoolToFront(FMalloc
 	return *Result;
 }
 
-FMallocBinned2::FMallocBinned2(uint32 InPageSize, uint64 AddressLimit)
-	: PageSize          (InPageSize)
-	, NumPoolsPerPage   (PageSize / sizeof(FPoolInfo))
-	, PtrToPoolMapping  (PageSize, NumPoolsPerPage, AddressLimit)
-	, HashBucketFreeList(nullptr)
+FMallocBinned2::FMallocBinned2()
+	: HashBucketFreeList(nullptr)
 {
-	checkf(FMath::IsPowerOfTwo(PageSize),                                                      TEXT("OS page size must be a power of two"));
-	checkf(FMath::IsPowerOfTwo(AddressLimit),                                                  TEXT("OS address limit must be a power of two"));
-	checkf(AddressLimit > PageSize,                                                            TEXT("OS address limit must be greater than the page size")); // Check to catch 32 bit overflow in AddressLimit
-	checkf(GMallocBinned2SmallBlockSizes[SMALL_POOL_COUNT - 1] == BINNED2_MAX_SMALL_POOL_SIZE, TEXT("BINNED2_MAX_SMALL_POOL_SIZE must equal the smallest block size"));
-	checkf(PageSize % BINNED2_LARGE_ALLOC == 0,                                                TEXT("OS page size must be a multiple of BINNED2_LARGE_ALLOC"));
-	checkf(sizeof(FMallocBinned2::FFreeBlock) <= GMallocBinned2SmallBlockSizes[0],             TEXT("Pool header must be able to fit into the smallest block"));
+	static bool bOnce = false;
+	check(!bOnce); // this is now a singleton-like thing and you cannot make multiple copies
+	bOnce = true;
+
+	for (uint32 Index = 0; Index != BINNED2_SMALL_POOL_COUNT; ++Index)
+	{
+		uint32 Partner = BINNED2_SMALL_POOL_COUNT - Index - 1;
+		SmallBlockSizesReversed[Index] = SmallBlockSizes[Partner];
+	}
+	FGenericPlatformMemoryConstants Constants = FPlatformMemory::GetConstants();
+	PageSize = Constants.PageSize;
+	OsAllocationGranularity = Constants.OsAllocationGranularity ? Constants.OsAllocationGranularity : PageSize;
+	NumPoolsPerPage = PageSize / sizeof(FPoolInfo);
+	PtrToPoolMapping.Init(PageSize, NumPoolsPerPage, Constants.AddressLimit);
+
+	checkf(FMath::IsPowerOfTwo(PageSize), TEXT("OS page size must be a power of two"));
+	checkf(FMath::IsPowerOfTwo(Constants.AddressLimit), TEXT("OS address limit must be a power of two"));
+	checkf(Constants.AddressLimit > PageSize, TEXT("OS address limit must be greater than the page size")); // Check to catch 32 bit overflow in AddressLimit
+	checkf(SmallBlockSizes[BINNED2_SMALL_POOL_COUNT - 1] == BINNED2_MAX_SMALL_POOL_SIZE, TEXT("BINNED2_MAX_SMALL_POOL_SIZE must equal the smallest block size"));
+	checkf(PageSize % BINNED2_LARGE_ALLOC == 0, TEXT("OS page size must be a multiple of BINNED2_LARGE_ALLOC"));
+	checkf(sizeof(FMallocBinned2::FFreeBlock) <= SmallBlockSizes[0], TEXT("Pool header must be able to fit into the smallest block"));
+	static_assert(ARRAY_COUNT(SmallBlockSizes) == BINNED2_SMALL_POOL_COUNT, "Small block size array size must match BINNED2_SMALL_POOL_COUNT");
+	static_assert(ARRAY_COUNT(SmallBlockSizes) <= 256, "Small block size array size must fit in a byte");
+	static_assert(sizeof(FFreeBlock) <= BINNED2_MINIMUM_ALIGNMENT, "Free block struct must be small enough to fit into a block.");
 
 	// Init pool tables.
-	for (uint32 Index = 0; Index != SMALL_POOL_COUNT; ++Index)
+	for (uint32 Index = 0; Index != BINNED2_SMALL_POOL_COUNT; ++Index)
 	{
-		checkf(Index == 0 || GMallocBinned2SmallBlockSizes[Index - 1] < GMallocBinned2SmallBlockSizes[Index], TEXT("Small block sizes must be strictly increasing"));
-		checkf(GMallocBinned2SmallBlockSizes[Index] <= PageSize,                                              TEXT("Small block size must be small enough to fit into a page"));
-		checkf(GMallocBinned2SmallBlockSizes[Index] % BINNED2_MINIMUM_ALIGNMENT == 0,                         TEXT("Small block size must be a multiple of BINNED2_MINIMUM_ALIGNMENT"));
+		checkf(Index == 0 || SmallBlockSizes[Index - 1] < SmallBlockSizes[Index], TEXT("Small block sizes must be strictly increasing"));
+		checkf(SmallBlockSizes[Index] <= PageSize, TEXT("Small block size must be small enough to fit into a page"));
+		checkf(SmallBlockSizes[Index] % BINNED2_MINIMUM_ALIGNMENT == 0, TEXT("Small block size must be a multiple of BINNED2_MINIMUM_ALIGNMENT"));
 
-		SmallPoolTables[Index].BlockSize = GMallocBinned2SmallBlockSizes[Index];
+		SmallPoolTables[Index].BlockSize = SmallBlockSizes[Index];
 	}
 
 	// Set up pool mappings
-	uint32* IndexEntry = GMallocBinned2MemSizeToIndex;
+	uint8* IndexEntry = MemSizeToIndex;
 	uint32  PoolIndex  = 0;
-	for (uint32 Index = 0; Index != BINNED2_MAX_SMALL_POOL_SIZE; ++Index)
+	for (uint32 Index = 0; Index != 1 + (BINNED2_MAX_SMALL_POOL_SIZE >> BINNED2_MINIMUM_ALIGNMENT_SHIFT); ++Index)
 	{
-		while (GMallocBinned2SmallBlockSizes[PoolIndex] - 1 < Index)
+		
+		uint32 BlockSize = Index << BINNED2_MINIMUM_ALIGNMENT_SHIFT; // inverse of int32 Index = int32((Size >> BINNED2_MINIMUM_ALIGNMENT_SHIFT));
+		while (SmallBlockSizes[PoolIndex] < BlockSize)
 		{
 			++PoolIndex;
-			check(PoolIndex != SMALL_POOL_COUNT);
+			check(PoolIndex != BINNED2_SMALL_POOL_COUNT);
 		}
-		*IndexEntry++ = PoolIndex;
+		check(PoolIndex < 256);
+		*IndexEntry++ = uint8(PoolIndex);
+	}
+	// now reverse the pool sizes for cache coherency
+
+	for (uint32 Index = 0; Index != BINNED2_SMALL_POOL_COUNT; ++Index)
+	{
+		uint32 Partner = BINNED2_SMALL_POOL_COUNT - Index - 1;
+		SmallBlockSizesReversed[Index] = SmallBlockSizes[Partner];
 	}
 
 	uint64 MaxHashBuckets = PtrToPoolMapping.GetMaxHashBuckets();
 
-	HashBuckets = (PoolHashBucket*)FPlatformMemory::BinnedAllocFromOS(Align(MaxHashBuckets * sizeof(PoolHashBucket), PageSize));
+	HashBuckets = (PoolHashBucket*)FPlatformMemory::BinnedAllocFromOS(Align(MaxHashBuckets * sizeof(PoolHashBucket), OsAllocationGranularity));
 	DefaultConstructItems<PoolHashBucket>(HashBuckets, MaxHashBuckets);
+	MallocBinned2 = this;
+	GFixedMallocLocationPtr = (FMalloc**)(&MallocBinned2);
 }
 
 FMallocBinned2::~FMallocBinned2()
@@ -774,26 +531,25 @@ bool FMallocBinned2::IsInternallyThreadSafe() const
 	return true;
 }
 
-void* FMallocBinned2::Malloc(SIZE_T Size, uint32 Alignment)
+void* FMallocBinned2::MallocExternal(SIZE_T Size, uint32 Alignment)
 {
-	Alignment = FMath::Max<uint32>(Alignment, BINNED2_MINIMUM_ALIGNMENT);
-	Size      = Align(FMath::Max((SIZE_T)1, Size), Alignment);
-
-	checkf(FMath::IsPowerOfTwo(Alignment), TEXT("Alignment must be a power of two"));
-	checkf(Alignment <= PageSize, TEXT("Alignment cannot exceed the system page size"));
+	static_assert(DEFAULT_ALIGNMENT <= BINNED2_MINIMUM_ALIGNMENT, "DEFAULT_ALIGNMENT is assumed to be zero"); // used below
 
 	// Only allocate from the small pools if the size is small enough and the alignment isn't crazy large.
 	// With large alignments, we'll waste a lot of memory allocating an entire page, but such alignments are highly unlikely in practice.
-	if (Size <= BINNED2_MAX_SMALL_POOL_SIZE && Alignment <= BINNED2_MINIMUM_ALIGNMENT)
+	if ((Size <= BINNED2_MAX_SMALL_POOL_SIZE) & (Alignment <= BINNED2_MINIMUM_ALIGNMENT)) // one branch, not two
 	{
-		uint32 PoolIndex = GMallocBinned2MemSizeToIndex[Size - 1];
+		uint32 PoolIndex = BoundSizeToPoolIndex(Size);
 
-		if (GMallocBinned2PerThreadCaches)
+		FPerThreadFreeBlockLists* Lists = GMallocBinned2PerThreadCaches ? FPerThreadFreeBlockLists::Get() : nullptr;
+		if (Lists)
 		{
-			Private::FPerThreadFreeBlockLists& Lists = Private::FPerThreadFreeBlockLists::Get(PageSize);
-			if (void* Result = Lists.Malloc(PoolIndex))
+			if (Lists->ObtainRecycledPartial(PoolIndex))
 			{
-				return Result;
+				if (void* Result = Lists->Malloc(PoolIndex))
+				{
+					return Result;
+				}
 			}
 		}
 
@@ -809,10 +565,24 @@ void* FMallocBinned2::Malloc(SIZE_T Size, uint32 Alignment)
 		}
 		else
 		{
-			Pool = &Table.ActivePools.PushNewPoolToFront(*this, Table.BlockSize);
+			Pool = &Table.ActivePools.PushNewPoolToFront(*this, Table.BlockSize, PoolIndex);
 		}
 
 		void* Result = Pool->AllocateRegularBlock();
+		if (GMallocBinned2AllocExtra)
+		{
+			if (Lists)
+			{
+				for (int32 Index = 0; Index < GMallocBinned2AllocExtra && Pool->HasFreeRegularBlock(); Index++)
+				{
+					if (!Lists->Free(Result, PoolIndex, Table.BlockSize))
+					{
+						break;
+					}
+					Result = Pool->AllocateRegularBlock();
+				}
+			}
+		}
 		if (!Pool->HasFreeRegularBlock())
 		{
 			Table.ExhaustedPools.LinkToFront(Pool);
@@ -820,16 +590,22 @@ void* FMallocBinned2::Malloc(SIZE_T Size, uint32 Alignment)
 
 		return Result;
 	}
+	Alignment = FMath::Max<uint32>(Alignment, BINNED2_MINIMUM_ALIGNMENT);
+	Size = Align(FMath::Max((SIZE_T)1, Size), Alignment);
+
+	checkSlow(FMath::IsPowerOfTwo(Alignment));
+	checkSlow(Alignment <= PageSize);
 
 	FScopeLock Lock(&Mutex);
 
 	// Use OS for non-pooled allocations.
-	UPTRINT AlignedSize = Align(Size, PageSize);
+	UPTRINT AlignedSize = Align(Size, OsAllocationGranularity);
 	void* Result = CachedOSPageAllocator.Allocate(AlignedSize);
 	if (!Result)
 	{
 		Private::OutOfMemory(AlignedSize);
 	}
+	check(IsAligned(Result, PageSize) && IsOSAllocation(Result));
 
 	// Create pool.
 	FPoolInfo* Pool = Private::GetOrCreatePoolInfo(*this, Result);
@@ -838,55 +614,58 @@ void* FMallocBinned2::Malloc(SIZE_T Size, uint32 Alignment)
 	return Result;
 }
 
-void* FMallocBinned2::Realloc(void* Ptr, SIZE_T NewSize, uint32 Alignment)
+
+void* FMallocBinned2::ReallocExternal(void* Ptr, SIZE_T NewSize, uint32 Alignment)
 {
 	if (NewSize == 0)
 	{
-		FMallocBinned2::Free(Ptr);
+		FMallocBinned2::FreeExternal(Ptr);
 		return nullptr;
 	}
+	static_assert(DEFAULT_ALIGNMENT <= BINNED2_MINIMUM_ALIGNMENT, "DEFAULT_ALIGNMENT is assumed to be zero"); // used below
+	checkSlow(FMath::IsPowerOfTwo(Alignment));
+	checkSlow(Alignment <= PageSize);
 
-	if (!Ptr)
-	{
-		void* Result = FMallocBinned2::Malloc(NewSize, Alignment);
-		return Result;
-	}
-
-	Alignment = FMath::Max<uint32>(Alignment, BINNED2_MINIMUM_ALIGNMENT);
-	NewSize   = Align(NewSize, Alignment);
-
-	checkf(FMath::IsPowerOfTwo(Alignment), TEXT("Alignment must be a power of two"));
-	checkf(Alignment <= PageSize, TEXT("Alignment cannot exceed the system page size"));
-
-	if (!Private::IsOSAllocation(Ptr))
+	if (!IsOSAllocation(Ptr))
 	{
+		checkSlow(Ptr); // null is 64k aligned so we should not be here
 		// Reallocate to a smaller/bigger pool if necessary
-		FFreeBlock* Free       = Private::GetPoolHeaderFromPointer(Ptr);
-		uint32      TableIndex = GMallocBinned2MemSizeToIndex[Free->BlockSize - 1];
-		FPoolTable& Table      = SmallPoolTables[TableIndex];
-		if (NewSize <= Table.BlockSize && (TableIndex == 0 || NewSize > SmallPoolTables[TableIndex - 1].BlockSize) && IsAligned(Ptr, Alignment))
+		FFreeBlock* Free = GetPoolHeaderFromPointer(Ptr);
+		Free->CanaryTest();
+		uint32 BlockSize = Free->BlockSize;
+		uint32 PoolIndex = Free->PoolIndex;
+		if (
+			((NewSize <= BlockSize) & (Alignment <= BINNED2_MINIMUM_ALIGNMENT)) && // one branch, not two
+			(PoolIndex == 0 || NewSize > PoolIndexToBlockSize(PoolIndex - 1)))
 		{
 			return Ptr;
 		}
 
 		// Reallocate and copy the data across
-		void* Result = FMallocBinned2::Malloc(NewSize, Alignment);
-		FMemory::Memcpy(Result, Ptr, FMath::Min<SIZE_T>(NewSize, Table.BlockSize));
-		FMallocBinned2::Free(Ptr);
+		void* Result = FMallocBinned2::MallocExternal(NewSize, Alignment);
+		FMemory::Memcpy(Result, Ptr, FMath::Min<SIZE_T>(NewSize, BlockSize));
+		FMallocBinned2::FreeExternal(Ptr);
+		return Result;
+	}
+	if (!Ptr)
+	{
+		void* Result = FMallocBinned2::MallocExternal(NewSize, Alignment);
 		return Result;
 	}
 
 	FScopeLock Lock(&Mutex);
 
 	// Allocated from OS.
-	FPoolInfo* Pool        = Private::FindPoolInfo(*this, Ptr);
-	UPTRINT    PoolOsBytes = Pool->GetOsAllocatedBytes();
-	if (NewSize > PoolOsBytes || NewSize * 3 < PoolOsBytes * 2)
+	FPoolInfo* Pool = Private::FindPoolInfo(*this, Ptr);
+	UPTRINT PoolOsBytes = Pool->GetOsAllocatedBytes();
+	if (NewSize > PoolOsBytes || // can't fit in the old block
+		(NewSize <= BINNED2_MAX_SMALL_POOL_SIZE && Alignment <= BINNED2_MINIMUM_ALIGNMENT) || // can switch to the small block allocator
+		Align(NewSize, OsAllocationGranularity) < PoolOsBytes) // we can get some pages back
 	{
 		// Grow or shrink.
-		void* Result = FMallocBinned2::Malloc(NewSize, Alignment);
+		void* Result = FMallocBinned2::MallocExternal(NewSize, Alignment);
 		FMemory::Memcpy(Result, Ptr, FMath::Min<SIZE_T>(NewSize, Pool->GetOSRequestedBytes()));
-		FMallocBinned2::Free(Ptr);
+		FMallocBinned2::FreeExternal(Ptr);
 		return Result;
 	}
 
@@ -895,85 +674,91 @@ void* FMallocBinned2::Realloc(void* Ptr, SIZE_T NewSize, uint32 Alignment)
 	return Ptr;
 }
 
-void FMallocBinned2::Free(void* Ptr)
+void FMallocBinned2::FreeExternal(void* Ptr)
 {
-	if (!Ptr)
-	{
-		return;
-	}
-
-	FPoolInfo* Pool = Private::FindPoolInfo(*this, Ptr);
-#if PLATFORM_IOS || PLATFORM_MAC
-	if (!Pool)
-	{
-		UE_LOG(LogMemory, Warning, TEXT("Attempting to free a pointer we didn't allocate!"));
-		return;
-	}
-#endif
-
-	checkSlow(Pool);
-	if (!Private::IsOSAllocation(Ptr))
+	if (!IsOSAllocation(Ptr))
 	{
-		FFreeBlock* BasePtr   = Private::GetPoolHeaderFromPointer(Ptr);
-		uint32      BlockSize = BasePtr->BlockSize;
-		uint32      PoolIndex = GMallocBinned2MemSizeToIndex[BlockSize - 1];
-
-		Private::FBundleNode* BundlesToRecycle;
-		if (GMallocBinned2PerThreadCaches)
+		checkSlow(Ptr); // null is 64k aligned so we should not be here
+		FFreeBlock* BasePtr = GetPoolHeaderFromPointer(Ptr);
+		BasePtr->CanaryTest();
+		uint32 BlockSize = BasePtr->BlockSize;
+		uint32 PoolIndex = BasePtr->PoolIndex;
+		FBundleNode* BundlesToRecycle = nullptr;
+		FPerThreadFreeBlockLists* Lists = GMallocBinned2PerThreadCaches ? FPerThreadFreeBlockLists::Get() : nullptr;
+		if (Lists)
 		{
-			Private::FPerThreadFreeBlockLists& Lists = Private::FPerThreadFreeBlockLists::Get(PageSize);
-			Lists.Free(Ptr, PoolIndex, BundlesToRecycle);
+			BundlesToRecycle = Lists->RecycleFullBundle(BasePtr->PoolIndex);
+			bool bPushed = Lists->Free(Ptr, PoolIndex, BlockSize);
+			check(bPushed);
 		}
 		else
 		{
-			BundlesToRecycle = (Private::FBundleNode*)Ptr;
+			BundlesToRecycle = (FBundleNode*)Ptr;
 			BundlesToRecycle->NextNodeInCurrentBundle = nullptr;
 		}
-
 		if (BundlesToRecycle)
 		{
+			BundlesToRecycle->NextBundle = nullptr;
 			FScopeLock Lock(&Mutex);
-
 			Private::FreeBundles(*this, BundlesToRecycle, BlockSize, PoolIndex);
 		}
 	}
-	else
+	else if (Ptr)
 	{
 		FScopeLock Lock(&Mutex);
+		FPoolInfo* Pool = Private::FindPoolInfo(*this, Ptr);
+		if (!Pool)
+		{
+			UE_LOG(LogMemory, Fatal, TEXT("FMallocBinned2 Attempt to free an unrecognized block %p"), Ptr);
+		}
 
 		// Free an OS allocation.
 		CachedOSPageAllocator.Free(Ptr, Pool->GetOsAllocatedBytes());
 	}
 }
 
-bool FMallocBinned2::GetAllocationSize(void* Ptr, SIZE_T& SizeOut)
+bool FMallocBinned2::GetAllocationSizeExternal(void* Ptr, SIZE_T& SizeOut)
 {
-	if (!Ptr)
+	if (!IsOSAllocation(Ptr))
 	{
-		return false;
+		checkSlow(Ptr); // null is 64k aligned so we should not be here
+		const FFreeBlock* Free = GetPoolHeaderFromPointer(Ptr);
+		Free->CanaryTest();
+		uint32 BlockSize = Free->BlockSize;
+		SizeOut = BlockSize;
+		return true;
 	}
-
-	if (!Private::IsOSAllocation(Ptr))
+	if (!Ptr)
 	{
-		const FFreeBlock* Free = Private::GetPoolHeaderFromPointer(Ptr);
-		SizeOut = SmallPoolTables[GMallocBinned2MemSizeToIndex[Free->BlockSize - 1]].BlockSize;
-		return true;
+		return false;
 	}
-
 	FScopeLock Lock(&Mutex);
-
 	FPoolInfo* Pool = Private::FindPoolInfo(*this, Ptr);
 	SizeOut = Pool->GetOSRequestedBytes();
 	return true;
 }
 
-void FMallocBinned2::FPoolList::Validate(TFunctionRef<void(FMallocBinned2::FPoolInfo&)> PoolValidate)
+void FMallocBinned2::FPoolList::ValidateActivePools()
+{
+	for (FPoolInfo** PoolPtr = &Front; *PoolPtr; PoolPtr = &(*PoolPtr)->Next)
+	{
+		FPoolInfo* Pool = *PoolPtr;
+		check(Pool->PtrToPrevNext == PoolPtr);
+		check(Pool->FirstFreeBlock);
+		for (FFreeBlock* Free = Pool->FirstFreeBlock; Free; Free = (FFreeBlock*)Free->NextFreeBlock)
+		{
+			check(Free->GetNumFreeRegularBlocks() > 0);
+		}
+	}
+}
+
+void FMallocBinned2::FPoolList::ValidateExhaustedPools()
 {
 	for (FPoolInfo** PoolPtr = &Front; *PoolPtr; PoolPtr = &(*PoolPtr)->Next)
 	{
 		FPoolInfo* Pool = *PoolPtr;
 		check(Pool->PtrToPrevNext == PoolPtr);
-		PoolValidate(*Pool);
+		check(!Pool->FirstFreeBlock);
 	}
 }
 
@@ -983,23 +768,8 @@ bool FMallocBinned2::ValidateHeap()
 
 	for (FPoolTable& Table : SmallPoolTables)
 	{
-		Table.ActivePools.Validate(
-			[](FMallocBinned2::FPoolInfo& Pool)
-			{
-				check(Pool.FirstFreeBlock);
-				for (FFreeBlock* Free = Pool.FirstFreeBlock; Free; Free = (FFreeBlock*)Free->NextFreeBlock)
-				{
-					check(Free->GetNumFreeRegularBlocks() > 0);
-				}
-			}
-		);
-
-		Table.ExhaustedPools.Validate(
-			[](FMallocBinned2::FPoolInfo& Pool)
-			{
-				check(!Pool.FirstFreeBlock);
-			}
-		);
+		Table.ActivePools.ValidateActivePools();
+		Table.ExhaustedPools.ValidateExhaustedPools();
 	}
 
 	return true;
@@ -1012,15 +782,155 @@ const TCHAR* FMallocBinned2::GetDescriptiveName()
 
 void FMallocBinned2::FlushCurrentThreadCache()
 {
-	FScopeLock Lock(&Mutex);
+	FPerThreadFreeBlockLists* Lists = FPerThreadFreeBlockLists::Get();
+	if (Lists)
+	{
+		FScopeLock Lock(&Mutex);
+		for (int32 PoolIndex = 0; PoolIndex != BINNED2_SMALL_POOL_COUNT; ++PoolIndex)
+		{
+			FBundleNode* Bundles = Lists->PopBundles(PoolIndex);
+			if (Bundles)
+			{
+				Private::FreeBundles(*this, Bundles, PoolIndexToBlockSize(PoolIndex), PoolIndex);
+			}
+		}
+	}
+}
+
+#include "TaskGraphInterfaces.h"
+
+void FMallocBinned2::Trim()
+{
+	QUICK_SCOPE_CYCLE_COUNTER(STAT_FMallocBinned2_Trim);
+
+	if (GMallocBinned2PerThreadCaches)
+	{
+		//double StartTime = FPlatformTime::Seconds();
+		TFunction<void(ENamedThreads::Type CurrentThread)> Broadcast =
+			[this](ENamedThreads::Type MyThread)
+		{
+			FlushCurrentThreadCache();
+		};
+		FTaskGraphInterface::BroadcastSlow_OnlyUseForSpecialPurposes(false, Broadcast);
+		//UE_LOG(LogTemp, Display, TEXT("Trim Broadcast = %6.2fms"), 1000.0f * float(FPlatformTime::Seconds() - StartTime));
+	}
+	{
+		//double StartTime = FPlatformTime::Seconds();
+		FScopeLock Lock(&Mutex);
+		CachedOSPageAllocator.FreeAll();
+		//UE_LOG(LogTemp, Display, TEXT("Trim CachedOSPageAllocator = %6.2fms"), 1000.0f * float(FPlatformTime::Seconds() - StartTime));
+	}
+}
+
+void FMallocBinned2::SetupTLSCachesOnCurrentThread()
+{
+	if (!BINNED2_ALLOW_RUNTIME_TWEAKING && !GMallocBinned2PerThreadCaches)
+	{
+		return;
+	}
+	if (!FMallocBinned2::Binned2TlsSlot)
+	{
+		FMallocBinned2::Binned2TlsSlot = FPlatformTLS::AllocTlsSlot();
+	}
+	check(FMallocBinned2::Binned2TlsSlot);
+	FPerThreadFreeBlockLists::SetTLS();
+}
 
-	Private::FPerThreadFreeBlockLists& Lists = Private::FPerThreadFreeBlockLists::Get(PageSize);
-	for (int32 PoolIndex = 0; PoolIndex != SMALL_POOL_COUNT; ++PoolIndex)
+void FMallocBinned2::ClearAndDisableTLSCachesOnCurrentThread()
+{
+	FlushCurrentThreadCache();
+	FPerThreadFreeBlockLists::ClearTLS();
+}
+
+
+bool FMallocBinned2::FFreeBlockList::ObtainPartial(uint32 InPoolIndex)
+{
+	if (!PartialBundle.Head)
 	{
-		Private::FBundleNode* Bundles = Lists.PopBundles(PoolIndex);
-		if (Bundles)
+		PartialBundle.Count = 0;
+		PartialBundle.Head = FMallocBinned2::Private::GGlobalRecycler.PopBundle(InPoolIndex);
+		if (PartialBundle.Head)
 		{
-			Private::FreeBundles(*this, Bundles, GMallocBinned2SmallBlockSizes[PoolIndex], PoolIndex);
+			PartialBundle.Count = PartialBundle.Head->Count;
+			PartialBundle.Head->NextBundle = nullptr;
+			return true;
 		}
+		return false;
+	}
+	return true;
+}
+
+FMallocBinned2::FBundleNode* FMallocBinned2::FFreeBlockList::RecyleFull(uint32 InPoolIndex)
+{
+	FMallocBinned2::FBundleNode* Result = nullptr;
+	if (FullBundle.Head)
+	{
+		FullBundle.Head->Count = FullBundle.Count;
+		if (!FMallocBinned2::Private::GGlobalRecycler.PushBundle(InPoolIndex, FullBundle.Head))
+		{
+			Result = FullBundle.Head;
+			Result->NextBundle = nullptr;
+		}
+		FullBundle.Reset();
 	}
+	return Result;
 }
+
+FMallocBinned2::FBundleNode* FMallocBinned2::FFreeBlockList::PopBundles(uint32 InPoolIndex)
+{
+	FBundleNode* Partial = PartialBundle.Head;
+	if (Partial)
+	{
+		PartialBundle.Reset();
+		Partial->NextBundle = nullptr;
+	}
+
+	FBundleNode* Full = FullBundle.Head;
+	if (Full)
+	{
+		FullBundle.Reset();
+		Full->NextBundle = nullptr;
+	}
+
+	FBundleNode* Result = Partial;
+	if (Result)
+	{
+		Result->NextBundle = Full;
+	}
+	else
+	{
+		Result = Full;
+	}
+
+	return Result;
+}
+
+void FMallocBinned2::FPerThreadFreeBlockLists::SetTLS()
+{
+	check(FMallocBinned2::Binned2TlsSlot);
+	FPerThreadFreeBlockLists* ThreadSingleton = (FPerThreadFreeBlockLists*)FPlatformTLS::GetTlsValue(FMallocBinned2::Binned2TlsSlot);
+	if (!ThreadSingleton)
+	{
+		ThreadSingleton = new (FPlatformMemory::BinnedAllocFromOS(Align(sizeof(FPerThreadFreeBlockLists), FMallocBinned2::OsAllocationGranularity))) FPerThreadFreeBlockLists();
+		FPlatformTLS::SetTlsValue(FMallocBinned2::Binned2TlsSlot, ThreadSingleton);
+	}
+}
+
+void FMallocBinned2::FPerThreadFreeBlockLists::ClearTLS()
+{
+	check(FMallocBinned2::Binned2TlsSlot);
+	FPlatformTLS::SetTlsValue(FMallocBinned2::Binned2TlsSlot, nullptr);
+}
+
+void FMallocBinned2::FFreeBlock::CanaryFail() const
+{
+	UE_LOG(LogMemory, Fatal, TEXT("FMallocBinned2 Attempt to realloc an unrecognized block %p   canary == 0x%x != 0x%x"), (void*)this, (int32)Canary, (int32)FMallocBinned2::FFreeBlock::CANARY_VALUE);
+}
+
+#if !BINNED2_INLINE
+	#if PLATFORM_USES_FIXED_GMalloc_CLASS && !FORCE_ANSI_ALLOCATOR && USE_MALLOC_BINNED2
+		//#define FMEMORY_INLINE_FUNCTION_DECORATOR  FORCEINLINE
+		#define FMEMORY_INLINE_GMalloc (FMallocBinned2::MallocBinned2)
+		#include "FMemory.inl"
+	#endif
+#endif
diff --git a/Engine/Source/Runtime/Core/Private/HAL/MallocLeakDetection.h b/Engine/Source/Runtime/Core/Private/HAL/MallocLeakDetection.h
--- a/Engine/Source/Runtime/Core/Private/HAL/MallocLeakDetection.h
+++ b/Engine/Source/Runtime/Core/Private/HAL/MallocLeakDetection.h
@@ -141,6 +141,25 @@ public:
 		return UsedMalloc->GetAllocationSize(Original, OutSize);
 	}
 
+	virtual SIZE_T QuantizeSize(SIZE_T Count, uint32 Alignment) override
+	{
+		return UsedMalloc->QuantizeSize(Count, Alignment);
+	}
+
+	virtual void Trim() override
+	{
+		return UsedMalloc->Trim();
+	}
+	virtual void SetupTLSCachesOnCurrentThread() override
+	{
+		return UsedMalloc->SetupTLSCachesOnCurrentThread();
+	}
+	virtual void ClearAndDisableTLSCachesOnCurrentThread() override
+	{
+		return UsedMalloc->ClearAndDisableTLSCachesOnCurrentThread();
+	}
+
+
 	virtual const TCHAR* GetDescriptiveName() override
 	{ 
 		return UsedMalloc->GetDescriptiveName();
diff --git a/Engine/Source/Runtime/Core/Private/HAL/MallocVerify.h b/Engine/Source/Runtime/Core/Private/HAL/MallocVerify.h
--- a/Engine/Source/Runtime/Core/Private/HAL/MallocVerify.h
+++ b/Engine/Source/Runtime/Core/Private/HAL/MallocVerify.h
@@ -104,6 +104,22 @@ public:
 	{
 		return UsedMalloc->GetAllocationSize(Original, OutSize);
 	}
+	virtual SIZE_T QuantizeSize(SIZE_T Count, uint32 Alignment) override
+	{
+		return UsedMalloc->QuantizeSize(Count, Alignment);
+	}
+	virtual void Trim() override
+	{
+		return UsedMalloc->Trim();
+	}
+	virtual void SetupTLSCachesOnCurrentThread() override
+	{
+		return UsedMalloc->SetupTLSCachesOnCurrentThread();
+	}
+	virtual void ClearAndDisableTLSCachesOnCurrentThread() override
+	{
+		return UsedMalloc->ClearAndDisableTLSCachesOnCurrentThread();
+	}
 
 	virtual const TCHAR* GetDescriptiveName() override
 	{ 
diff --git a/Engine/Source/Runtime/Core/Private/HAL/UnrealMemory.cpp b/Engine/Source/Runtime/Core/Private/HAL/UnrealMemory.cpp
--- a/Engine/Source/Runtime/Core/Private/HAL/UnrealMemory.cpp
+++ b/Engine/Source/Runtime/Core/Private/HAL/UnrealMemory.cpp
@@ -17,6 +17,36 @@
 #include "MallocLeakDetection.h"
 #include "PlatformMallocCrash.h"
 
+#if MALLOC_GT_HOOKS
+
+// This code is used to find memory allocations, you set up the lambda in the section of the code you are interested in and add a breakpoint to your lambda to see who is allocating memory
+
+//An example:
+//static TFunction<void(int32)> AnimHook(
+//	[](int32 Index)
+//{
+//	TickAnimationMallocStats[Index]++;
+//	if (++AllCount % 337 == 0)
+//	{
+//		BreakMe();
+//	}
+//}
+//);
+//extern CORE_API TFunction<void(int32)>* GGameThreadMallocHook;
+//TGuardValue<TFunction<void(int32)>*> RestoreHook(GGameThreadMallocHook, &AnimHook);
+
+
+CORE_API TFunction<void(int32)>* GGameThreadMallocHook = nullptr;
+
+void DoGamethreadHook(int32 Index)
+{
+	if (GIsRunning && GGameThreadMallocHook // otherwise our hook might not be initialized yet
+		&& IsInGameThread())
+	{
+		(*GGameThreadMallocHook)(Index);
+	}
+}
+#endif
 
 class FMallocPurgatoryProxy : public FMalloc
 {
@@ -155,6 +185,22 @@ public:
 		return UsedMalloc->GetAllocationSize(Original, SizeOut);
 	}
 
+	virtual SIZE_T QuantizeSize(SIZE_T Count, uint32 Alignment) override
+	{
+		return UsedMalloc->QuantizeSize(Count, Alignment);
+	}
+	virtual void Trim() override
+	{
+		return UsedMalloc->Trim();
+	}
+	virtual void SetupTLSCachesOnCurrentThread() override
+	{
+		return UsedMalloc->SetupTLSCachesOnCurrentThread();
+	}
+	virtual void ClearAndDisableTLSCachesOnCurrentThread() override
+	{
+		return UsedMalloc->ClearAndDisableTLSCachesOnCurrentThread();
+	}
 	virtual const TCHAR* GetDescriptiveName() override
 	{
 		return UsedMalloc->GetDescriptiveName();
@@ -163,6 +209,11 @@ public:
 
 void FMemory::EnablePurgatoryTests()
 {
+	if (PLATFORM_USES_FIXED_GMalloc_CLASS)
+	{
+		UE_LOG(LogMemory, Error, TEXT("Purgatory proxy cannot be turned on because we are using PLATFORM_USES_FIXED_GMalloc_CLASS"));
+		return;
+	}
 	static bool bOnce = false;
 	if (bOnce)
 	{
@@ -218,12 +269,21 @@ FConsoleCommandDelegate::CreateStatic(&FMemory::EnablePurgatoryTests)
 
 
 /** Helper function called on first allocation to create and initialize GMalloc */
-void GCreateMalloc()
+void FMemory::GCreateMalloc()
 {
 	GMalloc = FPlatformMemory::BaseAllocator();
 	// Setup malloc crash as soon as possible.
 	FPlatformMallocCrash::Get( GMalloc );
 
+#if PLATFORM_USES_FIXED_GMalloc_CLASS
+#if USE_MALLOC_PROFILER || MALLOC_VERIFY || MALLOC_LEAKDETECTION
+#error "Turn of PLATFORM_USES_FIXED_GMalloc_CLASS in order to use special allocator proxies"
+#endif
+	if (!GMalloc->IsInternallyThreadSafe())
+	{
+		UE_LOG(LogMemory, Fatal, TEXT("PLATFORM_USES_FIXED_GMalloc_CLASS only makes sense for allocators that are internally threadsafe."));
+	}
+#else
 // so now check to see if we are using a Mem Profiler which wraps the GMalloc
 #if USE_MALLOC_PROFILER
 	#if WITH_ENGINE && IS_MONOLITHIC
@@ -249,81 +309,143 @@ void GCreateMalloc()
 #if MALLOC_LEAKDETECTION
 	GMalloc = new FMallocLeakDetectionProxy(GMalloc);
 #endif
+#endif
 }
 
+#if TIME_MALLOC
 
-#if STATS
-	#define MALLOC_GT_HOOKS 1
-#else
-	#define MALLOC_GT_HOOKS 0
-#endif
-
-#if MALLOC_GT_HOOKS
+uint64 FScopedMallocTimer::GTotalCycles[4] = { 0 };
+uint64 FScopedMallocTimer::GTotalCount[4] = { 0 };
+uint64 FScopedMallocTimer::GTotalMisses[4] = { 0 };
 
-// This code is used to find memory allocations, you set up the lambda in the section of the code you are interested in and add a breakpoint to your lambda to see who is allocating memory
+void FScopedMallocTimer::Spew()
+{
+	static uint64 GLastTotalCycles[4] = { 0 };
+	static uint64 GLastTotalCount[4] = { 0 };
+	static uint64 GLastTotalMisses[4] = { 0 };
+	static uint64 GLastFrame = 0;
 
-//An example:
-//static TFunction<void(int32)> AnimHook(
-//	[](int32 Index)
-//{
-//	TickAnimationMallocStats[Index]++;
-//	if (++AllCount % 337 == 0)
-//	{
-//		BreakMe();
-//	}
-//}
-//);
-//extern CORE_API TFunction<void(int32)>* GGameThreadMallocHook;
-//TGuardValue<TFunction<void(int32)>*> RestoreHook(GGameThreadMallocHook, &AnimHook);
+	uint64 Frames = GFrameCounter - GLastFrame;
+	if (Frames)
+	{
+		GLastFrame = GFrameCounter;
+		// not atomic; we assume the error is minor
+		uint64 TotalCycles[4] = { 0 };
+		uint64 TotalCount[4] = { 0 };
+		uint64 TotalMisses[4] = { 0 };
+		for (int32 Comp = 0; Comp < 4; Comp++)
+		{
+			TotalCycles[Comp] = GTotalCycles[Comp] - GLastTotalCycles[Comp];
+			TotalCount[Comp] = GTotalCount[Comp] - GLastTotalCount[Comp];
+			TotalMisses[Comp] = GTotalMisses[Comp] - GLastTotalMisses[Comp];
+			GLastTotalCycles[Comp] = GTotalCycles[Comp];
+			GLastTotalCount[Comp] = GTotalCount[Comp];
+			GLastTotalMisses[Comp] = GTotalMisses[Comp];
+		}
+		auto PrintIt = [&TotalCycles, &TotalCount, &TotalMisses, &Frames](const TCHAR * Op, int32 InIndex)
+		{
+			if (TotalCount[InIndex])
+			{
+				UE_LOG(LogMemory, Display, TEXT("FMemory %8s  %5d count/frame   %6.2fms / frame (all threads)  %6.2fns / op    inline miss rate %5.2f%%"),
+					Op,
+					int32(TotalCount[InIndex] / Frames),
+					1000.0f * float(FPlatformTime::GetSecondsPerCycle64()) * float(TotalCycles[InIndex]) / float(Frames),
+					1000000000.0f * float(FPlatformTime::GetSecondsPerCycle64()) * float(TotalCycles[InIndex]) / float(TotalCount[InIndex]),
+					100.0f * float(TotalMisses[InIndex]) / float(TotalCount[InIndex])
+					);
+			}
+		};
+		PrintIt(TEXT("Malloc"), 0);
+		PrintIt(TEXT("Realloc"), 1);
+		PrintIt(TEXT("Free"), 2);
+		PrintIt(TEXT("NullFree"), 3);
+	}
+}
 
+#endif
 
-CORE_API TFunction<void(int32)>* GGameThreadMallocHook = nullptr;
+void* FMemory::MallocExternal(SIZE_T Count, uint32 Alignment)
+{
+	if (!GMalloc)
+	{
+		GCreateMalloc();
+		CA_ASSUME(GMalloc != NULL);	// Don't want to assert, but suppress static analysis warnings about potentially NULL GMalloc
+	}
+	return GMalloc->Malloc(Count, Alignment);
+}
 
-FORCEINLINE static void DoGamethreadHook(int32 Index)
+void* FMemory::ReallocExternal(void* Original, SIZE_T Count, uint32 Alignment)
 {
-	if (GIsRunning && GGameThreadMallocHook // otherwise our hook might not be initialized yet
-		&& IsInGameThread() ) 
+	if (!GMalloc)
 	{
-		(*GGameThreadMallocHook)(Index);
+		GCreateMalloc();
+		CA_ASSUME(GMalloc != NULL);	// Don't want to assert, but suppress static analysis warnings about potentially NULL GMalloc
 	}
+	return GMalloc->Realloc(Original, Count, Alignment);
 }
-#else
-FORCEINLINE static void DoGamethreadHook(int32 Index)
+
+void FMemory::FreeExternal(void* Original)
 {
+	if (!GMalloc)
+	{
+		GCreateMalloc();
+		CA_ASSUME(GMalloc != NULL);	// Don't want to assert, but suppress static analysis warnings about potentially NULL GMalloc
+	}
+	if (Original)
+	{
+		GMalloc->Free(Original);
+}
 }
-#endif
 
-void* FMemory::Malloc( SIZE_T Count, uint32 Alignment ) 
+SIZE_T FMemory::GetAllocSizeExternal(void* Original)
 { 
-	if( !GMalloc )
+	if (!GMalloc)
 	{
 		GCreateMalloc();
-		CA_ASSUME( GMalloc != NULL );	// Don't want to assert, but suppress static analysis warnings about potentially NULL GMalloc
+		CA_ASSUME(GMalloc != NULL);	// Don't want to assert, but suppress static analysis warnings about potentially NULL GMalloc
 	}
-	DoGamethreadHook(0);
-	return GMalloc->Malloc( Count, Alignment );
+	SIZE_T Size = 0;
+	return GMalloc->GetAllocationSize(Original, Size) ? Size : 0;
 }
 
-void* FMemory::Realloc( void* Original, SIZE_T Count, uint32 Alignment ) 
+SIZE_T FMemory::QuantizeSizeExternal(SIZE_T Count, uint32 Alignment)
 { 
-	if( !GMalloc )
+	if (!GMalloc)
 	{
 		GCreateMalloc();	
-		CA_ASSUME( GMalloc != NULL );	// Don't want to assert, but suppress static analysis warnings about potentially NULL GMalloc
+		CA_ASSUME(GMalloc != NULL);	// Don't want to assert, but suppress static analysis warnings about potentially NULL GMalloc
 	}
-	DoGamethreadHook(1);
-	return GMalloc->Realloc( Original, Count, Alignment );
+	return GMalloc->QuantizeSize(Count, Alignment);
 }	
 
-void FMemory::Free( void* Original )
+
+void FMemory::Trim()
 {
-	if( !GMalloc )
+	if (!GMalloc)
 	{
 		GCreateMalloc();
-		CA_ASSUME( GMalloc != NULL );	// Don't want to assert, but suppress static analysis warnings about potentially NULL GMalloc
+		CA_ASSUME(GMalloc != NULL);	// Don't want to assert, but suppress static analysis warnings about potentially NULL GMalloc
+	}
+	QUICK_SCOPE_CYCLE_COUNTER(STAT_FMemory_Trim);
+	GMalloc->Trim();
+}
+
+void FMemory::SetupTLSCachesOnCurrentThread()
+{
+	if (!GMalloc)
+	{
+		GCreateMalloc();	
+		CA_ASSUME(GMalloc != NULL);	// Don't want to assert, but suppress static analysis warnings about potentially NULL GMalloc
+	}
+	GMalloc->SetupTLSCachesOnCurrentThread();
+}
+	
+void FMemory::ClearAndDisableTLSCachesOnCurrentThread()
+{
+	if (GMalloc)
+	{
+		GMalloc->ClearAndDisableTLSCachesOnCurrentThread();
 	}
-	DoGamethreadHook(2);
-	return GMalloc->Free( Original );
 }
 
 void* FMemory::GPUMalloc(SIZE_T Count, uint32 Alignment /* = DEFAULT_ALIGNMENT */)
@@ -341,18 +463,6 @@ void FMemory::GPUFree(void* Original)
 	return FPlatformMemory::GPUFree(Original);
 }
 
-SIZE_T FMemory::GetAllocSize( void* Original )
-{
-	if( !GMalloc )
-	{
-		GCreateMalloc();	
-		CA_ASSUME( GMalloc != NULL );	// Don't want to assert, but suppress static analysis warnings about potentially NULL GMalloc
-	}
-	
-	SIZE_T Size = 0;
-	return GMalloc->GetAllocationSize( Original, Size ) ? Size : 0;
-}
-
 void FMemory::TestMemory()
 {
 #if !UE_BUILD_SHIPPING
@@ -400,3 +510,27 @@ void FMemory::TestMemory()
 	}
 #endif
 }
+
+void* FUseSystemMallocForNew::operator new(size_t Size)
+{
+	return FMemory::SystemMalloc(Size);
+}
+
+void FUseSystemMallocForNew::operator delete(void* Ptr)
+{
+	FMemory::SystemFree(Ptr);
+}
+
+void* FUseSystemMallocForNew::operator new[](size_t Size)
+{
+	return FMemory::SystemMalloc(Size);
+}
+
+void FUseSystemMallocForNew::operator delete[](void* Ptr)
+{
+	FMemory::SystemFree(Ptr);
+}
+
+#if !INLINE_FMEMORY_OPERATION && !PLATFORM_USES_FIXED_GMalloc_CLASS
+#include "FMemory.inl"
+#endif
\ No newline at end of file
diff --git a/Engine/Source/Runtime/Core/Private/Mac/MacMallocZone.cpp b/Engine/Source/Runtime/Core/Private/Mac/MacMallocZone.cpp
--- a/Engine/Source/Runtime/Core/Private/Mac/MacMallocZone.cpp
+++ b/Engine/Source/Runtime/Core/Private/Mac/MacMallocZone.cpp
@@ -85,6 +85,10 @@ void FMacMallocCrashHandler::Enable( FMacCrashContext* Context, uint32 CrashedTh
 	CrashContext = Context;
 	ThreadId = CrashedThreadId;
 	OriginalHeap = GMalloc;
+	if (PLATFORM_USES_FIXED_GMalloc_CLASS && GFixedMallocLocationPtr)
+	{
+		*GFixedMallocLocationPtr = nullptr; // this disables any fast-path inline allocators
+	}
 	GMalloc = this;
 }
 	
diff --git a/Engine/Source/Runtime/Core/Private/Misc/Core.cpp b/Engine/Source/Runtime/Core/Private/Misc/Core.cpp
--- a/Engine/Source/Runtime/Core/Private/Misc/Core.cpp
+++ b/Engine/Source/Runtime/Core/Private/Misc/Core.cpp
@@ -30,6 +30,7 @@ FConfigCacheIni*				GConfig						= nullptr;		/* Configuration database cache */
 ITransaction*				GUndo						= nullptr;		/* Transaction tracker, non-NULL when a transaction is in progress */
 FOutputDeviceConsole*		GLogConsole					= nullptr;		/* Console log hook */
 CORE_API FMalloc*			GMalloc						= nullptr;		/* Memory allocator */
+CORE_API FMalloc**			GFixedMallocLocationPtr = nullptr;		/* Memory allocator pointer location when PLATFORM_USES_FIXED_GMalloc_CLASS is true */
 
 class UPropertyWindowManager*	GPropertyWindowManager	= nullptr;		/* Manages and tracks property editing windows */
 
diff --git a/Engine/Source/Runtime/Core/Private/Misc/MemStack.cpp b/Engine/Source/Runtime/Core/Private/Misc/MemStack.cpp
--- a/Engine/Source/Runtime/Core/Private/Misc/MemStack.cpp
+++ b/Engine/Source/Runtime/Core/Private/Misc/MemStack.cpp
@@ -7,7 +7,6 @@ DECLARE_MEMORY_STAT(TEXT("PageAllocator Free"), STAT_PageAllocatorFree, STATGROU
 DECLARE_MEMORY_STAT(TEXT("PageAllocator Used"), STAT_PageAllocatorUsed, STATGROUP_Memory);
 
 FPageAllocator::TPageAllocator FPageAllocator::TheAllocator;
-FPageAllocator::TSmallPageAllocator FPageAllocator::TheSmallAllocator;
 
 
 #define USE_MEMSTACK_PUGATORY (0)
@@ -150,7 +149,7 @@ void *FPageAllocator::AllocSmall()
 	}
 	else
 	{
-		Result = TheSmallAllocator.Allocate();
+		Result = FMemory::Malloc(SmallPageSize);
 	}
 	STAT(UpdateStats());
 	return Result;
@@ -177,7 +176,7 @@ void FPageAllocator::FreeSmall(void *Mem)
 	}
 	else
 	{
-		TheSmallAllocator.Free(Mem);
+		FMemory::Free(Mem);
 	}
 	STAT(UpdateStats());
 }
@@ -215,25 +214,22 @@ void FPageAllocator::Free(void *Mem)
 
 void *FPageAllocator::AllocSmall()
 {
-	void *Result = TheSmallAllocator.Allocate();
-	STAT(UpdateStats());
-	return Result;
+	return FMemory::Malloc(SmallPageSize);
 }
 
 void FPageAllocator::FreeSmall(void *Mem)
 {
-	TheSmallAllocator.Free(Mem);
-	STAT(UpdateStats());
+	FMemory::Free(Mem);
 }
 
 uint64 FPageAllocator::BytesUsed()
 {
-	return uint64(TheAllocator.GetNumUsed().GetValue()) * PageSize + uint64(TheSmallAllocator.GetNumUsed().GetValue()) * SmallPageSize;
+	return uint64(TheAllocator.GetNumUsed().GetValue()) * PageSize;
 }
 
 uint64 FPageAllocator::BytesFree()
 {
-	return uint64(TheAllocator.GetNumFree().GetValue()) * PageSize + uint64(TheSmallAllocator.GetNumFree().GetValue()) * SmallPageSize;
+	return uint64(TheAllocator.GetNumFree().GetValue()) * PageSize;
 }
 
 #endif
diff --git a/Engine/Source/Runtime/Core/Private/Stats/Stats2.cpp b/Engine/Source/Runtime/Core/Private/Stats/Stats2.cpp
--- a/Engine/Source/Runtime/Core/Private/Stats/Stats2.cpp
+++ b/Engine/Source/Runtime/Core/Private/Stats/Stats2.cpp
@@ -844,8 +844,10 @@ public:
 	/** Attaches to the task graph stats thread, all processing will be handled by the task graph. */
 	virtual uint32 Run() override
 	{
+		FMemory::SetupTLSCachesOnCurrentThread();
 		FTaskGraphInterface::Get().AttachToThread(ENamedThreads::StatsThread);
 		FTaskGraphInterface::Get().ProcessThreadUntilRequestReturn(ENamedThreads::StatsThread);
+		FMemory::ClearAndDisableTLSCachesOnCurrentThread();
 		return 0;
 	}
 
diff --git a/Engine/Source/Runtime/Core/Public/Async/TaskGraphInterfaces.h b/Engine/Source/Runtime/Core/Public/Async/TaskGraphInterfaces.h
--- a/Engine/Source/Runtime/Core/Public/Async/TaskGraphInterfaces.h
+++ b/Engine/Source/Runtime/Core/Public/Async/TaskGraphInterfaces.h
@@ -329,10 +329,16 @@ public:
 	}
 
 	/**
-	*	When a set of tasks complete, fire a scoped event
+	*	Deletegates for shutdown
 	*	@param	Callback - function to call prior to shutting down the taskgraph
 	**/
 	virtual void AddShutdownCallback(TFunction<void()>& Callback) = 0;
+
+	/**
+	*	A (slow) function to call a function on every known thread, both named and workers
+	*	@param	Callback - function to call prior to shutting down the taskgraph
+	**/
+	static void BroadcastSlow_OnlyUseForSpecialPurposes(bool bDoBackgroundThreads, TFunction<void(ENamedThreads::Type CurrentThread)>& Callback);
 };
 
 /** 
diff --git a/Engine/Source/Runtime/Core/Public/Containers/Array.h b/Engine/Source/Runtime/Core/Public/Containers/Array.h
--- a/Engine/Source/Runtime/Core/Public/Containers/Array.h
+++ b/Engine/Source/Runtime/Core/Public/Containers/Array.h
@@ -18,6 +18,7 @@
 	#define TARRAY_RANGED_FOR_CHECKS 1
 #endif
 
+#define AGRESSIVE_ARRAY_FORCEINLINE
 
 /**
  * Generic iterator which can operate on types that expose the following:
@@ -229,15 +230,15 @@ private:
 class FScriptArray : protected FHeapAllocator::ForAnyElementType
 {
 public:
-	void* GetData()
+	FORCEINLINE void* GetData()
 	{
 		return this->GetAllocation();
 	}
-	const void* GetData() const
+	FORCEINLINE const void* GetData() const
 	{
 		return this->GetAllocation();
 	}
-	bool IsValidIndex( int32 i ) const
+	FORCEINLINE bool IsValidIndex(int32 i) const
 	{
 		return i>=0 && i<ArrayNum;
 	}
@@ -263,8 +264,7 @@ public:
 		const int32 OldNum = ArrayNum;
 		if( (ArrayNum+=Count)>ArrayMax )
 		{
-			ArrayMax = this->CalculateSlack( ArrayNum, ArrayMax, NumBytesPerElement );
-			this->ResizeAllocation(OldNum,ArrayMax,NumBytesPerElement);
+			ResizeGrow(OldNum, NumBytesPerElement);
 		}
 		FMemory::Memmove
 		(
@@ -273,7 +273,7 @@ public:
 			                                               (OldNum-Index)*NumBytesPerElement
 		);
 	}
-	int32 Add( int32 Count, int32 NumBytesPerElement )
+	AGRESSIVE_ARRAY_FORCEINLINE int32 Add( int32 Count, int32 NumBytesPerElement )
 	{
 		check(Count>=0);
 		checkSlow(ArrayNum>=0);
@@ -282,8 +282,7 @@ public:
 		const int32 OldNum = ArrayNum;
 		if( (ArrayNum+=Count)>ArrayMax )
 		{
-			ArrayMax = this->CalculateSlack( ArrayNum, ArrayMax, NumBytesPerElement );
-			this->ResizeAllocation(OldNum,ArrayMax,NumBytesPerElement);
+			ResizeGrow(OldNum, NumBytesPerElement);
 		}
 
 		return OldNum;
@@ -298,21 +297,18 @@ public:
 	{
 		checkSlow(ArrayNum>=0);
 		checkSlow(ArrayMax>=ArrayNum);
-		if( ArrayMax != ArrayNum )
+		if (ArrayNum != ArrayMax)
 		{
-			ArrayMax = ArrayNum;
-			this->ResizeAllocation(ArrayNum,ArrayMax,NumBytesPerElement);
+			ResizeTo(ArrayNum, NumBytesPerElement);
 		}
 	}
 	void Empty( int32 Slack, int32 NumBytesPerElement )
 	{
 		checkSlow(Slack>=0);
 		ArrayNum = 0;
-		// only reallocate if we need to, I don't trust realloc to the same size to work
-		if (ArrayMax != Slack)
+		if (Slack != ArrayMax)
 		{
-			ArrayMax = Slack;
-			this->ResizeAllocation(0,ArrayMax,NumBytesPerElement);
+			ResizeTo(Slack, NumBytesPerElement);
 		}
 	}
 	void SwapMemory(int32 A, int32 B, int32 NumBytesPerElement )
@@ -335,53 +331,86 @@ public:
 	/**
 	 * Returns the amount of slack in this array in elements.
 	 */
-	int32 GetSlack() const
+	FORCEINLINE int32 GetSlack() const
 	{
 		return ArrayMax - ArrayNum;
 	}
 		
 	void Remove( int32 Index, int32 Count, int32 NumBytesPerElement  )
 	{
-		checkSlow(Count >= 0);
-		checkSlow(Index >= 0); 
-		checkSlow(Index <= ArrayNum);
-		checkSlow(Index + Count <= ArrayNum);
-
-		// Skip memmove in the common case that there is nothing to move.
-		int32 NumToMove = ArrayNum - Index - Count;
-		if( NumToMove )
-		{
-			FMemory::Memmove
-			(
-				(uint8*)this->GetAllocation() + (Index      ) * NumBytesPerElement,
-				(uint8*)this->GetAllocation() + (Index+Count) * NumBytesPerElement,
-				NumToMove * NumBytesPerElement
-			);
-		}
-		ArrayNum -= Count;
-		
-		const int32 NewArrayMax = this->CalculateSlack(ArrayNum,ArrayMax,NumBytesPerElement);
-		if(NewArrayMax != ArrayMax)
+		if (Count)
 		{
-			ArrayMax = NewArrayMax;
-			this->ResizeAllocation(ArrayNum,ArrayMax,NumBytesPerElement);
+			checkSlow(Count >= 0);
+			checkSlow(Index >= 0);
+			checkSlow(Index <= ArrayNum);
+			checkSlow(Index + Count <= ArrayNum);
+
+			// Skip memmove in the common case that there is nothing to move.
+			int32 NumToMove = ArrayNum - Index - Count;
+			if (NumToMove)
+			{
+				FMemory::Memmove
+					(
+					(uint8*)this->GetAllocation() + (Index)* NumBytesPerElement,
+					(uint8*)this->GetAllocation() + (Index + Count) * NumBytesPerElement,
+					NumToMove * NumBytesPerElement
+					);
+			}
+			ArrayNum -= Count;
+
+			ResizeShrink(NumBytesPerElement);
+			checkSlow(ArrayNum >= 0);
+			checkSlow(ArrayMax >= ArrayNum);
 		}
-		checkSlow(ArrayNum >= 0);
-		checkSlow(ArrayMax >= ArrayNum);
 	}
 
 protected:
 
 	FScriptArray( int32 InNum, int32 NumBytesPerElement  )
-	:   ArrayNum( InNum )
+	:   ArrayNum( 0 )
 	,	ArrayMax( InNum )
 
 	{
-		this->ResizeAllocation(0,ArrayMax,NumBytesPerElement);
+		if (ArrayMax)
+		{
+			ResizeInit(NumBytesPerElement);
+		}
+		ArrayNum = InNum;
 	}
 	int32	  ArrayNum;
 	int32	  ArrayMax;
 
+	FORCENOINLINE void ResizeInit(int32 NumBytesPerElement)
+	{
+		ArrayMax = this->CalculateSlackReserve(ArrayMax, NumBytesPerElement);
+		this->ResizeAllocation(ArrayNum, ArrayMax, NumBytesPerElement);
+	}
+	FORCENOINLINE void ResizeGrow(int32 OldNum, int32 NumBytesPerElement)
+	{
+		ArrayMax = this->CalculateSlackGrow(ArrayNum, ArrayMax, NumBytesPerElement);
+		this->ResizeAllocation(OldNum, ArrayMax, NumBytesPerElement);
+	}
+	FORCENOINLINE void ResizeShrink(int32 NumBytesPerElement)
+	{
+		const int32 NewArrayMax = this->CalculateSlackShrink(ArrayNum, ArrayMax, NumBytesPerElement);
+		if (NewArrayMax != ArrayMax)
+		{
+			ArrayMax = NewArrayMax;
+			this->ResizeAllocation(ArrayNum, ArrayMax, NumBytesPerElement);
+		}
+	}
+	FORCENOINLINE void ResizeTo(int32 NewMax, int32 NumBytesPerElement)
+	{
+		if (NewMax)
+		{
+			NewMax = this->CalculateSlackReserve(NewMax, NumBytesPerElement);
+		}
+		if (NewMax != ArrayMax)
+		{
+			ArrayMax = NewMax;
+			this->ResizeAllocation(ArrayNum, ArrayMax, NumBytesPerElement);
+		}
+	}
 public:
 	// These should really be private, because they shouldn't be called, but there's a bunch of code
 	// that needs to be fixed first.
@@ -482,7 +511,7 @@ public:
 	/**
 	 * Constructor, initializes element number counters.
 	 */
-	TArray()
+	FORCEINLINE TArray()
 		: ArrayNum(0)
 		, ArrayMax(0)
 	{}
@@ -493,7 +522,7 @@ public:
 	 * @param Other The source array to copy.
 	 */
 	template <typename OtherElementType, typename OtherAllocator>
-	explicit TArray(const TArray<OtherElementType, OtherAllocator>& Other)
+	FORCEINLINE explicit TArray(const TArray<OtherElementType, OtherAllocator>& Other)
 	{
 		CopyToEmpty(Other, 0, 0);
 	}
@@ -503,7 +532,7 @@ public:
 	 *
 	 * @param Other The source array to copy.
 	 */
-	TArray(const TArray& Other)
+	FORCEINLINE TArray(const TArray& Other)
 	{
 		CopyToEmpty(Other, 0, 0);
 	}
@@ -515,7 +544,7 @@ public:
 	 * @param ExtraSlack Tells how much extra memory should be preallocated
 	 *                   at the end of the array in the number of elements.
 	 */
-	TArray(const TArray& Other, int32 ExtraSlack)
+	FORCEINLINE TArray(const TArray& Other, int32 ExtraSlack)
 	{
 		CopyToEmpty(Other, 0, ExtraSlack);
 	}
@@ -529,7 +558,7 @@ public:
 	 * @param Other The source array to assign from.
 	 */
 	template<typename OtherAllocator>
-	TArray& operator=(const TArray<ElementType, OtherAllocator>& Other)
+	AGRESSIVE_ARRAY_FORCEINLINE TArray& operator=(const TArray<ElementType, OtherAllocator>& Other)
 	{
 		DestructItems(GetData(), ArrayNum);
 		CopyToEmpty(Other, ArrayMax, 0);
@@ -542,7 +571,7 @@ public:
 	 *
 	 * @param Other The source array to assign from.
 	 */
-	TArray& operator=(const TArray& Other)
+	AGRESSIVE_ARRAY_FORCEINLINE TArray& operator=(const TArray& Other)
 	{
 		if (this != &Other)
 		{
@@ -652,7 +681,7 @@ public:
 	 *                   at the end of the array in the number of elements.
 	 */
 	template <typename OtherElementType>
-	TArray(TArray<OtherElementType, Allocator>&& Other, int32 ExtraSlack)
+	AGRESSIVE_ARRAY_FORCEINLINE TArray(TArray<OtherElementType, Allocator>&& Other, int32 ExtraSlack)
 	{
 		// We don't implement move semantics for general OtherAllocators, as there's no way
 		// to tell if they're compatible with the current one.  Probably going to be a pretty
@@ -666,7 +695,7 @@ public:
 	 *
 	 * @param Other Array to assign and move from.
 	 */
-	TArray& operator=(TArray&& Other)
+	AGRESSIVE_ARRAY_FORCEINLINE TArray& operator=(TArray&& Other)
 	{
 		if (this != &Other)
 		{
@@ -679,7 +708,7 @@ public:
 	/**
 	 * Destructor.
 	 */
-	~TArray()
+	AGRESSIVE_ARRAY_FORCEINLINE ~TArray()
 	{
 		DestructItems(GetData(), ArrayNum);
 
@@ -756,7 +785,7 @@ public:
 	/**
 	 * Returns the amount of slack in this array in elements.
 	 */
-	int32 GetSlack() const
+	FORCEINLINE int32 GetSlack() const
 	{
 		return ArrayMax - ArrayNum;
 	}
@@ -850,7 +879,7 @@ public:
 	 *
 	 * @returns Popped element.
 	 */
-	ElementType Pop(bool bAllowShrinking = true)
+	FORCEINLINE ElementType Pop(bool bAllowShrinking = true)
 	{
 		RangeCheck(0);
 		ElementType Result = MoveTemp(GetData()[ArrayNum - 1]);
@@ -863,7 +892,7 @@ public:
 	 *
 	 * @param Item Item to push.
 	 */
-	void Push(ElementType&& Item)
+	FORCEINLINE void Push(ElementType&& Item)
 	{
 		Add(MoveTemp(Item));
 	}
@@ -875,7 +904,7 @@ public:
 	 *
 	 * @param Item Item to push.
 	 */
-	void Push(const ElementType& Item)
+	FORCEINLINE void Push(const ElementType& Item)
 	{
 		Add(Item);
 	}
@@ -885,7 +914,7 @@ public:
 	 *
 	 * @returns Reference to the top element.
 	 */
-	ElementType& Top()
+	FORCEINLINE ElementType& Top()
 	{
 		return Last();
 	}
@@ -897,7 +926,7 @@ public:
 	 *
 	 * @returns Reference to the top element.
 	 */
-	const ElementType& Top() const
+	FORCEINLINE const ElementType& Top() const
 	{
 		return Last();
 	}
@@ -910,7 +939,7 @@ public:
 	 *
 	 * @returns Reference to n-th last element from the array.
 	 */
-	ElementType& Last(int32 IndexFromTheEnd = 0)
+	FORCEINLINE ElementType& Last(int32 IndexFromTheEnd = 0)
 	{
 		RangeCheck(ArrayNum - IndexFromTheEnd - 1);
 		return GetData()[ArrayNum - IndexFromTheEnd - 1];
@@ -926,7 +955,7 @@ public:
 	 *
 	 * @returns Reference to n-th last element from the array.
 	 */
-	const ElementType& Last(int32 IndexFromTheEnd = 0) const
+	FORCEINLINE const ElementType& Last(int32 IndexFromTheEnd = 0) const
 	{
 		RangeCheck(ArrayNum - IndexFromTheEnd - 1);
 		return GetData()[ArrayNum - IndexFromTheEnd - 1];
@@ -936,13 +965,12 @@ public:
 	 * Shrinks the array's used memory to smallest possible to store elements
 	 * currently in it.
 	 */
-	void Shrink()
+	FORCEINLINE void Shrink()
 	{
 		CheckInvariants();
 		if (ArrayMax != ArrayNum)
 		{
-			ArrayMax = ArrayNum;
-			AllocatorInstance.ResizeAllocation(ArrayNum, ArrayMax, sizeof(ElementType));
+			ResizeTo(ArrayNum);
 		}
 	}
 
@@ -967,7 +995,7 @@ public:
 	 *
 	 * @returns Index of the found element. INDEX_NONE otherwise.
 	 */
-	int32 Find(const ElementType& Item) const
+	AGRESSIVE_ARRAY_FORCEINLINE int32 Find(const ElementType& Item) const
 	{
 		const ElementType* RESTRICT Start = GetData();
 		for (const ElementType* RESTRICT Data = Start, *RESTRICT DataEnd = Data + ArrayNum; Data != DataEnd; ++Data)
@@ -1088,7 +1116,7 @@ public:
 	 *          found.
 	 */
 	template <typename KeyType>
-	int32 IndexOfByKey(const KeyType& Key) const
+	AGRESSIVE_ARRAY_FORCEINLINE int32 IndexOfByKey(const KeyType& Key) const
 	{
 		const ElementType* RESTRICT Start = GetData();
 		for (const ElementType* RESTRICT Data = Start, *RESTRICT DataEnd = Start + ArrayNum; Data != DataEnd; ++Data)
@@ -1110,7 +1138,7 @@ public:
 	 *          found.
 	 */
 	template <typename Predicate>
-	int32 IndexOfByPredicate(Predicate Pred) const
+	AGRESSIVE_ARRAY_FORCEINLINE int32 IndexOfByPredicate(Predicate Pred) const
 	{
 		const ElementType* RESTRICT Start = GetData();
 		for (const ElementType* RESTRICT Data = Start, *RESTRICT DataEnd = Start + ArrayNum; Data != DataEnd; ++Data)
@@ -1146,7 +1174,7 @@ public:
 	 * @returns Pointer to the first matching element, or nullptr if none is found.
 	 */
 	template <typename KeyType>
-	ElementType* FindByKey(const KeyType& Key)
+	AGRESSIVE_ARRAY_FORCEINLINE ElementType* FindByKey(const KeyType& Key)
 	{
 		for (ElementType* RESTRICT Data = GetData(), *RESTRICT DataEnd = Data + ArrayNum; Data != DataEnd; ++Data)
 		{
@@ -1182,7 +1210,7 @@ public:
 	 *         true, or nullptr if none is found.
 	 */
 	template <typename Predicate>
-	ElementType* FindByPredicate(Predicate Pred)
+	AGRESSIVE_ARRAY_FORCEINLINE ElementType* FindByPredicate(Predicate Pred)
 	{
 		for (ElementType* RESTRICT Data = GetData(), *RESTRICT DataEnd = Data + ArrayNum; Data != DataEnd; ++Data)
 		{
@@ -1223,7 +1251,7 @@ public:
 	 * @returns	True if found. False otherwise.
 	 */
 	template <typename ComparisonType>
-	bool Contains(const ComparisonType& Item) const
+	AGRESSIVE_ARRAY_FORCEINLINE bool Contains(const ComparisonType& Item) const
 	{
 		for (const ElementType* RESTRICT Data = GetData(), *RESTRICT DataEnd = Data + ArrayNum; Data != DataEnd; ++Data)
 		{
@@ -1255,7 +1283,7 @@ public:
 	 *
 	 * @returns True if this array is the same as OtherArray. False otherwise.
 	 */
-	bool operator==(const TArray& OtherArray) const
+	AGRESSIVE_ARRAY_FORCEINLINE bool operator==(const TArray& OtherArray) const
 	{
 		int32 Count = Num();
 
@@ -1269,7 +1297,7 @@ public:
 	 *
 	 * @returns True if this array is NOT the same as OtherArray. False otherwise.
 	 */
-	bool operator!=(const TArray& OtherArray) const
+	FORCEINLINE bool operator!=(const TArray& OtherArray) const
 	{
 		return !(*this == OtherArray);
 	}
@@ -1290,10 +1318,9 @@ public:
 			// Serialize simple bytes which require no construction or destruction.
 			Ar << A.ArrayNum;
 			check(A.ArrayNum >= 0);
-			if (Ar.IsLoading())
+			if ((A.ArrayNum || A.ArrayMax) && Ar.IsLoading())
 			{
-				A.ArrayMax = A.ArrayNum;
-				A.AllocatorInstance.ResizeAllocation(0, A.ArrayMax, sizeof(ElementType));
+				A.ResizeForCopy(A.ArrayNum, A.ArrayMax);
 			}
 			Ar.Serialize(A.GetData(), A.Num());
 		}
@@ -1406,7 +1433,7 @@ public:
 	 *
 	 * @returns Number of elements in array before addition.
 	 */
-	int32 AddUninitialized(int32 Count = 1)
+	FORCEINLINE int32 AddUninitialized(int32 Count = 1)
 	{
 		CheckInvariants();
 		checkSlow(Count >= 0);
@@ -1414,10 +1441,8 @@ public:
 		const int32 OldNum = ArrayNum;
 		if ((ArrayNum += Count) > ArrayMax)
 		{
-			ArrayMax = AllocatorInstance.CalculateSlack(ArrayNum, ArrayMax, sizeof(ElementType));
-			AllocatorInstance.ResizeAllocation(OldNum, ArrayMax, sizeof(ElementType));
+			ResizeGrow(OldNum);
 		}
-
 		return OldNum;
 	}
 
@@ -1440,8 +1465,7 @@ public:
 		const int32 OldNum = ArrayNum;
 		if ((ArrayNum += Count) > ArrayMax)
 		{
-			ArrayMax = AllocatorInstance.CalculateSlack(ArrayNum, ArrayMax, sizeof(ElementType));
-			AllocatorInstance.ResizeAllocation(OldNum, ArrayMax, sizeof(ElementType));
+			ResizeGrow(OldNum);
 		}
 		ElementType* Data = GetData() + Index;
 		RelocateConstructItems<ElementType>(Data + Count, Data, OldNum - Index);
@@ -1566,31 +1590,29 @@ public:
 	 */
 	void RemoveAt(int32 Index, int32 Count = 1, bool bAllowShrinking = true)
 	{
-		CheckInvariants();
-		checkSlow((Count >= 0) & (Index >= 0) & (Index + Count <= ArrayNum));
+		if (Count)
+		{
+			CheckInvariants();
+			checkSlow((Count >= 0) & (Index >= 0) & (Index + Count <= ArrayNum));
 
-		DestructItems(GetData() + Index, Count);
+			DestructItems(GetData() + Index, Count);
 
-		// Skip memmove in the common case that there is nothing to move.
-		int32 NumToMove = ArrayNum - Index - Count;
-		if (NumToMove)
-		{
-			FMemory::Memmove
-				(
-				(uint8*)AllocatorInstance.GetAllocation() + (Index)* sizeof(ElementType),
-				(uint8*)AllocatorInstance.GetAllocation() + (Index + Count) * sizeof(ElementType),
-				NumToMove * sizeof(ElementType)
-				);
-		}
-		ArrayNum -= Count;
+			// Skip memmove in the common case that there is nothing to move.
+			int32 NumToMove = ArrayNum - Index - Count;
+			if (NumToMove)
+			{
+				FMemory::Memmove
+					(
+					(uint8*)AllocatorInstance.GetAllocation() + (Index)* sizeof(ElementType),
+					(uint8*)AllocatorInstance.GetAllocation() + (Index + Count) * sizeof(ElementType),
+					NumToMove * sizeof(ElementType)
+					);
+			}
+			ArrayNum -= Count;
 
-		if (bAllowShrinking)
-		{
-			const int32 NewArrayMax = AllocatorInstance.CalculateSlack(ArrayNum, ArrayMax, sizeof(ElementType));
-			if (NewArrayMax != ArrayMax)
+			if (bAllowShrinking)
 			{
-				ArrayMax = NewArrayMax;
-				AllocatorInstance.ResizeAllocation(ArrayNum, ArrayMax, sizeof(ElementType));
+				ResizeShrink();
 			}
 		}
 	}
@@ -1607,34 +1629,32 @@ public:
 	 * @param bAllowShrinking (Optional) Tells if this call can shrink array if
 	 *                        suitable after remove. Default is true.
 	 */
-	void RemoveAtSwap(int32 Index, int32 Count = 1, bool bAllowShrinking = true)
+	AGRESSIVE_ARRAY_FORCEINLINE void RemoveAtSwap(int32 Index, int32 Count = 1, bool bAllowShrinking = true)
 	{
-		CheckInvariants();
-		checkSlow((Count >= 0) & (Index >= 0) & (Index + Count <= ArrayNum));
+		if (Count)
+		{
+			CheckInvariants();
+			checkSlow((Count >= 0) & (Index >= 0) & (Index + Count <= ArrayNum));
 
-		DestructItems(GetData() + Index, Count);
+			DestructItems(GetData() + Index, Count);
 
-		// Replace the elements in the hole created by the removal with elements from the end of the array, so the range of indices used by the array is contiguous.
-		const int32 NumElementsInHole = Count;
-		const int32 NumElementsAfterHole = ArrayNum - (Index + Count);
-		const int32 NumElementsToMoveIntoHole = FMath::Min(NumElementsInHole, NumElementsAfterHole);
-		if (NumElementsToMoveIntoHole)
-		{
-			FMemory::Memcpy(
-				(uint8*)AllocatorInstance.GetAllocation() + (Index)* sizeof(ElementType),
-				(uint8*)AllocatorInstance.GetAllocation() + (ArrayNum - NumElementsToMoveIntoHole) * sizeof(ElementType),
-				NumElementsToMoveIntoHole * sizeof(ElementType)
-				);
-		}
-		ArrayNum -= Count;
+			// Replace the elements in the hole created by the removal with elements from the end of the array, so the range of indices used by the array is contiguous.
+			const int32 NumElementsInHole = Count;
+			const int32 NumElementsAfterHole = ArrayNum - (Index + Count);
+			const int32 NumElementsToMoveIntoHole = FMath::Min(NumElementsInHole, NumElementsAfterHole);
+			if (NumElementsToMoveIntoHole)
+			{
+				FMemory::Memcpy(
+					(uint8*)AllocatorInstance.GetAllocation() + (Index)* sizeof(ElementType),
+					(uint8*)AllocatorInstance.GetAllocation() + (ArrayNum - NumElementsToMoveIntoHole) * sizeof(ElementType),
+					NumElementsToMoveIntoHole * sizeof(ElementType)
+					);
+			}
+			ArrayNum -= Count;
 
-		if (bAllowShrinking)
-		{
-			const int32 NewArrayMax = AllocatorInstance.CalculateSlack(ArrayNum, ArrayMax, sizeof(ElementType));
-			if (NewArrayMax != ArrayMax)
+			if (bAllowShrinking)
 			{
-				ArrayMax = NewArrayMax;
-				AllocatorInstance.ResizeAllocation(ArrayNum, ArrayMax, sizeof(ElementType));
+				ResizeShrink();
 			}
 		}
 	}
@@ -1645,7 +1665,7 @@ public:
 	 *
 	 * @param NewSize The expected usage size after calling this function.
 	 */
-	void Reset(int32 NewSize = 0)
+	AGRESSIVE_ARRAY_FORCEINLINE void Reset(int32 NewSize = 0)
 	{
 		// If we have space to hold the excepted size, then don't reallocate
 		if (NewSize <= ArrayMax)
@@ -1664,17 +1684,15 @@ public:
 	 *
 	 * @param Slack (Optional) The expected usage size after empty operation. Default is 0.
 	 */
-	void Empty(int32 Slack = 0)
+	AGRESSIVE_ARRAY_FORCEINLINE void Empty(int32 Slack = 0)
 	{
 		DestructItems(GetData(), ArrayNum);
 
 		checkSlow(Slack >= 0);
 		ArrayNum = 0;
-		// only reallocate if we need to, I don't trust realloc to the same size to work
 		if (ArrayMax != Slack)
 		{
-			AllocatorInstance.ResizeAllocation(0, Slack, sizeof(ElementType));
-			ArrayMax = Slack;
+			ResizeTo(Slack);
 		}
 	}
 
@@ -1720,7 +1738,7 @@ public:
 	 *
 	 * @param NewNum New size of the array.
 	 */
-	void SetNumUninitialized(int32 NewNum, bool bAllowShrinking = true)
+	AGRESSIVE_ARRAY_FORCEINLINE void SetNumUninitialized(int32 NewNum, bool bAllowShrinking = true)
 	{
 		if (NewNum > Num())
 		{
@@ -1750,7 +1768,7 @@ public:
 	 * @param Source The array to append.
 	 */
 	template <typename OtherElementType, typename OtherAllocator>
-	FORCEINLINE void Append(const TArray<OtherElementType, OtherAllocator>& Source)
+	AGRESSIVE_ARRAY_FORCEINLINE void Append(const TArray<OtherElementType, OtherAllocator>& Source)
 	{
 		check((void*)this != (void*)&Source);
 
@@ -1776,7 +1794,7 @@ public:
 	 * @param Source The array to append.
 	 */
 	template <typename OtherElementType, typename OtherAllocator>
-	FORCEINLINE void Append(TArray<OtherElementType, OtherAllocator>&& Source)
+	AGRESSIVE_ARRAY_FORCEINLINE void Append(TArray<OtherElementType, OtherAllocator>&& Source)
 	{
 		check((void*)this != (void*)&Source);
 
@@ -1803,7 +1821,7 @@ public:
 	 * @param Ptr   A pointer to an array of elements to add.
 	 * @param Count The number of elements to insert from Ptr.
 	 */
-	void Append(const ElementType* Ptr, int32 Count)
+	AGRESSIVE_ARRAY_FORCEINLINE void Append(const ElementType* Ptr, int32 Count)
 	{
 		check(Ptr != nullptr);
 
@@ -1819,7 +1837,7 @@ public:
 	 *
 	 * @param Other The array to append.
 	 */
-	FORCEINLINE TArray& operator+=(TArray&& Other)
+	AGRESSIVE_ARRAY_FORCEINLINE TArray& operator+=(TArray&& Other)
 	{
 		Append(MoveTemp(Other));
 		return *this;
@@ -1831,7 +1849,7 @@ public:
 	 *
 	 * @param Other The array to append.
 	 */
-	FORCEINLINE TArray& operator+=(const TArray& Other)
+	AGRESSIVE_ARRAY_FORCEINLINE TArray& operator+=(const TArray& Other)
 	{
 		Append(Other);
 		return *this;
@@ -1846,7 +1864,7 @@ public:
 	#if PLATFORM_COMPILER_HAS_VARIADIC_TEMPLATES
 
 		template <typename... ArgsType>
-		int32 Emplace(ArgsType&&... Args)
+		FORCEINLINE int32 Emplace(ArgsType&&... Args)
 		{
 			const int32 Index = AddUninitialized(1);
 			new(GetData() + Index) ElementType(Forward<ArgsType>(Args)...);
@@ -1855,7 +1873,7 @@ public:
 
 	#else
 
-		int32 Emplace()
+		FORCEINLINE int32 Emplace()
 		{
 			const int32 Index = AddUninitialized(1);
 			new(GetData() + Index) ElementType();
@@ -1863,7 +1881,7 @@ public:
 		}
 
 		template <typename Arg0Type>
-		int32 Emplace(Arg0Type&& Arg0)
+		FORCEINLINE int32 Emplace(Arg0Type&& Arg0)
 		{
 			const int32 Index = AddUninitialized(1);
 			new(GetData() + Index) ElementType(Forward<Arg0Type>(Arg0));
@@ -1871,7 +1889,7 @@ public:
 		}
 
 		template <typename Arg0Type, typename Arg1Type>
-		int32 Emplace(Arg0Type&& Arg0, Arg1Type&& Arg1)
+		FORCEINLINE int32 Emplace(Arg0Type&& Arg0, Arg1Type&& Arg1)
 		{
 			const int32 Index = AddUninitialized(1);
 			new(GetData() + Index) ElementType(Forward<Arg0Type>(Arg0), Forward<Arg1Type>(Arg1));
@@ -1879,7 +1897,7 @@ public:
 		}
 
 		template <typename Arg0Type, typename Arg1Type, typename Arg2Type>
-		int32 Emplace(Arg0Type&& Arg0, Arg1Type&& Arg1, Arg2Type&& Arg2)
+		FORCEINLINE int32 Emplace(Arg0Type&& Arg0, Arg1Type&& Arg1, Arg2Type&& Arg2)
 		{
 			const int32 Index = AddUninitialized(1);
 			new(GetData() + Index) ElementType(Forward<Arg0Type>(Arg0), Forward<Arg1Type>(Arg1), Forward<Arg2Type>(Arg2));
@@ -1887,7 +1905,7 @@ public:
 		}
 
 		template <typename Arg0Type, typename Arg1Type, typename Arg2Type, typename Arg3Type>
-		int32 Emplace(Arg0Type&& Arg0, Arg1Type&& Arg1, Arg2Type&& Arg2, Arg3Type&& Arg3)
+		FORCEINLINE int32 Emplace(Arg0Type&& Arg0, Arg1Type&& Arg1, Arg2Type&& Arg2, Arg3Type&& Arg3)
 		{
 			const int32 Index = AddUninitialized(1);
 			new(GetData() + Index) ElementType(Forward<Arg0Type>(Arg0), Forward<Arg1Type>(Arg1), Forward<Arg2Type>(Arg2), Forward<Arg3Type>(Arg3));
@@ -1895,7 +1913,7 @@ public:
 		}
 
 		template <typename Arg0Type, typename Arg1Type, typename Arg2Type, typename Arg3Type, typename Arg4Type>
-		int32 Emplace(Arg0Type&& Arg0, Arg1Type&& Arg1, Arg2Type&& Arg2, Arg3Type&& Arg3, Arg4Type&& Arg4)
+		FORCEINLINE int32 Emplace(Arg0Type&& Arg0, Arg1Type&& Arg1, Arg2Type&& Arg2, Arg3Type&& Arg3, Arg4Type&& Arg4)
 		{
 			const int32 Index = AddUninitialized(1);
 			new(GetData() + Index) ElementType(Forward<Arg0Type>(Arg0), Forward<Arg1Type>(Arg1), Forward<Arg2Type>(Arg2), Forward<Arg3Type>(Arg3), Forward<Arg4Type>(Arg4));
@@ -1903,7 +1921,7 @@ public:
 		}
 
 		template <typename Arg0Type, typename Arg1Type, typename Arg2Type, typename Arg3Type, typename Arg4Type, typename Arg5Type>
-		int32 Emplace(Arg0Type&& Arg0, Arg1Type&& Arg1, Arg2Type&& Arg2, Arg3Type&& Arg3, Arg4Type&& Arg4, Arg5Type&& Arg5)
+		FORCEINLINE int32 Emplace(Arg0Type&& Arg0, Arg1Type&& Arg1, Arg2Type&& Arg2, Arg3Type&& Arg3, Arg4Type&& Arg4, Arg5Type&& Arg5)
 		{
 			const int32 Index = AddUninitialized(1);
 			new(GetData() + Index) ElementType(Forward<Arg0Type>(Arg0), Forward<Arg1Type>(Arg1), Forward<Arg2Type>(Arg2), Forward<Arg3Type>(Arg3), Forward<Arg4Type>(Arg4), Forward<Arg5Type>(Arg5));
@@ -1942,7 +1960,7 @@ public:
 	 *
 	 * @return Index to the first of the new items.
 	 */
-	int32 AddZeroed(int32 Count = 1)
+	AGRESSIVE_ARRAY_FORCEINLINE int32 AddZeroed(int32 Count = 1)
 	{
 		const int32 Index = AddUninitialized(Count);
 		FMemory::Memzero((uint8*)AllocatorInstance.GetAllocation() + Index*sizeof(ElementType), Count*sizeof(ElementType));
@@ -1957,7 +1975,7 @@ public:
 	 *
 	 * @return Index to the first of the new items.
 	 */
-	int32 AddDefaulted(int32 Count = 1)
+	AGRESSIVE_ARRAY_FORCEINLINE int32 AddDefaulted(int32 Count = 1)
 	{
 		const int32 Index = AddUninitialized(Count);
 		DefaultConstructItems<ElementType>((uint8*)AllocatorInstance.GetAllocation() + Index * sizeof(ElementType), Count);
@@ -1973,7 +1991,7 @@ private:
 	 * @returns Index of the element in the array.
 	 */
 	template <typename ArgsType>
-	int32 AddUniqueImpl(ArgsType&& Args)
+	AGRESSIVE_ARRAY_FORCEINLINE int32 AddUniqueImpl(ArgsType&& Args)
 	{
 		int32 Index;
 		if (Find(Args, Index))
@@ -2011,12 +2029,11 @@ public:
 	 * @param Number The number of elements that the array should be able to
 	 *               contain after allocation.
 	 */
-	void Reserve(int32 Number)
+	FORCEINLINE void Reserve(int32 Number)
 	{
 		if (Number > ArrayMax)
 		{
-			ArrayMax = Number;
-			AllocatorInstance.ResizeAllocation(ArrayNum, ArrayMax, sizeof(ElementType));
+			ResizeTo(Number);
 		}
 	}
 	
@@ -2040,7 +2057,7 @@ public:
 	 * @param Number The number of elements that the array should be able to
 	 *               contain after allocation.
 	 */
-	void Init(const ElementType& Element, int32 Number)
+	AGRESSIVE_ARRAY_FORCEINLINE void Init(const ElementType& Element, int32 Number)
 	{
 		Empty(Number);
 		for (int32 Index = 0; Index < Number; ++Index)
@@ -2220,7 +2237,7 @@ public:
 	 * @param FirstIndexToSwap Position of the first element to swap.
 	 * @param SecondIndexToSwap Position of the second element to swap.
 	 */
-	void SwapMemory(int32 FirstIndexToSwap, int32 SecondIndexToSwap)
+	FORCEINLINE void SwapMemory(int32 FirstIndexToSwap, int32 SecondIndexToSwap)
 	{
 		FMemory::Memswap(
 			(uint8*)AllocatorInstance.GetAllocation() + (sizeof(ElementType)*FirstIndexToSwap),
@@ -2237,7 +2254,7 @@ public:
 	 * @param FirstIndexToSwap Position of the first element to swap.
 	 * @param SecondIndexToSwap Position of the second element to swap.
 	 */
-	void Swap(int32 FirstIndexToSwap, int32 SecondIndexToSwap)
+	FORCEINLINE void Swap(int32 FirstIndexToSwap, int32 SecondIndexToSwap)
 	{
 		check((FirstIndexToSwap >= 0) && (SecondIndexToSwap >= 0));
 		check((ArrayNum > FirstIndexToSwap) && (ArrayNum > SecondIndexToSwap));
@@ -2392,6 +2409,47 @@ private:
 
 private:
 
+	FORCENOINLINE void ResizeGrow(int32 OldNum)
+	{
+		ArrayMax = AllocatorInstance.CalculateSlackGrow(ArrayNum, ArrayMax, sizeof(ElementType));
+		AllocatorInstance.ResizeAllocation(OldNum, ArrayMax, sizeof(ElementType));
+	}
+	FORCENOINLINE void ResizeShrink()
+	{
+		const int32 NewArrayMax = AllocatorInstance.CalculateSlackShrink(ArrayNum, ArrayMax, sizeof(ElementType));
+		if (NewArrayMax != ArrayMax)
+		{
+			ArrayMax = NewArrayMax;
+			check(ArrayMax >= ArrayNum);
+			AllocatorInstance.ResizeAllocation(ArrayNum, ArrayMax, sizeof(ElementType));
+		}
+	}
+	FORCENOINLINE void ResizeTo(int32 NewMax)
+	{
+		if (NewMax)
+		{
+			NewMax = AllocatorInstance.CalculateSlackReserve(NewMax, sizeof(ElementType));
+		}
+		if (NewMax != ArrayMax)
+		{
+			ArrayMax = NewMax;
+			AllocatorInstance.ResizeAllocation(ArrayNum, ArrayMax, sizeof(ElementType));
+		}
+	}
+	FORCENOINLINE void ResizeForCopy(int32 NewMax, int32 PrevMax)
+	{
+		if (NewMax)
+		{
+			NewMax = AllocatorInstance.CalculateSlackReserve(NewMax, sizeof(ElementType));
+		}
+		if (NewMax != PrevMax)
+		{
+			AllocatorInstance.ResizeAllocation(0, NewMax, sizeof(ElementType));
+		}
+		ArrayMax = NewMax;
+	}
+
+
 	/**
 	 * Copies data from one array into this array. Uses the fast path if the
 	 * data in question does not need a constructor.
@@ -2403,20 +2461,19 @@ private:
 	 *                   default.
 	 */
 	template <typename OtherElementType, typename OtherAllocator>
-	void CopyToEmpty(const TArray<OtherElementType, OtherAllocator>& Source, int32 PrevMax, int32 ExtraSlack)
+	AGRESSIVE_ARRAY_FORCEINLINE void CopyToEmpty(const TArray<OtherElementType, OtherAllocator>& Source, int32 PrevMax, int32 ExtraSlack)
 	{
-		check(ExtraSlack >= 0);
-
-		int32 SourceCount = Source.Num();
-		if (SourceCount + ExtraSlack != PrevMax)
+		checkSlow(ExtraSlack >= 0);
+		ArrayNum = Source.Num();
+		if (ArrayNum || ExtraSlack || PrevMax)
 		{
-			AllocatorInstance.ResizeAllocation(0, SourceCount + ExtraSlack, sizeof(ElementType));
+			ResizeForCopy(ArrayNum + ExtraSlack, PrevMax);
+			ConstructItems<ElementType>(GetData(), Source.GetData(), ArrayNum);
+		}
+		else
+		{
+			ArrayMax = 0;
 		}
-
-		ConstructItems<ElementType>(GetData(), Source.GetData(), SourceCount);
-
-		ArrayNum = SourceCount;
-		ArrayMax = SourceCount + ExtraSlack;
 	}
 
 protected:
diff --git a/Engine/Source/Runtime/Core/Public/Containers/BitArray.h b/Engine/Source/Runtime/Core/Public/Containers/BitArray.h
--- a/Engine/Source/Runtime/Core/Public/Containers/BitArray.h
+++ b/Engine/Source/Runtime/Core/Public/Containers/BitArray.h
@@ -287,7 +287,7 @@ public:
 		if(bReallocate)
 		{
 			// Allocate memory for the new bits.
-			const uint32 MaxDWORDs = AllocatorInstance.CalculateSlack(
+			const uint32 MaxDWORDs = AllocatorInstance.CalculateSlackGrow(
 				FMath::DivideAndRoundUp(NumBits, NumBitsPerDWORD),
 				FMath::DivideAndRoundUp(MaxBits, NumBitsPerDWORD),
 				sizeof(uint32)
@@ -309,6 +309,7 @@ public:
 	{
 		NumBits = 0;
 
+		ExpectedNumBits = FMath::DivideAndRoundUp(ExpectedNumBits, NumBitsPerDWORD) * NumBitsPerDWORD;
 		// If the expected number of bits doesn't match the allocated number of bits, reallocate.
 		if(MaxBits != ExpectedNumBits)
 		{
@@ -701,7 +702,7 @@ private:
 	int32         NumBits;
 	int32         MaxBits;
 
-	void Realloc(int32 PreviousNumBits)
+	FORCENOINLINE void Realloc(int32 PreviousNumBits)
 	{
 		const int32 PreviousNumDWORDs = FMath::DivideAndRoundUp(PreviousNumBits, NumBitsPerDWORD);
 		const int32 MaxDWORDs = FMath::DivideAndRoundUp(MaxBits, NumBitsPerDWORD);
@@ -984,6 +985,7 @@ public:
 	{
 		NumBits = 0;
 
+		Slack = FMath::DivideAndRoundUp(Slack, NumBitsPerDWORD) * NumBitsPerDWORD;
 		// If the expected number of bits doesn't match the allocated number of bits, reallocate.
 		if (MaxBits != Slack)
 		{
@@ -995,24 +997,12 @@ public:
 	int32 Add(const bool Value)
 	{
 		const int32 Index = NumBits;
-		const bool bReallocate = (NumBits + 1) > MaxBits;
-
 		NumBits++;
-
-		if(bReallocate)
+		if (NumBits > MaxBits)
 		{
-			// Allocate memory for the new bits.
-			const uint32 MaxDWORDs = AllocatorInstance.CalculateSlack(
-				FMath::DivideAndRoundUp(NumBits, NumBitsPerDWORD),
-				FMath::DivideAndRoundUp(MaxBits, NumBitsPerDWORD),
-				sizeof(uint32)
-				);
-			MaxBits = MaxDWORDs * NumBitsPerDWORD;
-			Realloc(NumBits - 1);
+			ReallocGrow(NumBits - 1);
 		}
-
 		(*this)[Index] = Value;
-
 		return Index;
 	}
 
@@ -1054,14 +1044,35 @@ private:
 		return (const uint32*)AllocatorInstance.GetAllocation();
 	}
 
-	void Realloc(int32 PreviousNumBits)
+	FORCENOINLINE void Realloc(int32 PreviousNumBits)
 	{
+		const uint32 MaxDWORDs = AllocatorInstance.CalculateSlackReserve(
+			FMath::DivideAndRoundUp(MaxBits, NumBitsPerDWORD),
+			sizeof(uint32)
+			);
+		MaxBits = MaxDWORDs * NumBitsPerDWORD;
 		const int32 PreviousNumDWORDs = FMath::DivideAndRoundUp(PreviousNumBits, NumBitsPerDWORD);
-		const int32 MaxDWORDs = FMath::DivideAndRoundUp(MaxBits, NumBitsPerDWORD);
 
 		AllocatorInstance.ResizeAllocation(PreviousNumDWORDs, MaxDWORDs, sizeof(uint32));
 
-		if (MaxDWORDs)
+		if (MaxDWORDs && MaxDWORDs - PreviousNumDWORDs > 0)
+		{
+			// Reset the newly allocated slack DWORDs.
+			FMemory::Memzero((uint32*)AllocatorInstance.GetAllocation() + PreviousNumDWORDs, (MaxDWORDs - PreviousNumDWORDs) * sizeof(uint32));
+		}
+	}
+	FORCENOINLINE void ReallocGrow(int32 PreviousNumBits)
+	{
+		// Allocate memory for the new bits.
+		const uint32 MaxDWORDs = AllocatorInstance.CalculateSlackGrow(
+			FMath::DivideAndRoundUp(NumBits, NumBitsPerDWORD),
+			FMath::DivideAndRoundUp(MaxBits, NumBitsPerDWORD),
+			sizeof(uint32)
+			);
+		MaxBits = MaxDWORDs * NumBitsPerDWORD;
+		const int32 PreviousNumDWORDs = FMath::DivideAndRoundUp(PreviousNumBits, NumBitsPerDWORD);
+		AllocatorInstance.ResizeAllocation(PreviousNumDWORDs, MaxDWORDs, sizeof(uint32));
+		if (MaxDWORDs && MaxDWORDs - PreviousNumDWORDs > 0)
 		{
 			// Reset the newly allocated slack DWORDs.
 			FMemory::Memzero((uint32*)AllocatorInstance.GetAllocation() + PreviousNumDWORDs, (MaxDWORDs - PreviousNumDWORDs) * sizeof(uint32));
diff --git a/Engine/Source/Runtime/Core/Public/Containers/ContainerAllocationPolicies.h b/Engine/Source/Runtime/Core/Public/Containers/ContainerAllocationPolicies.h
--- a/Engine/Source/Runtime/Core/Public/Containers/ContainerAllocationPolicies.h
+++ b/Engine/Source/Runtime/Core/Public/Containers/ContainerAllocationPolicies.h
@@ -73,9 +73,90 @@ ReferencedType* IfAThenAElseB(ReferencedType* A,ReferencedType* B);
 template<typename PredicateType,typename ReferencedType>
 ReferencedType* IfPThenAElseB(PredicateType Predicate,ReferencedType* A,ReferencedType* B);
 
-/** The default slack calculation heuristic. */
-extern CORE_API int32 DefaultCalculateSlack(int32 NumElements,int32 NumAllocatedElements,SIZE_T BytesPerElement);
+#if TRACK_ARRAY_SLACK 
+CORE_API void TrackSlack(int32 NumElements, int32 NumAllocatedElements, SIZE_T BytesPerElement, int32 Retval);
+#endif
+
+FORCEINLINE int32 DefaultCalculateSlackShrink(int32 NumElements, int32 NumAllocatedElements, SIZE_T BytesPerElement, bool bAllowQuantize, uint32 Alignment = DEFAULT_ALIGNMENT)
+{
+	int32 Retval;
+	check(NumElements < NumAllocatedElements);
+
+	// If the container has too much slack, shrink it to exactly fit the number of elements.
+	const uint32 CurrentSlackElements = NumAllocatedElements - NumElements;
+	const SIZE_T CurrentSlackBytes = (NumAllocatedElements - NumElements)*BytesPerElement;
+	const bool bTooManySlackBytes = CurrentSlackBytes >= 16384;
+	const bool bTooManySlackElements = 3 * NumElements < 2 * NumAllocatedElements;
+	if ((bTooManySlackBytes || bTooManySlackElements) && (CurrentSlackElements > 64 || !NumElements)) //  hard coded 64 :-(
+	{
+		Retval = NumElements;
+		if (Retval > 0)
+		{
+			if (bAllowQuantize)
+			{
+				Retval = FMemory::QuantizeSize(Retval * BytesPerElement, Alignment) / BytesPerElement;
+			}
+		}
+	}
+	else
+	{
+		Retval = NumAllocatedElements;
+	}
+#if TRACK_ARRAY_SLACK 
+	void TrackSlack(NumElements, NumAllocatedElements, BytesPerElement, Retval);
+#endif
+	return Retval;
+}
+
+FORCEINLINE int32 DefaultCalculateSlackGrow(int32 NumElements, int32 NumAllocatedElements, SIZE_T BytesPerElement, bool bAllowQuantize, uint32 Alignment = DEFAULT_ALIGNMENT)
+{
+	int32 Retval;
+	check(NumElements > NumAllocatedElements && NumElements > 0);
 
+	SIZE_T Grow = 4; // this is the amount for the first alloc
+	if (NumAllocatedElements || SIZE_T(NumElements) > Grow)
+	{
+		// Allocate slack for the array proportional to its size.
+		Grow = SIZE_T(NumElements) + 3 * SIZE_T(NumElements) / 8 + 16;
+	}
+	if (bAllowQuantize)
+	{
+		Retval = FMemory::QuantizeSize(Grow * BytesPerElement, Alignment) / BytesPerElement;
+	}
+	else
+	{
+		Retval = Grow;
+	}
+	// NumElements and MaxElements are stored in 32 bit signed integers so we must be careful not to overflow here.
+	if (NumElements > Retval)
+	{
+		Retval = MAX_int32;
+	}
+
+#if TRACK_ARRAY_SLACK 
+	void TrackSlack(NumElements, NumAllocatedElements, BytesPerElement, Retval);
+#endif
+	return Retval;
+}
+
+FORCEINLINE int32 DefaultCalculateSlackReserve(int32 NumElements, SIZE_T BytesPerElement, bool bAllowQuantize, uint32 Alignment = DEFAULT_ALIGNMENT)
+{
+	int32 Retval = NumElements;
+	check(NumElements > 0);
+	if (bAllowQuantize)
+	{
+		Retval = FMemory::QuantizeSize(SIZE_T(Retval) * SIZE_T(BytesPerElement), Alignment) / BytesPerElement;
+		// NumElements and MaxElements are stored in 32 bit signed integers so we must be careful not to overflow here.
+		if (NumElements > Retval)
+		{
+			Retval = MAX_int32;
+		}
+	}
+#if TRACK_ARRAY_SLACK 
+	void TrackSlack(NumElements, NumAllocatedElements, BytesPerElement, Retval);
+#endif
+	return Retval;
+}
 
 /** A type which is used to represent a script type that is unknown at compile time. */
 struct FScriptContainerElement
@@ -189,9 +270,9 @@ public:
 			);
 
 		/**
-		 * Calculates the amount of slack to allocate for an array that has just grown to a given number of elements.
+		 * Calculates the amount of slack to allocate for an array that has just grown or shrunk to a given number of elements.
 		 * @param NumElements - The number of elements to allocate space for.
-		 * @param CurrentNumSlackElements - The current number of slack elements allocated.
+		 * @param CurrentNumSlackElements - The current number of elements allocated.
 		 * @param NumBytesPerElement - The number of bytes/element.
 		 */
 		int32 CalculateSlack(
@@ -200,6 +281,30 @@ public:
 			SIZE_T NumBytesPerElement
 			) const;
 
+		/**
+		* Calculates the amount of slack to allocate for an array that has just shrunk to a given number of elements.
+		* @param NumElements - The number of elements to allocate space for.
+		* @param CurrentNumSlackElements - The current number of elements allocated.
+		* @param NumBytesPerElement - The number of bytes/element.
+		*/
+		int32 CalculateSlackShrink(
+			int32 NumElements,
+			int32 CurrentNumSlackElements,
+			SIZE_T NumBytesPerElement
+			) const;
+
+		/**
+		* Calculates the amount of slack to allocate for an array that has just grown to a given number of elements.
+		* @param NumElements - The number of elements to allocate space for.
+		* @param CurrentNumSlackElements - The current number of elements allocated.
+		* @param NumBytesPerElement - The number of bytes/element.
+		*/
+		int32 CalculateSlackGrow(
+			int32 NumElements,
+			int32 CurrentNumSlackElements,
+			SIZE_T NumBytesPerElement
+			) const;
+
 		SIZE_T GetAllocatedSize(int32 NumAllocatedElements, SIZE_T NumBytesPerElement) const;
 	};
 
@@ -285,13 +390,17 @@ public:
 #endif
 			}
 		}
-		int32 CalculateSlack(
-			int32 NumElements,
-			int32 NumAllocatedElements,
-			SIZE_T NumBytesPerElement
-			) const
+		FORCEINLINE int32 CalculateSlackReserve(int32 NumElements, int32 NumBytesPerElement) const
 		{
-			return DefaultCalculateSlack(NumElements,NumAllocatedElements,NumBytesPerElement);
+			return DefaultCalculateSlackReserve(NumElements, NumBytesPerElement, true, Alignment);
+		}
+		FORCEINLINE int32 CalculateSlackShrink(int32 NumElements, int32 NumAllocatedElements, int32 NumBytesPerElement) const
+		{
+			return DefaultCalculateSlackShrink(NumElements, NumAllocatedElements, NumBytesPerElement, true, Alignment);
+		}
+		FORCEINLINE int32 CalculateSlackGrow(int32 NumElements, int32 NumAllocatedElements, int32 NumBytesPerElement) const
+		{
+			return DefaultCalculateSlackGrow(NumElements, NumAllocatedElements, NumBytesPerElement, true, Alignment);
 		}
 
 		SIZE_T GetAllocatedSize(int32 NumAllocatedElements, SIZE_T NumBytesPerElement) const
@@ -299,6 +408,11 @@ public:
 			return NumAllocatedElements * NumBytesPerElement;
 		}
 
+		bool HasAllocation()
+		{
+			return !!Data;
+		}
+
 	private:
 		ForAnyElementType(const ForAnyElementType&);
 		ForAnyElementType& operator=(const ForAnyElementType&);
@@ -378,18 +492,26 @@ public:
 		{
 			return Data;
 		}
-		void ResizeAllocation(int32 PreviousNumElements,int32 NumElements,SIZE_T NumBytesPerElement)
+		FORCEINLINE void ResizeAllocation(int32 PreviousNumElements, int32 NumElements, SIZE_T NumBytesPerElement)
 		{
 			// Avoid calling FMemory::Realloc( nullptr, 0 ) as ANSI C mandates returning a valid pointer which is not what we want.
-			if( Data || NumElements )
+			if (Data || NumElements)
 			{
 				//checkSlow(((uint64)NumElements*(uint64)ElementTypeInfo.GetSize() < (uint64)INT_MAX));
 				Data = (FScriptContainerElement*)FMemory::Realloc( Data, NumElements*NumBytesPerElement );
 			}
 		}
-		int32 CalculateSlack(int32 NumElements,int32 NumAllocatedElements,SIZE_T NumBytesPerElement) const
+		FORCEINLINE int32 CalculateSlackReserve(int32 NumElements, int32 NumBytesPerElement) const
+		{
+			return DefaultCalculateSlackReserve(NumElements, NumBytesPerElement, true);
+		}
+		FORCEINLINE int32 CalculateSlackShrink(int32 NumElements, int32 NumAllocatedElements, int32 NumBytesPerElement) const
+		{
+			return DefaultCalculateSlackShrink(NumElements, NumAllocatedElements, NumBytesPerElement, true);
+		}
+		FORCEINLINE int32 CalculateSlackGrow(int32 NumElements, int32 NumAllocatedElements, int32 NumBytesPerElement) const
 		{
-			return DefaultCalculateSlack(NumElements,NumAllocatedElements,NumBytesPerElement);
+			return DefaultCalculateSlackGrow(NumElements, NumAllocatedElements, NumBytesPerElement, true);
 		}
 
 		SIZE_T GetAllocatedSize(int32 NumAllocatedElements, SIZE_T NumBytesPerElement) const
@@ -397,6 +519,11 @@ public:
 			return NumAllocatedElements * NumBytesPerElement;
 		}
 
+		bool HasAllocation()
+		{
+			return !!Data;
+		}
+
 	private:
 		ForAnyElementType(const ForAnyElementType&);
 		ForAnyElementType& operator=(const ForAnyElementType&);
@@ -511,12 +638,26 @@ public:
 			}
 		}
 
-		int32 CalculateSlack(int32 NumElements,int32 NumAllocatedElements,SIZE_T NumBytesPerElement) const
+		FORCEINLINE int32 CalculateSlackReserve(int32 NumElements, SIZE_T NumBytesPerElement) const
 		{
 			// If the elements use less space than the inline allocation, only use the inline allocation as slack.
 			return NumElements <= NumInlineElements ?
 				NumInlineElements :
-				SecondaryData.CalculateSlack(NumElements,NumAllocatedElements,NumBytesPerElement);
+								  SecondaryData.CalculateSlackReserve(NumElements, NumBytesPerElement);
+		}
+		FORCEINLINE int32 CalculateSlackShrink(int32 NumElements, int32 NumAllocatedElements, int32 NumBytesPerElement) const
+		{
+			// If the elements use less space than the inline allocation, only use the inline allocation as slack.
+			return NumElements <= NumInlineElements ?
+			NumInlineElements :
+							  SecondaryData.CalculateSlackShrink(NumElements, NumAllocatedElements, NumBytesPerElement);
+		}
+		FORCEINLINE int32 CalculateSlackGrow(int32 NumElements, int32 NumAllocatedElements, int32 NumBytesPerElement) const
+		{
+			// If the elements use less space than the inline allocation, only use the inline allocation as slack.
+			return NumElements <= NumInlineElements ?
+			NumInlineElements :
+							  SecondaryData.CalculateSlackGrow(NumElements, NumAllocatedElements, NumBytesPerElement);
 		}
 
 		SIZE_T GetAllocatedSize(int32 NumAllocatedElements, SIZE_T NumBytesPerElement) const
@@ -524,6 +665,11 @@ public:
 			return SecondaryData.GetAllocatedSize(NumAllocatedElements, NumBytesPerElement);
 		}
 
+		bool HasAllocation()
+		{
+			return SecondaryData.HasAllocation();
+		}
+
 	private:
 		ForElementType(const ForElementType&);
 		ForElementType& operator=(const ForElementType&);
@@ -597,12 +743,23 @@ public:
 			check(NumElements <= NumInlineElements);
 		}
 
-		int32 CalculateSlack(int32 NumElements,int32 NumAllocatedElements,SIZE_T NumBytesPerElement) const
+		FORCEINLINE int32 CalculateSlackReserve(int32 NumElements, SIZE_T NumBytesPerElement) const
 		{
 			// Ensure the requested allocation will fit in the inline data area.
 			check(NumElements <= NumInlineElements);
-
-			return NumElements;
+			return NumInlineElements;
+		}
+		FORCEINLINE int32 CalculateSlackShrink(int32 NumElements, int32 NumAllocatedElements, int32 NumBytesPerElement) const
+		{
+			// Ensure the requested allocation will fit in the inline data area.
+			check(NumAllocatedElements <= NumInlineElements);
+			return NumInlineElements;
+		}
+		FORCEINLINE int32 CalculateSlackGrow(int32 NumElements, int32 NumAllocatedElements, int32 NumBytesPerElement) const
+		{
+			// Ensure the requested allocation will fit in the inline data area.
+			check(NumElements <= NumInlineElements);
+			return NumInlineElements;
 		}
 
 		SIZE_T GetAllocatedSize(int32 NumAllocatedElements, SIZE_T NumBytesPerElement) const
@@ -610,6 +767,12 @@ public:
 			return 0;
 		}
 
+		bool HasAllocation()
+		{
+			return false;
+		}
+
+
 	private:
 		ForElementType(const ForElementType&);
 		ForElementType& operator=(const ForElementType&);
diff --git a/Engine/Source/Runtime/Core/Public/GenericPlatform/GenericPlatformMemory.h b/Engine/Source/Runtime/Core/Public/GenericPlatform/GenericPlatformMemory.h
--- a/Engine/Source/Runtime/Core/Public/GenericPlatform/GenericPlatformMemory.h
+++ b/Engine/Source/Runtime/Core/Public/GenericPlatform/GenericPlatformMemory.h
@@ -33,6 +33,17 @@ struct FGenericPlatformMemoryConstants
 	/** The size of a page, in bytes. */
 	SIZE_T PageSize;
 
+	/** 
+	* For platforms that support multiple page sizes this is non-zero and smaller than PageSize.
+	* If non-zero, then BinnedAllocFromOS will take allocation requests aligned to this size and return blocks aligned to PageSize
+	*/
+	SIZE_T OsAllocationGranularity;
+
+	// AddressLimit - Second parameter is estimate of the range of addresses expected to be returns by BinnedAllocFromOS(). Binned
+	// Malloc will adjust its internal structures to make lookups for memory allocations O(1) for this range. 
+	// It is ok to go outside this range, lookups will just be a little slower
+	uint64 AddressLimit;
+
 	/** Approximate physical RAM in GB; 1 on everything except PC. Used for "course tuning", like FPlatformMisc::NumberOfCores(). */
 	uint32 TotalPhysicalGB;
 
@@ -41,6 +52,8 @@ struct FGenericPlatformMemoryConstants
 		: TotalPhysical( 0 )
 		, TotalVirtual( 0 )
 		, PageSize( 0 )
+		, OsAllocationGranularity(0)
+		, AddressLimit((uint64)0xffffffff + 1)
 		, TotalPhysicalGB( 1 )
 	{}
 
@@ -49,7 +62,9 @@ struct FGenericPlatformMemoryConstants
 		: TotalPhysical( Other.TotalPhysical )
 		, TotalVirtual( Other.TotalVirtual )
 		, PageSize( Other.PageSize )
-		, TotalPhysicalGB( Other.TotalPhysicalGB )
+		, OsAllocationGranularity(Other.OsAllocationGranularity)
+		, AddressLimit(Other.AddressLimit)
+		, TotalPhysicalGB(Other.TotalPhysicalGB)
 	{}
 };
 
diff --git a/Engine/Source/Runtime/Core/Public/GenericPlatform/GenericPlatformTime.h b/Engine/Source/Runtime/Core/Public/GenericPlatform/GenericPlatformTime.h
--- a/Engine/Source/Runtime/Core/Public/GenericPlatform/GenericPlatformTime.h
+++ b/Engine/Source/Runtime/Core/Public/GenericPlatform/GenericPlatformTime.h
@@ -67,6 +67,12 @@ struct CORE_API FGenericPlatformTime
 		gettimeofday( &tv, NULL );
 		return (uint32) ((((uint64)tv.tv_sec) * 1000000ULL) + (((uint64)tv.tv_usec)));
 	}
+	static FORCEINLINE uint64 Cycles64()
+	{
+		struct timeval tv;
+		gettimeofday( &tv, NULL );
+		return ((((uint64)tv.tv_sec) * 1000000ULL) + (((uint64)tv.tv_usec)));
+	}
 
 	/** Returns the system time. */
 	static void SystemTime( int32& Year, int32& Month, int32& DayOfWeek, int32& Day, int32& Hour, int32& Min, int32& Sec, int32& MSec );
@@ -122,6 +128,13 @@ struct CORE_API FGenericPlatformTime
 		return FCPUTime( 0.0f, 0.0f );
 	}
 
+	/**
+	* @return seconds per cycle.
+	*/
+	static double GetSecondsPerCycle()
+	{
+		return SecondsPerCycle;
+	}
 	/** Converts cycles to milliseconds. */
 	static float ToMilliseconds( const uint32 Cycles )
 	{
@@ -133,16 +146,24 @@ struct CORE_API FGenericPlatformTime
 	{
 		return (float)double( SecondsPerCycle * Cycles );
 	}
-
 	/**
 	 * @return seconds per cycle.
 	 */
-	static double GetSecondsPerCycle()
+	static double GetSecondsPerCycle64();
+	/** Converts cycles to milliseconds. */
+	static double ToMilliseconds64(const uint64 Cycles)
 	{
-		return SecondsPerCycle;
+		return ToSeconds64(Cycles) * 1000.0;
+	}
+
+	/** Converts cycles to seconds. */
+	static double ToSeconds64(const uint64 Cycles)
+	{
+		return GetSecondsPerCycle64() * double(Cycles);
 	}
 
 protected:
 
 	static double SecondsPerCycle;
+	static double SecondsPerCycle64;
 };
diff --git a/Engine/Source/Runtime/Core/Public/HAL/Allocators/CachedOSPageAllocator.h b/Engine/Source/Runtime/Core/Public/HAL/Allocators/CachedOSPageAllocator.h
--- a/Engine/Source/Runtime/Core/Public/HAL/Allocators/CachedOSPageAllocator.h
+++ b/Engine/Source/Runtime/Core/Public/HAL/Allocators/CachedOSPageAllocator.h
@@ -19,6 +19,7 @@ protected:
 
 	void* AllocateImpl(SIZE_T Size, FFreePageBlock* First, FFreePageBlock* Last, uint32& FreedPageBlocksNum, uint32& CachedTotal);
 	void FreeImpl(void* Ptr, SIZE_T Size, uint32 NumCacheBlocks, uint32 CachedByteLimit, FFreePageBlock* First, uint32& FreedPageBlocksNum, uint32& CachedTotal);
+	void FreeAllImpl(FFreePageBlock* First, uint32& FreedPageBlocksNum, uint32& CachedTotal);
 };
 
 template <uint32 NumCacheBlocks, uint32 CachedByteLimit>
@@ -39,6 +40,10 @@ struct TCachedOSPageAllocator : private FCachedOSPageAllocator
 	{
 		return FreeImpl(Ptr, Size, NumCacheBlocks, CachedByteLimit, FreedPageBlocks, FreedPageBlocksNum, CachedTotal);
 	}
+	void FreeAll()
+	{
+		return FreeAllImpl(FreedPageBlocks, FreedPageBlocksNum, CachedTotal);
+	}
 
 private:
 	FFreePageBlock FreedPageBlocks[NumCacheBlocks];
diff --git a/Engine/Source/Runtime/Core/Public/HAL/FMemory.inl b/Engine/Source/Runtime/Core/Public/HAL/FMemory.inl
new file mode 100644
--- /dev/null
+++ b/Engine/Source/Runtime/Core/Public/HAL/FMemory.inl
@@ -0,0 +1,70 @@
+// Copyright 1998-2016 Epic Games, Inc. All Rights Reserved.
+
+#if !defined(FMEMORY_INLINE_FUNCTION_DECORATOR)
+	#define FMEMORY_INLINE_FUNCTION_DECORATOR 
+#endif
+
+#if !defined(FMEMORY_INLINE_GMalloc)
+	#define FMEMORY_INLINE_GMalloc GMalloc
+#endif
+
+FMEMORY_INLINE_FUNCTION_DECORATOR void* FMemory::Malloc(SIZE_T Count, uint32 Alignment)
+{
+	if (!FMEMORY_INLINE_GMalloc)
+	{
+		return MallocExternal(Count, Alignment);
+	}
+	DoGamethreadHook(0);
+	FScopedMallocTimer Timer(0);
+	return FMEMORY_INLINE_GMalloc->Malloc(Count, Alignment);
+}
+
+FMEMORY_INLINE_FUNCTION_DECORATOR void* FMemory::Realloc(void* Original, SIZE_T Count, uint32 Alignment)
+{
+	if (!FMEMORY_INLINE_GMalloc)
+	{
+		return ReallocExternal(Original, Count, Alignment);
+	}
+	DoGamethreadHook(1);
+	FScopedMallocTimer Timer(1);
+	return FMEMORY_INLINE_GMalloc->Realloc(Original, Count, Alignment);
+}
+
+FMEMORY_INLINE_FUNCTION_DECORATOR void FMemory::Free(void* Original)
+{
+	if (!Original)
+	{
+		FScopedMallocTimer Timer(3);
+		return;
+	}
+
+	if (!FMEMORY_INLINE_GMalloc)
+	{
+		FreeExternal(Original);
+		return;
+	}
+	DoGamethreadHook(2);
+	FScopedMallocTimer Timer(2);
+	FMEMORY_INLINE_GMalloc->Free(Original);
+}
+
+FMEMORY_INLINE_FUNCTION_DECORATOR SIZE_T FMemory::GetAllocSize(void* Original)
+{
+	if (!FMEMORY_INLINE_GMalloc)
+	{
+		return GetAllocSizeExternal(Original);
+	}
+
+	SIZE_T Size = 0;
+	return FMEMORY_INLINE_GMalloc->GetAllocationSize(Original, Size) ? Size : 0;
+}
+
+FMEMORY_INLINE_FUNCTION_DECORATOR SIZE_T FMemory::QuantizeSize(SIZE_T Count, uint32 Alignment)
+{
+	if (!FMEMORY_INLINE_GMalloc)
+	{
+		return Count;
+	}
+	return FMEMORY_INLINE_GMalloc->QuantizeSize(Count, Alignment);
+}
+
diff --git a/Engine/Source/Runtime/Core/Public/HAL/MallocBinned2.h b/Engine/Source/Runtime/Core/Public/HAL/MallocBinned2.h
--- a/Engine/Source/Runtime/Core/Public/HAL/MallocBinned2.h
+++ b/Engine/Source/Runtime/Core/Public/HAL/MallocBinned2.h
@@ -2,8 +2,8 @@
 
 #pragma once
 
+#include "AlignmentTemplates.h"
 #include "Allocators/CachedOSPageAllocator.h"
-#include "Function.h"
 
 #define BINNED2_MAX_CACHED_OS_FREES (64)
 #if PLATFORM_64BITS
@@ -12,23 +12,107 @@
 	#define BINNED2_MAX_CACHED_OS_FREES_BYTE_LIMIT (16*1024*1024)
 #endif
 
+#define BINNED2_LARGE_ALLOC					65536		// Alignment of OS-allocated pointer - pool-allocated pointers will have a non-aligned pointer
+#define BINNED2_MINIMUM_ALIGNMENT_SHIFT		4			// Alignment of blocks, expressed as a shift
+#define BINNED2_MINIMUM_ALIGNMENT			16			// Alignment of blocks
+#define BINNED2_MAX_SMALL_POOL_SIZE			(32768-16)	// Maximum block size in GMallocBinned2SmallBlockSizes
+#define BINNED2_SMALL_POOL_COUNT			45
+
+
+#define DEFAULT_GMallocBinned2PerThreadCaches 1
+#define DEFAULT_GMallocBinned2LockFreeCaches 0
+#define DEFAULT_GMallocBinned2BundleSize BINNED2_LARGE_ALLOC
+#define DEFAULT_GMallocBinned2BundleCount 64
+#define DEFAULT_GMallocBinned2AllocExtra 32
+#define BINNED2_MAX_GMallocBinned2MaxBundlesBeforeRecycle 8
+
+#define BINNED2_ALLOW_RUNTIME_TWEAKING 0
+#if BINNED2_ALLOW_RUNTIME_TWEAKING
+	extern CORE_API int32 GMallocBinned2PerThreadCaches;
+	extern CORE_API int32 GMallocBinned2BundleSize = DEFAULT_GMallocBinned2BundleSize;
+	extern CORE_API int32 GMallocBinned2BundleCount = DEFAULT_GMallocBinned2BundleCount;
+	extern CORE_API int32 GMallocBinned2MaxBundlesBeforeRecycle = BINNED2_MAX_GMallocBinned2MaxBundlesBeforeRecycle;
+	extern CORE_API int32 GMallocBinned2AllocExtra = DEFAULT_GMallocBinned2AllocExtra;
+#else
+	#define GMallocBinned2PerThreadCaches DEFAULT_GMallocBinned2PerThreadCaches
+	#define GMallocBinned2BundleSize DEFAULT_GMallocBinned2BundleSize
+	#define GMallocBinned2BundleCount DEFAULT_GMallocBinned2BundleCount
+	#define GMallocBinned2MaxBundlesBeforeRecycle BINNED2_MAX_GMallocBinned2MaxBundlesBeforeRecycle
+	#define GMallocBinned2AllocExtra DEFAULT_GMallocBinned2AllocExtra
+#endif
+
+
+
 //
 // Optimized virtual memory allocator.
 //
-class FMallocBinned2 : public FMalloc
+class CORE_API FMallocBinned2 final : public FMalloc
 {
 	struct Private;
 
-private:
-	// Counts.
-	enum { SMALL_POOL_COUNT = 41 };
-
 	// Forward declares.
-	struct FFreeBlock;
 	struct FPoolInfo;
-	struct FPoolTable;
 	struct PoolHashBucket;
 
+	/** Information about a piece of free memory. */
+	struct FFreeBlock
+	{
+		enum
+		{
+			CANARY_VALUE = 0xe3
+		};
+
+		FORCEINLINE FFreeBlock(uint32 InPageSize, uint32 InBlockSize, uint32 InPoolIndex)
+			: BlockSize(InBlockSize)
+			, PoolIndex(InPoolIndex)
+			, Canary(CANARY_VALUE)
+			, NextFreeBlock(nullptr)
+		{
+			check(InPoolIndex < MAX_uint8 && InBlockSize <= MAX_uint16);
+			NumFreeBlocks = InPageSize / InBlockSize;
+			if (NumFreeBlocks * InBlockSize + BINNED2_MINIMUM_ALIGNMENT > InPageSize)
+			{
+				NumFreeBlocks--;
+			}
+			check(NumFreeBlocks * InBlockSize + BINNED2_MINIMUM_ALIGNMENT <= InPageSize);
+		}
+
+		FORCEINLINE uint32 GetNumFreeRegularBlocks() const
+		{
+			return NumFreeBlocks;
+		}
+		FORCEINLINE bool IsCanaryOk() const
+		{
+			return Canary == FFreeBlock::CANARY_VALUE;
+		}
+
+		FORCEINLINE void CanaryTest() const
+		{
+			if (!IsCanaryOk())
+			{
+				CanaryFail();
+			}
+			//checkSlow(PoolIndex == BoundSizeToPoolIndex(BlockSize));
+		}
+		void CanaryFail() const;
+
+		FORCEINLINE void* AllocateRegularBlock()
+		{
+			--NumFreeBlocks;
+			if (IsAligned(this, BINNED2_LARGE_ALLOC))
+			{
+				return (uint8*)this + BINNED2_LARGE_ALLOC - (NumFreeBlocks + 1) * BlockSize;
+			}
+			return (uint8*)this + (NumFreeBlocks)* BlockSize;
+		}
+
+		uint16 BlockSize;				// Size of the blocks that this list points to
+		uint8 PoolIndex;				// Index of this pool
+		uint8 Canary;					// Constant value of 0xe3
+		uint32 NumFreeBlocks;          // Number of consecutive free blocks here, at least 1.
+		void*  NextFreeBlock;          // Next free block in another pool
+	};
+
 	struct FPoolList
 	{
 		FPoolList();
@@ -40,9 +124,10 @@ private:
 
 		void LinkToFront(FPoolInfo* Pool);
 
-		FPoolInfo& PushNewPoolToFront(FMallocBinned2& Allocator, uint32 InBytes);
+		FPoolInfo& PushNewPoolToFront(FMallocBinned2& Allocator, uint32 InBytes, uint32 InPoolIndex);
 
-		void Validate(TFunctionRef<void(FPoolInfo&)>);
+		void ValidateActivePools();
+		void ValidateExhaustedPools();
 
 	private:
 		FPoolInfo* Front;
@@ -58,12 +143,21 @@ private:
 		FPoolTable();
 	};
 
-	const uint32 PageSize;
-	const uint64 NumPoolsPerPage;
-
 	struct FPtrToPoolMapping
 	{
+		FPtrToPoolMapping()
+			: PtrToPoolPageBitShift(0)
+			, HashKeyShift(0)
+			, PoolMask(0)
+			, MaxHashBuckets(0)
+		{
+		}
 		explicit FPtrToPoolMapping(uint32 InPageSize, uint64 InNumPoolsPerPage, uint64 AddressLimit)
+		{
+			Init(InPageSize, InNumPoolsPerPage, AddressLimit);
+		}
+
+		void Init(uint32 InPageSize, uint64 InNumPoolsPerPage, uint64 AddressLimit)
 		{
 			uint64 PoolPageToPoolBitShift = FPlatformMath::CeilLogTwo(InNumPoolsPerPage);
 
@@ -99,64 +193,327 @@ private:
 		uint64 MaxHashBuckets;
 	};
 
-	const FPtrToPoolMapping PtrToPoolMapping;
+	FPtrToPoolMapping PtrToPoolMapping;
 
 	// Pool tables for different pool sizes
-	FPoolTable SmallPoolTables[SMALL_POOL_COUNT];
+	FPoolTable SmallPoolTables[BINNED2_SMALL_POOL_COUNT];
 
 	PoolHashBucket* HashBuckets;
 	PoolHashBucket* HashBucketFreeList;
+	uint64 NumPoolsPerPage;
 
 	TCachedOSPageAllocator<BINNED2_MAX_CACHED_OS_FREES, BINNED2_MAX_CACHED_OS_FREES_BYTE_LIMIT> CachedOSPageAllocator;
 
 	FCriticalSection Mutex;
 
+	FORCEINLINE static bool IsOSAllocation(const void* Ptr)
+	{
+		return IsAligned(Ptr, BINNED2_LARGE_ALLOC);
+	}
+
+	struct FBundleNode
+	{
+		FBundleNode* NextNodeInCurrentBundle;
+		union
+		{
+			FBundleNode* NextBundle;
+			int32 Count;
+		};
+	};
+
+	struct FBundle
+	{
+		FORCEINLINE FBundle()
+		{
+			Reset();
+		}
+
+		FORCEINLINE void Reset()
+		{
+			Head = nullptr;
+			Count = 0;
+		}
+
+		FORCEINLINE void PushHead(FBundleNode* Node)
+		{
+			Node->NextNodeInCurrentBundle = Head;
+			Node->NextBundle = nullptr;
+			Head = Node;
+			Count++;
+		}
+
+		FORCEINLINE FBundleNode* PopHead()
+		{
+			FBundleNode* Result = Head;
+
+			Count--;
+			Head = Head->NextNodeInCurrentBundle;
+			return Result;
+		}
+
+		FBundleNode* Head;
+		uint32       Count;
+	};
+	static_assert(sizeof(FBundleNode) <= BINNED2_MINIMUM_ALIGNMENT, "Bundle nodes must fit into the smallest block size");
+
+	struct FFreeBlockList
+	{
+		// return true if we actually pushed it
+		FORCEINLINE bool PushToFront(void* InPtr, uint32 InPoolIndex, uint32 InBlockSize)
+		{
+			checkSlow(InPtr);
+
+			if (PartialBundle.Count >= (uint32)GMallocBinned2BundleCount || PartialBundle.Count * InBlockSize >= (uint32)GMallocBinned2BundleSize)
+			{
+				if (FullBundle.Head)
+				{
+					return false;
+				}
+				FullBundle = PartialBundle;
+				PartialBundle.Reset();
+			}
+			PartialBundle.PushHead((FBundleNode*)InPtr);
+			return true;
+		}
+		FORCEINLINE bool CanPushToFront(uint32 InPoolIndex, uint32 InBlockSize)
+		{
+			if (FullBundle.Head && (PartialBundle.Count >= (uint32)GMallocBinned2BundleCount || PartialBundle.Count * InBlockSize >= (uint32)GMallocBinned2BundleSize))
+			{
+				return false;
+			}
+			return true;
+		}
+		FORCEINLINE void* PopFromFront(uint32 InPoolIndex)
+		{
+			if (!PartialBundle.Head)
+			{
+				if (FullBundle.Head)
+				{
+					PartialBundle = FullBundle;
+					FullBundle.Reset();
+				}
+			}
+			return PartialBundle.Head ? PartialBundle.PopHead() : nullptr;
+		}
+
+		// tries to recycle the full bundle, if that fails, it is returned for freeing
+		FBundleNode* RecyleFull(uint32 InPoolIndex);
+		bool ObtainPartial(uint32 InPoolIndex);
+		FBundleNode* PopBundles(uint32 InPoolIndex);
+	private:
+		FBundle PartialBundle;
+		FBundle FullBundle;
+	};
+
+	struct FPerThreadFreeBlockLists
+	{
+		FORCEINLINE static FPerThreadFreeBlockLists* Get()
+		{
+			return FMallocBinned2::Binned2TlsSlot ? (FPerThreadFreeBlockLists*)FPlatformTLS::GetTlsValue(FMallocBinned2::Binned2TlsSlot) : nullptr;
+		}
+		static void SetTLS();
+		static void ClearTLS();
+
+		FORCEINLINE void* Malloc(uint32 InPoolIndex)
+		{
+			return FreeLists[InPoolIndex].PopFromFront(InPoolIndex);
+		}
+		// return true if the pointer was pushed
+		FORCEINLINE bool Free(void* InPtr, uint32 InPoolIndex, uint32 InBlockSize)
+		{
+			return FreeLists[InPoolIndex].PushToFront(InPtr, InPoolIndex, InBlockSize);
+		}		
+		// return true if a pointer can be pushed
+		FORCEINLINE bool CanFree(uint32 InPoolIndex, uint32 InBlockSize)
+		{
+			return FreeLists[InPoolIndex].CanPushToFront(InPoolIndex, InBlockSize);
+		}
+		// returns a bundle that needs to be freed if it can't be recycled
+		FBundleNode* RecycleFullBundle(uint32 InPoolIndex)
+		{
+			return FreeLists[InPoolIndex].RecyleFull(InPoolIndex);
+		}
+		// returns true if we have anything to pop
+		bool ObtainRecycledPartial(uint32 InPoolIndex)
+		{
+			return FreeLists[InPoolIndex].ObtainPartial(InPoolIndex);
+		}
+		FBundleNode* PopBundles(uint32 InPoolIndex)
+		{
+			return FreeLists[InPoolIndex].PopBundles(InPoolIndex);
+		}
+	private:
+		FFreeBlockList FreeLists[BINNED2_SMALL_POOL_COUNT];
+	};
+
+	static FORCEINLINE FFreeBlock* GetPoolHeaderFromPointer(void* Ptr)
+	{
+		return (FFreeBlock*)AlignDown(Ptr, BINNED2_LARGE_ALLOC);
+	}
+
 public:
-	// FMalloc interface.
-	// InPageSize - First parameter is page size, all allocs from BinnedAllocFromOS() MUST be aligned to this size
-	// AddressLimit - Second parameter is estimate of the range of addresses expected to be returns by BinnedAllocFromOS(). Binned
-	// Malloc will adjust its internal structures to make lookups for memory allocations O(1) for this range. 
-	// It is ok to go outside this range, lookups will just be a little slower
-	FMallocBinned2(uint32 InPageSize, uint64 AddressLimit);
+
+
+	FMallocBinned2();
 
 	virtual ~FMallocBinned2();
 
-	/**
-	 * Returns if the allocator is guaranteed to be thread-safe and therefore
-	 * doesn't need a unnecessary thread-safety wrapper around it.
-	 */
+	// FMalloc interface.
 	virtual bool IsInternallyThreadSafe() const override;
+	FORCEINLINE virtual void* Malloc(SIZE_T Size, uint32 Alignment) override
+	{
+		// Only allocate from the small pools if the size is small enough and the alignment isn't crazy large.
+		// With large alignments, we'll waste a lot of memory allocating an entire page, but such alignments are highly unlikely in practice.
+		if ((Size <= BINNED2_MAX_SMALL_POOL_SIZE) & (Alignment <= BINNED2_MINIMUM_ALIGNMENT)) // one branch, not two
+		{
+			FPerThreadFreeBlockLists* Lists = GMallocBinned2PerThreadCaches ? FPerThreadFreeBlockLists::Get() : nullptr;
+			if (Lists)
+			{
+				if (void* Result = Lists->Malloc(BoundSizeToPoolIndex(Size)))
+				{
+					return Result;
+				}
+			}
+		}
+		return MallocExternal(Size, Alignment);
+	}
 
-	/** 
-	 * Malloc
-	 */
-	virtual void* Malloc(SIZE_T Size, uint32 Alignment) override;
-
-	/** 
-	 * Realloc
-	 */
-	virtual void* Realloc( void* Ptr, SIZE_T NewSize, uint32 Alignment ) override;
-
-	/** 
-	 * Free
-	 */
-	virtual void Free( void* Ptr ) override;
-
-	/**
-	 * If possible determine the size of the memory allocated at the given address
-	 *
-	 * @param Original - Pointer to memory we are checking the size of
-	 * @param SizeOut - If possible, this value is set to the size of the passed in pointer
-	 * @return true if succeeded
-	 */
-	virtual bool GetAllocationSize(void *Original, SIZE_T &SizeOut) override;
-
-	/**
-	 * Validates the allocator's heap
-	 */
-	virtual bool ValidateHeap() override;
+	FORCEINLINE virtual void* Realloc(void* Ptr, SIZE_T NewSize, uint32 Alignment) override
+	{
+		if (NewSize <= BINNED2_MAX_SMALL_POOL_SIZE && Alignment <= BINNED2_MINIMUM_ALIGNMENT) // one branch, not two
+		{
+			FPerThreadFreeBlockLists* Lists = GMallocBinned2PerThreadCaches ? FPerThreadFreeBlockLists::Get() : nullptr;
+			if (Lists && (!Ptr || !IsOSAllocation(Ptr)))
+			{
+				uint32 BlockSize = 0;
+				uint32 PoolIndex = 0;
+
+				bool bCanFree = true; // the nullptr is always "freeable"
+				if (Ptr)
+				{
+					// Reallocate to a smaller/bigger pool if necessary
+					FFreeBlock* Free = GetPoolHeaderFromPointer(Ptr);
+					BlockSize = Free->BlockSize;
+					PoolIndex = Free->PoolIndex;
+					bCanFree = Free->IsCanaryOk();
+					if (NewSize && bCanFree && NewSize <= BlockSize && (PoolIndex == 0 || NewSize > PoolIndexToBlockSize(PoolIndex - 1)))
+					{
+						return Ptr;
+					}
+					bCanFree = bCanFree && Lists->CanFree(PoolIndex, BlockSize);
+				}
+				if (bCanFree)
+				{
+					void* Result = NewSize ? Lists->Malloc(BoundSizeToPoolIndex(NewSize)) : nullptr;
+					if (Result || !NewSize)
+					{
+						if (Result && Ptr)
+						{
+							FMemory::Memcpy(Result, Ptr, FMath::Min<SIZE_T>(NewSize, BlockSize));
+						}
+						if (Ptr)
+						{
+							bool bDidPush = Lists->Free(Ptr, PoolIndex, BlockSize);
+							checkSlow(bDidPush);
+						}
+						return Result;
+					}
+				}
+			}
+		}
+		return ReallocExternal(Ptr, NewSize, Alignment);
+	}
 
-	virtual const TCHAR* GetDescriptiveName() override;
+	FORCEINLINE virtual void Free(void* Ptr) override
+	{
+		if (!IsOSAllocation(Ptr))
+		{
+			FPerThreadFreeBlockLists* Lists = GMallocBinned2PerThreadCaches ? FPerThreadFreeBlockLists::Get() : nullptr;
+			if (Lists)
+			{
+				FFreeBlock* BasePtr = GetPoolHeaderFromPointer(Ptr);
+				if (BasePtr->IsCanaryOk() && Lists->Free(Ptr, BasePtr->PoolIndex, BasePtr->BlockSize))
+				{
+					return;
+				}
+			}
+		}
+		FreeExternal(Ptr);
+	}
+
+	FORCEINLINE virtual bool GetAllocationSize(void *Ptr, SIZE_T &SizeOut) override
+	{
+		if (!IsOSAllocation(Ptr))
+		{
+			const FFreeBlock* Free = GetPoolHeaderFromPointer(Ptr);
+			if (Free->IsCanaryOk())
+			{
+				SizeOut = Free->BlockSize;
+				return true;
+			}
+		}
+		return GetAllocationSizeExternal(Ptr, SizeOut);
+	}
+
+	FORCEINLINE virtual SIZE_T QuantizeSize(SIZE_T Count, uint32 Alignment) override
+	{
+		static_assert(DEFAULT_ALIGNMENT <= BINNED2_MINIMUM_ALIGNMENT, "DEFAULT_ALIGNMENT is assumed to be zero"); // used below
+		checkSlow(FMath::IsPowerOfTwo(Alignment));
+		SIZE_T SizeOut;
+		if ((Count <= BINNED2_MAX_SMALL_POOL_SIZE) & (Alignment <= BINNED2_MINIMUM_ALIGNMENT)) // one branch, not two
+		{
+			SizeOut = PoolIndexToBlockSize(BoundSizeToPoolIndex(Count));
+		}
+		else
+		{
+			Alignment = FMath::Max<uint32>(Alignment, OsAllocationGranularity);
+			checkSlow(Alignment <= PageSize);
+			SizeOut = Align(Count, Alignment);
+		}
+		check(SizeOut >= Count);
+		return SizeOut;
+	}
 
-	virtual void FlushCurrentThreadCache() override;
+	virtual bool ValidateHeap() override;
+	virtual void Trim() override;
+	virtual void SetupTLSCachesOnCurrentThread() override;
+	virtual void ClearAndDisableTLSCachesOnCurrentThread() override;
+	virtual const TCHAR* GetDescriptiveName() override;
+	// End FMalloc interface.
+
+	void FlushCurrentThreadCache();
+	void* MallocExternal(SIZE_T Size, uint32 Alignment);
+	void* ReallocExternal(void* Ptr, SIZE_T NewSize, uint32 Alignment);
+	void FreeExternal(void *Ptr);
+	bool GetAllocationSizeExternal(void* Ptr, SIZE_T& SizeOut);
+
+	static uint16 SmallBlockSizesReversed[BINNED2_SMALL_POOL_COUNT]; // this is reversed to get the smallest elements on our main cache line
+	static FMallocBinned2* MallocBinned2;
+	static uint32 Binned2TlsSlot;
+	static uint32 PageSize;
+	static uint32 OsAllocationGranularity;
+	// Mapping of sizes to small table indices
+	static uint8 MemSizeToIndex[1 + (BINNED2_MAX_SMALL_POOL_SIZE >> BINNED2_MINIMUM_ALIGNMENT_SHIFT)];
+
+	FORCEINLINE uint32 BoundSizeToPoolIndex(SIZE_T Size) 
+	{
+		auto Index = ((Size + BINNED2_MINIMUM_ALIGNMENT - 1) >> BINNED2_MINIMUM_ALIGNMENT_SHIFT);
+		checkSlow(Index >= 0 && Index <= (BINNED2_MAX_SMALL_POOL_SIZE >> BINNED2_MINIMUM_ALIGNMENT_SHIFT)); // and it should be in the table
+		uint32 PoolIndex = uint32(MemSizeToIndex[Index]);
+		checkSlow(PoolIndex >= 0 && PoolIndex < BINNED2_SMALL_POOL_COUNT);
+		return PoolIndex;
+	}
+	FORCEINLINE uint32 PoolIndexToBlockSize(uint32 PoolIndex)
+	{
+		return SmallBlockSizesReversed[BINNED2_SMALL_POOL_COUNT - PoolIndex - 1];
+	}
 };
+
+#define BINNED2_INLINE (1)
+#if BINNED2_INLINE // during development, it helps with iteration time to not include these here, but rather in the .cpp
+	#if PLATFORM_USES_FIXED_GMalloc_CLASS && !FORCE_ANSI_ALLOCATOR && USE_MALLOC_BINNED2
+		#define FMEMORY_INLINE_FUNCTION_DECORATOR  FORCEINLINE
+		#define FMEMORY_INLINE_GMalloc (FMallocBinned2::MallocBinned2)
+		#include "FMemory.inl"
+	#endif
+#endif
\ No newline at end of file
diff --git a/Engine/Source/Runtime/Core/Public/HAL/MemoryBase.h b/Engine/Source/Runtime/Core/Public/HAL/MemoryBase.h
--- a/Engine/Source/Runtime/Core/Public/HAL/MemoryBase.h
+++ b/Engine/Source/Runtime/Core/Public/HAL/MemoryBase.h
@@ -4,12 +4,21 @@
 
 #include "UnrealMemory.h"
 
-#ifndef __MEMORYBASE_H__
-#define __MEMORYBASE_H__
+enum
+{
+	// Default allocator alignment. If the default is specified, the allocator applies to engine rules.
+	// Blocks >= 16 bytes will be 16-byte-aligned, Blocks < 16 will be 8-byte aligned. If the allocator does
+	// not support allocation alignment, the alignment will be ignored.
+	DEFAULT_ALIGNMENT = 0,
+
+	// Minimum allocator alignment
+	MIN_ALIGNMENT = 8,
+};
 
 
 /** The global memory allocator. */
 CORE_API extern class FMalloc* GMalloc;
+CORE_API extern class FMalloc** GFixedMallocLocationPtr;
 
 /** Global FMallocProfiler variable to allow multiple malloc profilers to communicate. */
 MALLOC_PROFILER( CORE_API extern class FMallocProfiler* GMallocProfiler; )
@@ -32,20 +41,14 @@ public:
 	 * @param	Size	Amount of memory to allocate (in bytes)
 	 * @return			A pointer to a block of memory with size Size or nullptr
 	 */
-	void* operator new( size_t Size )
-	{
-		return FMemory::SystemMalloc( Size );
-	}
+	void* operator new(size_t Size);
 
 	/**
 	 * Overloaded delete operator using the system allocator
 	 *
 	 * @param	Ptr		Pointer to delete
 	 */
-	void operator delete( void* Ptr )
-	{
-		FMemory::SystemFree( Ptr );
-	}
+	void operator delete(void* Ptr);
 
 	/**
 	 * Overloaded array new operator using the system allocator.
@@ -53,20 +56,14 @@ public:
 	 * @param	Size	Amount of memory to allocate (in bytes)
 	 * @return			A pointer to a block of memory with size Size or nullptr
 	 */
-	void* operator new[]( size_t Size )
-	{
-		return FMemory::SystemMalloc( Size );
-	}
+	void* operator new[](size_t Size);
 
 	/**
 	 * Overloaded array delete operator using the system allocator
 	 *
 	 * @param	Ptr		Pointer to delete
 	 */
-	void operator delete[]( void* Ptr )
-	{
-		FMemory::SystemFree( Ptr );
-	}
+	void operator delete[](void* Ptr);
 };
 
 /** The global memory allocator's interface. */
@@ -76,12 +73,6 @@ class CORE_API FMalloc  :
 {
 public:
 	/**
-	 *	Initializes stats metadata. We need to do this as soon as possible, but cannot be done in the constructor
-	 *	due to the FName::StaticInit
-	 */
-	virtual void InitializeStatsMetadata();
-
-	/** 
 	 * Malloc
 	 */
 	virtual void* Malloc( SIZE_T Count, uint32 Alignment=DEFAULT_ALIGNMENT ) = 0;
@@ -97,6 +88,56 @@ public:
 	virtual void Free( void* Original ) = 0;
 		
 	/** 
+	* For some allocators this will return the actual size that should be requested to eliminate
+	* internal fragmentation. The return value will always be >= Count. This can be used to grow
+	* and shrink containers to optimal sizes.
+	* This call is always fast and threadsafe with no locking.
+	*/
+	virtual SIZE_T QuantizeSize(SIZE_T Count, uint32 Alignment)
+	{
+		return Count; // Default implementation has no way of determining this
+	}
+
+	/**
+	* If possible determine the size of the memory allocated at the given address
+	*
+	* @param Original - Pointer to memory we are checking the size of
+	* @param SizeOut - If possible, this value is set to the size of the passed in pointer
+	* @return true if succeeded
+	*/
+	virtual bool GetAllocationSize(void *Original, SIZE_T &SizeOut)
+	{
+		return false; // Default implementation has no way of determining this
+	}
+
+	/**
+	* Releases as much memory as possible. Must be called from the main thread.
+	*/
+	virtual void Trim()
+	{
+	}
+
+	/**
+	* Set up TLS caches on the current thread. These are the threads that we can trim.
+	*/
+	virtual void SetupTLSCachesOnCurrentThread()
+	{
+	}
+
+	/**
+	* Clears the TLS caches on the current thread and disables any future caching.
+	*/
+	virtual void ClearAndDisableTLSCachesOnCurrentThread()
+	{
+	}
+
+	/**
+	*	Initializes stats metadata. We need to do this as soon as possible, but cannot be done in the constructor
+	*	due to the FName::StaticInit
+	*/
+	virtual void InitializeStatsMetadata();
+
+	/**
 	 * Handles any commands passed in on the command line
 	 */
 	virtual bool Exec( UWorld* InWorld, const TCHAR* Cmd, FOutputDevice& Ar ) override
@@ -133,18 +174,6 @@ public:
 		return( true );
 	}
 
-	/**
-	* If possible determine the size of the memory allocated at the given address
-	*
-	* @param Original - Pointer to memory we are checking the size of
-	* @param SizeOut - If possible, this value is set to the size of the passed in pointer
-	* @return true if succeeded
-	*/
-	virtual bool GetAllocationSize(void *Original, SIZE_T &SizeOut)
-	{
-		return false; // Default implementation has no way of determining this
-	}
-
 	/**
 	 * Gets descriptive name for logging purposes.
 	 *
@@ -155,13 +184,6 @@ public:
 		return TEXT("Unspecified allocator");
 	}
 
-	/**
-	 * Triggers a cache flush of any memory being held by the allocator on the current thread.
-	 */
-	virtual void FlushCurrentThreadCache()
-	{
-	}
-
 protected:
 	friend struct FCurrentFrameCalls;
 
@@ -192,4 +214,3 @@ protected:
 };
 
 
-#endif
diff --git a/Engine/Source/Runtime/Core/Public/HAL/Platform.h b/Engine/Source/Runtime/Core/Public/HAL/Platform.h
--- a/Engine/Source/Runtime/Core/Public/HAL/Platform.h
+++ b/Engine/Source/Runtime/Core/Public/HAL/Platform.h
@@ -273,6 +273,10 @@
 	#define PLATFORM_USES_FIXED_RHI_CLASS		0
 #endif
 
+#ifndef PLATFORM_USES_FIXED_GMalloc_CLASS
+	#define PLATFORM_USES_FIXED_GMalloc_CLASS		0
+#endif
+
 #ifndef PLATFORM_SUPPORTS_MULTIPLE_NATIVE_WINDOWS
 	#define PLATFORM_SUPPORTS_MULTIPLE_NATIVE_WINDOWS	1
 #endif
diff --git a/Engine/Source/Runtime/Core/Public/HAL/UnrealMemory.h b/Engine/Source/Runtime/Core/Public/HAL/UnrealMemory.h
--- a/Engine/Source/Runtime/Core/Public/HAL/UnrealMemory.h
+++ b/Engine/Source/Runtime/Core/Public/HAL/UnrealMemory.h
@@ -2,18 +2,77 @@
 
 #pragma once
 
+#include "MemoryBase.h"
+
 // Sizes.
 
-enum 
+#if STATS
+#define MALLOC_GT_HOOKS 1
+#else
+#define MALLOC_GT_HOOKS 0
+#endif
+
+#if MALLOC_GT_HOOKS
+CORE_API void DoGamethreadHook(int32 Index);
+#else
+FORCEINLINE void DoGamethreadHook(int32 Index)
 { 
-	// Default allocator alignment. If the default is specified, the allocator applies to engine rules.
-	// Blocks >= 16 bytes will be 16-byte-aligned, Blocks < 16 will be 8-byte aligned. If the allocator does
-	// not support allocation alignment, the alignment will be ignored.
-	DEFAULT_ALIGNMENT = 0,
+}
+#endif
+
+#define TIME_MALLOC (0 && PLATFORM_PS4)
+
+#if TIME_MALLOC
+
+struct CORE_API FScopedMallocTimer
+{
+	static uint64 GTotalCycles[4];
+	static uint64 GTotalCount[4];
+	static uint64 GTotalMisses[4];
+
+	int32 Index;
+	uint64 Cycles;
+
+	FORCEINLINE FScopedMallocTimer(int32 InIndex)
+		: Index(InIndex)
+		, Cycles(FPlatformTime::Cycles64())
+	{
+	}
+	FORCEINLINE ~FScopedMallocTimer()
+	{
+		uint64 Add = uint64(FPlatformTime::Cycles64() - Cycles);
+		FPlatformAtomics::InterlockedAdd((volatile int64*)&GTotalCycles[Index], Add);
+		uint64 Was = FPlatformAtomics::InterlockedAdd((volatile int64*)&GTotalCount[Index], 1);
 
-	// Minimum allocator alignment
-	MIN_ALIGNMENT = 8,
+		extern CORE_API bool GIsRunning;
+		if (GIsRunning && Index == 1 && (Was & 0xffff) == 0)
+		{
+			Spew();
+		}
+	}
+	static FORCEINLINE void Miss(int32 InIndex)
+	{
+		FPlatformAtomics::InterlockedAdd((volatile int64*)&GTotalMisses[InIndex], 1);
+	}
+	static void Spew();
+}; 
+#else
+struct FScopedMallocTimer
+{
+	FORCEINLINE FScopedMallocTimer(int32 InIndex)
+	{
+	}
+	FORCEINLINE ~FScopedMallocTimer()
+	{
+	}
+	FORCEINLINE void Hit(int32 InIndex)
+	{
+	}
 }; 
+#endif
+
+
+
 
 /*-----------------------------------------------------------------------------
 	FMemory.
@@ -125,9 +184,32 @@ struct CORE_API FMemory
 	// C style memory allocation stubs.
 	//
 
-	static void* Malloc( SIZE_T Count, uint32 Alignment=DEFAULT_ALIGNMENT );
-	static void* Realloc( void* Original, SIZE_T Count, uint32 Alignment=DEFAULT_ALIGNMENT );
-	static void Free( void* Original );
+	static void* Malloc(SIZE_T Count, uint32 Alignment = DEFAULT_ALIGNMENT);
+	static void* Realloc(void* Original, SIZE_T Count, uint32 Alignment = DEFAULT_ALIGNMENT);
+	static void Free(void* Original);
+	static SIZE_T GetAllocSize(void* Original);
+	/**
+	* For some allocators this will return the actual size that should be requested to eliminate
+	* internal fragmentation. The return value will always be >= Count. This can be used to grow
+	* and shrink containers to optimal sizes.
+	* This call is always fast and threadsafe with no locking.
+	*/
+	static SIZE_T QuantizeSize(SIZE_T Count, uint32 Alignment = DEFAULT_ALIGNMENT);
+
+	/**
+	* Releases as much memory as possible. Must be called from the main thread.
+	*/
+	static void Trim();
+
+	/**
+	* Set up TLS caches on the current thread. These are the threads that we can trim.
+	*/
+	static void SetupTLSCachesOnCurrentThread();
+
+	/**
+	* Clears the TLS caches on the current thread and disables any future caching.
+	*/
+	static void ClearAndDisableTLSCachesOnCurrentThread();
 
 	//
 	// Malloc for GPU mapped memory on UMA systems (XB1/PS4/etc)
@@ -138,8 +220,6 @@ struct CORE_API FMemory
 	static void* GPURealloc(void* Original, SIZE_T Count, uint32 Alignment = DEFAULT_ALIGNMENT);
 	static void GPUFree(void* Original);
 
-	static SIZE_T GetAllocSize( void* Original );
-
 	/**
 	 * A helper function that will perform a series of random heap allocations to test
 	 * the internal validity of the heap. Note, this function will "leak" memory, but another call
@@ -153,5 +233,27 @@ struct CORE_API FMemory
 	* This uses the purgatory malloc proxy to check if things are writing to stale pointers.
 	*/
 	static void EnablePurgatoryTests();
+private:
+	static void GCreateMalloc();
+	// These versions are called either at startup or in the event of a crash
+	static void* MallocExternal(SIZE_T Count, uint32 Alignment = DEFAULT_ALIGNMENT);
+	static void* ReallocExternal(void* Original, SIZE_T Count, uint32 Alignment = DEFAULT_ALIGNMENT);
+	static void FreeExternal(void* Original);
+	static SIZE_T GetAllocSizeExternal(void* Original);
+	static SIZE_T QuantizeSizeExternal(SIZE_T Count, uint32 Alignment = DEFAULT_ALIGNMENT);
 };
 
+#define INLINE_FMEMORY_OPERATION (0) // untested, but should work. Inlines FMemory::Malloc, etc
+
+#if INLINE_FMEMORY_OPERATION
+	#if PLATFORM_USES_FIXED_GMalloc_CLASS
+		#error "PLATFORM_USES_FIXED_GMalloc_CLASS and INLINE_FMEMORY_OPERATION are not compatible. PLATFORM_USES_FIXED_GMalloc_CLASS is inlined below."
+	#endif
+
+	#define FMEMORY_INLINE_FUNCTION_DECORATOR FORCEINLINE
+	#include "FMemory.inl"
+#endif
+
+#if PLATFORM_USES_FIXED_GMalloc_CLASS && !FORCE_ANSI_ALLOCATOR && USE_MALLOC_BINNED2
+	#include "MallocBinned2.h"
+#endif
\ No newline at end of file
diff --git a/Engine/Source/Runtime/Core/Public/Misc/MemStack.h b/Engine/Source/Runtime/Core/Public/Misc/MemStack.h
--- a/Engine/Source/Runtime/Core/Public/Misc/MemStack.h
+++ b/Engine/Source/Runtime/Core/Public/Misc/MemStack.h
@@ -24,14 +24,12 @@ public:
 	enum
 	{
 		PageSize = 64 * 1024,
-		SmallPageSize = 1024
+		SmallPageSize = 1024-16 // allow a little extra space for allocator headers, etc
 	};
 #if UE_BUILD_SHIPPING
 	typedef TLockFreeFixedSizeAllocator<PageSize, PLATFORM_CACHE_LINE_SIZE, FNoopCounter> TPageAllocator;
-	typedef TLockFreeFixedSizeAllocator<SmallPageSize, PLATFORM_CACHE_LINE_SIZE, FNoopCounter> TSmallPageAllocator;
 #else
 	typedef TLockFreeFixedSizeAllocator<PageSize, PLATFORM_CACHE_LINE_SIZE, FThreadSafeCounter> TPageAllocator;
-	typedef TLockFreeFixedSizeAllocator<SmallPageSize, PLATFORM_CACHE_LINE_SIZE, FThreadSafeCounter> TSmallPageAllocator;
 #endif
 
 	static void *Alloc();
@@ -47,7 +45,6 @@ private:
 	static void UpdateStats();
 #endif
 	static TPageAllocator TheAllocator;
-	static TSmallPageAllocator TheSmallAllocator;
 };
 
 
@@ -295,15 +292,28 @@ public:
 				}
 			}
 		}
-		int32 CalculateSlack(int32 NumElements,int32 NumAllocatedElements,int32 NumBytesPerElement) const
+		FORCEINLINE int32 CalculateSlackReserve(int32 NumElements, int32 NumBytesPerElement) const
 		{
-			return DefaultCalculateSlack(NumElements,NumAllocatedElements,NumBytesPerElement);
+			return DefaultCalculateSlackReserve(NumElements, NumBytesPerElement, false, Alignment);
+		}
+		FORCEINLINE int32 CalculateSlackShrink(int32 NumElements, int32 NumAllocatedElements, int32 NumBytesPerElement) const
+		{
+			return DefaultCalculateSlackShrink(NumElements, NumAllocatedElements, NumBytesPerElement, false, Alignment);
+		}
+		FORCEINLINE int32 CalculateSlackGrow(int32 NumElements, int32 NumAllocatedElements, int32 NumBytesPerElement) const
+		{
+			return DefaultCalculateSlackGrow(NumElements, NumAllocatedElements, NumBytesPerElement, false, Alignment);
 		}
 
-		int32 GetAllocatedSize(int32 NumAllocatedElements, int32 NumBytesPerElement) const
+		FORCEINLINE int32 GetAllocatedSize(int32 NumAllocatedElements, int32 NumBytesPerElement) const
 		{
 			return NumAllocatedElements * NumBytesPerElement;
 		}
+
+		bool HasAllocation()
+		{
+			return !!Data;
+		}
 			
 	private:
 
diff --git a/Engine/Source/Runtime/Core/Public/Templates/AlignmentTemplates.h b/Engine/Source/Runtime/Core/Public/Templates/AlignmentTemplates.h
new file mode 100644
--- /dev/null
+++ b/Engine/Source/Runtime/Core/Public/Templates/AlignmentTemplates.h
@@ -0,0 +1,54 @@
+// Copyright 1998-2016 Epic Games, Inc. All Rights Reserved.
+
+#pragma once
+
+/**
+ * Aligns a value to the nearest higher multiple of 'Alignment', which must be a power of two.
+ *
+ * @param Ptr			Value to align
+ * @param Alignment		Alignment, must be a power of two
+ * @return				Aligned value
+ */
+template <typename T>
+inline CONSTEXPR T Align( const T Ptr, int32 Alignment )
+{
+	return (T)(((PTRINT)Ptr + Alignment - 1) & ~(Alignment-1));
+}
+
+/**
+ * Aligns a value to the nearest lower multiple of 'Alignment', which must be a power of two.
+ *
+ * @param Ptr			Value to align
+ * @param Alignment		Alignment, must be a power of two
+ * @return				Aligned value
+ */
+template <typename T>
+inline CONSTEXPR T AlignDown( const T Ptr, int32 Alignment )
+{
+	return (T)(((PTRINT)Ptr) & ~(Alignment-1));
+}
+
+/**
+ * Checks if a pointer is aligned to the specified alignment.
+ *
+ * @param Ptr - The pointer to check.
+ *
+ * @return true if the pointer is aligned, false otherwise.
+ */
+static FORCEINLINE bool IsAligned(const volatile void* Ptr, const uint32 Alignment)
+{
+	return !(UPTRINT(Ptr) & (Alignment - 1));
+}
+
+/**
+ * Aligns a value to the nearest higher multiple of 'Alignment'.
+ *
+ * @param Ptr			Value to align
+ * @param Alignment		Alignment, can be any arbitrary uint32
+ * @return				Aligned value
+ */
+template< class T > inline T AlignArbitrary( const T Ptr, uint32 Alignment )
+{
+	return (T) ( ( ((UPTRINT)Ptr + Alignment - 1) / Alignment ) * Alignment );
+}
+
diff --git a/Engine/Source/Runtime/Core/Public/Templates/UnrealTemplate.h b/Engine/Source/Runtime/Core/Public/Templates/UnrealTemplate.h
--- a/Engine/Source/Runtime/Core/Public/Templates/UnrealTemplate.h
+++ b/Engine/Source/Runtime/Core/Public/Templates/UnrealTemplate.h
@@ -8,62 +8,13 @@
 #include "HAL/Platform.h"
 #include "Templates/UnrealTypeTraits.h"
 #include "UnrealMemory.h"
+#include "AlignmentTemplates.h"
 
 
 /*-----------------------------------------------------------------------------
 	Standard templates.
 -----------------------------------------------------------------------------*/
 
-/**
- * Aligns a value to the nearest higher multiple of 'Alignment', which must be a power of two.
- *
- * @param Ptr			Value to align
- * @param Alignment		Alignment, must be a power of two
- * @return				Aligned value
- */
-template <typename T>
-inline CONSTEXPR T Align( const T Ptr, int32 Alignment )
-{
-	return (T)(((PTRINT)Ptr + Alignment - 1) & ~(Alignment-1));
-}
-
-/**
- * Aligns a value to the nearest lower multiple of 'Alignment', which must be a power of two.
- *
- * @param Ptr			Value to align
- * @param Alignment		Alignment, must be a power of two
- * @return				Aligned value
- */
-template <typename T>
-inline CONSTEXPR T AlignDown( const T Ptr, int32 Alignment )
-{
-	return (T)(((PTRINT)Ptr) & ~(Alignment-1));
-}
-
-/**
- * Checks if a pointer is aligned to the specified alignment.
- *
- * @param Ptr - The pointer to check.
- *
- * @return true if the pointer is aligned, false otherwise.
- */
-static FORCEINLINE bool IsAligned(const volatile void* Ptr, const uint32 Alignment)
-{
-	return !(UPTRINT(Ptr) & (Alignment - 1));
-}
-
-/**
- * Aligns a value to the nearest higher multiple of 'Alignment'.
- *
- * @param Ptr			Value to align
- * @param Alignment		Alignment, can be any arbitrary uint32
- * @return				Aligned value
- */
-template< class T > inline T AlignArbitrary( const T Ptr, uint32 Alignment )
-{
-	return (T) ( ( ((UPTRINT)Ptr + Alignment - 1) / Alignment ) * Alignment );
-}
-
 /**
  * Chooses between the two parameters based on whether the first is nullptr or not.
  * @return If the first parameter provided is non-nullptr, it is returned; otherwise the second parameter is returned.
diff --git a/Engine/Source/Runtime/CoreUObject/Private/Serialization/AsyncLoading.cpp b/Engine/Source/Runtime/CoreUObject/Private/Serialization/AsyncLoading.cpp
--- a/Engine/Source/Runtime/CoreUObject/Private/Serialization/AsyncLoading.cpp
+++ b/Engine/Source/Runtime/CoreUObject/Private/Serialization/AsyncLoading.cpp
@@ -2266,6 +2266,7 @@ void FlushAsyncLoading(int32 PackageID /* = INDEX_NONE */)
 
 		// Reset min priority again.
 		FIOSystem::Get().SetMinPriority( AIOP_MIN );
+		FMemory::Trim();
 	}
 }
 
diff --git a/Engine/Source/Runtime/CoreUObject/Private/UObject/GarbageCollection.cpp b/Engine/Source/Runtime/CoreUObject/Private/UObject/GarbageCollection.cpp
--- a/Engine/Source/Runtime/CoreUObject/Private/UObject/GarbageCollection.cpp
+++ b/Engine/Source/Runtime/CoreUObject/Private/UObject/GarbageCollection.cpp
@@ -1164,6 +1164,7 @@ void CollectGarbageInternal(EObjectFlags KeepFlags, bool bPerformFullPurge)
 	if( GObjIncrementalPurgeIsInProgress || GObjPurgeIsRequired )
 	{
 		IncrementalPurgeGarbage( false );
+		FMemory::Trim();
 	}
 	check( !GObjIncrementalPurgeIsInProgress );
 	check( !GObjPurgeIsRequired );
@@ -1326,6 +1327,7 @@ void CollectGarbageInternal(EObjectFlags KeepFlags, bool bPerformFullPurge)
 	{
 		IncrementalPurgeGarbage( false );	
 	}
+	FMemory::Trim();
 
 	// Destroy all pending delete linkers
 	DeleteLoaders();
diff --git a/Engine/Source/Runtime/Engine/Private/Particles/ParticleEmitterInstances.cpp b/Engine/Source/Runtime/Engine/Private/Particles/ParticleEmitterInstances.cpp
--- a/Engine/Source/Runtime/Engine/Private/Particles/ParticleEmitterInstances.cpp
+++ b/Engine/Source/Runtime/Engine/Private/Particles/ParticleEmitterInstances.cpp
@@ -144,84 +144,15 @@ DECLARE_CYCLE_STAT(TEXT("EmitterInstance PrepPerInstanceBlock"), STAT_PrepPerIns
 DECLARE_CYCLE_STAT(TEXT("EmitterInstance Resize"), STAT_ParticleEmitterInstance_Resize, STATGROUP_Particles);
 
 
-/*-----------------------------------------------------------------------------
-	Fast allocators for small block data being sent to the render thread
------------------------------------------------------------------------------*/
-
-static TLockFreeFixedSizeAllocator<256, PLATFORM_CACHE_LINE_SIZE> FastParticleSmallBlockAlloc256;
-static TLockFreeFixedSizeAllocator<384, PLATFORM_CACHE_LINE_SIZE> FastParticleSmallBlockAlloc384;
-static TLockFreeFixedSizeAllocator<512, PLATFORM_CACHE_LINE_SIZE> FastParticleSmallBlockAlloc512;
-static TLockFreeFixedSizeAllocator<768, PLATFORM_CACHE_LINE_SIZE> FastParticleSmallBlockAlloc768;
-static TLockFreeFixedSizeAllocator<1024, PLATFORM_CACHE_LINE_SIZE> FastParticleSmallBlockAlloc1024;
-static TLockFreeFixedSizeAllocator<1792, PLATFORM_CACHE_LINE_SIZE> FastParticleSmallBlockAlloc1792;
-static TLockFreeFixedSizeAllocator<2048, PLATFORM_CACHE_LINE_SIZE> FastParticleSmallBlockAlloc2048;
-
 FORCEINLINE static void* FastParticleSmallBlockAlloc(size_t AllocSize)
 {
-	check(AllocSize > 0);
-	if (AllocSize <= 768)
-	{
-		if (AllocSize <= 256)
-		{
-			return FastParticleSmallBlockAlloc256.Allocate();
-		}
-		if (AllocSize <= 384)
-		{
-			return FastParticleSmallBlockAlloc384.Allocate();
-		}
-		if (AllocSize <= 512)
-		{
-			return FastParticleSmallBlockAlloc512.Allocate();
-		}
-		return FastParticleSmallBlockAlloc768.Allocate();
-	}
-	if (AllocSize <= 1024)
-	{
-		return FastParticleSmallBlockAlloc1024.Allocate();
-	}
-	if (AllocSize <= 1792)
-	{
-		return FastParticleSmallBlockAlloc1792.Allocate();
-	}
-	if (AllocSize <= 2048)
-	{
-		return FastParticleSmallBlockAlloc2048.Allocate();
-	}
+	QUICK_SCOPE_CYCLE_COUNTER(STAT_PARTALLOC);
 	return FMemory::Malloc(AllocSize);
 }
 
 FORCEINLINE static void FastParticleSmallBlockFree(void *RawMemory, size_t AllocSize)
 {
-	check(AllocSize > 0);
-	if (AllocSize <= 768)
-	{
-		if (AllocSize <= 256)
-		{
-			return FastParticleSmallBlockAlloc256.Free(RawMemory);
-		}
-		if (AllocSize <= 384)
-		{
-			return FastParticleSmallBlockAlloc384.Free(RawMemory);
-		}
-		if (AllocSize <= 512)
-		{
-			return FastParticleSmallBlockAlloc512.Free(RawMemory);
-		}
-		return FastParticleSmallBlockAlloc768.Free(RawMemory);
-	}
-
-	if (AllocSize <= 1024)
-	{
-		return FastParticleSmallBlockAlloc1024.Free(RawMemory);
-	}
-	if (AllocSize <= 1792)
-	{
-		return FastParticleSmallBlockAlloc1792.Free(RawMemory);
-	}
-	if (AllocSize <= 2048)
-	{
-		return FastParticleSmallBlockAlloc2048.Free(RawMemory);
-	}
+	QUICK_SCOPE_CYCLE_COUNTER(STAT_PARTALLOC);
 	FMemory::Free(RawMemory);
 }
 
@@ -791,9 +722,10 @@ void FParticleEmitterInstance::CheckEmitterFinished()
 			const UParticleModuleRequired* RequiredModule = LODLevel->RequiredModule;
 			check(RequiredModule);
 
-			if (SpawnModule->GetMaximumSpawnRate() == 0
+			if (HasCompleted() || 
+				(SpawnModule->GetMaximumSpawnRate() == 0
 				&& RequiredModule->EmitterDuration == 0
-				&& RequiredModule->EmitterLoops == 0
+				&& RequiredModule->EmitterLoops == 0)
 				)
 			{
 				bEmitterIsDone = true;
@@ -3769,20 +3701,14 @@ void FDynamicEmitterDataBase::operator delete(void *RawMemory, size_t AllocSize)
 	FastParticleSmallBlockFree(RawMemory, AllocSize);
 }	
 
-static TLockFreeFixedSizeAllocator<sizeof(FParticleDynamicData), PLATFORM_CACHE_LINE_SIZE> ParticleDynamicDataAllocator;
-
 void* FParticleDynamicData::operator new(size_t AllocSize)
 {
-	check(AllocSize == sizeof(FParticleDynamicData));
-	return ParticleDynamicDataAllocator.Allocate();
-	//return FMemory::Malloc(AllocSize);
+	return FMemory::Malloc(AllocSize);
 }
 
 void FParticleDynamicData::operator delete(void *RawMemory, size_t AllocSize)
 {
-	check(AllocSize == sizeof(FParticleDynamicData));
-	ParticleDynamicDataAllocator.Free(RawMemory);
-	//FMemory::Free(RawMemory);
+	FMemory::Free(RawMemory);
 }	
 
 FDynamicEmitterDataBase::FDynamicEmitterDataBase(const UParticleModuleRequired* RequiredModule)
diff --git a/Engine/Source/Runtime/Engine/Private/Particles/ParticleGpuSimulation.cpp b/Engine/Source/Runtime/Engine/Private/Particles/ParticleGpuSimulation.cpp
--- a/Engine/Source/Runtime/Engine/Private/Particles/ParticleGpuSimulation.cpp
+++ b/Engine/Source/Runtime/Engine/Private/Particles/ParticleGpuSimulation.cpp
@@ -2681,66 +2681,16 @@ public:
 	}
 };
 
-struct FNewParticleAlloc
-{
-	TLockFreeFixedSizeAllocator<sizeof(TArray<FNewParticle>), PLATFORM_CACHE_LINE_SIZE, FThreadSafeCounter> FreeTArrayFNewParticleArrays;
-	TLockFreePointerListUnordered<TArray<FNewParticle>, PLATFORM_CACHE_LINE_SIZE>	FreeTArrayFNewParticle;
-	~FNewParticleAlloc()
-	{
-		while (true)
-		{
-			TArray<FNewParticle>* Recycle = FreeTArrayFNewParticle.Pop();
-			if (!Recycle)
-			{
-				break;
-			}
-			Recycle->Empty();
-			FreeTArrayFNewParticleArrays.Free(Recycle);
-		}
-	}
-
-};
-
-static FNewParticleAlloc& GNewParticleAlloc()
-{
-	static FNewParticleAlloc Singleton;
-	return Singleton;
-}
-
-static const int MaxNumParticlesToRecycle  = 512; // these can be quite large
-
 // recycle memory blocks for the NewParticle array
 static void FreeNewParticleArray(TArray<FNewParticle>& NewParticles)
 {
 	NewParticles.Reset();
-	const int MaxNumToRecycledArrays  = 100; 
-	int32 CurrentSize = NewParticles.GetSlack();
-	if (CurrentSize > 0 && CurrentSize <= MaxNumParticlesToRecycle && GNewParticleAlloc().FreeTArrayFNewParticleArrays.GetNumUsed().GetValue() < MaxNumToRecycledArrays)
-	{
-		TArray<FNewParticle>* Recycle = new (GNewParticleAlloc().FreeTArrayFNewParticleArrays.Allocate()) TArray<FNewParticle>;
-		Exchange(*Recycle, NewParticles);
-		check(Recycle->Num() == 0 && Recycle->GetSlack());
-		check(NewParticles.Num() == 0 && NewParticles.GetSlack() == 0);
-		GNewParticleAlloc().FreeTArrayFNewParticle.Push(Recycle);
-	}
 }
 
 static void GetNewParticleArray(TArray<FNewParticle>& NewParticles, int32 NumParticlesNeeded = -1)
 {
-	if (NumParticlesNeeded <= MaxNumParticlesToRecycle)
-	{
-		TArray<FNewParticle>* Recycle = GNewParticleAlloc().FreeTArrayFNewParticle.Pop();
-		if (Recycle)
-		{
-			Exchange(*Recycle, NewParticles);
-			Recycle->~TArray<FNewParticle>(); // this probably doesn't do anything, but type safety and all
-			GNewParticleAlloc().FreeTArrayFNewParticleArrays.Free(Recycle);
-			check(NewParticles.Num() == 0 && NewParticles.GetSlack());
-		}
-	}
 	if (NumParticlesNeeded > 0)
 	{
-		// this might realloc, but we need to get the small blocks out of the recycle list
 		NewParticles.Reserve(NumParticlesNeeded);
 	}
 }
diff --git a/Engine/Source/Runtime/Engine/Private/ShaderCompiler/ShaderCompiler.cpp b/Engine/Source/Runtime/Engine/Private/ShaderCompiler/ShaderCompiler.cpp
--- a/Engine/Source/Runtime/Engine/Private/ShaderCompiler/ShaderCompiler.cpp
+++ b/Engine/Source/Runtime/Engine/Private/ShaderCompiler/ShaderCompiler.cpp
@@ -2507,17 +2507,21 @@ bool RecompileShaders(const TCHAR* Cmd, FOutputDevice& Ar)
 			{
 				FRecompileShadersTimer TestTimer(TEXT("RecompileShaders Changed"));
 
-				// Kick off global shader recompiles
 				UMaterialInterface::IterateOverActiveFeatureLevels([&](ERHIFeatureLevel::Type InFeatureLevel) {
 					auto ShaderPlatform = GShaderPlatformForFeatureLevel[InFeatureLevel];
 					BeginRecompileGlobalShaders(OutdatedShaderTypes, OutdatedShaderPipelineTypes, ShaderPlatform);
-					UMaterial::UpdateMaterialShaders(OutdatedShaderTypes, OutdatedShaderPipelineTypes, OutdatedFactoryTypes, ShaderPlatform);
 				});
 
-				GWarn->StatusUpdate(0, 1, NSLOCTEXT("ShaderCompilingManager", "CompilingGlobalShaderStatus", "Compiling global shaders..."));
-
 				// Block on global shaders
 				FinishRecompileGlobalShaders();
+
+				// Kick off global shader recompiles
+				UMaterialInterface::IterateOverActiveFeatureLevels([&](ERHIFeatureLevel::Type InFeatureLevel) {
+					auto ShaderPlatform = GShaderPlatformForFeatureLevel[InFeatureLevel];
+					UMaterial::UpdateMaterialShaders(OutdatedShaderTypes, OutdatedShaderPipelineTypes, OutdatedFactoryTypes, ShaderPlatform);
+				});
+
+				GWarn->StatusUpdate(0, 1, NSLOCTEXT("ShaderCompilingManager", "CompilingGlobalShaderStatus", "Compiling global shaders..."));
 			}
 			else
 			{
diff --git a/Engine/Source/Runtime/Engine/Private/SkeletalRenderGPUSkin.cpp b/Engine/Source/Runtime/Engine/Private/SkeletalRenderGPUSkin.cpp
--- a/Engine/Source/Runtime/Engine/Private/SkeletalRenderGPUSkin.cpp
+++ b/Engine/Source/Runtime/Engine/Private/SkeletalRenderGPUSkin.cpp
@@ -1080,8 +1080,6 @@ const FTwoVectors& FSkeletalMeshObjectGPUSkin::GetCustomLeftRightVectors(int32 S
 FDynamicSkelMeshObjectDataGPUSkin
 -----------------------------------------------------------------------------*/
 
-static TLockFreePointerListUnordered<FDynamicSkelMeshObjectDataGPUSkin, PLATFORM_CACHE_LINE_SIZE>	FreeDynamicSkelMeshObjectDataGPUSkins;
-
 void FDynamicSkelMeshObjectDataGPUSkin::Clear()
 {
 	ReferenceToLocal.Reset();
@@ -1099,19 +1097,12 @@ void FDynamicSkelMeshObjectDataGPUSkin::Clear()
 
 FDynamicSkelMeshObjectDataGPUSkin* FDynamicSkelMeshObjectDataGPUSkin::AllocDynamicSkelMeshObjectDataGPUSkin()
 {
-	FDynamicSkelMeshObjectDataGPUSkin* Result = FreeDynamicSkelMeshObjectDataGPUSkins.Pop();
-	if (!Result)
-	{
-		Result = new FDynamicSkelMeshObjectDataGPUSkin;
-	}
-	return Result;
+	return new FDynamicSkelMeshObjectDataGPUSkin;
 }
 
 void FDynamicSkelMeshObjectDataGPUSkin::FreeDynamicSkelMeshObjectDataGPUSkin(FDynamicSkelMeshObjectDataGPUSkin* Who)
 {
-	check(Who);
-	Who->Clear();
-	FreeDynamicSkelMeshObjectDataGPUSkins.Push(Who);
+	delete Who;
 }
 
 void FDynamicSkelMeshObjectDataGPUSkin::InitDynamicSkelMeshObjectDataGPUSkin(
diff --git a/Engine/Source/Runtime/Engine/Private/SkeletalRenderGPUSkin.h b/Engine/Source/Runtime/Engine/Private/SkeletalRenderGPUSkin.h
--- a/Engine/Source/Runtime/Engine/Private/SkeletalRenderGPUSkin.h
+++ b/Engine/Source/Runtime/Engine/Private/SkeletalRenderGPUSkin.h
@@ -27,8 +27,6 @@ class FDynamicSkelMeshObjectDataGPUSkin
 
 	virtual ~FDynamicSkelMeshObjectDataGPUSkin()
 	{
-		// we leak these
-		check(0);
 	}
 
 	void Clear();
diff --git a/Engine/Source/Runtime/Engine/Private/UnrealEngine.cpp b/Engine/Source/Runtime/Engine/Private/UnrealEngine.cpp
--- a/Engine/Source/Runtime/Engine/Private/UnrealEngine.cpp
+++ b/Engine/Source/Runtime/Engine/Private/UnrealEngine.cpp
@@ -12085,7 +12085,25 @@ static void SetupThreadAffinity(const TArray<FString>& Args)
 	{
 		FSimpleDelegateGraphTask::CreateAndDispatchWhenReady(
 			FSimpleDelegateGraphTask::FDelegate::CreateStatic(&SetAffinityOnThread),
-			TStatId(), NULL, ENamedThreads::AnyThread);
+			TStatId(), NULL, ENamedThreads::AnyNormalThreadHiPriTask);
+	}
+	if (ENamedThreads::bHasHighPriorityThreads)
+	{
+		for (int32 Index = 0; Index < FTaskGraphInterface::Get().GetNumWorkerThreads(); Index++)
+		{
+			FSimpleDelegateGraphTask::CreateAndDispatchWhenReady(
+				FSimpleDelegateGraphTask::FDelegate::CreateStatic(&SetAffinityOnThread),
+				TStatId(), NULL, ENamedThreads::AnyHiPriThreadHiPriTask);
+		}
+	}
+	if (ENamedThreads::bHasBackgroundThreads)
+	{
+		for (int32 Index = 0; Index < FTaskGraphInterface::Get().GetNumWorkerThreads(); Index++)
+		{
+			FSimpleDelegateGraphTask::CreateAndDispatchWhenReady(
+				FSimpleDelegateGraphTask::FDelegate::CreateStatic(&SetAffinityOnThread),
+				TStatId(), NULL, ENamedThreads::AnyBackgroundHiPriTask);
+		}
 	}
 	FSimpleDelegateGraphTask::CreateAndDispatchWhenReady(
 		FSimpleDelegateGraphTask::FDelegate::CreateStatic(&SetAffinityOnThread),
diff --git a/Engine/Source/Runtime/Launch/Private/LaunchEngineLoop.cpp b/Engine/Source/Runtime/Launch/Private/LaunchEngineLoop.cpp
--- a/Engine/Source/Runtime/Launch/Private/LaunchEngineLoop.cpp
+++ b/Engine/Source/Runtime/Launch/Private/LaunchEngineLoop.cpp
@@ -634,6 +634,8 @@ FEngineLoop::FEngineLoop()
 
 int32 FEngineLoop::PreInit(int32 ArgC, TCHAR* ArgV[], const TCHAR* AdditionalCommandline)
 {
+	FMemory::SetupTLSCachesOnCurrentThread();
+
 	FString CmdLine;
 
 	// loop over the parameters, skipping the first one (which is the executable name)
@@ -747,6 +749,10 @@ int32 FEngineLoop::PreInit( const TCHAR* CmdLine )
 	// Create the stats malloc profiler proxy.
 	if( FStatsMallocProfilerProxy::HasMemoryProfilerToken() )
 	{
+		if (PLATFORM_USES_FIXED_GMalloc_CLASS)
+		{
+			UE_LOG(LogMemory, Fatal, TEXT("Cannot do malloc profiling with PLATFORM_USES_FIXED_GMalloc_CLASS."));
+		}
 		// Assumes no concurrency here.
 		GMalloc = FStatsMallocProfilerProxy::Get();
 	}
diff --git a/Engine/Source/Runtime/RHI/Private/RHICommandList.cpp b/Engine/Source/Runtime/RHI/Private/RHICommandList.cpp
--- a/Engine/Source/Runtime/RHI/Private/RHICommandList.cpp
+++ b/Engine/Source/Runtime/RHI/Private/RHICommandList.cpp
@@ -1634,37 +1634,24 @@ void FRHICommandListBase::HandleRTThreadTaskCompletion(const FGraphEventRef& MyC
 	RTTasks.Empty();
 }
 
-static TLockFreeFixedSizeAllocator<sizeof(FRHICommandList), PLATFORM_CACHE_LINE_SIZE, FThreadSafeCounter> RHICommandListAllocator;
-static TLockFreeFixedSizeAllocator<sizeof(FRHIAsyncComputeCommandList), PLATFORM_CACHE_LINE_SIZE, FThreadSafeCounter> RHIAsyncComputeCommandListAllocator;
-
 void* FRHICommandList::operator new(size_t Size)
 {
-	// doesn't support derived classes with a different size
-	check(Size == sizeof(FRHICommandList));
-	return RHICommandListAllocator.Allocate();
-	//return FMemory::Malloc(Size);
+	return FMemory::Malloc(Size);
 }
 
 void FRHICommandList::operator delete(void *RawMemory)
 {
-	check(RawMemory != (void*) &GRHICommandList.GetImmediateCommandList());
-	RHICommandListAllocator.Free(RawMemory);
-	//FMemory::Free(RawMemory);
+	FMemory::Free(RawMemory);
 }
 
 void* FRHIAsyncComputeCommandList::operator new(size_t Size)
 {
-	// doesn't support derived classes with a different size
-	check(Size == sizeof(FRHICommandList));
-	return RHIAsyncComputeCommandListAllocator.Allocate();
-	//return FMemory::Malloc(Size);
+	return FMemory::Malloc(Size);
 }
 
 void FRHIAsyncComputeCommandList::operator delete(void *RawMemory)
 {
-	check(RawMemory != (void*)&GRHICommandList.GetImmediateAsyncComputeCommandList());
-	RHIAsyncComputeCommandListAllocator.Free(RawMemory);
-	//FMemory::Free(RawMemory);
+	FMemory::Free(RawMemory);
 }
 
 void* FRHICommandListBase::operator new(size_t Size)
@@ -1675,9 +1662,7 @@ void* FRHICommandListBase::operator new(size_t Size)
 
 void FRHICommandListBase::operator delete(void *RawMemory)
 {
-	check(RawMemory != (void*) &GRHICommandList.GetImmediateCommandList());
-	RHICommandListAllocator.Free(RawMemory);
-	//FMemory::Free(RawMemory);
+	FMemory::Free(RawMemory);
 }	
 
 ///////// Pass through functions that allow RHIs to optimize certain calls.
diff --git a/Engine/Source/Runtime/RenderCore/Private/RenderingThread.cpp b/Engine/Source/Runtime/RenderCore/Private/RenderingThread.cpp
--- a/Engine/Source/Runtime/RenderCore/Private/RenderingThread.cpp
+++ b/Engine/Source/Runtime/RenderCore/Private/RenderingThread.cpp
@@ -267,8 +267,10 @@ public:
 
 	virtual uint32 Run()
 	{
+		FMemory::SetupTLSCachesOnCurrentThread();
 		FTaskGraphInterface::Get().AttachToThread(ENamedThreads::RHIThread);
 		FTaskGraphInterface::Get().ProcessThreadUntilRequestReturn(ENamedThreads::RHIThread);
+		FMemory::ClearAndDisableTLSCachesOnCurrentThread();
 		return 0;
 	}
 
@@ -424,6 +426,7 @@ public:
 
 	virtual uint32 Run(void) override
 	{
+		FMemory::SetupTLSCachesOnCurrentThread();
 		FPlatformProcess::SetupGameOrRenderThread(true);
 
 #if PLATFORM_WINDOWS
@@ -453,6 +456,7 @@ public:
 		{
 			RenderingThreadMain( TaskGraphBoundSyncEvent );
 		}
+		FMemory::ClearAndDisableTLSCachesOnCurrentThread();
 		return 0;
 	}
 };
diff --git a/Engine/Source/Runtime/Renderer/Private/SceneRendering.cpp b/Engine/Source/Runtime/Renderer/Private/SceneRendering.cpp
--- a/Engine/Source/Runtime/Renderer/Private/SceneRendering.cpp
+++ b/Engine/Source/Runtime/Renderer/Private/SceneRendering.cpp
@@ -546,6 +546,25 @@ static void SetBlack3DIfNull(FTextureRHIParamRef& Tex)
 	}
 }
 
+void UpdateNoiseTextureParameters(FFrameUniformShaderParameters& FrameUniformShaderParameters)
+{
+	if (GSystemTextures.PerlinNoiseGradient.GetReference())
+	{
+		FrameUniformShaderParameters.PerlinNoiseGradientTexture = (FTexture2DRHIRef&)GSystemTextures.PerlinNoiseGradient->GetRenderTargetItem().ShaderResourceTexture;
+		SetBlack2DIfNull(FrameUniformShaderParameters.PerlinNoiseGradientTexture);
+	}
+	check(FrameUniformShaderParameters.PerlinNoiseGradientTexture);
+	FrameUniformShaderParameters.PerlinNoiseGradientTextureSampler = TStaticSamplerState<SF_Point, AM_Wrap, AM_Wrap, AM_Wrap>::GetRHI();
+
+	if (GSystemTextures.PerlinNoise3D.GetReference())
+	{
+		FrameUniformShaderParameters.PerlinNoise3DTexture = (FTexture3DRHIRef&)GSystemTextures.PerlinNoise3D->GetRenderTargetItem().ShaderResourceTexture;
+		SetBlack3DIfNull(FrameUniformShaderParameters.PerlinNoise3DTexture);
+	}
+	check(FrameUniformShaderParameters.PerlinNoise3DTexture);
+	FrameUniformShaderParameters.PerlinNoise3DTextureSampler = TStaticSamplerState<SF_Bilinear, AM_Wrap, AM_Wrap, AM_Wrap>::GetRHI();
+}
+
 
 /** Creates the view's uniform buffers given a set of view transforms. */
 void FViewInfo::CreateUniformBuffer(
@@ -855,21 +874,8 @@ void FViewInfo::CreateUniformBuffer(
 	FrameUniformShaderParameters.AtmosphereIrradianceTextureSampler_UB = TStaticSamplerState<SF_Bilinear>::GetRHI();
 	FrameUniformShaderParameters.AtmosphereInscatterTextureSampler_UB = TStaticSamplerState<SF_Bilinear>::GetRHI();
 
-	if (GSystemTextures.PerlinNoiseGradient.GetReference())
-	{
-		FrameUniformShaderParameters.PerlinNoiseGradientTexture = (FTexture2DRHIRef&)GSystemTextures.PerlinNoiseGradient->GetRenderTargetItem().ShaderResourceTexture;
-		SetBlack2DIfNull(FrameUniformShaderParameters.PerlinNoiseGradientTexture);
-	}
-	check(FrameUniformShaderParameters.PerlinNoiseGradientTexture);
-	FrameUniformShaderParameters.PerlinNoiseGradientTextureSampler = TStaticSamplerState<SF_Point, AM_Wrap, AM_Wrap, AM_Wrap>::GetRHI();
-
-	if (GSystemTextures.PerlinNoise3D.GetReference())
-	{
-		FrameUniformShaderParameters.PerlinNoise3DTexture = (FTexture3DRHIRef&)GSystemTextures.PerlinNoise3D->GetRenderTargetItem().ShaderResourceTexture;
-		SetBlack3DIfNull(FrameUniformShaderParameters.PerlinNoise3DTexture);
-	}
-	check(FrameUniformShaderParameters.PerlinNoise3DTexture);
-	FrameUniformShaderParameters.PerlinNoise3DTextureSampler = TStaticSamplerState<SF_Bilinear, AM_Wrap, AM_Wrap, AM_Wrap>::GetRHI();
+	// Update the Texture Parameters used for noise
+	UpdateNoiseTextureParameters(FrameUniformShaderParameters);
 
 	for (int32 Index = 0; Index < GMaxGlobalDistanceFieldClipmaps; Index++)
 	{
diff --git a/Engine/Source/Runtime/Renderer/Private/SceneRendering.h b/Engine/Source/Runtime/Renderer/Private/SceneRendering.h
--- a/Engine/Source/Runtime/Renderer/Private/SceneRendering.h
+++ b/Engine/Source/Runtime/Renderer/Private/SceneRendering.h
@@ -1001,3 +1001,6 @@ private:
 	bool bModulatedShadowsInUse;
 	bool bCSMShadowsInUse;
 };
+
+// The noise textures need to be set in Slate too.
+RENDERER_API void UpdateNoiseTextureParameters(FFrameUniformShaderParameters& FrameUniformShaderParameters);
diff --git a/Engine/Source/Runtime/Renderer/Private/ShaderBaseClasses.cpp b/Engine/Source/Runtime/Renderer/Private/ShaderBaseClasses.cpp
--- a/Engine/Source/Runtime/Renderer/Private/ShaderBaseClasses.cpp
+++ b/Engine/Source/Runtime/Renderer/Private/ShaderBaseClasses.cpp
@@ -79,8 +79,6 @@ FUniformBufferRHIParamRef FMaterialShader::GetParameterCollectionBuffer(const FG
 	return Scene ? Scene->GetParameterCollectionBuffer(Id) : FUniformBufferRHIParamRef();
 }
 
-static TLockFreeFixedSizeAllocator<sizeof(FUniformExpressionCache), PLATFORM_CACHE_LINE_SIZE, FNoopCounter> TempUniformCache;
-
 template<typename ShaderRHIParamRef>
 void FMaterialShader::SetParameters(
 	FRHICommandList& RHICmdList,
@@ -112,7 +110,7 @@ void FMaterialShader::SetParameters(
 	{
 		FMaterialRenderContext MaterialRenderContext(MaterialRenderProxy, Material, &View);
 		bUniformExpressionCacheNeedsDelete = true;
-		UniformExpressionCache = new (TempUniformCache.Allocate()) FUniformExpressionCache();
+		UniformExpressionCache = new FUniformExpressionCache();
 		MaterialRenderProxy->EvaluateUniformExpressions(*UniformExpressionCache, MaterialRenderContext, &RHICmdList);
 		SetLocalUniformBufferParameter(RHICmdList, ShaderRHI, MaterialUniformBuffer, UniformExpressionCache->LocalUniformBuffer);
 	}
@@ -334,8 +332,7 @@ void FMaterialShader::SetParameters(
 
 	if (bUniformExpressionCacheNeedsDelete)
 	{
-		UniformExpressionCache->~FUniformExpressionCache();
-		TempUniformCache.Free(UniformExpressionCache);
+		delete UniformExpressionCache;
 	}
 }
 
diff --git a/Engine/Source/Runtime/SlateRHIRenderer/Private/SlateRHIRenderingPolicy.cpp b/Engine/Source/Runtime/SlateRHIRenderer/Private/SlateRHIRenderingPolicy.cpp
--- a/Engine/Source/Runtime/SlateRHIRenderer/Private/SlateRHIRenderingPolicy.cpp
+++ b/Engine/Source/Runtime/SlateRHIRenderer/Private/SlateRHIRenderingPolicy.cpp
@@ -15,6 +15,8 @@
 #include "SlateElementIndexBuffer.h"
 #include "SlateElementVertexBuffer.h"
 
+extern void UpdateNoiseTextureParameters(FFrameUniformShaderParameters& FrameUniformShaderParameters);
+
 DECLARE_CYCLE_STAT(TEXT("Update Buffers RT"), STAT_SlateUpdateBufferRTTime, STATGROUP_Slate);
 DECLARE_CYCLE_STAT(TEXT("PreFill Buffers RT"), STAT_SlatePreFullBufferRTTime, STATGROUP_Slate);
 DECLARE_DWORD_COUNTER_STAT(TEXT("Num Layers"), STAT_SlateNumLayers, STATGROUP_Slate);
@@ -240,6 +242,9 @@ static FSceneView& CreateSceneView( FSceneViewFamilyContext& ViewFamilyContext,
 	FrameUniformShaderParameters.Random = FMath::Rand();
 	FrameUniformShaderParameters.FrameNumber = View->Family->FrameNumber;
 
+	// Update noise buffers
+	UpdateNoiseTextureParameters(FrameUniformShaderParameters);
+
 	{
 		// setup a matrix to transform float4(SvPosition.xyz,1) directly to TranslatedWorld (quality, performance as we don't need to convert or use interpolator)
 
diff --git a/Engine/Source/Runtime/Windows/D3D12RHI/Private/D3D12Device.cpp b/Engine/Source/Runtime/Windows/D3D12RHI/Private/D3D12Device.cpp
--- a/Engine/Source/Runtime/Windows/D3D12RHI/Private/D3D12Device.cpp
+++ b/Engine/Source/Runtime/Windows/D3D12RHI/Private/D3D12Device.cpp
@@ -182,18 +182,14 @@ public:
 	}
 };
 
-static TLockFreeFixedSizeAllocator<sizeof(FD3D12CommandContextContainer), PLATFORM_CACHE_LINE_SIZE, FThreadSafeCounter> FD3D12CommandContextContainerAllocator;
-
 void* FD3D12CommandContextContainer::operator new(size_t Size)
 {
-	// doesn't support derived classes with a different size
-	check(Size == sizeof(FD3D12CommandContextContainer));
-	return FD3D12CommandContextContainerAllocator.Allocate();
+	return FMemory::Malloc(Size);
 }
 
 void FD3D12CommandContextContainer::operator delete(void* RawMemory)
 {
-	FD3D12CommandContextContainerAllocator.Free(RawMemory);
+	FMemory::Free(RawMemory);
 }
 
 IRHICommandContextContainer* FD3D12DynamicRHI::RHIGetCommandContextContainer()
