Fix for missing ENGINE_NET_VERSION define due to Version.h being unsubmittable

diff --git a/Engine/Source/Runtime/Core/Private/Misc/EngineVersion.cpp b/Engine/Source/Runtime/Core/Private/Misc/EngineVersion.cpp
--- a/Engine/Source/Runtime/Core/Private/Misc/EngineVersion.cpp
+++ b/Engine/Source/Runtime/Core/Private/Misc/EngineVersion.cpp
@@ -8,7 +8,7 @@
 #include "UObject/ReleaseObjectVersion.h"
 
 /** Version numbers for networking */
-int32 GEngineNetVersion = ENGINE_NET_VERSION;
+int32 GEngineNetVersion = BUILT_FROM_CHANGELIST;
 const int32 GEngineMinNetVersion		= 7038;
 const int32 GEngineNegotiationVersion	= 3077;
 
