Copying //UE4/Dev-Core to //UE4/Main

diff --git a/Engine/Source/Runtime/Core/Private/HAL/MallocVerify.h b/Engine/Source/Runtime/Core/Private/HAL/MallocVerify.h
--- a/Engine/Source/Runtime/Core/Private/HAL/MallocVerify.h
+++ b/Engine/Source/Runtime/Core/Private/HAL/MallocVerify.h
@@ -19,7 +19,6 @@ class FMallocVerify
 {
 	/** List of all currently allocated pointers */
 	TSet<void*, DefaultKeyFuncs<void*>, FAnsiSetAllocator> AllocatedPointers;
-	FCriticalSection AllocatedPointersCritical;
 
 public:
 
@@ -46,6 +45,9 @@ private:
 	/* Verifier object */
 	FMallocVerify Verify;
 
+	/** Malloc critical section */
+	FCriticalSection VerifyCritical;
+
 public:
 	explicit FMallocVerifyProxy(FMalloc* InMalloc)
 		: UsedMalloc(InMalloc)
@@ -54,6 +56,7 @@ public:
 
 	virtual void* Malloc(SIZE_T Size, uint32 Alignment) override
 	{
+		FScopeLock VerifyLock(&VerifyCritical);
 		void* Result = UsedMalloc->Malloc(Size, Alignment);
 		Verify.Malloc(Result);
 		return Result;
@@ -61,6 +64,7 @@ public:
 
 	virtual void* Realloc(void* Ptr, SIZE_T NewSize, uint32 Alignment) override
 	{
+		FScopeLock VerifyLock(&VerifyCritical);
 		void* Result = UsedMalloc->Realloc(Ptr, NewSize, Alignment);
 		Verify.Realloc(Ptr, Result);
 		return Result;
@@ -70,6 +74,7 @@ public:
 	{
 		if (Ptr)
 		{
+			FScopeLock VerifyLock(&VerifyCritical);
 			Verify.Free(Ptr);
 			UsedMalloc->Free(Ptr);
 		}
@@ -125,6 +130,11 @@ public:
 	{ 
 		return UsedMalloc->GetDescriptiveName();
 	}
+
+	virtual void FlushCurrentThreadCache() override
+	{
+		UsedMalloc->FlushCurrentThreadCache();
+	}
 };
 
 #endif // MALLOC_VERIFY
diff --git a/Engine/Source/Runtime/Core/Private/HAL/PThreadRunnableThread.h b/Engine/Source/Runtime/Core/Private/HAL/PThreadRunnableThread.h
--- a/Engine/Source/Runtime/Core/Private/HAL/PThreadRunnableThread.h
+++ b/Engine/Source/Runtime/Core/Private/HAL/PThreadRunnableThread.h
@@ -161,6 +161,8 @@ protected:
 		// cache the thread ID for this thread (defined by the platform)
 		ThisThread->ThreadID = FPlatformTLS::GetCurrentThreadId();
 
+		FThreadManager::Get().AddThread(ThisThread->ThreadID, ThisThread);
+
 		// set the affinity.  This function sets affinity on the current thread, so don't call in the Create function which will trash the main thread affinity.
 		FPlatformProcess::SetThreadAffinityMask(ThisThread->ThreadAffinityMask);		
 
diff --git a/Engine/Source/Runtime/Core/Private/HAL/ThreadHeartBeat.cpp b/Engine/Source/Runtime/Core/Private/HAL/ThreadHeartBeat.cpp
new file mode 100644
--- /dev/null
+++ b/Engine/Source/Runtime/Core/Private/HAL/ThreadHeartBeat.cpp
@@ -0,0 +1,113 @@
+// Copyright 1998-2016 Epic Games, Inc. All Rights Reserved.
+#include "CorePrivatePCH.h"
+#include "ThreadHeartBeat.h"
+
+FThreadHeartBeat::FThreadHeartBeat()
+: Thread(nullptr)
+{
+	// We don't care about programs for now so no point in spawning the extra thread
+#if !IS_PROGRAM
+	Thread = FRunnableThread::Create(this, TEXT("FHeartBeatThread"), 0, TPri_BelowNormal);
+#endif
+}
+
+FThreadHeartBeat::~FThreadHeartBeat()
+{
+	delete Thread;
+	Thread = nullptr;
+}
+
+FThreadHeartBeat& FThreadHeartBeat::Get()
+{
+	static FThreadHeartBeat Singleton;
+	return Singleton;
+}
+
+	//~ Begin FRunnable Interface.
+bool FThreadHeartBeat::Init()
+{
+	return true;
+}
+
+uint32 FThreadHeartBeat::Run()
+{
+	while (StopTaskCounter.GetValue() == 0)
+	{
+		uint32 ThreadThatHung = CheckHeartBeat();
+		if (ThreadThatHung != FThreadHeartBeat::InvalidThreadId)
+		{
+			const SIZE_T StackTraceSize = 65535;
+			ANSICHAR* StackTrace = (ANSICHAR*)GMalloc->Malloc(StackTraceSize);
+			StackTrace[0] = 0;
+			// Walk the stack and dump it to the allocated memory. This process usually allocates a lot of memory.
+			FPlatformStackWalk::ThreadStackWalkAndDump(StackTrace, StackTraceSize, 0, ThreadThatHung);
+			FString StackTraceText(StackTrace);
+			TArray<FString> StackLines;
+			StackTraceText.ParseIntoArrayLines(StackLines);
+
+			FString ThreadName(ThreadThatHung == GGameThreadId ? TEXT("GameThread") : FThreadManager::Get().GetThreadName(ThreadThatHung));
+			if (ThreadName.IsEmpty())
+			{
+				ThreadName = FString::Printf(TEXT("unknown thread (%u)"), ThreadThatHung);
+			}
+			UE_LOG(LogCore, Error, TEXT("Infinite stall detected on %s:"), *ThreadName);
+			for (FString& StackLine : StackLines)
+			{
+				UE_LOG(LogCore, Error, TEXT("%s"), *StackLine);
+			}
+			
+			FString StackTrimmed;
+			for (int32 LineIndex = 0; LineIndex < StackLines.Num() && StackTrimmed.Len() < 512; ++LineIndex)
+			{
+				StackTrimmed += StackLines[LineIndex];
+				StackTrimmed += TEXT("\n");
+			}
+			UE_LOG(LogCore, Fatal, TEXT("Infinite stall detected on %s:\n%s"), *ThreadName, *StackTrimmed);
+		}
+		FPlatformProcess::SleepNoStats(0.5f);
+	}
+
+	return 0;
+}
+
+void FThreadHeartBeat::Stop()
+{
+	StopTaskCounter.Increment();
+}
+
+void FThreadHeartBeat::HeartBeat()
+{
+	uint32 ThreadId = FPlatformTLS::GetCurrentThreadId();
+	FScopeLock HeartBeatLock(&HeartBeatCritical);
+	double& HeartBeatTime = ThreadHeartBeat.FindOrAdd(ThreadId);
+	HeartBeatTime = FPlatformTime::Seconds();
+}
+
+uint32 FThreadHeartBeat::CheckHeartBeat()
+{
+	// Editor and debug builds run too slow to measure them correctly
+#if !WITH_EDITORONLY_DATA && !IS_PROGRAM && !UE_BUILD_DEBUG
+	if (!GIsRequestingExit)
+	{
+		const double HangDuration = 10.0;
+		const double CurrentTime = FPlatformTime::Seconds();
+		FScopeLock HeartBeatLock(&HeartBeatCritical);
+		for (TPair<uint32, double>& LastHeartBeat : ThreadHeartBeat)
+		{
+			if ((CurrentTime - LastHeartBeat.Value) > HangDuration)
+			{
+				LastHeartBeat.Value = CurrentTime;
+				return LastHeartBeat.Key;
+			}
+		}
+	}
+#endif
+	return InvalidThreadId;
+}
+
+void FThreadHeartBeat::KillHeartBeat()
+{
+	uint32 ThreadId = FPlatformTLS::GetCurrentThreadId();
+	FScopeLock HeartBeatLock(&HeartBeatCritical);
+	ThreadHeartBeat.Remove(ThreadId);
+}
\ No newline at end of file
diff --git a/Engine/Source/Runtime/Core/Private/HAL/ThreadingBase.cpp b/Engine/Source/Runtime/Core/Private/HAL/ThreadingBase.cpp
--- a/Engine/Source/Runtime/Core/Private/HAL/ThreadingBase.cpp
+++ b/Engine/Source/Runtime/Core/Private/HAL/ThreadingBase.cpp
@@ -87,18 +87,18 @@ public:
 	{
 		ThreadID = ThreadIdCounter++;
 		// Auto register with single thread manager.
-		FSingleThreadManager::Get().AddThread(this);
+		FThreadManager::Get().AddThread(ThreadID, this);
 	}
 
 	/** Virtual destructor. */
 	virtual ~FFakeThread()
 	{
 		// Remove from the manager.
-		FSingleThreadManager::Get().RemoveThread(this);
+		FThreadManager::Get().RemoveThread(this);
 	}
 
 	/** Tick one time per frame. */
-	void Tick()
+	virtual void Tick() override
 	{
 		if (Runnable && !bIsSuspended)
 		{
@@ -122,13 +122,13 @@ public:
 
 	virtual bool Kill(bool bShouldWait) override
 	{
-		FSingleThreadManager::Get().RemoveThread(this);
+		FThreadManager::Get().RemoveThread(this);
 		return true;
 	}
 
 	virtual void WaitForCompletion() override
 	{
-		FSingleThreadManager::Get().RemoveThread(this);
+		FThreadManager::Get().RemoveThread(this);
 	}
 
 	virtual bool CreateInternal(FRunnable* InRunnable, const TCHAR* InThreadName,
@@ -147,30 +147,57 @@ public:
 uint32 FFakeThread::ThreadIdCounter = 0xffff;
 
 
-void FSingleThreadManager::AddThread(FFakeThread* Thread)
+void FThreadManager::AddThread(uint32 ThreadId, FRunnableThread* Thread)
 {
-	ThreadList.Add(Thread);
+	FScopeLock ThreadsLock(&ThreadsCritical);
+	// Some platforms do not support TLS
+	if (!Threads.Contains(ThreadId))
+	{
+		Threads.Add(ThreadId, Thread);
+	}
 }
 
-void FSingleThreadManager::RemoveThread(FFakeThread* Thread)
+void FThreadManager::RemoveThread(FRunnableThread* Thread)
 {
-	ThreadList.Remove(Thread);
+	FScopeLock ThreadsLock(&ThreadsCritical);
+	const uint32* ThreadId = Threads.FindKey(Thread);
+	if (ThreadId)
+	{
+		Threads.Remove(*ThreadId);
+	}
 }
 
-void FSingleThreadManager::Tick()
-{
-	QUICK_SCOPE_CYCLE_COUNTER(STAT_FSingleThreadManager_Tick);
+void FThreadManager::Tick()
+{	
+	if (!FPlatformProcess::SupportsMultithreading())
+	{
+		QUICK_SCOPE_CYCLE_COUNTER(STAT_FSingleThreadManager_Tick);
+
+		FScopeLock ThreadsLock(&ThreadsCritical);
+
+		// Tick all registered threads.
+		for (TPair<uint32, FRunnableThread*>& ThreadPair : Threads)
+		{
+			ThreadPair.Value->Tick();
+		}
+	}
+}
 
-	// Tick all registered threads.
-	for (int32 RunnableIndex = 0; RunnableIndex < ThreadList.Num(); ++RunnableIndex)
+const FString& FThreadManager::GetThreadName(uint32 ThreadId)
+{
+	static FString NoThreadName;
+	FScopeLock ThreadsLock(&ThreadsCritical);
+	FRunnableThread** Thread = Threads.Find(ThreadId);
+	if (Thread)
 	{
-		ThreadList[RunnableIndex]->Tick();
+		return (*Thread)->GetThreadName();
 	}
+	return NoThreadName;
 }
 
-FSingleThreadManager& FSingleThreadManager::Get()
+FThreadManager& FThreadManager::Get()
 {
-	static FSingleThreadManager Singleton;
+	static FThreadManager Singleton;
 	return Singleton;
 }
 
@@ -259,6 +286,14 @@ FRunnableThread::FRunnableThread()
 {
 }
 
+FRunnableThread::~FRunnableThread()
+{
+	if (!GIsRequestingExit)
+	{
+		FThreadManager::Get().RemoveThread(this);
+	}
+}
+
 FRunnableThread* FRunnableThread::Create(
 	class FRunnable* InRunnable,
 	const TCHAR* ThreadName,
diff --git a/Engine/Source/Runtime/Core/Private/Misc/App.cpp b/Engine/Source/Runtime/Core/Private/Misc/App.cpp
--- a/Engine/Source/Runtime/Core/Private/Misc/App.cpp
+++ b/Engine/Source/Runtime/Core/Private/Misc/App.cpp
@@ -183,12 +183,12 @@ bool FApp::ShouldUseThreadingForPerformance()
 #endif // HAVE_RUNTIME_THREADING_SWITCHES
 
 
-static bool GUnfocusedVolumeMultiplierItialised = false;
+static bool GUnfocusedVolumeMultiplierInitialised = false;
 float FApp::GetUnfocusedVolumeMultiplier()
 {
-	if (!GUnfocusedVolumeMultiplierItialised)
+	if (!GUnfocusedVolumeMultiplierInitialised)
 	{
-		GUnfocusedVolumeMultiplierItialised = true;
+		GUnfocusedVolumeMultiplierInitialised = true;
 		GConfig->GetFloat(TEXT("Audio"), TEXT("UnfocusedVolumeMultiplier"), UnfocusedVolumeMultiplier, GEngineIni);
 	}
 	return UnfocusedVolumeMultiplier;
@@ -197,7 +197,7 @@ float FApp::GetUnfocusedVolumeMultiplier()
 void FApp::SetUnfocusedVolumeMultiplier(float InVolumeMultiplier)
 {
 	UnfocusedVolumeMultiplier = InVolumeMultiplier;
-	GUnfocusedVolumeMultiplierItialised = true;
+	GUnfocusedVolumeMultiplierInitialised = true;
 }
 
 void FApp::SetUseVRFocus(bool bInUseVRFocus)
diff --git a/Engine/Source/Runtime/Core/Private/Modules/ModuleManager.cpp b/Engine/Source/Runtime/Core/Private/Modules/ModuleManager.cpp
--- a/Engine/Source/Runtime/Core/Private/Modules/ModuleManager.cpp
+++ b/Engine/Source/Runtime/Core/Private/Modules/ModuleManager.cpp
@@ -226,9 +226,38 @@ void FModuleManager::AddModule(const FName InModuleName)
 		return;
 	}
 
+	FString ModuleFilename = MoveTemp(TMap<FName, FString>::TConstIterator(ModulePathMap).Value());
+
+	const int32 MatchPos = ModuleFilename.Find(ModuleNameString, ESearchCase::IgnoreCase, ESearchDir::FromEnd);
+	if (!ensureMsgf(MatchPos != INDEX_NONE, TEXT("Could not find module name '%s' in module filename '%s'"), InModuleName, *ModuleFilename))
+	{
+		return;
+	}
+
+	// Skip any existing module number suffix
+	const int32 SuffixStart = MatchPos + ModuleNameString.Len();
+	int32 SuffixEnd = SuffixStart;
+	if (ModuleFilename[SuffixEnd] == TEXT('-'))
+	{
+		++SuffixEnd;
+		while (FCString::Strchr(TEXT("0123456789"), ModuleFilename[SuffixEnd]))
+		{
+			++SuffixEnd;
+		}
+
+		// Only skip the suffix if it was a number
+		if (SuffixEnd - SuffixStart == 1)
+		{
+			--SuffixEnd;
+		}
+	}
+
+	const FString Prefix = ModuleFilename.Left(SuffixStart);
+	const FString Suffix = ModuleFilename.Right(ModuleFilename.Len() - SuffixEnd);
+
 	// Add this module to the set of modules that we know about
-	ModuleInfo->OriginalFilename = TMap<FName, FString>::TConstIterator(ModulePathMap).Value();
-	ModuleInfo->Filename = ModuleInfo->OriginalFilename;
+	ModuleInfo->OriginalFilename = Prefix + Suffix;
+	ModuleInfo->Filename         = ModuleInfo->OriginalFilename;
 
 	// When iterating on code during development, it's possible there are multiple rolling versions of this
 	// module's DLL file.  This can happen if the programmer is recompiling DLLs while the game is loaded.  In
@@ -246,18 +275,6 @@ void FModuleManager::AddModule(const FName InModuleName)
 		return;
 	}
 
-	const FString ModuleName = *InModuleName.ToString();
-	const int32 MatchPos = ModuleInfo->OriginalFilename.Find(ModuleName, ESearchCase::IgnoreCase, ESearchDir::FromEnd);
-	if (!ensureMsgf(MatchPos != INDEX_NONE, TEXT("Could not find module name '%s' in module filename '%s'"), *ModuleName, *ModuleInfo->OriginalFilename))
-	{
-		return;
-	}
-
-	const int32 SuffixPos = MatchPos + ModuleName.Len();
-
-	const FString Prefix = ModuleInfo->OriginalFilename.Left(SuffixPos);
-	const FString Suffix = ModuleInfo->OriginalFilename.Right(ModuleInfo->OriginalFilename.Len() - SuffixPos);
-
 	const FString ModuleFileSearchString = FString::Printf(TEXT("%s-*%s"), *Prefix, *Suffix);
 
 	// Search for module files
diff --git a/Engine/Source/Runtime/Core/Private/ProfilingDebugging/LoadTimeTracker.cpp b/Engine/Source/Runtime/Core/Private/ProfilingDebugging/LoadTimeTracker.cpp
--- a/Engine/Source/Runtime/Core/Private/ProfilingDebugging/LoadTimeTracker.cpp
+++ b/Engine/Source/Runtime/Core/Private/ProfilingDebugging/LoadTimeTracker.cpp
@@ -54,7 +54,7 @@ void FLoadTimeTracker::DumpLoadTimes()
 
 void FLoadTimeTracker::ResetLoadTimes()
 {
-	static bool bActuallyReset = !FParse::Param(FCommandLine::Get(), TEXT("-NoLoadTrackClear"));
+	static bool bActuallyReset = !FParse::Param(FCommandLine::Get(), TEXT("NoLoadTrackClear"));
 	if(bActuallyReset)
 	{
 		TimeInfo.Reset();
diff --git a/Engine/Source/Runtime/Core/Private/Stats/StatsCommand.cpp b/Engine/Source/Runtime/Core/Private/Stats/StatsCommand.cpp
--- a/Engine/Source/Runtime/Core/Private/Stats/StatsCommand.cpp
+++ b/Engine/Source/Runtime/Core/Private/Stats/StatsCommand.cpp
@@ -252,14 +252,24 @@ void DumpHistoryFrame(FStatsThreadState const& StatsData, int64 TargetFrame, flo
 	}
 }
 
-void DumpNonFrame(FStatsThreadState const& StatsData)
+void DumpNonFrame(FStatsThreadState const& StatsData, const FName OptionalGroup)
 {
-	UE_LOG(LogStats, Log, TEXT("Full non-frame data ---------------------------------"));
+	if (OptionalGroup == NAME_None)
+	{
+		UE_LOG(LogStats, Log, TEXT("Full non-frame data ---------------------------------"));
+	}
+	else
+	{
+		UE_LOG(LogStats, Log, TEXT("Filtered non-frame data ---------------------------------"));
+	}
 
 	TArray<FStatMessage> Stats;
 	for (auto It = StatsData.NotClearedEveryFrame.CreateConstIterator(); It; ++It)
 	{
-		Stats.Add(It.Value());
+		if (OptionalGroup == NAME_None || OptionalGroup == It.Value().NameAndInfo.GetGroupName())
+		{
+			Stats.Add(It.Value());
+		}
 	}
 	Stats.Sort(FGroupSort());
 	FName LastGroup = NAME_None;
@@ -1724,7 +1734,10 @@ static void StatCmd(FString InCmd, bool bStatCommand)
 		}
 		else if (FParse::Command(&Cmd, TEXT("DUMPNONFRAME")))
 		{
-			DumpNonFrame(Stats);
+			FString MaybeGroup;
+			FParse::Token(Cmd, MaybeGroup, false);
+
+			DumpNonFrame(Stats, MaybeGroup.Len() == 0 ? NAME_None : FName(*(FString(TEXT("STATGROUP_")) + MaybeGroup)));
 		}
 		else if (FParse::Command(&Cmd, TEXT("DUMPCPU")))
 		{
diff --git a/Engine/Source/Runtime/Core/Private/WinRT/WinRTRunnableThread.h b/Engine/Source/Runtime/Core/Private/WinRT/WinRTRunnableThread.h
--- a/Engine/Source/Runtime/Core/Private/WinRT/WinRTRunnableThread.h
+++ b/Engine/Source/Runtime/Core/Private/WinRT/WinRTRunnableThread.h
@@ -74,6 +74,7 @@ class FRunnableThreadWinRT : public FRunnableThread
 		check(Runnable);
 
 		ThreadID = GetCurrentThreadId();
+		FThreadManager::Get().AddThread(ThreadID, this);
 
 		// Initialize the runnable object
 		if (Runnable->Init() == true)
diff --git a/Engine/Source/Runtime/Core/Private/Windows/WindowsPlatformMemory.cpp b/Engine/Source/Runtime/Core/Private/Windows/WindowsPlatformMemory.cpp
--- a/Engine/Source/Runtime/Core/Private/Windows/WindowsPlatformMemory.cpp
+++ b/Engine/Source/Runtime/Core/Private/Windows/WindowsPlatformMemory.cpp
@@ -13,7 +13,7 @@
 #endif // ENABLE_WIN_ALLOC_TRACKING
 
 #if !FORCE_ANSI_ALLOCATOR
-#include "MallocBinned.h"
+#include "MallocBinnedRedirect.h"
 #endif
 
 #include "AllowWindowsPlatformTypes.h"
@@ -79,7 +79,7 @@ FMalloc* FWindowsPlatformMemory::BaseAllocator()
 #elif (WITH_EDITORONLY_DATA || IS_PROGRAM) && TBB_ALLOCATOR_ALLOWED
 	return new FMallocTBB();
 #else
-	return new FMallocBinned((uint32)(GetConstants().PageSize&MAX_uint32), (uint64)MAX_uint32+1);
+	return new FMallocBinnedRedirect((uint32)(GetConstants().PageSize&MAX_uint32), (uint64)MAX_uint32+1);
 #endif
 }
 
diff --git a/Engine/Source/Runtime/Core/Private/Windows/WindowsPlatformMisc.cpp b/Engine/Source/Runtime/Core/Private/Windows/WindowsPlatformMisc.cpp
--- a/Engine/Source/Runtime/Core/Private/Windows/WindowsPlatformMisc.cpp
+++ b/Engine/Source/Runtime/Core/Private/Windows/WindowsPlatformMisc.cpp
@@ -30,7 +30,6 @@
 #include "VarargsHelper.h"
 
 #if !FORCE_ANSI_ALLOCATOR
-	#include "MallocBinned2.h"
 	#include "AllowWindowsPlatformTypes.h"
 		#include <psapi.h>
 	#include "HideWindowsPlatformTypes.h"
@@ -2124,6 +2123,13 @@ void FWindowsPlatformMisc::PromptForRemoteDebugging(bool bIsEnsure)
 			return;
 		}
 
+		if (GIsCriticalError && !GIsGuarded)
+		{
+			// A fatal error occurred.
+			// We have not ability to debug, this does not make sense to ask.
+			return;
+		}
+
 		// Upload locally compiled files for remote debugging
 		FPlatformStackWalk::UploadLocalSymbols();
 
@@ -2408,7 +2414,7 @@ FString FWindowsPlatformMisc::GetPrimaryGPUBrand()
 	}
 
 	return PrimaryGPUBrand;
-}
+	}
 
 static void GetVideoDriverDetails(const FString& Key, FGPUDriverInfo& Out)
 {
@@ -2865,10 +2871,10 @@ IPlatformChunkInstall* FWindowsPlatformMisc::GetPlatformChunkInstall()
 		{
 			PlatformChunkInstallModule = FModuleManager::LoadModulePtr<IPlatformChunkInstallModule>("HTTPChunkInstaller");
 			if (PlatformChunkInstallModule != nullptr)
-			{
-				// Attempt to grab the platform installer
-				ChunkInstall = PlatformChunkInstallModule->GetPlatformChunkInstall();
-			}
+		{
+			// Attempt to grab the platform installer
+			ChunkInstall = PlatformChunkInstallModule->GetPlatformChunkInstall();
+		}
 		}
 
 		if (PlatformChunkInstallModule == nullptr)
diff --git a/Engine/Source/Runtime/Core/Private/Windows/WindowsPlatformStackWalk.cpp b/Engine/Source/Runtime/Core/Private/Windows/WindowsPlatformStackWalk.cpp
--- a/Engine/Source/Runtime/Core/Private/Windows/WindowsPlatformStackWalk.cpp
+++ b/Engine/Source/Runtime/Core/Private/Windows/WindowsPlatformStackWalk.cpp
@@ -176,6 +176,19 @@ void FWindowsPlatformStackWalk::StackWalkAndDump( ANSICHAR* HumanReadableString,
 	InitStackWalking();
 	FGenericPlatformStackWalk::StackWalkAndDump(HumanReadableString, HumanReadableStringSize, IgnoreCount, Context);
 }
+
+void FWindowsPlatformStackWalk::ThreadStackWalkAndDump(ANSICHAR* HumanReadableString, SIZE_T HumanReadableStringSize, int32 IgnoreCount, uint32 ThreadId)
+{
+	InitStackWalking();
+	HANDLE ThreadHandle = OpenThread(THREAD_GET_CONTEXT | THREAD_SET_CONTEXT | THREAD_TERMINATE | THREAD_SUSPEND_RESUME, false, ThreadId);
+	CONTEXT ThreadContext;
+	ThreadContext.ContextFlags = CONTEXT_CONTROL;
+	if (GetThreadContext(ThreadHandle, &ThreadContext))
+	{
+		FGenericPlatformStackWalk::StackWalkAndDump(HumanReadableString, HumanReadableStringSize, IgnoreCount, &ThreadContext);
+	}
+}
+
 // #CrashReport: 2014-09-05 Switch to TArray<uint64,TFixedAllocator<100>>>
 /**
  * Capture a stack backtrace and optionally use the passed in exception pointers.
diff --git a/Engine/Source/Runtime/Core/Private/Windows/WindowsRunnableThread.h b/Engine/Source/Runtime/Core/Private/Windows/WindowsRunnableThread.h
--- a/Engine/Source/Runtime/Core/Private/Windows/WindowsRunnableThread.h
+++ b/Engine/Source/Runtime/Core/Private/Windows/WindowsRunnableThread.h
@@ -184,6 +184,8 @@ protected:
 		}
 		else
 		{
+			FThreadManager::Get().AddThread(ThreadID, this);
+
 			// Let the thread start up, then set the name for debug purposes.
 			ThreadInitSyncEvent->Wait(INFINITE);
 			ThreadName = InThreadName ? InThreadName : TEXT("Unnamed UE4");
diff --git a/Engine/Source/Runtime/Core/Public/Containers/Queue.h b/Engine/Source/Runtime/Core/Public/Containers/Queue.h
--- a/Engine/Source/Runtime/Core/Public/Containers/Queue.h
+++ b/Engine/Source/Runtime/Core/Public/Containers/Queue.h
@@ -141,7 +141,7 @@ public:
 	 * @return true if an item was returned, false if the queue was empty.
 	 * @see Dequeue, Enqueue, IsEmpty
 	 */
-	bool Peek( ItemType& OutItem )
+	bool Peek( ItemType& OutItem ) const
 	{
 		if (Tail->NextNode == nullptr)
 		{
diff --git a/Engine/Source/Runtime/Core/Public/Containers/Union.h b/Engine/Source/Runtime/Core/Public/Containers/Union.h
--- a/Engine/Source/Runtime/Core/Public/Containers/Union.h
+++ b/Engine/Source/Runtime/Core/Public/Containers/Union.h
@@ -9,7 +9,7 @@
 #include "ArchiveBase.h"
 
 
-DECLARE_LOG_CATEGORY_EXTERN(LogUnion, Log, All);
+CORE_API DECLARE_LOG_CATEGORY_EXTERN(LogUnion, Log, All);
 
 
 /** Used to disambiguate methods that are overloaded for all possible subtypes of a TUnion where the subtypes may not be distinct. */
@@ -105,7 +105,7 @@ public:
 		// Copy the value of the appropriate subtype from the other union
 		switch(Other.CurrentSubtypeIndex)
 		{
-		case -1: break;
+		case (uint8)-1: break;
 		case 0: SetSubtype<TypeA>(Other.GetSubtype<TypeA>()); break;
 		case 1: SetSubtype<TypeB>(Other.GetSubtype<TypeB>()); break;
 		case 2: SetSubtype<TypeC>(Other.GetSubtype<TypeC>()); break;
diff --git a/Engine/Source/Runtime/Core/Public/GenericPlatform/GenericPlatformFile.h b/Engine/Source/Runtime/Core/Public/GenericPlatform/GenericPlatformFile.h
--- a/Engine/Source/Runtime/Core/Public/GenericPlatform/GenericPlatformFile.h
+++ b/Engine/Source/Runtime/Core/Public/GenericPlatform/GenericPlatformFile.h
@@ -178,6 +178,14 @@ public:
 		}
 	}
 
+	virtual void		BypassSecurity(bool bInBypass)
+	{
+		if (GetLowerLevel() != nullptr)
+		{
+			GetLowerLevel()->BypassSecurity(bInBypass);
+		}
+	}
+
 	/** Gets the platform file wrapped by this file. */
 	virtual IPlatformFile* GetLowerLevel() = 0;
 		/** Gets this platform file type name. */
diff --git a/Engine/Source/Runtime/Core/Public/GenericPlatform/GenericPlatformMemory.h b/Engine/Source/Runtime/Core/Public/GenericPlatform/GenericPlatformMemory.h
--- a/Engine/Source/Runtime/Core/Public/GenericPlatform/GenericPlatformMemory.h
+++ b/Engine/Source/Runtime/Core/Public/GenericPlatform/GenericPlatformMemory.h
@@ -335,13 +335,16 @@ private:
 		B = Tmp;
 	}
 
-	static void MemswapImpl( void* Ptr1, void* Ptr2, SIZE_T Size );
+	static void MemswapGreaterThan8( void* Ptr1, void* Ptr2, SIZE_T Size );
 
 public:
 	static inline void Memswap( void* Ptr1, void* Ptr2, SIZE_T Size )
 	{
 		switch (Size)
 		{
+			case 0:
+				break;
+
 			case 1:
 				Valswap(*(uint8*)Ptr1, *(uint8*)Ptr2);
 				break;
@@ -350,10 +353,31 @@ public:
 				Valswap(*(uint16*)Ptr1, *(uint16*)Ptr2);
 				break;
 
+			case 3:
+				Valswap(*((uint16*&)Ptr1)++, *((uint16*&)Ptr2)++);
+				Valswap(*(uint8*)Ptr1, *(uint8*)Ptr2);
+				break;
+
 			case 4:
 				Valswap(*(uint32*)Ptr1, *(uint32*)Ptr2);
 				break;
 
+			case 5:
+				Valswap(*((uint32*&)Ptr1)++, *((uint32*&)Ptr2)++);
+				Valswap(*(uint8*)Ptr1, *(uint8*)Ptr2);
+				break;
+
+			case 6:
+				Valswap(*((uint32*&)Ptr1)++, *((uint32*&)Ptr2)++);
+				Valswap(*(uint16*)Ptr1, *(uint16*)Ptr2);
+				break;
+
+			case 7:
+				Valswap(*((uint32*&)Ptr1)++, *((uint32*&)Ptr2)++);
+				Valswap(*((uint16*&)Ptr1)++, *((uint16*&)Ptr2)++);
+				Valswap(*(uint8*)Ptr1, *(uint8*)Ptr2);
+				break;
+
 			case 8:
 				Valswap(*(uint64*)Ptr1, *(uint64*)Ptr2);
 				break;
@@ -364,7 +388,7 @@ public:
 				break;
 
 			default:
-				MemswapImpl(Ptr1, Ptr2, Size);
+				MemswapGreaterThan8(Ptr1, Ptr2, Size);
 				break;
 		}
 	}
diff --git a/Engine/Source/Runtime/Core/Public/GenericPlatform/GenericPlatformStackWalk.h b/Engine/Source/Runtime/Core/Public/GenericPlatform/GenericPlatformStackWalk.h
--- a/Engine/Source/Runtime/Core/Public/GenericPlatform/GenericPlatformStackWalk.h
+++ b/Engine/Source/Runtime/Core/Public/GenericPlatform/GenericPlatformStackWalk.h
@@ -182,6 +182,19 @@ struct CORE_API FGenericPlatformStackWalk
 	 */ 
 	static void StackWalkAndDump( ANSICHAR* HumanReadableString, SIZE_T HumanReadableStringSize, int32 IgnoreCount, void* Context = nullptr );
 
+	/**
+	* Walks the stack for the specified thread and appends the human readable string to the passed in one.
+	* @warning: The code assumes that HumanReadableString is large enough to contain the information.
+	*
+	* @param	HumanReadableString	String to concatenate information with
+	* @param	HumanReadableStringSize size of string in characters
+	* @param	IgnoreCount			Number of stack entries to ignore (some are guaranteed to be in the stack walking code)
+	* @param	ThreadId				ThreadId to walk the strack for.
+	*/
+	static void ThreadStackWalkAndDump(ANSICHAR* HumanReadableString, SIZE_T HumanReadableStringSize, int32 IgnoreCount, uint32 ThreadId)
+	{
+	}
+
 	/**
 	 * Walks the stack and appends the human readable string to the passed in one.
 	 * @warning: The code assumes that HumanReadableString is large enough to contain the information.
diff --git a/Engine/Source/Runtime/Core/Public/HAL/MallocBinnedRedirect.h b/Engine/Source/Runtime/Core/Public/HAL/MallocBinnedRedirect.h
new file mode 100644
--- /dev/null
+++ b/Engine/Source/Runtime/Core/Public/HAL/MallocBinnedRedirect.h
@@ -0,0 +1,19 @@
+// Copyright 1998-2016 Epic Games, Inc. All Rights Reserved.
+
+#pragma once
+
+#ifndef MALLOC_BINNED_VERSION_OVERRIDE
+	#define MALLOC_BINNED_VERSION_OVERRIDE 1
+#endif
+
+#if MALLOC_BINNED_VERSION_OVERRIDE != 1 && MALLOC_BINNED_VERSION_OVERRIDE != 2
+	#error MALLOC_BINNED_VERSION_OVERRIDE should be set to 1 or 2
+#endif
+
+#if MALLOC_BINNED_VERSION_OVERRIDE == 1
+	#include "MallocBinned.h"
+	typedef FMallocBinned FMallocBinnedRedirect;
+#else
+	#include "MallocBinned2.h"
+	typedef FMallocBinned2 FMallocBinnedRedirect;
+#endif
diff --git a/Engine/Source/Runtime/Core/Public/HAL/MallocDebug.h b/Engine/Source/Runtime/Core/Public/HAL/MallocDebug.h
--- a/Engine/Source/Runtime/Core/Public/HAL/MallocDebug.h
+++ b/Engine/Source/Runtime/Core/Public/HAL/MallocDebug.h
@@ -54,7 +54,7 @@ public:
 	 */
 	virtual void* Malloc( SIZE_T Size, uint32 Alignment ) override
 	{
-		check(Alignment == DEFAULT_ALIGNMENT && "Alignment currently unsupported in this allocator");
+		check(Alignment <= ALLOCATION_ALIGNMENT && "Alignment currently unsupported in this allocator");
 		FMemDebug* Ptr = (FMemDebug*)malloc( AllocatorOverhead + Size );
 		check(Ptr);
 		uint8* AlignedPtr = Align( (uint8*) Ptr + sizeof(FMemDebug) + sizeof(FMemDebug*) + sizeof(int32), ALLOCATION_ALIGNMENT );
diff --git a/Engine/Source/Runtime/Core/Public/HAL/MallocThreadSafeProxy.h b/Engine/Source/Runtime/Core/Public/HAL/MallocThreadSafeProxy.h
--- a/Engine/Source/Runtime/Core/Public/HAL/MallocThreadSafeProxy.h
+++ b/Engine/Source/Runtime/Core/Public/HAL/MallocThreadSafeProxy.h
@@ -111,4 +111,11 @@ public:
 		check(UsedMalloc);
 		return UsedMalloc->GetDescriptiveName(); 
 	}
+
+	virtual void Trim() override
+	{
+		FScopeLock ScopeLock(&SynchronizationObject);
+		check(UsedMalloc);
+		UsedMalloc->Trim();
+	}
 };
diff --git a/Engine/Source/Runtime/Core/Public/HAL/ThreadHeartBeat.h b/Engine/Source/Runtime/Core/Public/HAL/ThreadHeartBeat.h
new file mode 100644
--- /dev/null
+++ b/Engine/Source/Runtime/Core/Public/HAL/ThreadHeartBeat.h
@@ -0,0 +1,45 @@
+// Copyright 1998-2016 Epic Games, Inc. All Rights Reserved.
+#pragma once
+
+/**
+ * Thread heartbeat check class.
+ * Used by crash handling code to check for hangs.
+ */
+class CORE_API FThreadHeartBeat : public FRunnable
+{
+	/** Thread to run the worker FRunnable on */
+	FRunnableThread* Thread;
+	/** Stops this thread */
+	FThreadSafeCounter StopTaskCounter;
+	/** Synch object for the heartbeat */
+	FCriticalSection HeartBeatCritical;
+	/** Keeps track of the last heartbeat time for threads */
+	TMap<uint32, double> ThreadHeartBeat;
+
+	FThreadHeartBeat();
+	virtual ~FThreadHeartBeat();
+
+public:
+
+	enum EConstants
+	{
+		/** Invalid thread Id used by CheckHeartBeat */
+		InvalidThreadId = (uint32)-1
+	};
+
+	/** Gets the heartbeat singleton */
+	static FThreadHeartBeat& Get();
+
+	/** Called from a thread once per frame to update the heartbeat time */
+	void HeartBeat();
+	/** Called by a supervising thread to check the threads' health */
+	uint32 CheckHeartBeat();
+	/** Called by a thread when it's no longer expecting to be ticked */
+	void KillHeartBeat();
+
+	//~ Begin FRunnable Interface.
+	virtual bool Init();
+	virtual uint32 Run();
+	virtual void Stop();
+	//~ End FRunnable Interface
+};
diff --git a/Engine/Source/Runtime/Core/Public/HAL/ThreadingBase.h b/Engine/Source/Runtime/Core/Public/HAL/ThreadingBase.h
--- a/Engine/Source/Runtime/Core/Public/HAL/ThreadingBase.h
+++ b/Engine/Source/Runtime/Core/Public/HAL/ThreadingBase.h
@@ -258,6 +258,7 @@ class CORE_API FRunnableThread
 {
 	friend class FThreadSingletonInitializer;
 	friend class FTlsAutoCleanup;
+	friend class FThreadManager;
 
 	/** Index of TLS slot for FRunnableThread pointer. */
 	static uint32 RunnableTlsSlot;
@@ -360,7 +361,7 @@ public:
 	FRunnableThread();
 
 	/** Virtual destructor */
-	virtual ~FRunnableThread(){}
+	virtual ~FRunnableThread();
 
 protected:
 
@@ -412,6 +413,11 @@ protected:
 
 	/** ID set during thread creation. */
 	uint32 ThreadID;
+
+private:
+
+	/** Used by the thread manager to tick threads in single-threaded mode */
+	virtual void Tick() {}
 };
 
 
@@ -487,40 +493,45 @@ public:
 
 
 /**
- * Manages runnables and runnable threads when multithreading is disabled.
+ * Manages runnables and runnable threads.
  */
-class CORE_API FSingleThreadManager
+class CORE_API FThreadManager
 {
 	/** List of thread objects to be ticked. */
-	TArray<class FFakeThread*> ThreadList;
+	TMap<uint32, class FRunnableThread*> Threads;
+	/** Critical section for ThreadList */
+	FCriticalSection ThreadsCritical;
 
 public:
 
 	/**
-	 * Used internally to add a new thread object when multithreading is disabled.
-	 *
-	 * @param Thread Fake thread object.
-	 * @see RemoveThread
-	 */
-	void AddThread( class FFakeThread* Thread );
+	* Used internally to add a new thread object.
+	*
+	* @param Thread thread object.
+	* @see RemoveThread
+	*/
+	void AddThread(uint32 ThreadId, class FRunnableThread* Thread);
 
 	/**
-	 * Used internally to remove fake thread object.
-	 *
-	 * @param Thread Fake thread object to be removed.
-	 * @see AddThread
-	 */
-	void RemoveThread( class FFakeThread* Thread );
+	* Used internally to remove thread object.
+	*
+	* @param Thread thread object to be removed.
+	* @see AddThread
+	*/
+	void RemoveThread(class FRunnableThread* Thread);
 
 	/** Ticks all fake threads and their runnable objects. */
 	void Tick();
 
+	/** Returns the name of a thread given its TLS id */
+	const FString& GetThreadName(uint32 ThreadId);
+
 	/**
 	 * Access to the singleton object.
 	 *
 	 * @return Thread manager object.
 	 */
-	static FSingleThreadManager& Get();
+	static FThreadManager& Get();
 };
 
 
@@ -650,6 +661,7 @@ extern CORE_API FQueuedThreadPool* GLargeThreadPool;
 class FThreadSafeCounter
 {
 public:
+	typedef int32 IntegerType;
 
 	/**
 	 * Default constructor.
@@ -779,6 +791,7 @@ private:
 class FThreadSafeCounter64
 {
 public:
+	typedef int64 IntegerType;
 
 	/**
 	* Default constructor.
diff --git a/Engine/Source/Runtime/Core/Public/Math/BigInt.h b/Engine/Source/Runtime/Core/Public/Math/BigInt.h
--- a/Engine/Source/Runtime/Core/Public/Math/BigInt.h
+++ b/Engine/Source/Runtime/Core/Public/Math/BigInt.h
@@ -10,9 +10,15 @@
  * Minimum sanity checks.
  * Can convert from int64 and back (by truncating the result, this is mostly for testing)
  */
-template <int32 NumBits>
+template <int32 NumBits, bool bSigned = true>
 class TBigInt
 {
+public:
+
+	typedef TBigInt<NumBits, bSigned> BigInt;
+
+private:
+
 	enum
 	{
 		/** Word size. */
@@ -33,18 +39,21 @@ class TBigInt
 	 * @param FactorB second factor
 	 * @param SignB sign of the second pactor
 	 */
-	FORCEINLINE static void MakePositiveFactors(TBigInt<NumBits>& FactorA, int32& SignA, TBigInt<NumBits>& FactorB, int32& SignB)
-	{
-		SignA = FactorA.Sign();
-		SignB = FactorB.Sign();
-		if (SignA < 0)
-		{
-			FactorA.Negate();
-		}
-		if (SignB < 0)
-		{
-			FactorB.Negate();
-		}
+	FORCEINLINE static void MakePositiveFactors(BigInt& FactorA, int32& SignA, BigInt& FactorB, int32& SignB)
+	{
+		if (bSigned)
+	  {
+		  SignA = FactorA.Sign();
+		  SignB = FactorB.Sign();
+		  if (SignA < 0)
+		  {
+			  FactorA.Negate();
+		  }
+		  if (SignB < 0)
+		  {
+			  FactorB.Negate();
+		  }
+	  }
 	}
 
 	/**
@@ -54,9 +63,9 @@ class TBigInt
 	 * @param SignA sign of the first factor
 	 * @param SignB sign of the second factor
 	 */
-	FORCEINLINE static void RestoreSign(TBigInt<NumBits>& Result, int32 SignA, int32 SignB)
+	FORCEINLINE static void RestoreSign(BigInt& Result, int32 SignA, int32 SignB)
 	{
-		if ((SignA * SignB) < 0)
+		if (bSigned && (SignA * SignB) < 0)
 		{
 			Result.Negate();
 		}
@@ -64,6 +73,11 @@ class TBigInt
 
 public:
 
+	FORCEINLINE uint32* GetBits()
+	{
+		return Bits;
+	}
+
 	/** Sets this integer to 0. */
 	FORCEINLINE void Zero()
 	{
@@ -123,7 +137,7 @@ public:
 	{
 		checkSlow(BitCount > 0);
 
-		TBigInt<NumBits> Result;
+		BigInt Result;
 		if (BitCount & (BitsPerWord - 1))
 		{
 			const int32 LoWordOffset = (BitCount - 1) / BitsPerWord;
@@ -149,6 +163,21 @@ public:
 		*this = Result;
 	}
 
+	/**
+	 * Shift left by 1 bit.
+	 */
+	FORCEINLINE void ShiftLeftByOneInternal()
+	{
+		const int32 HiWordShift = BitsPerWord - 1;
+		Bits[NumWords - 1] = Bits[NumWords - 1] << 1;
+		for (int32 WordIndex = NumWords - 2; WordIndex >= 0; --WordIndex)
+		{
+			const uint32 Value = Bits[WordIndex];
+			Bits[WordIndex + 0] = Value << 1;
+			Bits[WordIndex + 1] |= Value >> HiWordShift;
+		}
+	}
+
 	/**
 	 * Shift right by the specified amount of bits. Does not check if BitCount is valid.
 	 *
@@ -158,7 +187,7 @@ public:
 	{
 		checkSlow(BitCount > 0);
 
-		TBigInt<NumBits> Result;
+		BigInt Result;
 		if (BitCount & (BitsPerWord - 1))
 		{
 			const int32 HiWordOffset = (BitCount - 1) / BitsPerWord;
@@ -184,10 +213,25 @@ public:
 		*this = Result;
 	}
 
+	/**
+	 * Shift right by 1 bit.
+	 */
+	FORCEINLINE void ShiftRightByOneInternal()
+	{
+		const int32 LoWordShift = BitsPerWord - 1;
+		Bits[0] = Bits[0] >> 1;
+		for (int32 WordIndex = 1; WordIndex < NumWords; ++WordIndex)
+		{
+			const uint32 Value = Bits[WordIndex];
+			Bits[WordIndex - 0] = Value >> 1;
+			Bits[WordIndex - 1] |= Value << LoWordShift;
+		}
+	}
+
 	/**
 	 * Adds two integers.
 	 */
-	FORCEINLINE void Add(const TBigInt<NumBits>& Other)
+	FORCEINLINE void Add(const BigInt& Other)
 	{
 		int64 CarryOver = 0;
 		for (int32 WordIndex = 0; WordIndex < NumWords; ++WordIndex)
@@ -202,9 +246,9 @@ public:
 	/**
 	 * Subtracts two integers.
 	 */
-	FORCEINLINE void Subtract(const TBigInt<NumBits>& Other)
+	FORCEINLINE void Subtract(const BigInt& Other)
 	{
-		TBigInt<NumBits> NegativeOther(Other);
+		BigInt NegativeOther(Other);
 		NegativeOther.Negate();
 		Add(NegativeOther);
 	}
@@ -214,7 +258,14 @@ public:
 	 */
 	FORCEINLINE bool IsNegative() const
 	{
-		return !!(Bits[NumWords - 1] & (1U << (BitsPerWord - 1)));
+		if (bSigned)
+		{
+			return !!(Bits[NumWords - 1] & (1U << (BitsPerWord - 1)));
+		}
+		else
+		{
+			return false;
+		}
 	}
 
 	/**
@@ -230,7 +281,7 @@ public:
 	 */
 	void Negate()
 	{
-		static const TBigInt<NumBits> One(1LL);
+		static const BigInt One(1LL);
 		for (int32 WordIndex = 0; WordIndex < NumWords; ++WordIndex)
 		{
 			Bits[WordIndex] = ~Bits[WordIndex];
@@ -238,57 +289,76 @@ public:
 		Add(One);
 	}
 
+	/**
+	* Returns the index of the highest word that is not zero. -1 if no such word exists.
+	*/
+	FORCEINLINE int32 GetHighestNonZeroWord() const
+	{
+		int32 WordIndex;
+		for (WordIndex = NumWords - 1; WordIndex >= 0 && Bits[WordIndex] == 0; --WordIndex);
+		return WordIndex;
+	}
 
 	/**
-	 * Multiplies two integers.
+	* Multiplies two positive integers.
 	 */
-	void Multiply(const TBigInt<NumBits>& Factor)
+	FORCEINLINE void MultiplyFast(const BigInt& Factor)
 	{
-		TBigInt<NumBits> Result;
-		TBigInt<NumBits> Temp = *this;
-		TBigInt<NumBits> Other = Factor;
+		const int64 Base = 2LL << BitsPerWord;
+		TBigInt<NumBits * 2, bSigned> Result; // Twice as many bits for overflow protection
+		uint32* ResultBits = Result.GetBits();
+		BigInt Temp = *this;
 		
-		int32 ThisSign;
-		int32 OtherSign;
-		MakePositiveFactors(Temp, ThisSign, Other, OtherSign);
+		const int32 NumWordsA = Temp.GetHighestNonZeroWord() + 1;
+		const int32 NumWordsB = Factor.GetHighestNonZeroWord() + 1;
 
-		int32 ShiftCount = 0;
-		for (int32 WordIndex = 0; WordIndex < NumWords; ++WordIndex)
-		{			
-			if (Other.Bits[WordIndex])
-			{
-				for (int32 BitIndex = 0; BitIndex < BitsPerWord; ++BitIndex)
-				{					
-					if (!!(Other.Bits[WordIndex] & (1 << BitIndex)))
-					{
-						if (ShiftCount)
-						{
-							Temp.ShiftLeftInternal(ShiftCount);
-							ShiftCount = 0;
-						}
-						Result += Temp;
-					}
-					ShiftCount++;
-				}
-			}
-			else
+		for (int32 WordIndexA = 0; WordIndexA < NumWordsA; ++WordIndexA)
+		{
+			const uint64 A = Temp.Bits[WordIndexA];
+			uint64 Carry = 0;
+			for (int32 WordIndexB = 0; WordIndexB < NumWordsB; ++WordIndexB)
 			{
-				ShiftCount += BitsPerWord;
+				const uint64 B = Factor.Bits[WordIndexB];
+				const uint64 Product = ResultBits[WordIndexA + WordIndexB] + Carry + A * B;
+				Carry = Product >> BitsPerWord;
+				ResultBits[WordIndexA + WordIndexB] = (uint32)(Product & (Base - 1));
 			}
+			ResultBits[WordIndexA + NumWordsB] += Carry;
+		}
+
+		for (int32 WordIndex = 0; WordIndex < NumWords; ++WordIndex)
+		{
+			Bits[WordIndex] = ResultBits[WordIndex];
 		}
+	}
+
+	/**
+	 * Multiplies two integers.
+	 */
+	FORCEINLINE void Multiply(const BigInt& Factor)
+	{
+		BigInt Result = *this;
+		BigInt Other = Factor;
+		
+		int32 ResultSign;
+		int32 OtherSign;
+		MakePositiveFactors(Result, ResultSign, Other, OtherSign);
+
+		Result.MultiplyFast(Other);
+
 		// Restore the sign if necessary		
-		RestoreSign(Result, OtherSign, ThisSign);
+		RestoreSign(Result, OtherSign, ResultSign);
 		*this = Result;		
 	}
 
 	/**
 	 * Divides two integers with remainder.
 	 */
-	void DivideWithRemainder(const TBigInt<NumBits>& Divisor, TBigInt<NumBits>& Remainder) 
+	void DivideWithRemainder(const BigInt& Divisor, BigInt& Remainder) 
 	{ 
-		static const TBigInt<NumBits> One(1LL);
-		TBigInt<NumBits> Denominator(Divisor);
-		TBigInt<NumBits> Dividend(*this);		
+		static const BigInt One(1LL);
+		BigInt Denominator(Divisor);
+		BigInt Dividend(*this);		
 		
 		int32 DenominatorSign;
 		int32 DividendSign;
@@ -308,27 +378,49 @@ public:
 			return;
 		}
 
-		TBigInt<NumBits> Current(One);
-		TBigInt<NumBits> Quotient;
+		BigInt Current(One);
+		BigInt Quotient;
+
+		int32 ShiftCount = Dividend.GetHighestNonZeroBit() - Denominator.GetHighestNonZeroBit();
+
+		if (ShiftCount > 0)
+		{
+			Denominator.ShiftLeftInternal(ShiftCount);
+		}
 
 		while (Denominator <= Dividend) 
 		{
-			Denominator.ShiftLeftInternal(1);	
-			Current.ShiftLeftInternal(1);
+			Denominator.ShiftLeftByOneInternal();
+			ShiftCount++;			
 		}
 
-		Denominator.ShiftRightInternal(1);
-		Current.ShiftRightInternal(1);
+		Denominator.ShiftRightByOneInternal();
 
-		while (!Current.IsZero()) 
+		ShiftCount--; // equivalent of a shift right
+		if (ShiftCount)
+		{
+			Current.ShiftLeftInternal(ShiftCount);
+		}
+
+		// Reuse ShiftCount to track number of pending shifts
+		ShiftCount = 0;
+		const int32 NumLoops = Current.GetHighestNonZeroBit() + 1;
+
+		for (int32 i = 0; i < NumLoops; ++i)
 		{
 			if (Dividend >= Denominator) 
 			{
+				if (ShiftCount != 0)
+				{
+					Current.ShiftRightInternal(ShiftCount);
+					ShiftCount = 0;
+				}
+
 				Dividend -= Denominator;
 				Quotient |= Current;
 			}
-			Current.ShiftRightInternal(1);
-			Denominator.ShiftRightInternal(1);
+			Denominator.ShiftRightByOneInternal();
+			ShiftCount++;
 		}    
 		RestoreSign(Quotient, DenominatorSign, DividendSign);
 		Remainder = Dividend;
@@ -338,23 +430,29 @@ public:
 	/**
 	 * Divides two integers.
 	 */
-	void Divide(const TBigInt<NumBits>& Divisor) 
+	void Divide(const BigInt& Divisor) 
 	{
-		TBigInt<NumBits> Remainder;
+		BigInt Remainder;
 		DivideWithRemainder(Divisor, Remainder);
 	}
 
 	/**
 	 * Performs modulo operation on this integer.
 	 */
-	void Modulo(const TBigInt<NumBits>& Modulus)
+	FORCEINLINE void Modulo(const BigInt& Modulus)
 	{
 		// a mod b = a - floor(a/b)*b
 		check(!IsNegative());
-		TBigInt<NumBits> Dividend(*this);
-		Dividend.Divide(Modulus);
-		Dividend.Multiply(Modulus);
-		Subtract(Dividend);
+		BigInt Dividend(*this);
+
+		// Dividend.Divide(Modulus);
+		BigInt Temp;
+		Dividend.DivideWithRemainder(Modulus, Temp);
+		Dividend.MultiplyFast(Modulus);
+
+		// Subtract(Dividend);
+		Dividend.Negate();
+		Add(Dividend);
 	}
 
 	/**
@@ -362,17 +460,17 @@ public:
 	 */
 	void Sqrt() 
 	{
-		TBigInt<NumBits> Number(*this);
-    TBigInt<NumBits> Result;
+		BigInt Number(*this);
+    BigInt Result;
 
-    TBigInt<NumBits> Bit(1);
+    BigInt Bit(1);
 		Bit.ShiftLeftInternal(NumBits - 2);
     while (Bit > Number)
 		{
 			Bit.ShiftRightInternal(2);
 		}
  
-		TBigInt<NumBits> Temp;
+		BigInt Temp;
     while (!Bit.IsZero()) 
 		{
 			Temp = Result;
@@ -401,16 +499,6 @@ public:
 		return *this;
 	}
 
-	/**
-	 * Returns the index of the highest word that is not zero. -1 if no such word exists.
-	 */
-	FORCEINLINE int32 GetHighestNonZeroWord() const
-	{
-		int32 WordIndex;
-		for (WordIndex = NumWords - 1; WordIndex >= 0 && Bits[WordIndex] == 0; --WordIndex);
-		return WordIndex;
-	}
-
 	/**
 	 * Returns the index of the highest non-zero bit. -1 if no such bit exists.
 	 */
@@ -508,7 +596,7 @@ public:
 	/**
 	 * Bitwise 'or'
 	 */
-	FORCEINLINE void BitwiseOr(const TBigInt<NumBits>& Other)
+	FORCEINLINE void BitwiseOr(const BigInt& Other)
 	{
 		for (int32 WordIndex = 0; WordIndex < NumWords; ++WordIndex)
 		{
@@ -519,7 +607,7 @@ public:
 	/**
 	 * Bitwise 'and'
 	 */
-	FORCEINLINE void BitwiseAnd(const TBigInt<NumBits>& Other)
+	FORCEINLINE void BitwiseAnd(const BigInt& Other)
 	{
 		for (int32 WordIndex = 0; WordIndex < NumWords; ++WordIndex)
 		{
@@ -541,7 +629,7 @@ public:
 	/**
 	 * Checks if two integers are equal.
 	 */
-	FORCEINLINE bool IsEqual(const TBigInt<NumBits>& Other) const
+	FORCEINLINE bool IsEqual(const BigInt& Other) const
 	{
 		for (int32 WordIndex = 0; WordIndex < NumWords; ++WordIndex)
 		{
@@ -557,7 +645,7 @@ public:
 	/**
 	 * this < Other
 	 */
-	bool IsLess(const TBigInt<NumBits>& Other) const
+	FORCEINLINE bool IsLess(const BigInt& Other) const
 	{
 		if (IsNegative())
 		{
@@ -578,7 +666,7 @@ public:
 	/**
 	 * this <= Other
 	 */
-	bool IsLessOrEqual(const TBigInt<NumBits>& Other) const
+	FORCEINLINE bool IsLessOrEqual(const BigInt& Other) const
 	{
 		if (IsNegative())
 		{
@@ -599,7 +687,7 @@ public:
 	/**
 	 * this > Other
 	 */
-	bool IsGreater(const TBigInt<NumBits>& Other) const
+	FORCEINLINE bool IsGreater(const BigInt& Other) const
 	{
 		if (IsNegative())
 		{
@@ -620,7 +708,7 @@ public:
 	/**
 	 * this >= Other
 	 */
-	bool IsGreaterOrEqual(const TBigInt<NumBits>& Other) const
+	FORCEINLINE bool IsGreaterOrEqual(const BigInt& Other) const
 	{
 		if (IsNegative())
 		{
@@ -664,6 +752,10 @@ public:
 		return IsNegative() && !IsZero();
 	}
 
+	FORCEINLINE bool IsFirstBitSet() const
+	{
+		return !!(Bits[0] & 1);
+	}
 	/**
 	 * Bit indexing operator
 	 */
@@ -673,170 +765,170 @@ public:
 	}
 
 	// Begin operator overloads
-	FORCEINLINE TBigInt<NumBits> operator>>(int32 Count) const
+	FORCEINLINE BigInt operator>>(int32 Count) const
 	{
-		TBigInt<NumBits> Result = *this;
+		BigInt Result = *this;
 		Result.ShiftRight(Count);
 		return Result;
 	}
 
-	FORCEINLINE TBigInt<NumBits>& operator>>=(int32 Count)
+	FORCEINLINE BigInt& operator>>=(int32 Count)
 	{
 		ShiftRight(Count);
 		return *this;
 	}
 
-	FORCEINLINE TBigInt<NumBits> operator<<(int32 Count) const
+	FORCEINLINE BigInt operator<<(int32 Count) const
 	{
-		TBigInt<NumBits> Result = *this;
+		BigInt Result = *this;
 		Result.ShiftLeft(Count);
 		return Result;
 	}
 
-	FORCEINLINE TBigInt<NumBits>& operator<<=(int32 Count)
+	FORCEINLINE BigInt& operator<<=(int32 Count)
 	{
 		ShiftLeft(Count);
 		return *this;
 	}
 
-	FORCEINLINE TBigInt<NumBits> operator+(const TBigInt<NumBits>& Other) const
+	FORCEINLINE BigInt operator+(const BigInt& Other) const
 	{
-		TBigInt<NumBits> Result(*this);
+		BigInt Result(*this);
 		Result.Add(Other);
 		return Result;
 	}
 
-	FORCEINLINE TBigInt<NumBits>& operator++()
+	FORCEINLINE BigInt& operator++()
 	{
-		static const TBigInt<NumBits> One(1);
+		static const BigInt One(1);
 		Add(One);
 		return *this;
 	}
 
-	FORCEINLINE TBigInt<NumBits>& operator+=(const TBigInt<NumBits>& Other)
+	FORCEINLINE BigInt& operator+=(const BigInt& Other)
 	{
 		Add(Other);
 		return *this;
 	}
 
-	FORCEINLINE TBigInt<NumBits> operator-(const TBigInt<NumBits>& Other) const
+	FORCEINLINE BigInt operator-(const BigInt& Other) const
 	{
-		TBigInt<NumBits> Result(*this);
+		BigInt Result(*this);
 		Result.Subtract(Other);
 		return Result;
 	}
 
-	FORCEINLINE TBigInt<NumBits>& operator--()
+	FORCEINLINE BigInt& operator--()
 	{
-		static const TBigInt<NumBits> One(1);
+		static const BigInt One(1);
 		Subtract(One);
 		return *this;
 	}
 
-	FORCEINLINE TBigInt<NumBits>& operator-=(const TBigInt<NumBits>& Other)
+	FORCEINLINE BigInt& operator-=(const BigInt& Other)
 	{
 		Subtract(Other);
 		return *this;
 	}
 
-	FORCEINLINE TBigInt<NumBits> operator*(const TBigInt<NumBits>& Other) const
+	FORCEINLINE BigInt operator*(const BigInt& Other) const
 	{
-		TBigInt<NumBits> Result(*this);
+		BigInt Result(*this);
 		Result.Multiply(Other);
 		return Result;
 	}
 
-	FORCEINLINE TBigInt<NumBits>& operator*=(const TBigInt<NumBits>& Other)
+	FORCEINLINE BigInt& operator*=(const BigInt& Other)
 	{
 		Multiply(Other);
 		return *this;
 	}
 
-	FORCEINLINE TBigInt<NumBits> operator/(const TBigInt<NumBits>& Divider) const
+	FORCEINLINE BigInt operator/(const BigInt& Divider) const
 	{
-		TBigInt<NumBits> Result(*this);
+		BigInt Result(*this);
 		Result.Divide(Divider);
 		return Result;
 	}
 
-	FORCEINLINE TBigInt<NumBits>& operator/=(const TBigInt<NumBits>& Divider)
+	FORCEINLINE BigInt& operator/=(const BigInt& Divider)
 	{
 		Divide(Divider);
 		return *this;
 	}
 
-	FORCEINLINE TBigInt<NumBits> operator%(const TBigInt<NumBits>& Modulus) const
+	FORCEINLINE BigInt operator%(const BigInt& Modulus) const
 	{
-		TBigInt<NumBits> Result(*this);
+		BigInt Result(*this);
 		Result.Modulo(Modulus);
 		return Result;
 	}
 
-	FORCEINLINE TBigInt<NumBits>& operator%=(const TBigInt<NumBits>& Modulus)
+	FORCEINLINE BigInt& operator%=(const BigInt& Modulus)
 	{
 		Modulo(Modulus);
 		return *this;
 	}
 
-	FORCEINLINE bool operator<(const TBigInt<NumBits>& Other) const
+	FORCEINLINE bool operator<(const BigInt& Other) const
 	{
 		return IsLess(Other);
 	}
 
-	FORCEINLINE bool operator<=(const TBigInt<NumBits>& Other) const
+	FORCEINLINE bool operator<=(const BigInt& Other) const
 	{
 		return IsLessOrEqual(Other);
 	}
 
-	FORCEINLINE bool operator>(const TBigInt<NumBits>& Other) const
+	FORCEINLINE bool operator>(const BigInt& Other) const
 	{
 		return IsGreater(Other);
 	}
 
-	FORCEINLINE bool operator>=(const TBigInt<NumBits>& Other) const
+	FORCEINLINE bool operator>=(const BigInt& Other) const
 	{
 		return IsGreaterOrEqual(Other);
 	}
 
-	FORCEINLINE bool operator==(const TBigInt<NumBits>& Other) const
+	FORCEINLINE bool operator==(const BigInt& Other) const
 	{
 		return IsEqual(Other);
 	}
 
-	FORCEINLINE bool operator!=(const TBigInt<NumBits>& Other) const
+	FORCEINLINE bool operator!=(const BigInt& Other) const
 	{
 		return !IsEqual(Other);
 	}
 
-	FORCEINLINE TBigInt<NumBits> operator&(const TBigInt<NumBits>& Other) const
+	FORCEINLINE BigInt operator&(const BigInt& Other) const
 	{
-		TBigInt<NumBits> Result(*this);
+		BigInt Result(*this);
 		Result.BitwiseAnd(Other);
 		return Result;
 	}
 
-	FORCEINLINE TBigInt<NumBits>& operator&=(const TBigInt<NumBits>& Other)
+	FORCEINLINE BigInt& operator&=(const BigInt& Other)
 	{
 		BitwiseAnd(Other);
 		return *this;
 	}
 
-	FORCEINLINE TBigInt<NumBits> operator|(const TBigInt<NumBits>& Other) const
+	FORCEINLINE BigInt operator|(const BigInt& Other) const
 	{
-		TBigInt<NumBits> Result(*this);
+		BigInt Result(*this);
 		Result.BitwiseOr(Other);
 		return Result;
 	}
 
-	FORCEINLINE TBigInt<NumBits>& operator|=(const TBigInt<NumBits>& Other)
+	FORCEINLINE BigInt& operator|=(const BigInt& Other)
 	{
 		BitwiseOr(Other);
 		return *this;
 	}
 
-	FORCEINLINE TBigInt<NumBits> operator~() const
+	FORCEINLINE BigInt operator~() const
 	{
-		TBigInt<NumBits> Result(*this);
+		BigInt Result(*this);
 		Result.BitwiseNot();
 		return Result;
 	}
@@ -854,7 +946,7 @@ public:
 		}
 		else
 		{
-			TBigInt<NumBits> Positive(*this);
+			BigInt Positive(*this);
 			Positive.Negate();
 			Result = (int64)Positive.Bits[0] + ((int64)Positive.Bits[1] << BitsPerWord);
 		}
@@ -909,7 +1001,7 @@ public:
 	/**
 	 * Serialization.
 	 */
-	friend FArchive& operator<<(FArchive& Ar, TBigInt<NumBits>& Value)
+	friend FArchive& operator<<(FArchive& Ar, BigInt& Value)
 	{
 		for (int32 Index = 0; Index < NumWords; ++Index)
 		{
@@ -922,6 +1014,7 @@ public:
 // Predefined big int types
 typedef TBigInt<256> int256;
 typedef TBigInt<512> int512;
+typedef TBigInt<512> TEncryptionInt;
 
 /**
  * Encyption key - exponent and modulus pair
@@ -932,12 +1025,12 @@ struct TEncryptionKey
 	IntType Exponent;
 	IntType Modulus;
 };
-typedef TEncryptionKey<int256> FEncryptionKey;
+typedef TEncryptionKey<TEncryptionInt> FEncryptionKey;
 
 /**
  * Math utils for encryption.
  */
-struct FEncryption
+namespace FEncryption
 {
 	/**
 	 * Greatest common divisor of ValueA and ValueB.
@@ -1067,7 +1160,7 @@ struct FEncryption
 		IntType Result(One);
 		while (Exponent > 0)
 		{
-			if ((Exponent & One) != 0)
+			if (Exponent.IsFirstBitSet())
 			{
 				Result = (Result * Base) % Modulus;
 			}
@@ -1076,91 +1169,118 @@ struct FEncryption
 		}
 		return Result;
 	}
-};
+
+	/**
+	* Encrypts a stream of bytes
+	*/
+	template <typename IntType>
+	FORCEINLINE static void EncryptBytes(IntType* EncryptedData, const uint8* Data, const int64 DataLength, const FEncryptionKey& EncryptionKey)
+	{
+		for (int Index = 0; Index < DataLength; ++Index)
+		{
+			EncryptedData[Index] = FEncryption::ModularPow(IntType(Data[Index]), EncryptionKey.Exponent, EncryptionKey.Modulus);
+		}
+	}
+
+	/**
+	* Decrypts a stream of bytes
+	*/
+	template <typename IntType>
+	FORCEINLINE static void DecryptBytes(uint8* DecryptedData, const IntType* Data, const int64 DataLength, const FEncryptionKey& DecryptionKey)
+	{
+		for (int64 Index = 0; Index < DataLength; ++Index)
+		{
+			IntType DecryptedByte = FEncryption::ModularPow(Data[Index], DecryptionKey.Exponent, DecryptionKey.Modulus);
+			DecryptedData[Index] = (uint8)DecryptedByte.ToInt();
+		}
+	}
+}
 
 /**
- * Specialization for int256 (performance). Avoids using temporary results and most of the operations are inplace.
+ * Specialization for int type used in encryption (performance). Avoids using temporary results and most of the operations are inplace.
  */
 template <>
-FORCEINLINE int256 FEncryption::ModularPow(int256 Base, int256 Exponent, int256 Modulus)
+FORCEINLINE TEncryptionInt FEncryption::ModularPow(TEncryptionInt Base, TEncryptionInt Exponent, TEncryptionInt Modulus)
 {
-	static const int256 One(1);
-	int256 Result(One);
+	TEncryptionInt Result(1LL);
 	while (Exponent.IsGreaterThanZero())
 	{
-		if (!(Exponent & One).IsZero())
+		if (Exponent.IsFirstBitSet())
 		{
-			Result.Multiply(Base);
+			Result.MultiplyFast(Base);
 			Result.Modulo(Modulus);
 		}
-		Exponent.ShiftRightInternal(1);
-		Base.Multiply(Base);
+		Exponent.ShiftRightByOneInternal();
+		Base.MultiplyFast(Base);
 		Base.Modulo(Modulus);
 	}
 	return Result;
 }
 
-template <class TYPE, int64 NUM_VALUES>
+template <class TYPE>
 struct FSignatureBase
 {
-	TYPE Data[NUM_VALUES];
+	TYPE Data;
 
 	FSignatureBase()
 	{
-		for (int32 Index = 0; Index < ARRAY_COUNT(Data); ++Index)
-		{
-			Data[Index] = 0;
-		}
+		Data = 0;
 	}
 
 	void Serialize(FArchive& Ar)
 	{
-		for (int32 Index = 0; Index < ARRAY_COUNT(Data); ++Index)
-		{
-			Ar << Data[Index];
-		}
+		Ar << Data;
 	}
 
 	static int64 Size()
 	{
-		return sizeof(TYPE) * NUM_VALUES;
+		return sizeof(TYPE);
 	}
 
 	bool IsValid() const
 	{
-		for (int64 Index = 0; Index < NUM_VALUES; ++Index)
-		{
-			if (Data[Index] != 0)
-			{
-				return true;
-			}
-		}
-
-		return false;
+		return Data != 0;
 	}
 
 	bool operator != (const FSignatureBase& InOther)
 	{
-		for (int64 i = 0; i < NUM_VALUES; ++i)
-		{
-			if (Data[i] != InOther.Data[i])
-			{
-				return true;
-			}
-		}
+		return Data != InOther.Data;
+	}
+
+	bool operator == (const FSignatureBase& InOther)
+	{
+		return Data == InOther.Data;
+	}
 
-		return false;
+	/**
+	* Serialization.
+	*/
+	friend FArchive& operator<<(FArchive& Ar, FSignatureBase& Value)
+	{
+		Value.Serialize(Ar);
+		return Ar;
 	}
 };
 
-template <int64 NUM_VALUES>
-struct FEncryptedSignature : public FSignatureBase<int256, NUM_VALUES>
+struct FEncryptedSignature : public FSignatureBase<TEncryptionInt>
 {
 
 };
 
-template <int64 NUM_VALUES>
-struct FDecryptedSignature : public FSignatureBase<uint8, NUM_VALUES>
+struct FDecryptedSignature : public FSignatureBase<uint32>
 {
 
-};
\ No newline at end of file
+};
+
+namespace FEncryption
+{
+	FORCEINLINE static void EncryptSignature(const FDecryptedSignature& InUnencryptedSignature, FEncryptedSignature& OutEncryptedSignature, const FEncryptionKey& EncryptionKey)
+	{
+		OutEncryptedSignature.Data = FEncryption::ModularPow(TEncryptionInt(InUnencryptedSignature.Data), EncryptionKey.Exponent, EncryptionKey.Modulus);
+	}
+
+	FORCEINLINE static void DecryptSignature(const FEncryptedSignature& InEncryptedSignature, FDecryptedSignature& OutUnencryptedSignature, const FEncryptionKey& EncryptionKey)
+	{
+		OutUnencryptedSignature.Data = FEncryption::ModularPow(TEncryptionInt(InEncryptedSignature.Data), EncryptionKey.Exponent, EncryptionKey.Modulus).ToInt();
+	}
+}
\ No newline at end of file
diff --git a/Engine/Source/Runtime/Core/Public/Misc/CString.h b/Engine/Source/Runtime/Core/Public/Misc/CString.h
--- a/Engine/Source/Runtime/Core/Public/Misc/CString.h
+++ b/Engine/Source/Runtime/Core/Public/Misc/CString.h
@@ -258,6 +258,16 @@ struct TCString
 	static FORCEINLINE const CharType* Strrstr( const CharType* String, const CharType* Find );
 	static FORCEINLINE CharType* Strrstr( CharType* String, const CharType* Find );
 
+	/**
+	 * strspn wrapper
+	 */
+	static FORCEINLINE int32 Strspn( const CharType* String, const CharType* Mask );
+
+	/**
+	 * strcspn wrapper
+	 */
+	static FORCEINLINE int32 Strcspn( const CharType* String, const CharType* Mask );
+
 	/**
 	 * atoi wrapper
 	 */
@@ -599,6 +609,46 @@ typename TCString<T>::CharType* TCString<T>::Strrstr( CharType* String, const Ch
 	}
 }
 
+template <typename T> FORCEINLINE
+int32 TCString<T>::Strspn( const CharType* String, const CharType* Mask )
+{
+	const TCHAR* StringIt = String;
+	while (*StringIt)
+	{
+		for (const TCHAR* MaskIt = Mask; *MaskIt; ++MaskIt)
+		{
+			if (*StringIt != *MaskIt)
+			{
+				return StringIt - String;
+			}
+		}
+
+		++StringIt;
+	}
+
+	return StringIt - String;
+}
+
+template <typename T> FORCEINLINE
+int32 TCString<T>::Strcspn( const CharType* String, const CharType* Mask )
+{
+	const TCHAR* StringIt = String;
+	while (*StringIt)
+	{
+		for (const TCHAR* MaskIt = Mask; *MaskIt; ++MaskIt)
+		{
+			if (*StringIt == *MaskIt)
+			{
+				return StringIt - String;
+			}
+		}
+
+		++StringIt;
+	}
+
+	return StringIt - String;
+}
+
 template <typename T> FORCEINLINE 
 int32 TCString<T>::Atoi( const CharType* String ) 
 {
diff --git a/Engine/Source/Runtime/Core/Public/Templates/UnrealTemplate.h b/Engine/Source/Runtime/Core/Public/Templates/UnrealTemplate.h
--- a/Engine/Source/Runtime/Core/Public/Templates/UnrealTemplate.h
+++ b/Engine/Source/Runtime/Core/Public/Templates/UnrealTemplate.h
@@ -73,6 +73,20 @@ FORCEINLINE void Move(T& A,typename TMoveSupportTraits<T>::Move B)
 	// Use placement new and a copy constructor so types with const members will work.
 	new(&A) T(MoveTemp(B));
 }
+
+/**
+ * Tests if an index is valid for a given array.
+ *
+ * @param Array  The array to check against.
+ * @param Index  The index to check.
+ * @return true if the index is valid, false otherwise.
+ */
+template <typename T, SIZE_T N>
+static FORCEINLINE bool IsValidArrayIndex(const T(&Array)[N], SIZE_T Index)
+{
+	return Index < N;
+}
+
 /*----------------------------------------------------------------------------
 	Standard macros.
 ----------------------------------------------------------------------------*/
diff --git a/Engine/Source/Runtime/Core/Public/Windows/WindowsPlatformStackWalk.h b/Engine/Source/Runtime/Core/Public/Windows/WindowsPlatformStackWalk.h
--- a/Engine/Source/Runtime/Core/Public/Windows/WindowsPlatformStackWalk.h
+++ b/Engine/Source/Runtime/Core/Public/Windows/WindowsPlatformStackWalk.h
@@ -15,6 +15,7 @@ struct CORE_API FWindowsPlatformStackWalk
 	static void ProgramCounterToSymbolInfo( uint64 ProgramCounter, FProgramCounterSymbolInfo& out_SymbolInfo );
 	static void CaptureStackBackTrace( uint64* BackTrace, uint32 MaxDepth, void* Context = nullptr );
 	static void StackWalkAndDump( ANSICHAR* HumanReadableString, SIZE_T HumanReadableStringSize, int32 IgnoreCount, void* Context = nullptr );
+	static void ThreadStackWalkAndDump(ANSICHAR* HumanReadableString, SIZE_T HumanReadableStringSize, int32 IgnoreCount, uint32 ThreadId);
 
 	static int32 GetProcessModuleCount();
 	static int32 GetProcessModuleSignatures(FStackWalkModuleInfo *ModuleSignatures, const int32 ModuleSignaturesSize);
diff --git a/Engine/Source/Runtime/CoreUObject/Private/Serialization/AsyncLoading.cpp b/Engine/Source/Runtime/CoreUObject/Private/Serialization/AsyncLoading.cpp
--- a/Engine/Source/Runtime/CoreUObject/Private/Serialization/AsyncLoading.cpp
+++ b/Engine/Source/Runtime/CoreUObject/Private/Serialization/AsyncLoading.cpp
@@ -14,6 +14,7 @@
 #include "ExclusiveLoadPackageTimeTracker.h"
 #include "AssetRegistryInterface.h"
 #include "BlueprintSupport.h"
+#include "HAL/ThreadHeartBeat.h"
 #include "HAL/ExceptionHandling.h"
 
 #define FIND_MEMORY_STOMPS (1 && PLATFORM_WINDOWS && !WITH_EDITORONLY_DATA)
@@ -637,6 +638,9 @@ EAsyncPackageState::Type FAsyncLoadingThread::ProcessAsyncLoading(int32& OutPack
 {
 	SCOPE_CYCLE_COUNTER(STAT_FAsyncLoadingThread_ProcessAsyncLoading);
 	
+	// If we're not multithreaded and flushing async loading, update the thread heartbeat
+	const bool bNeedsHeartbeatTick = !bUseTimeLimit && !FAsyncLoadingThread::IsMultithreaded();
+
 	EAsyncPackageState::Type LoadingState = EAsyncPackageState::Complete;
 	OutPackagesProcessed = 0;
 
@@ -690,6 +694,12 @@ EAsyncPackageState::Type FAsyncLoadingThread::ProcessAsyncLoading(int32& OutPack
 		{
 			AsyncPackagesCounter.Decrement();
 		}
+
+		if (bNeedsHeartbeatTick)
+		{
+			// Update heartbeat after each package has been processed
+			FThreadHeartBeat::Get().HeartBeat();
+		}
 	}
 
 	return LoadingState;
@@ -2260,6 +2270,8 @@ void FlushAsyncLoading(int32 PackageID /* = INDEX_NONE */)
 
 				if (AsyncThread.IsMultithreaded())
 				{
+					// Update the heartbeat and sleep. If we're not multithreading, the heartbeat is updated after each package has been processed
+					FThreadHeartBeat::Get().HeartBeat();
 					FPlatformProcess::SleepNoStats(0.0001f);
 				}
 			}
@@ -2309,11 +2321,13 @@ bool IsAsyncLoadingMultithreadedCoreUObjectInternal()
 
 void SuspendAsyncLoadingInternal()
 {
+	check(IsInGameThread() && !IsInSlateThread());
 	FAsyncLoadingThread::Get().SuspendLoading();
 }
 
 void ResumeAsyncLoadingInternal()
 {
+	check(IsInGameThread() && !IsInSlateThread());
 	FAsyncLoadingThread::Get().ResumeLoading();
 }
 
diff --git a/Engine/Source/Runtime/CoreUObject/Private/UObject/LinkerLoad.cpp b/Engine/Source/Runtime/CoreUObject/Private/UObject/LinkerLoad.cpp
--- a/Engine/Source/Runtime/CoreUObject/Private/UObject/LinkerLoad.cpp
+++ b/Engine/Source/Runtime/CoreUObject/Private/UObject/LinkerLoad.cpp
@@ -14,6 +14,7 @@
 #include "GatherableTextData.h"
 #include "Serialization/AsyncLoading.h"
 #include "ModuleManager.h"
+#include "HAL/ThreadHeartBeat.h"
 
 #define LOCTEXT_NAMESPACE "LinkerLoad"
 
@@ -2867,6 +2868,9 @@ void FLinkerLoad::LoadAllObjects( bool bForcePreload )
 		MetaDataIndex = LoadMetaDataFromExportMap(bForcePreload);
 	}
 	
+	// Tick the heartbeat if we're loading on the game thread
+	const bool bShouldTickHeartBeat = IsInGameThread();
+
 	for(int32 ExportIndex = 0; ExportIndex < ExportMap.Num(); ++ExportIndex)
 	{
 #if WITH_EDITOR
@@ -2896,6 +2900,12 @@ void FLinkerLoad::LoadAllObjects( bool bForcePreload )
 #endif // USE_CIRCULAR_DEPENDENCY_LOAD_DEFERRING
 
 		CreateExportAndPreload(ExportIndex, bForcePreload);
+
+		// If needed send a heartbeat, but no need to do it too often
+		if (bShouldTickHeartBeat && (ExportIndex % 10) == 0)
+		{
+			FThreadHeartBeat::Get().HeartBeat();
+		}
 	}
 
 	// Mark package as having been fully loaded.
diff --git a/Engine/Source/Runtime/CoreUObject/Private/UObject/UObjectArray.cpp b/Engine/Source/Runtime/CoreUObject/Private/UObject/UObjectArray.cpp
--- a/Engine/Source/Runtime/CoreUObject/Private/UObject/UObjectArray.cpp
+++ b/Engine/Source/Runtime/CoreUObject/Private/UObject/UObjectArray.cpp
@@ -116,7 +116,6 @@ void FUObjectArray::DisableDisregardForGC()
 {
 	MaxObjectsNotConsideredByGC = 0;
 	ObjFirstGCIndex = 0;
-	ObjLastNonGCIndex = INDEX_NONE;
 	if (IsOpenForDisregardForGC())
 	{
 		CloseDisregardForGC();
diff --git a/Engine/Source/Runtime/Engine/Private/UnrealEngine.cpp b/Engine/Source/Runtime/Engine/Private/UnrealEngine.cpp
--- a/Engine/Source/Runtime/Engine/Private/UnrealEngine.cpp
+++ b/Engine/Source/Runtime/Engine/Private/UnrealEngine.cpp
@@ -4536,7 +4536,7 @@ bool UEngine::HandleDebugCommand( const TCHAR* Cmd, FOutputDevice& Ar )
 {
 	if (FParse::Command(&Cmd, TEXT("RESETLOADERS")))
 	{
-		ResetLoaders(NULL);
+		ResetLoaders( NULL );
 		return true;
 	}
 
@@ -6202,6 +6202,12 @@ bool UEngine::PerformError(const TCHAR* Cmd, FOutputDevice& Ar)
 		}
 		return true;
 	}
+	else if (FParse::Command(&Cmd, TEXT("SLEEP")))
+	{
+		Ar.Log(TEXT("Sleep for 1 hour. This should crash after a few seconds in cooked builds."));
+		FPlatformProcess::Sleep(3600);
+		return true;
+	}
 
 	return false;
 }
diff --git a/Engine/Source/Runtime/Launch/Private/Launch.cpp b/Engine/Source/Runtime/Launch/Private/Launch.cpp
--- a/Engine/Source/Runtime/Launch/Private/Launch.cpp
+++ b/Engine/Source/Runtime/Launch/Private/Launch.cpp
@@ -4,6 +4,7 @@
 #include "PhysicsPublic.h"
 #include "ExceptionHandling.h"
 #include "ModuleManager.h"
+#include "LoadTimeTracker.h"
 
 
 IMPLEMENT_MODULE(FDefaultModuleImpl, Launch);
@@ -137,7 +138,9 @@ int32 GuardedMain( const TCHAR* CmdLine )
 		}
 	}
 
-	UE_LOG(LogLoad, Log, TEXT("Full Startup: %.2f seconds (BP compile: %.2f seconds)"), FPlatformTime::Seconds() - GStartTime, GBlueprintCompileTime);
+	double EngineInitializationTime = FPlatformTime::Seconds() - GStartTime;
+	UE_LOG(LogLoad, Log, TEXT("Full Startup: %.2f seconds (BP compile: %.2f seconds)"), EngineInitializationTime, GBlueprintCompileTime);
+	ACCUM_LOADTIME(TEXT("EngineInitialization"), EngineInitializationTime);
 
 	while( !GIsRequestingExit )
 	{
diff --git a/Engine/Source/Runtime/Launch/Private/LaunchEngineLoop.cpp b/Engine/Source/Runtime/Launch/Private/LaunchEngineLoop.cpp
--- a/Engine/Source/Runtime/Launch/Private/LaunchEngineLoop.cpp
+++ b/Engine/Source/Runtime/Launch/Private/LaunchEngineLoop.cpp
@@ -17,6 +17,7 @@
 #include "../Resources/Version.h"
 #include "VersionManifest.h"
 #include "UObject/DevObjectVersion.h"
+#include "HAL/ThreadHeartBeat.h"
 
 #if WITH_EDITOR
 	#include "EditorStyle.h"
@@ -2567,6 +2568,10 @@ void FEngineLoop::Tick()
 #if !UE_BUILD_SHIPPING && !UE_BUILD_TEST
 	FScopedSampleMallocChurn ChurnTracker;
 #endif
+
+	// Send a heartbeat for the diagnostics thread
+	FThreadHeartBeat::Get().HeartBeat();
+
 	// Ensure we aren't starting a frame while loading or playing a loading movie
 	ensure(GetMoviePlayer()->IsLoadingFinished() && !GetMoviePlayer()->IsMovieCurrentlyPlaying());
 
@@ -2785,7 +2790,7 @@ void FEngineLoop::Tick()
 
 			FTicker::GetCoreTicker().Tick(FApp::GetDeltaTime());
 
-			FSingleThreadManager::Get().Tick();
+			FThreadManager::Get().Tick();
 
 			GEngine->TickDeferredCommands();		
 		}
@@ -2806,7 +2811,7 @@ void FEngineLoop::Tick()
 #if UE_GC_TRACK_OBJ_AVAILABLE
 		SET_DWORD_STAT(STAT_Hash_NumObjects, GUObjectArray.GetObjectArrayNumMinusAvailable());
 #endif
-	}
+	}	
 }
 
 
@@ -3158,7 +3163,7 @@ bool FEngineLoop::AppInit( )
 	UE_LOG(LogInit, Log, TEXT("Installed Engine Build: %d"), FApp::IsEngineInstalled() ? 1 : 0);
 
 	// if a logging build, clear out old log files
-#if !NO_LOGGING && !(UE_BUILD_SHIPPING || UE_BUILD_TEST)
+#if !NO_LOGGING
 	FMaintenance::DeleteOldLogs();
 #endif
 
diff --git a/Engine/Source/Runtime/PakFile/Private/IPlatformFilePak.cpp b/Engine/Source/Runtime/PakFile/Private/IPlatformFilePak.cpp
--- a/Engine/Source/Runtime/PakFile/Private/IPlatformFilePak.cpp
+++ b/Engine/Source/Runtime/PakFile/Private/IPlatformFilePak.cpp
@@ -240,7 +240,7 @@ bool FPakEntry::VerifyPakEntriesMatch(const FPakEntry& FileEntryA, const FPakEnt
 
 bool FPakPlatformFile::IsFilenameAllowed(const FString& InFilename)
 {
-	if (!bSecurityEnabled)
+	if (bForceSecurityBypass || !bSecurityEnabled)
 	{
 		return true;
 	}
@@ -254,8 +254,6 @@ bool FPakPlatformFile::IsFilenameAllowed(const FString& InFilename)
 	}
 #endif
 
-	static TSet<FName> AllowedExtensions;
-
 	if (AllowedExtensions.Num() == 0)
 	{
 		static const FName PakExtension(TEXT("pak"));
@@ -327,26 +325,29 @@ FPakFile::~FPakFile()
 FArchive* FPakFile::CreatePakReader(const TCHAR* Filename)
 {
 	FArchive* ReaderArchive = IFileManager::Get().CreateFileReader(Filename);
-	return SetupSignedPakReader(ReaderArchive);
+	return SetupSignedPakReader(ReaderArchive, Filename);
 }
 
 FArchive* FPakFile::CreatePakReader(IFileHandle& InHandle, const TCHAR* Filename)
 {
 	FArchive* ReaderArchive = new FArchiveFileReaderGeneric(&InHandle, Filename, InHandle.Size());
-	return SetupSignedPakReader(ReaderArchive);
+	return SetupSignedPakReader(ReaderArchive, Filename);
 }
 
-FArchive* FPakFile::SetupSignedPakReader(FArchive* ReaderArchive)
+FArchive* FPakFile::SetupSignedPakReader(FArchive* ReaderArchive, const TCHAR* Filename)
 {
+	if (FPlatformProperties::RequiresCookedData())
+	{
 #if !USING_SIGNED_CONTENT
-	if (bSigned || FParse::Param(FCommandLine::Get(), TEXT("signedpak")) || FParse::Param(FCommandLine::Get(), TEXT("signed")))
+		if (bSigned || FParse::Param(FCommandLine::Get(), TEXT("signedpak")) || FParse::Param(FCommandLine::Get(), TEXT("signed")))
 #endif
-	{	
-		if (!Decryptor.IsValid())
 		{
-			Decryptor = new FChunkCacheWorker(ReaderArchive);
+			if (!Decryptor.IsValid())
+			{
+				Decryptor = new FChunkCacheWorker(ReaderArchive, Filename);
+			}
+			ReaderArchive = new FSignedArchiveReader(ReaderArchive, Decryptor);
 		}
-		ReaderArchive = new FSignedArchiveReader(ReaderArchive, Decryptor);
 	}
 	return ReaderArchive;
 }
@@ -559,6 +560,7 @@ void FPakPlatformFile::HandlePakListCommand(const TCHAR* Cmd, FOutputDevice& Ar)
 FPakPlatformFile::FPakPlatformFile()
 	: LowerLevel(NULL)
 	, bSigned(false)
+	, bForceSecurityBypass(false)
 {
 
 #if PLATFORM_DESKTOP && UE_BUILD_SHIPPING
@@ -707,7 +709,7 @@ bool FPakPlatformFile::Initialize(IPlatformFile* Inner, const TCHAR* CmdLine)
 		// Even if -signed is not provided in the command line, use signed reader if the hardcoded key is non-zero.
 		FEncryptionKey DecryptionKey;
 		DecryptionKey.Exponent.Parse(DECRYPTION_KEY_EXPONENT);
-		DecryptionKey.Modulus.Parse(DECYRPTION_KEY_MODULUS);
+		DecryptionKey.Modulus.Parse(DECRYPTION_KEY_MODULUS);
 		bSigned = !DecryptionKey.Exponent.IsZero() && !DecryptionKey.Modulus.IsZero();
 	}
 #else
@@ -971,13 +973,9 @@ bool FPakPlatformFile::CopyFile(const TCHAR* To, const TCHAR* From)
 	return Result;
 }
 
-void ComputePakChunkHash(const uint8* InData, const int64 InDataSize, uint8* OutHash)
+uint32 ComputePakChunkHash(const uint8* InData, const int64 InDataSize)
 {
-#if PAKFILE_USE_CRC_FOR_CHUNK_HASHES
-	*(reinterpret_cast<uint32*>(OutHash)) = FCrc::MemCrc32(InData, InDataSize);
-#else
-	FSHA1::HashBuffer(InData, InDataSize, OutHash);
-#endif
+	return FCrc::MemCrc32(InData, InDataSize);
 }
 
 /**
diff --git a/Engine/Source/Runtime/PakFile/Private/PublicKey.inl b/Engine/Source/Runtime/PakFile/Private/PublicKey.inl
--- a/Engine/Source/Runtime/PakFile/Private/PublicKey.inl
+++ b/Engine/Source/Runtime/PakFile/Private/PublicKey.inl
@@ -2,10 +2,10 @@
 
 #pragma once
 
-/**
- * Do not modify this file manually. It will be modified by the build machine with the correct key value.
- * The below values should always be "0x0" in P4, otherwise pak code will always assume paks are signed.
- */
-
+#ifndef DECRYPTION_KEY_EXPONENT
 #define DECRYPTION_KEY_EXPONENT "0x0"
-#define DECYRPTION_KEY_MODULUS "0x0"
+#endif
+
+#ifndef DECRYPTION_KEY_MODULUS
+#define DECRYPTION_KEY_MODULUS "0x0"
+#endif
diff --git a/Engine/Source/Runtime/PakFile/Private/SignedArchiveReader.cpp b/Engine/Source/Runtime/PakFile/Private/SignedArchiveReader.cpp
--- a/Engine/Source/Runtime/PakFile/Private/SignedArchiveReader.cpp
+++ b/Engine/Source/Runtime/PakFile/Private/SignedArchiveReader.cpp
@@ -24,10 +24,23 @@ DECLARE_DWORD_ACCUMULATOR_STAT(TEXT("FChunkCacheWorker.SignaturesDecryptedDuring
 
 DECLARE_FLOAT_ACCUMULATOR_STAT(TEXT("FSignedArchiveReader.Serialize"), STAT_SignedArchiveReader_Serialize, STATGROUP_PakFile);
 
-FChunkCacheWorker::FChunkCacheWorker(FArchive* InReader)
+FChunkCacheWorker::FChunkCacheWorker(FArchive* InReader, const TCHAR* Filename)
 	: Reader(InReader)
 	, QueuedRequestsEvent(NULL)
 {
+	FString SigFileFilename = FPaths::ChangeExtension(Filename, TEXT("sig"));
+	FArchive* SigFileReader = IFileManager::Get().CreateFileReader(*SigFileFilename);
+
+	if (SigFileReader == nullptr)
+	{
+		UE_LOG(LogPakFile, Fatal, TEXT("Couldn't find pak signature file '%s'"), *Filename);
+	}
+
+	*SigFileReader << EncryptedSignatures;
+	delete SigFileReader;
+
+	DecryptedSignatures.AddDefaulted(EncryptedSignatures.Num());
+
 	SetupDecryptionKey();
 
 	QueuedRequestsEvent = FPlatformProcess::GetSynchEventFromPool();
@@ -50,11 +63,11 @@ bool FChunkCacheWorker::Init()
 void FChunkCacheWorker::SetupDecryptionKey()
 {
 	DecryptionKey.Exponent.Parse(DECRYPTION_KEY_EXPONENT);
-	DecryptionKey.Modulus.Parse(DECYRPTION_KEY_MODULUS);
+	DecryptionKey.Modulus.Parse(DECRYPTION_KEY_MODULUS);
 	// Public key should never be zero at this point. Check PublicKey.inl for more details.
 	UE_CLOG(DecryptionKey.Exponent.IsZero() || DecryptionKey.Modulus.IsZero(), LogPakFile, Fatal, TEXT("Invalid decryption key detected"));
 	// Public key should produce decrypted results - check for identity keys
-	static int256 TestValues[] = 
+	static TEncryptionInt TestValues[] = 
 	{
 		11,
 		23,
@@ -66,7 +79,7 @@ void FChunkCacheWorker::SetupDecryptionKey()
 	bool bIdentical = true;
 	for (int32 Index = 0; bIdentical && Index < ARRAY_COUNT(TestValues); ++Index)
 	{
-		int256 DecryptedValue = FEncryption::ModularPow(TestValues[Index], DecryptionKey.Exponent, DecryptionKey.Modulus);
+		TEncryptionInt DecryptedValue = FEncryption::ModularPow(TestValues[Index], DecryptionKey.Exponent, DecryptionKey.Modulus);
 		bIdentical = (DecryptedValue == TestValues[Index]);
 	}	
 	UE_CLOG(bIdentical, LogPakFile, Fatal, TEXT("Decryption key produces identical results to source data."));
@@ -74,15 +87,15 @@ void FChunkCacheWorker::SetupDecryptionKey()
 
 void FChunkCacheWorker::DecryptSignatures(int32 NextIndexToDecrypt)
 {
-	FDecryptedSignature<GPakFileChunkHashSize>& DecryptedSignature = DecryptedSignatures[NextIndexToDecrypt];
+	FDecryptedSignature& DecryptedSignature = DecryptedSignatures[NextIndexToDecrypt];
 
 	// Check to see if this signature was already decrypted.
 	if (!DecryptedSignature.IsValid())
 	{
 		INC_DWORD_STAT(STAT_FChunkCacheWorker_SignaturesDecryptedDuringIdle);
 		SCOPE_SECONDS_ACCUMULATOR(STAT_FChunkCacheWorker_DecryptDuringIdle);
-		FEncryptedSignature<GPakFileChunkHashSize>& EncryptedSignature = EncryptedSignatures[NextIndexToDecrypt];
-		Decrypt(DecryptedSignature.Data, EncryptedSignature.Data, GPakFileChunkHashSize);
+		FEncryptedSignature& EncryptedSignature = EncryptedSignatures[NextIndexToDecrypt];
+		FEncryption::DecryptSignature(EncryptedSignature, DecryptedSignature, DecryptionKey);
 	}
 }
 
@@ -239,20 +252,11 @@ int32 FChunkCacheWorker::ProcessQueue()
 	return ProcessedRequests;
 }
 
-void FChunkCacheWorker::Decrypt(uint8* DecryptedData, const int256* Data, const int64 DataLength)
-{
-	for (int64 Index = 0; Index < DataLength; ++Index)
-	{
-		int256 DecryptedByte = FEncryption::ModularPow(Data[Index], DecryptionKey.Exponent, DecryptionKey.Modulus);
-		DecryptedData[Index] = (uint8)DecryptedByte.ToInt();
-	}
-}
-
 bool FChunkCacheWorker::CheckSignature(const FChunkRequest& ChunkInfo)
 {
 	SCOPE_SECONDS_ACCUMULATOR(STAT_FChunkCacheWorker_CheckSignature);
 
-	FDecryptedSignature<GPakFileChunkHashSize> Hash;
+	FDecryptedSignature SourceSignature;
 
 	{
 		SCOPE_SECONDS_ACCUMULATOR(STAT_FChunkCacheWorker_Serialize);
@@ -261,21 +265,21 @@ bool FChunkCacheWorker::CheckSignature(const FChunkRequest& ChunkInfo)
 	}
 	{
 		SCOPE_SECONDS_ACCUMULATOR(STAT_FChunkCacheWorker_HashBuffer);
-		ComputePakChunkHash(ChunkInfo.Buffer->Data, ChunkInfo.Size, Hash.Data);
+		SourceSignature.Data = FCrc::MemCrc32(ChunkInfo.Buffer->Data, ChunkInfo.Size);
 	}
 
 	// Decrypt serialized hash
-	FDecryptedSignature<GPakFileChunkHashSize>& DecryptedSignature = DecryptedSignatures[ChunkInfo.Index];
+	FDecryptedSignature& DecryptedSignature = DecryptedSignatures[ChunkInfo.Index];
 	if (!DecryptedSignature.IsValid())
 	{
 		INC_DWORD_STAT(STAT_FChunkCacheWorker_SignaturesDecryptedDuringVerify);
 		SCOPE_SECONDS_ACCUMULATOR(STAT_FChunkCacheWorker_DecryptDuringVerify);
-		Decrypt(DecryptedSignature.Data, EncryptedSignatures[ChunkInfo.Index].Data, GPakFileChunkHashSize);
+		FEncryption::DecryptSignature(EncryptedSignatures[ChunkInfo.Index], DecryptedSignature, DecryptionKey);
 		check(DecryptedSignature.IsValid());
 	}
 
 	// Compare hashes
-	if (DecryptedSignature != Hash)
+	if (DecryptedSignature != SourceSignature)
 	{
 		UE_LOG(LogPakFile, Fatal, TEXT("Pak file has been tampered with."));
 	}
@@ -310,36 +314,6 @@ void FChunkCacheWorker::ReleaseChunk(FChunkRequest& Chunk)
 	Chunk.RefCount.Decrement();
 }
 
-int64 FChunkCacheWorker::AttachToArchive(FArchive& Ar)
-{
-	// Only serialize signature data if this was the first time this worker was attached
-	if (DecryptedSignatures.Num() > 0 || EncryptedSignatures.Num() > 0)
-	{
-		return Ar.TotalSize() - sizeof(int64)-(DecryptedSignatures.Num() * FEncryptedSignature<GPakFileChunkHashSize>::Size());
-	}
-
-	// Read num signatures from the end of the file
-	int64 OriginalPos = Ar.Tell();
-	Ar.Seek(Ar.TotalSize() - sizeof(int64));
-	int64 NumSignatures = 0;
-	Ar << NumSignatures;
-
-	// Read signature data
-	int64 StartOfSignatureData = Ar.TotalSize() - sizeof(int64)-(NumSignatures * FEncryptedSignature<GPakFileChunkHashSize>::Size());
-	DecryptedSignatures.AddDefaulted(NumSignatures);
-	EncryptedSignatures.AddUninitialized(NumSignatures);
-
-	Ar.Seek(StartOfSignatureData);
-	for (int64 SignatureIndex = 0; SignatureIndex < NumSignatures; ++SignatureIndex)
-	{
-		EncryptedSignatures[SignatureIndex].Serialize(Ar);
-	}
-
-	Ar.Seek(OriginalPos);
-
-	return StartOfSignatureData;
-}
-
 FSignedArchiveReader::FSignedArchiveReader(FArchive* InPakReader, FChunkCacheWorker* InSignatureChecker)
 	: ChunkCount(0)
 	, PakReader(InPakReader)
@@ -349,7 +323,7 @@ FSignedArchiveReader::FSignedArchiveReader(FArchive* InPakReader, FChunkCacheWor
 	, SignatureChecker(InSignatureChecker)
 {
 	// Cache global info about the archive
-	SizeOnDisk = InSignatureChecker->AttachToArchive(*PakReader);
+	SizeOnDisk = PakReader->TotalSize();
 	ChunkCount = SizeOnDisk / FPakInfo::MaxChunkDataSize + ((SizeOnDisk % FPakInfo::MaxChunkDataSize) ? 1 : 0);
 	PakSize = SizeOnDisk;
 }
diff --git a/Engine/Source/Runtime/PakFile/Private/SignedArchiveReader.h b/Engine/Source/Runtime/PakFile/Private/SignedArchiveReader.h
--- a/Engine/Source/Runtime/PakFile/Private/SignedArchiveReader.h
+++ b/Engine/Source/Runtime/PakFile/Private/SignedArchiveReader.h
@@ -99,9 +99,9 @@ class FChunkCacheWorker : public FRunnable
 	};
 
 	/** Encrypted signatures */
-	TArray<FEncryptedSignature<GPakFileChunkHashSize>> EncryptedSignatures;
+	TArray<FEncryptedSignature> EncryptedSignatures;
 	/** Decrypted signatures */
-	TArray<FDecryptedSignature<GPakFileChunkHashSize>> DecryptedSignatures;
+	TArray<FDecryptedSignature> DecryptedSignatures;
 	/** Thread to run the worker FRunnable on */
 	FRunnableThread* Thread;
 	/** Archive reader */
@@ -124,7 +124,7 @@ class FChunkCacheWorker : public FRunnable
 	FEncryptionKey DecryptionKey;
 	/** The time at which the chunk cache worker started, for debugging purposes */
 	double StartTime;
-	
+
 	/** 
 	 * Process requested chunks 
 	 *
@@ -136,10 +136,6 @@ class FChunkCacheWorker : public FRunnable
 	 */
 	bool CheckSignature(const FChunkRequest& ChunkInfo);
 	/** 
-	 * Decrypts a signature [*]
-	 */
-	void Decrypt(uint8* DecryptedData, const int256* Data, const int64 DataLength);
-	/**
 	* Decrypts multiple signature [*]
 	*/
 	void DecryptSignatures(int32 NextIndexToDecrypt);
@@ -162,7 +158,7 @@ class FChunkCacheWorker : public FRunnable
 
 public:
 
-	FChunkCacheWorker(FArchive* InReader);
+	FChunkCacheWorker(FArchive* InReader, const TCHAR* Filename);
 	virtual ~FChunkCacheWorker();
 
 	//~ Begin FRunnable Interface.
@@ -184,12 +180,6 @@ public:
 	 * Releases the requested chunk buffer
 	 */
 	void ReleaseChunk(FChunkRequest& Chunk);
-	/**
-	 * Attach this chunck cache worker to the given archive [*]
- 	 * @param Archive to attach to
-	 * @return End position in the archive of the actual file data, ignoring signature data 
-	 */
-	int64 AttachToArchive(FArchive& InAr);
 };
 
 /////////////////////////////////////////////////////////////////////////////////////////
@@ -265,7 +255,7 @@ class FSignedArchiveReader : public FArchive
 	int64 CalculateChunkSize(int64 ChunkIndex) const;
 
 	/** 
-	 * Queues chunks on the worker thread
+	 * Queues chunks on the worker thread 
 	 * @return Number of chunks in the output array which are actually required for the requested length. The rest are precache chunks 
 	 */
 	int64 PrecacheChunks(TArray<FReadInfo>& Chunks, int64 Length);
diff --git a/Engine/Source/Runtime/PakFile/Public/IPlatformFilePak.h b/Engine/Source/Runtime/PakFile/Public/IPlatformFilePak.h
--- a/Engine/Source/Runtime/PakFile/Public/IPlatformFilePak.h
+++ b/Engine/Source/Runtime/PakFile/Public/IPlatformFilePak.h
@@ -7,21 +7,8 @@ DECLARE_FLOAT_ACCUMULATOR_STAT_EXTERN(TEXT("Total pak file read time"), STAT_Pak
 
 DECLARE_DWORD_ACCUMULATOR_STAT_EXTERN(TEXT("Num open pak file handles"), STAT_PakFile_NumOpenHandles, STATGROUP_PakFile, PAKFILE_API);
 
-// Whether to use CRC (1) for chunk hashes, or SHA1 (0)
-#define PAKFILE_USE_CRC_FOR_CHUNK_HASHES 1
-#if PAKFILE_USE_CRC_FOR_CHUNK_HASHES
-static const int64 GPakFileChunkHashSize = 4;
-#else
-static const int64 GPakFileChunkHashSize = 20;
-#endif
-
 #define PAKFILE_TRACK_SECURITY_EXCLUDED_FILES 0
 
-/**
-* Compute a hash for a given block of data, as used for pak chunk signing
-*/
-void PAKFILE_API ComputePakChunkHash(const uint8* InData, const int64 InDataSize, uint8* OutHash);
-
 /**
  * Struct which holds pak file info (version, index offset, hash value).
  */
@@ -293,7 +280,7 @@ class PAKFILE_API FPakFile : FNoncopyable
 
 	FArchive* CreatePakReader(const TCHAR* Filename);
 	FArchive* CreatePakReader(IFileHandle& InHandle, const TCHAR* Filename);
-	FArchive* SetupSignedPakReader(FArchive* Reader);
+	FArchive* SetupSignedPakReader(FArchive* Reader, const TCHAR* Filename);
 
 public:
 
@@ -831,10 +818,14 @@ class PAKFILE_API FPakPlatformFile : public IPlatformFile
 	FCriticalSection PakListCritical;
 	/** true if we aren't allowing loose files. */
 	bool bSecurityEnabled;
+	/** External override of pak file security */
+	bool bForceSecurityBypass;
 #if PAKFILE_TRACK_SECURITY_EXCLUDED_FILES
 	/** List of files already excluded by the pak security system */
 	TSet<FName> SecurityExcludedFiles;
 #endif
+	/** List of allowed extension when security is enabled */
+	TSet<FName> AllowedExtensions;
 
 	/**
 	 * Gets mounted pak files
@@ -939,6 +930,15 @@ public:
 	 */
 	static void GetPakFolders(const TCHAR* CmdLine, TArray<FString>& OutPakFolders);
 
+	/**
+	 * Allows bypassing of loose file security
+	 */
+	virtual void BypassSecurity(bool bInBypassed) override
+	{
+		bForceSecurityBypass = bInBypassed;
+		IPlatformFile::BypassSecurity(bInBypassed);
+	}
+
 	/**
 	 * Constructor.
 	 * 
@@ -1231,11 +1231,7 @@ public:
 			return nullptr;
 		}
 		// Use lower level to handle writing.
-		if (IsFilenameAllowed(Filename))
-		{
-			return LowerLevel->OpenWrite(Filename, bAppend, bAllowRead);
-		}
-		return nullptr;
+		return LowerLevel->OpenWrite(Filename, bAppend, bAllowRead);
 	}
 
 	virtual bool DirectoryExists(const TCHAR* Directory) override
diff --git a/Engine/Source/Runtime/Windows/D3D11RHI/Private/D3D11Util.cpp b/Engine/Source/Runtime/Windows/D3D11RHI/Private/D3D11Util.cpp
--- a/Engine/Source/Runtime/Windows/D3D11RHI/Private/D3D11Util.cpp
+++ b/Engine/Source/Runtime/Windows/D3D11RHI/Private/D3D11Util.cpp
@@ -156,7 +156,14 @@ static void TerminateOnDeviceRemoved(HRESULT D3DResult)
 {
 	if (D3DResult == DXGI_ERROR_DEVICE_REMOVED)
 	{
-		FPlatformMisc::MessageBoxExt(EAppMsgType::Ok, *LOCTEXT("DeviceRemoved", "Video driver crashed and was reset!  Make sure your video drivers are up to date.  Exiting...").ToString(), TEXT("Error"));
+		if (!FApp::IsUnattended())
+		{
+			FPlatformMisc::MessageBoxExt(EAppMsgType::Ok, *LOCTEXT("DeviceRemoved", "Video driver crashed and was reset!  Make sure your video drivers are up to date.  Exiting...").ToString(), TEXT("Error"));
+		}
+		else
+		{
+			UE_LOG(LogD3D11RHI, Error, TEXT("%s"), *LOCTEXT("DeviceRemoved", "Video driver crashed and was reset!  Make sure your video drivers are up to date.  Exiting...").ToString());
+		}
 		FPlatformMisc::RequestExit(true);
 	}
 }
diff --git a/Engine/Source/Runtime/Windows/D3D12RHI/Private/D3D12Util.cpp b/Engine/Source/Runtime/Windows/D3D12RHI/Private/D3D12Util.cpp
--- a/Engine/Source/Runtime/Windows/D3D12RHI/Private/D3D12Util.cpp
+++ b/Engine/Source/Runtime/Windows/D3D12RHI/Private/D3D12Util.cpp
@@ -177,7 +177,14 @@ static void TerminateOnDeviceRemoved(HRESULT D3DResult)
 {
 	if (D3DResult == DXGI_ERROR_DEVICE_REMOVED)
 	{
-		FPlatformMisc::MessageBoxExt(EAppMsgType::Ok, *LOCTEXT("DeviceRemoved", "Video driver crashed and was reset!  Make sure your video drivers are up to date.  Exiting...").ToString(), TEXT("Error"));
+		if (!FApp::IsUnattended())
+		{
+			FPlatformMisc::MessageBoxExt(EAppMsgType::Ok, *LOCTEXT("DeviceRemoved", "Video driver crashed and was reset!  Make sure your video drivers are up to date.  Exiting...").ToString(), TEXT("Error"));
+		}
+		else
+		{
+			UE_LOG(LogD3D12RHI, Error, TEXT("%s"), *LOCTEXT("DeviceRemoved", "Video driver crashed and was reset!  Make sure your video drivers are up to date.  Exiting...").ToString());
+		}
 		FPlatformMisc::RequestExit(true);
 	}
 }
