Copying //UE4/Dev-Framework to Dev-Main (//UE4/Dev-Main) @ 2926658

diff --git a/Engine/Source/Runtime/Core/Private/Misc/ConfigCacheIni.cpp b/Engine/Source/Runtime/Core/Private/Misc/ConfigCacheIni.cpp
--- a/Engine/Source/Runtime/Core/Private/Misc/ConfigCacheIni.cpp
+++ b/Engine/Source/Runtime/Core/Private/Misc/ConfigCacheIni.cpp
@@ -50,11 +50,44 @@ namespace
 #endif
 }
 
+/*-----------------------------------------------------------------------------
+FConfigValue
+-----------------------------------------------------------------------------*/
+
+void FConfigValue::ExpandValue()
+{
+	// Replace %GAME% with game name.
+	FString PotentiallyExpandedValue = SavedValue.Replace(TEXT("%GAME%"), FApp::GetGameName(), ESearchCase::CaseSensitive);
+
+	// Replace %GAMEDIR% with the game directory.
+	PotentiallyExpandedValue = PotentiallyExpandedValue.Replace( TEXT("%GAMEDIR%"), *FPaths::GameDir(), ESearchCase::CaseSensitive );
+
+	// Replace %ENGINEUSERDIR% with the user's engine directory.
+	PotentiallyExpandedValue = PotentiallyExpandedValue.Replace(TEXT("%ENGINEUSERDIR%"), *FPaths::EngineUserDir(), ESearchCase::CaseSensitive);
+
+	// Replace %ENGINEVERSIONAGNOSTICUSERDIR% with the user's engine directory.
+	PotentiallyExpandedValue = PotentiallyExpandedValue.Replace(TEXT("%ENGINEVERSIONAGNOSTICUSERDIR%"), *FPaths::EngineVersionAgnosticUserDir(), ESearchCase::CaseSensitive);
+
+	// Replace %APPSETTINGSDIR% with the game directory.
+	FString AppSettingsDir = FPlatformProcess::ApplicationSettingsDir();
+	FPaths::NormalizeFilename(AppSettingsDir);
+	PotentiallyExpandedValue = PotentiallyExpandedValue.Replace( TEXT("%APPSETTINGSDIR%"), *AppSettingsDir, ESearchCase::CaseSensitive );
+
+	if (PotentiallyExpandedValue.Compare(SavedValue) != 0)
+	{
+		ExpandedValue = MoveTemp(PotentiallyExpandedValue);
+	}
+	else
+	{
+		ExpandedValue.Empty();
+	}
+}
+
 /*-----------------------------------------------------------------------------
 	FConfigSection
 -----------------------------------------------------------------------------*/
 
-bool FConfigSection::HasQuotes( const FString& Test ) const
+bool FConfigSection::HasQuotes( const FString& Test )
 {
 	return Test.Left(1) == TEXT("\"") && Test.Right(1) == TEXT("\"");
 }
@@ -62,19 +95,25 @@ bool FConfigSection::HasQuotes( const FString& Test ) const
 bool FConfigSection::operator==( const FConfigSection& Other ) const
 {
 	if ( Pairs.Num() != Other.Pairs.Num() )
+	{
 		return 0;
+	}
 
 	FConfigSectionMap::TConstIterator My(*this), Their(Other);
 	while ( My && Their )
 	{
 		if (My.Key() != Their.Key())
+		{
 			return 0;
+		}
 
-		const FString& MyValue = My.Value(), &TheirValue = Their.Value();
+		const FString& MyValue = My.Value().GetValue(), &TheirValue = Their.Value().GetValue();
 		if ( FCString::Strcmp(*MyValue,*TheirValue) &&
 			(!HasQuotes(MyValue) || FCString::Strcmp(*TheirValue,*MyValue.Mid(1,MyValue.Len()-2))) &&
 			(!HasQuotes(TheirValue) || FCString::Strcmp(*MyValue,*TheirValue.Mid(1,TheirValue.Len()-2))) )
+		{
 			return 0;
+		}
 
 		++My, ++Their;
 	}
@@ -93,15 +132,15 @@ FConfigFile::FConfigFile()
 : Dirty( false )
 , NoSave( false )
 , Name( NAME_None )
-, SourceConfigFile(NULL)
+, SourceConfigFile(nullptr)
 {}
 	
 FConfigFile::~FConfigFile()
 {
-	if( SourceConfigFile != NULL )
+	if( SourceConfigFile != nullptr )
 	{
 		delete SourceConfigFile;
-		SourceConfigFile = NULL;
+		SourceConfigFile = nullptr;
 	}
 }
 
@@ -144,25 +183,8 @@ bool FConfigFile::Combine(const FString& Filename)
 
 void FConfigFile::CombineFromBuffer(const FString& Buffer)
 {
-	// Replace %GAME% with game name.
-	FString Text = Buffer.Replace(TEXT("%GAME%"), FApp::GetGameName(), ESearchCase::CaseSensitive);
-
-	// Replace %GAMEDIR% with the game directory.
-	Text = Text.Replace( TEXT("%GAMEDIR%"), *FPaths::GameDir(), ESearchCase::CaseSensitive );
-
-	// Replace %ENGINEUSERDIR% with the user's engine directory.
-	Text = Text.Replace(TEXT("%ENGINEUSERDIR%"), *FPaths::EngineUserDir(), ESearchCase::CaseSensitive);
-
-	// Replace %ENGINEVERSIONAGNOSTICUSERDIR% with the user's engine directory.
-	Text = Text.Replace(TEXT("%ENGINEVERSIONAGNOSTICUSERDIR%"), *FPaths::EngineVersionAgnosticUserDir(), ESearchCase::CaseSensitive);
-
-	// Replace %APPSETTINGSDIR% with the game directory.
-	FString AppSettingsDir = FPlatformProcess::ApplicationSettingsDir();
-	FPaths::NormalizeFilename(AppSettingsDir);
-	Text = Text.Replace( TEXT("%APPSETTINGSDIR%"), *AppSettingsDir, ESearchCase::CaseSensitive );
-
-	const TCHAR* Ptr = *Text;
-	FConfigSection* CurrentSection = NULL;
+	const TCHAR* Ptr = *Buffer;
+	FConfigSection* CurrentSection = nullptr;
 	bool Done = false;
 	while( !Done )
 	{
@@ -176,7 +198,7 @@ void FConfigFile::CombineFromBuffer(const FString& Buffer)
 		FString TheLine;
 		int32 LinesConsumed = 0;
 		FParse::LineExtended(&Ptr, TheLine, LinesConsumed, false);
-		if (Ptr == NULL || *Ptr == 0)
+		if (Ptr == nullptr || *Ptr == 0)
 		{
 			Done = true;
 		}
@@ -303,17 +325,17 @@ void FConfigFile::CombineFromBuffer(const FString& Buffer)
 				if( Cmd=='+' ) 
 				{
 					// Add if not already present.
-					CurrentSection->AddUnique( Start, *ProcessedValue );
+					CurrentSection->AddUnique( Start, ProcessedValue );
 				}
 				else if( Cmd=='-' )	
 				{
 					// Remove if present.
-					CurrentSection->RemoveSingle( Start, *ProcessedValue );
+					CurrentSection->RemoveSingle( Start, ProcessedValue );
 					CurrentSection->Compact();
 				}
 				else if ( Cmd=='.' )
 				{
-					CurrentSection->Add( Start, *ProcessedValue );
+					CurrentSection->Add( Start, ProcessedValue );
 				}
 				else if( Cmd=='!' )
 				{
@@ -322,19 +344,19 @@ void FConfigFile::CombineFromBuffer(const FString& Buffer)
 				else
 				{
 					// Add if not present and replace if present.
-					FString* Str = CurrentSection->Find( Start );
-					if( !Str )
+					FConfigValue* ConfigValue = CurrentSection->Find( Start );
+					if( !ConfigValue )
 					{
-						CurrentSection->Add( Start, *ProcessedValue );
+						CurrentSection->Add( Start, ProcessedValue );
 					}
 					else
 					{
-						*Str = ProcessedValue;
+						*ConfigValue = FConfigValue(ProcessedValue);
 					}
 				}
 
 				// Mark as dirty so "Write" will actually save the changes.
-				Dirty = 1;
+				Dirty = true;
 			}
 		}
 	}
@@ -354,24 +376,10 @@ void FConfigFile::CombineFromBuffer(const FString& Buffer)
  */
 void FConfigFile::ProcessInputFileContents(const FString& Contents)
 {
-	// Replace %GAME% with game name.
-	FString Text = Contents.Replace(TEXT("%GAME%"), FApp::GetGameName(), ESearchCase::CaseSensitive);
-
-	// Replace %GAMEDIR% with the game directory.
-	Text = Text.Replace( TEXT("%GAMEDIR%"), *FPaths::GameDir(), ESearchCase::CaseSensitive );
-
-	// Replace %ENGINEUSERDIR% with the user's engine directory.
-	Text = Text.Replace(TEXT("%ENGINEUSERDIR%"), *FPaths::EngineUserDir(), ESearchCase::CaseSensitive);
-
-	// Replace %APPSETTINGSDIR% with the game directory.
-	FString AppSettingsDir = FPlatformProcess::ApplicationSettingsDir();
-	FPaths::NormalizeFilename(AppSettingsDir);
-	Text = Text.Replace( TEXT("%APPSETTINGSDIR%"), *AppSettingsDir, ESearchCase::CaseSensitive );
-
-	const TCHAR* Ptr = Text.Len() > 0 ? *Text : NULL;
-	FConfigSection* CurrentSection = NULL;
+	const TCHAR* Ptr = Contents.Len() > 0 ? *Contents : nullptr;
+	FConfigSection* CurrentSection = nullptr;
 	bool Done = false;
-	while( !Done && Ptr != NULL )
+	while( !Done && Ptr != nullptr )
 	{
 		// Advance past new line characters
 		while( *Ptr=='\r' || *Ptr=='\n' )
@@ -382,7 +390,7 @@ void FConfigFile::ProcessInputFileContents(const FString& Contents)
 		FString TheLine;
 		int32 LinesConsumed = 0;
 		FParse::LineExtended(&Ptr, TheLine, LinesConsumed, false);
-		if (Ptr == NULL || *Ptr == 0)
+		if (Ptr == nullptr || *Ptr == 0)
 		{
 			Done = true;
 		}
@@ -510,7 +518,7 @@ void FConfigFile::ProcessInputFileContents(const FString& Contents)
 void FConfigFile::Read( const FString& Filename )
 {
 	// we can't read in a file if file IO is disabled
-	if (GConfig == NULL || !GConfig->AreFileOperationsDisabled())
+	if (GConfig == nullptr || !GConfig->AreFileOperationsDisabled())
 	{
 		Empty();
 		FString Text;
@@ -527,7 +535,7 @@ bool FConfigFile::ShouldExportQuotedString(const FString& PropertyValue) const
 {
 	// The value should be exported as quoted string if it begins with a space (which is stripped on import) or
 	// when it contains '//' (interpreted as a comment when importing).
-	return **PropertyValue == TEXT(' ') || FCString::Strstr(*PropertyValue, TEXT("//")) != NULL;
+	return **PropertyValue == TEXT(' ') || FCString::Strstr(*PropertyValue, TEXT("//")) != nullptr;
 }
 
 
@@ -626,7 +634,7 @@ static bool LoadIniFileHierarchy(const FConfigFileHierarchy& HierarchyToLoad, FC
 		{
 			const FIniFilename& IniToLoad = HierarchyIt.Value;
 			if (IniToLoad.bRequired == false &&
-				 (!IsUsingLocalIniFile(*IniToLoad.Filename, NULL) || IFileManager::Get().FileSize(*IniToLoad.Filename) >= 0))
+				 (!IsUsingLocalIniFile(*IniToLoad.Filename, nullptr) || IFileManager::Get().FileSize(*IniToLoad.Filename) >= 0))
 			{
 				NumExistingOptionalInis++;
 			}
@@ -684,7 +692,7 @@ static bool LoadIniFileHierarchy(const FConfigFileHierarchy& HierarchyToLoad, FC
 			if (bDoProcess)
 			{
 				// Spit out friendly error if there was a problem locating .inis (e.g. bad command line parameter or missing folder, ...).
-				if (IsUsingLocalIniFile(*IniFileName, NULL) && (IFileManager::Get().FileSize(*IniFileName) < 0))
+				if (IsUsingLocalIniFile(*IniFileName, nullptr) && (IFileManager::Get().FileSize(*IniFileName) < 0))
 				{
 					if (IniToLoad.bRequired)
 					{
@@ -749,7 +757,7 @@ bool DoesConfigPropertyValueMatch( FConfigFile* InConfigFile, const FString& InS
 			// Start Array check, if the property is in an array, we need to iterate over all properties.
 			for (FConfigSection::TConstKeyIterator It(*Section, InPropertyName); It && !bFoundAMatch; ++It)
 			{
-				const FString& PropertyValue = It.Value();
+				const FString& PropertyValue = It.Value().GetSavedValue();
 				bFoundAMatch = PropertyValue == InPropertyValue;
 
 				// if our properties don't match, run further checks
@@ -856,7 +864,7 @@ bool FConfigFile::Write( const FString& Filename, bool bDoRemoteWrite/* = true*/
 		for( FConfigSection::TConstIterator It2(Section); It2; ++It2 )
 		{
 			const FName PropertyName = It2.Key();
-			const FString& PropertyValue = It2.Value();
+			const FString& PropertyValue = It2.Value().GetSavedValue();
 
 			// Check if the we've already processed a property of this name. If it was part of an array we may have already written it out.
 			if( !PropertiesAddedLookup.Contains( PropertyName ) )
@@ -865,13 +873,13 @@ bool FConfigFile::Write( const FString& Filename, bool bDoRemoteWrite/* = true*/
 				// This also catches the case where the property doesn't exist in the source in non-array cases
 				bool bDifferentNumberOfElements = false;
 				{
-					const FConfigSection* SourceSection = NULL;
+					const FConfigSection* SourceSection = nullptr;
 					if( SourceSection )
 					{
-						TArray< FString > SourceMatchingProperties;
+						TArray<FConfigValue> SourceMatchingProperties;
 						SourceSection->MultiFind( PropertyName, SourceMatchingProperties );
 
-						TArray< FString > DestMatchingProperties;
+						TArray<FConfigValue> DestMatchingProperties;
 						Section.MultiFind( PropertyName, DestMatchingProperties );
 
 						bDifferentNumberOfElements = SourceMatchingProperties.Num() != DestMatchingProperties.Num();
@@ -905,7 +913,7 @@ bool FConfigFile::Write( const FString& Filename, bool bDoRemoteWrite/* = true*/
 					}
 
 					// Write out our property, if it is an array we need to write out the entire array.
-					TArray< FString > CompletePropertyToWrite;
+					TArray<FConfigValue> CompletePropertyToWrite;
 					Section.MultiFind( PropertyName, CompletePropertyToWrite, true );
 
 					if( bIsADefaultIniWrite )
@@ -914,10 +922,10 @@ bool FConfigFile::Write( const FString& Filename, bool bDoRemoteWrite/* = true*/
 					}
 					else
 					{
-						for( int32 Idx = 0; Idx < CompletePropertyToWrite.Num(); Idx++ )
+						for (const FConfigValue& ConfigValue : CompletePropertyToWrite)
 						{
 							Text += FString::Printf( TEXT("%s=%s%s%s") LINE_TERMINATOR, 
-								*PropertyName.ToString(), QuoteString, *CompletePropertyToWrite[ Idx ], QuoteString);	
+								*PropertyName.ToString(), QuoteString, *ConfigValue.GetSavedValue(), QuoteString);	
 						}
 					}
 
@@ -978,7 +986,7 @@ void FConfigFile::AddMissingProperties( const FConfigFile& InSourceFile )
 		{
 			// If we don't already have this section, go ahead and add it now
 			FConfigSection* DestSection = Find( SourceSectionName );
-			if( DestSection == NULL )
+			if( DestSection == nullptr )
 			{
 				DestSection = &Add( SourceSectionName, FConfigSection() );
 				Dirty = true;
@@ -989,13 +997,14 @@ void FConfigFile::AddMissingProperties( const FConfigFile& InSourceFile )
 				const FName SourcePropertyName = SourcePropertyIt.Key();
 				
 				// If we don't already have this property, go ahead and add it now
-				if( DestSection->Find( SourcePropertyName ) == NULL )
+				if( DestSection->Find( SourcePropertyName ) == nullptr )
 				{
-					TArray<FString> Results;
+					TArray<FConfigValue> Results;
 					SourceSection.MultiFind(SourcePropertyName, Results, true);
-					for (int32 ResultsIdx = 0; ResultsIdx < Results.Num(); ++ResultsIdx)
+					for (const FConfigValue& Result : Results)
 					{
-						DestSection->Add(SourcePropertyName, Results[ResultsIdx]);
+
+						DestSection->Add(SourcePropertyName, Result.GetSavedValue());
 						Dirty = true;
 					}
 				}
@@ -1021,19 +1030,19 @@ void FConfigFile::Dump(FOutputDevice& Ar)
 		{
 			const FName KeyName = *KeyNameIt;
 
-			TArray<FString> Values;
+			TArray<FConfigValue> Values;
 			Section.MultiFind(KeyName,Values,true);
 
 			if ( Values.Num() > 1 )
 			{
 				for ( int32 ValueIndex = 0; ValueIndex < Values.Num(); ValueIndex++ )
 				{
-					Ar.Logf(TEXT("	%s[%i]=%s"), *KeyName.ToString(), ValueIndex, *Values[ValueIndex].ReplaceCharWithEscapedChar());
+					Ar.Logf(TEXT("	%s[%i]=%s"), *KeyName.ToString(), ValueIndex, *Values[ValueIndex].GetValue().ReplaceCharWithEscapedChar());
 				}
 			}
 			else
 			{
-				Ar.Logf(TEXT("	%s=%s"), *KeyName.ToString(), *Values[0].ReplaceCharWithEscapedChar());
+				Ar.Logf(TEXT("	%s=%s"), *KeyName.ToString(), *Values[0].GetValue().ReplaceCharWithEscapedChar());
 			}
 		}
 
@@ -1044,32 +1053,32 @@ void FConfigFile::Dump(FOutputDevice& Ar)
 bool FConfigFile::GetString( const TCHAR* Section, const TCHAR* Key, FString& Value ) const
 {
 	const FConfigSection* Sec = Find( Section );
-	if( Sec == NULL )
+	if( Sec == nullptr )
 	{
 		return false;
 	}
-	const FString* PairString = Sec->Find( Key );
-	if( PairString == NULL )
+	const FConfigValue* PairString = Sec->Find( Key );
+	if( PairString == nullptr )
 	{
 		return false;
 	}
-	Value = **PairString;
+	Value = PairString->GetValue();
 	return true;
 }
 
 bool FConfigFile::GetText( const TCHAR* Section, const TCHAR* Key, FText& Value ) const
 {
 	const FConfigSection* Sec = Find( Section );
-	if( Sec == NULL )
+	if( Sec == nullptr )
 	{
 		return false;
 	}
-	const FString* PairString = Sec->Find( Key );
-	if( PairString == NULL )
+	const FConfigValue* PairString = Sec->Find( Key );
+	if( PairString == nullptr )
 	{
 		return false;
 	}
-	return FTextStringHelper::ReadFromString( **PairString, Value, Section );
+	return FTextStringHelper::ReadFromString( *PairString->GetValue(), Value, Section );
 }
 
 bool FConfigFile::GetInt64( const TCHAR* Section, const TCHAR* Key, int64& Value ) const
@@ -1088,28 +1097,28 @@ bool FConfigFile::GetInt64( const TCHAR* Section, const TCHAR* Key, int64& Value
 void FConfigFile::SetString( const TCHAR* Section, const TCHAR* Key, const TCHAR* Value )
 {
 	FConfigSection* Sec  = Find( Section );
-	if( Sec == NULL )
+	if( Sec == nullptr )
 	{
 		Sec = &Add( Section, FConfigSection() );
 	}
 
-	FString* Str = Sec->Find( Key );
-	if( Str == NULL )
+	FConfigValue* ConfigValue = Sec->Find( Key );
+	if( ConfigValue == nullptr )
 	{
 		Sec->Add( Key, Value );
-		Dirty = 1;
+		Dirty = true;
 	}
-	else if( FCString::Strcmp(**Str,Value)!=0 )
+	else if( FCString::Strcmp(*ConfigValue->GetSavedValue(),Value)!=0 )
 	{
 		Dirty = true;
-		*Str = Value;
+		*ConfigValue = FConfigValue(Value);
 	}
 }
 
 void FConfigFile::SetText( const TCHAR* Section, const TCHAR* Key, const FText& Value )
 {
 	FConfigSection* Sec  = Find( Section );
-	if( Sec == NULL )
+	if( Sec == nullptr )
 	{
 		Sec = &Add( Section, FConfigSection() );
 	}
@@ -1117,16 +1126,16 @@ void FConfigFile::SetText( const TCHAR* Section, const TCHAR* Key, const FText&
 	FString StrValue;
 	FTextStringHelper::WriteToString(StrValue, Value);
 
-	FString* Str = Sec->Find( Key );
-	if( Str == NULL )
+	FConfigValue* ConfigValue = Sec->Find( Key );
+	if( ConfigValue == nullptr )
 	{
 		Sec->Add( Key, StrValue );
-		Dirty = 1;
+		Dirty = true;
 	}
-	else if( FCString::Strcmp(**Str, *StrValue)!=0 )
+	else if( FCString::Strcmp(*ConfigValue->GetSavedValue(), *StrValue)!=0 )
 	{
 		Dirty = true;
-		*Str = StrValue;
+		*ConfigValue = FConfigValue(StrValue);
 	}
 }
 
@@ -1155,7 +1164,7 @@ void FConfigFile::SaveSourceToBackupFile()
 		for( FConfigSection::TConstIterator PropertyIterator(Section); PropertyIterator; ++PropertyIterator )
 		{
 			const FName PropertyName = PropertyIterator.Key();
-			const FString& PropertyValue = PropertyIterator.Value();
+			const FString& PropertyValue = PropertyIterator.Value().GetSavedValue();
 
 			// Print the property to the file.
 			TCHAR QuoteString[2] = {0,0};
@@ -1205,7 +1214,7 @@ void FConfigFile::ProcessSourceAndCheckAgainstBackup()
 }
 
 
-void FConfigFile::ProcessPropertyAndWriteForDefaults( const TArray< FString >& InCompletePropertyToProcess, FString& OutText, const FString& SectionName, const FString& PropertyName )
+void FConfigFile::ProcessPropertyAndWriteForDefaults( const TArray< FConfigValue >& InCompletePropertyToProcess, FString& OutText, const FString& SectionName, const FString& PropertyName )
 {
 	// Only process against a hierarchy if this config file has one.
 	if (SourceIniHierarchy.Num() > 0)
@@ -1247,9 +1256,9 @@ void FConfigFile::ProcessPropertyAndWriteForDefaults( const TArray< FString >& I
 	}
 
 	// Write the properties out to a file.
-	for ( auto& PropertyIt : InCompletePropertyToProcess)
+	for ( const FConfigValue& PropertyIt : InCompletePropertyToProcess)
 	{
-		FString PropertyValue = *PropertyIt;
+		FString PropertyValue = PropertyIt.GetSavedValue();
 
 		if (ShouldExportQuotedString(PropertyValue))
 		{
@@ -1294,7 +1303,7 @@ FConfigFile* FConfigCacheIni::Find( const FString& Filename, bool CreateIfNotFou
 	// check for non-filenames
 	if(Filename.Len() == 0)
 	{
-		return NULL;
+		return nullptr;
 	}
 
 	// Get file.
@@ -1406,33 +1415,33 @@ void FConfigCacheIni::Parse1ToNSectionOfNames(const TCHAR* Section, const TCHAR*
 		return;
 	}
 
-	TArray<FName>* WorkingList = NULL;
+	TArray<FName>* WorkingList = nullptr;
 	for( FConfigSectionMap::TIterator It(*ConfigSection); It; ++It )
 	{
 		// is the current key the 1 key?
 		if (It.Key() == KeyOne)
 		{
-			FName KeyName(*It.Value());
+			const FName KeyName(*It.Value().GetValue());
 
 			// look for existing set in the map
 			WorkingList = OutMap.Find(KeyName);
 
 			// make a new one if it wasn't there
-			if (WorkingList == NULL)
+			if (WorkingList == nullptr)
 			{
 				WorkingList = &OutMap.Add(KeyName, TArray<FName>());
 			}
 		}
 		// is the current key the N key?
-		else if (It.Key() == KeyN && WorkingList != NULL)
+		else if (It.Key() == KeyN && WorkingList != nullptr)
 		{
 			// if so, add it to the N list for the current 1 key
-			WorkingList->Add(FName(*It.Value()));
+			WorkingList->Add(FName(*It.Value().GetValue()));
 		}
 		// if it's neither, then reset
 		else
 		{
-			WorkingList = NULL;
+			WorkingList = nullptr;
 		}
 	}
 }
@@ -1471,31 +1480,31 @@ void FConfigCacheIni::Parse1ToNSectionOfStrings(const TCHAR* Section, const TCHA
 		return;
 	}
 
-	TArray<FString>* WorkingList = NULL;
+	TArray<FString>* WorkingList = nullptr;
 	for( FConfigSectionMap::TIterator It(*ConfigSection); It; ++It )
 	{
 		// is the current key the 1 key?
 		if (It.Key() == KeyOne)
 		{
 			// look for existing set in the map
-			WorkingList = OutMap.Find(*It.Value());
+			WorkingList = OutMap.Find(It.Value().GetValue());
 
 			// make a new one if it wasn't there
-			if (WorkingList == NULL)
+			if (WorkingList == nullptr)
 			{
-				WorkingList = &OutMap.Add(*It.Value(), TArray<FString>());
+				WorkingList = &OutMap.Add(It.Value().GetValue(), TArray<FString>());
 			}
 		}
 		// is the current key the N key?
-		else if (It.Key() == KeyN && WorkingList != NULL)
+		else if (It.Key() == KeyN && WorkingList != nullptr)
 		{
 			// if so, add it to the N list for the current 1 key
-			WorkingList->Add(It.Value());
+			WorkingList->Add(It.Value().GetValue());
 		}
 		// if it's neither, then reset
 		else
 		{
-			WorkingList = NULL;
+			WorkingList = nullptr;
 		}
 	}
 }
@@ -1503,7 +1512,7 @@ void FConfigCacheIni::Parse1ToNSectionOfStrings(const TCHAR* Section, const TCHA
 void FConfigCacheIni::LoadFile( const FString& Filename, const FConfigFile* Fallback, const TCHAR* PlatformString )
 {
 	// if the file has some data in it, read it in
-	if( !IsUsingLocalIniFile(*Filename, NULL) || (IFileManager::Get().FileSize(*Filename) >= 0) )
+	if( !IsUsingLocalIniFile(*Filename, nullptr) || (IFileManager::Get().FileSize(*Filename) >= 0) )
 	{
 		FConfigFile* Result = &Add(Filename, FConfigFile());
 		bool bDoEmptyConfig = false;
@@ -1592,12 +1601,12 @@ bool FConfigCacheIni::GetString( const TCHAR* Section, const TCHAR* Key, FString
 #endif
 		return false;
 	}
-	FString* PairString = Sec->Find( Key );
-	if( !PairString )
+	const FConfigValue* ConfigValue = Sec->Find( Key );
+	if( !ConfigValue )
 	{
 		return false;
 	}
-	Value = **PairString;
+	Value = ConfigValue->GetValue();
 	return true;
 }
 
@@ -1617,40 +1626,53 @@ bool FConfigCacheIni::GetText( const TCHAR* Section, const TCHAR* Key, FText& Va
 #endif
 		return false;
 	}
-	FString* PairString = Sec->Find( Key );
-	if( !PairString )
+	const FConfigValue* ConfigValue = Sec->Find( Key );
+	if( !ConfigValue )
 	{
 		return false;
 	}
-	return FTextStringHelper::ReadFromString( **PairString, Value, Section );
+	return FTextStringHelper::ReadFromString( *ConfigValue->GetValue(), Value, Section );
 }
 
 bool FConfigCacheIni::GetSection( const TCHAR* Section, TArray<FString>& Result, const FString& Filename )
 {
 	FRemoteConfig::Get()->FinishRead(*Filename); // Ensure the remote file has been loaded and processed
-	Result.Empty();
-	FConfigFile* File = Find( Filename, 0 );
-	if( !File )
-		return 0;
+	Result.Reset();
+	FConfigFile* File = Find( Filename, false );
+	if (!File)
+	{
+		return false;
+	}
 	FConfigSection* Sec = File->Find( Section );
-	if( !Sec )
-		return 0;
-	for( FConfigSection::TIterator It(*Sec); It; ++It )
-		new(Result) FString(FString::Printf( TEXT("%s=%s"), *It.Key().ToString(), *It.Value() ));
-	return 1;
+	if (!Sec)
+	{
+		return false;
+	}
+	Result.Reserve(Sec->Num());
+	for (FConfigSection::TIterator It(*Sec); It; ++It)
+	{
+		Result.Add(FString::Printf(TEXT("%s=%s"), *It.Key().ToString(), *It.Value().GetValue()));
+	}
+	return true;
 }
 
 FConfigSection* FConfigCacheIni::GetSectionPrivate( const TCHAR* Section, bool Force, bool Const, const FString& Filename )
 {
 	FRemoteConfig::Get()->FinishRead(*Filename); // Ensure the remote file has been loaded and processed
 	FConfigFile* File = Find( Filename, Force );
-	if( !File )
-		return NULL;
+	if (!File)
+	{
+		return nullptr;
+	}
 	FConfigSection* Sec = File->Find( Section );
-	if( !Sec && Force )
-		Sec = &File->Add( Section, FConfigSection() );
-	if( Sec && (Force || !Const) )
-		File->Dirty = 1;
+	if (!Sec && Force)
+	{
+		Sec = &File->Add(Section, FConfigSection());
+	}
+	if (Sec && (Force || !Const))
+	{
+		File->Dirty = true;
+	}
 	return Sec;
 }
 
@@ -1661,7 +1683,7 @@ bool FConfigCacheIni::DoesSectionExist(const TCHAR* Section, const FString& File
 	FRemoteConfig::Get()->FinishRead(*Filename); // Ensure the remote file has been loaded and processed
 	FConfigFile* File = Find(Filename, false);
 
-	bReturnVal = File != NULL && File->Find(Section) != NULL;
+	bReturnVal = (File != nullptr && File->Find(Section) != nullptr);
 
 	return bReturnVal;
 }
@@ -1676,18 +1698,20 @@ void FConfigCacheIni::SetString( const TCHAR* Section, const TCHAR* Key, const T
 	}
 
 	FConfigSection* Sec = File->Find( Section );
-	if( !Sec )
-		Sec = &File->Add( Section, FConfigSection() );
-	FString* Str = Sec->Find( Key );
-	if( !Str )
+	if (!Sec)
+	{
+		Sec = &File->Add(Section, FConfigSection());
+	}
+	FConfigValue* ConfigValue = Sec->Find( Key );
+	if( !ConfigValue )
 	{
 		Sec->Add( Key, Value );
-		File->Dirty = 1;
+		File->Dirty = true;
 	}
-	else if( FCString::Stricmp(**Str,Value)!=0 )
+	else if( FCString::Stricmp(*ConfigValue->GetSavedValue(),Value)!=0 )
 	{
-		File->Dirty = (FCString::Strcmp(**Str,Value)!=0);
-		*Str = Value;
+		File->Dirty = true;
+		*ConfigValue = FConfigValue(Value);
 	}
 }
 
@@ -1709,16 +1733,16 @@ void FConfigCacheIni::SetText( const TCHAR* Section, const TCHAR* Key, const FTe
 	FString StrValue;
 	FTextStringHelper::WriteToString(StrValue, Value);
 
-	FString* Str = Sec->Find( Key );
-	if( !Str )
+	FConfigValue* ConfigValue = Sec->Find( Key );
+	if( !ConfigValue )
 	{
 		Sec->Add( Key, StrValue );
-		File->Dirty = 1;
+		File->Dirty = true;
 	}
-	else if( FCString::Stricmp(**Str, *StrValue)!=0 )
+	else if( FCString::Stricmp(*ConfigValue->GetSavedValue(), *StrValue)!=0 )
 	{
-		File->Dirty = (FCString::Strcmp(**Str, *StrValue)!=0);
-		*Str = StrValue;
+		File->Dirty = true;
+		*ConfigValue = FConfigValue(StrValue);
 	}
 }
 
@@ -1814,7 +1838,7 @@ bool FConfigCacheIni::GetSectionNames( const FString& Filename, TArray<FString>&
 	bool bResult = false;
 
 	FConfigFile* File = Find(Filename, false);
-	if ( File != NULL )
+	if ( File != nullptr )
 	{
 		out_SectionNames.Empty(Num());
 		for ( FConfigFile::TIterator It(*File); It; ++It )
@@ -1844,7 +1868,7 @@ bool FConfigCacheIni::GetPerObjectConfigSections( const FString& Filename, const
 
 	MaxResults = FMath::Max(0, MaxResults);
 	FConfigFile* File = Find(Filename, false);
-	if ( File != NULL )
+	if ( File != nullptr )
 	{
 		out_SectionNames.Empty();
 		for ( FConfigFile::TIterator It(*File); It && out_SectionNames.Num() < MaxResults; ++It )
@@ -1877,7 +1901,7 @@ void FConfigCacheIni::Exit()
 
 void FConfigCacheIni::Dump(FOutputDevice& Ar, const TCHAR* BaseIniName)
 {
-	if (BaseIniName == NULL)
+	if (BaseIniName == nullptr)
 	{
 		Ar.Log( TEXT("Files map:") );
 		TMap<FString,FConfigFile>::Dump( Ar );
@@ -1885,7 +1909,7 @@ void FConfigCacheIni::Dump(FOutputDevice& Ar, const TCHAR* BaseIniName)
 
 	for ( TIterator It(*this); It; ++It )
 	{
-		if (BaseIniName == NULL || FPaths::GetBaseFilename(It.Key()) == BaseIniName)
+		if (BaseIniName == nullptr || FPaths::GetBaseFilename(It.Key()) == BaseIniName)
 		{
 			Ar.Logf(TEXT("FileName: %s"), *It.Key());
 			FConfigFile& File = It.Value();
@@ -1893,8 +1917,10 @@ void FConfigCacheIni::Dump(FOutputDevice& Ar, const TCHAR* BaseIniName)
 			{
 				FConfigSection& Sec = FileIt.Value();
 				Ar.Logf(TEXT("   [%s]"), *FileIt.Key());
-				for ( FConfigSectionMap::TConstIterator SecIt(Sec); SecIt; ++SecIt )
-					Ar.Logf(TEXT("   %s=%s"), *SecIt.Key().ToString(), *SecIt.Value());
+				for (FConfigSectionMap::TConstIterator SecIt(Sec); SecIt; ++SecIt)
+				{
+					Ar.Logf(TEXT("   %s=%s"), *SecIt.Key().ToString(), *SecIt.Value().GetValue());
+				}
 
 				Ar.Log(LINE_TERMINATOR);
 			}
@@ -1986,19 +2012,19 @@ int32 FConfigCacheIni::GetArray
 	FRemoteConfig::Get()->FinishRead(*Filename); // Ensure the remote file has been loaded and processed
 	out_Arr.Empty();
 	FConfigFile* File = Find( Filename, 0 );
-	if ( File != NULL )
+	if ( File != nullptr )
 	{
 		FConfigSection* Sec = File->Find( Section );
-		if ( Sec != NULL )
+		if ( Sec != nullptr )
 		{
-			TArray<FString> RemapArray;
+			TArray<FConfigValue> RemapArray;
 			Sec->MultiFind(Key, RemapArray);
 
 			// TMultiMap::MultiFind will return the results in reverse order
 			out_Arr.AddZeroed(RemapArray.Num());
 			for ( int32 RemapIndex = RemapArray.Num() - 1, Index = 0; RemapIndex >= 0; RemapIndex--, Index++ )
 			{
-				out_Arr[Index] = RemapArray[RemapIndex];
+				out_Arr[Index] = RemapArray[RemapIndex].GetValue();
 			}
 		}
 #if !UE_BUILD_SHIPPING
@@ -2473,7 +2499,7 @@ bool FConfigCacheIni::ForEachEntry(const FKeyValueSink& Visitor, const TCHAR* Se
 
 	for(FConfigSectionMap::TIterator It(*Sec); It; ++It)
 	{
-		Visitor.Execute(*It.Key().GetPlainNameString(), *It.Value());
+		Visitor.Execute(*It.Key().GetPlainNameString(), *It.Value().GetValue());
 	}
 
 	return true;
@@ -2491,7 +2517,7 @@ static void LoadAnIniFile(const FString& FilenameToLoad, FConfigFile& ConfigFile
 	// This shouldn't be getting called if seekfree is enabled on console.
 	check(!GUseSeekFreeLoading || !FPlatformProperties::RequiresCookedData());
 
-	if( !IsUsingLocalIniFile(*FilenameToLoad, NULL) || (IFileManager::Get().FileSize( *FilenameToLoad ) >= 0) )
+	if( !IsUsingLocalIniFile(*FilenameToLoad, nullptr) || (IFileManager::Get().FileSize( *FilenameToLoad ) >= 0) )
 	{
 		ProcessIniContents(*FilenameToLoad, *FilenameToLoad, &ConfigFile, false, false);
 	}
@@ -2611,7 +2637,7 @@ static bool GenerateDestIniFile(FConfigFile& DestConfigFile, const FString& Dest
 		DestConfigFile.Dirty = true;
 	}
 		
-	if (!IsUsingLocalIniFile(*DestIniFilename, NULL))
+	if (!IsUsingLocalIniFile(*DestIniFilename, nullptr))
 	{
 		// Save off a copy of the local file prior to overwriting it with the contents of a remote file
 		MakeLocalCopy(*DestIniFilename);
@@ -2651,7 +2677,7 @@ static FString GetSourceIniFilename(const TCHAR* ConfigDir, const TCHAR* Prefix,
  * Creates a chain of ini filenames to load and combine.
  *
  * @param InBaseIniName Ini name.
- * @param InPlatformName Platform name, NULL means to use the current platform
+ * @param InPlatformName Platform name, nullptr means to use the current platform
  * @param OutHierarchy An array which is to receive the generated hierachy of ini filenames.
  */
 static void GetSourceIniHierarchyFilenames(const TCHAR* InBaseIniName, const TCHAR* InPlatformName, const TCHAR* EngineConfigDir, const TCHAR* SourceConfigDir, FConfigFileHierarchy& OutHierarchy, bool bRequireDefaultIni)
@@ -2739,7 +2765,7 @@ static void GetSourceIniHierarchyFilenames(const TCHAR* InBaseIniName, const TCH
  * Calculates the name of a dest (generated) .ini file for a given base (ie Engine, Game, etc)
  * 
  * @param IniBaseName Base name of the .ini (Engine, Game)
- * @param PlatformName Name of the platform to get the .ini path for (NULL means to use the current platform)
+ * @param PlatformName Name of the platform to get the .ini path for (nullptr means to use the current platform)
  * @param GeneratedConfigDir The base folder that will contain the generated config files.
  * 
  * @return Standardized .ini filename
@@ -2774,7 +2800,7 @@ void FConfigCacheIni::InitializeConfigSystem()
 	// load the main .ini files (unless we're running a program or a gameless UE4Editor.exe, DefaultEngine.ini is required).
 	const bool bIsGamelessExe = !FApp::HasGameName();
 	const bool bDefaultEngineIniRequired = !bIsGamelessExe && (GIsGameAgnosticExe || FApp::IsGameNameEmpty());
-	bool bEngineConfigCreated = FConfigCacheIni::LoadGlobalIniFile(GEngineIni, TEXT("Engine"), NULL, bDefaultEngineIniRequired);
+	bool bEngineConfigCreated = FConfigCacheIni::LoadGlobalIniFile(GEngineIni, TEXT("Engine"), nullptr, bDefaultEngineIniRequired);
 
 	if ( !bIsGamelessExe )
 	{
@@ -2810,9 +2836,9 @@ void FConfigCacheIni::InitializeConfigSystem()
 
 	// Project agnostic editor ini files
 	static const FString EditorSettingsDir = FPaths::Combine(*FPaths::GameAgnosticSavedDir(), TEXT("Config")) + TEXT("/");
-	FConfigCacheIni::LoadGlobalIniFile(GEditorSettingsIni, TEXT("EditorSettings"), NULL, false, false, true, *EditorSettingsDir);
-	FConfigCacheIni::LoadGlobalIniFile(GEditorLayoutIni, TEXT("EditorLayout"), NULL, false, false, true, *EditorSettingsDir);
-	FConfigCacheIni::LoadGlobalIniFile(GEditorKeyBindingsIni, TEXT("EditorKeyBindings"), NULL, false, false, true, *EditorSettingsDir);
+	FConfigCacheIni::LoadGlobalIniFile(GEditorSettingsIni, TEXT("EditorSettings"), nullptr, false, false, true, *EditorSettingsDir);
+	FConfigCacheIni::LoadGlobalIniFile(GEditorLayoutIni, TEXT("EditorLayout"), nullptr, false, false, true, *EditorSettingsDir);
+	FConfigCacheIni::LoadGlobalIniFile(GEditorKeyBindingsIni, TEXT("EditorKeyBindings"), nullptr, false, false, true, *EditorSettingsDir);
 
 #endif
 #if PLATFORM_DESKTOP
@@ -2853,7 +2879,7 @@ bool FConfigCacheIni::LoadGlobalIniFile(FString& FinalIniFilename, const TCHAR*
 
 	// need to check to see if the file already exists in the GConfigManager's cache
 	// if it does exist then we are done, nothing else to do
-	if (!bForceReload && GConfig->FindConfigFile(*FinalIniFilename) != NULL)
+	if (!bForceReload && GConfig->FindConfigFile(*FinalIniFilename) != nullptr)
 	{
 		//UE_LOG(LogConfig, Log,  TEXT( "Request to load a config file that was already loaded: %s" ), GeneratedIniFile );
 		return true;
@@ -2988,7 +3014,7 @@ void FConfigCacheIni::LoadConsoleVariablesFromINI()
 			for(FConfigSectionMap::TConstIterator It(*Section); It; ++It)
 			{
 				const FString& KeyString = It.Key().GetPlainNameString(); 
-				const FString& ValueString = It.Value();
+				const FString& ValueString = It.Value().GetValue();
 
 				IConsoleVariable* CVar = IConsoleManager::Get().FindConsoleVariable(*KeyString);
 
@@ -3014,7 +3040,7 @@ void FConfigCacheIni::LoadConsoleVariablesFromINI()
 	}
 }
 
-void FConfigFile::UpdateSections(const TCHAR* DiskFilename, const TCHAR* IniRootName/*=NULL*/)
+void FConfigFile::UpdateSections(const TCHAR* DiskFilename, const TCHAR* IniRootName/*=nullptr*/)
 {
 	// since we don't want any modifications to other sections, we manually process the file, not read into sections, etc
 	FString DiskFile;
@@ -3023,7 +3049,7 @@ void FConfigFile::UpdateSections(const TCHAR* DiskFilename, const TCHAR* IniRoot
 	if (FFileHelper::LoadFileToString(DiskFile, DiskFilename))
 	{
 		// walk each line
-		const TCHAR* Ptr = DiskFile.Len() > 0 ? *DiskFile : NULL;
+		const TCHAR* Ptr = DiskFile.Len() > 0 ? *DiskFile : nullptr;
 		bool bDone = Ptr ? false : true;
 		bool bIsSkippingSection = true;
 		while (!bDone)
@@ -3057,11 +3083,11 @@ void FConfigFile::UpdateSections(const TCHAR* DiskFilename, const TCHAR* IniRoot
 	}
 
 	// load the hierarchy up to right before this file
-	if (IniRootName != NULL)
+	if (IniRootName != nullptr)
 	{
 		// get the standard ini files
 		SourceIniHierarchy.Empty();
-		GetSourceIniHierarchyFilenames(IniRootName, NULL, *FPaths::EngineConfigDir(), *FPaths::SourceConfigDir(), SourceIniHierarchy, false);
+		GetSourceIniHierarchyFilenames(IniRootName, nullptr, *FPaths::EngineConfigDir(), *FPaths::SourceConfigDir(), SourceIniHierarchy, false);
 
 		// now chop off this file and any after it
 		for (auto& HierarchyFileIt : SourceIniHierarchy)
@@ -3298,9 +3324,9 @@ bool FConfigFile::UpdateSinglePropertyInSection(const TCHAR* DiskFilename, const
 	FString PropertyValue;
 	if (const FConfigSection* LocalSection = this->Find(SectionName))
 	{
-		if (const FString* LocalSectionPropertyStr = LocalSection->Find(PropertyName))
+		if (const FConfigValue* ConfigValue = LocalSection->Find(PropertyName))
 		{
-			PropertyValue = *LocalSectionPropertyStr;
+			PropertyValue = ConfigValue->GetSavedValue();
 			FSinglePropertyConfigHelper SinglePropertyConfigHelper(DiskFilename, SectionName, PropertyName, PropertyValue);
 			bSuccessfullyUpdatedFile = SinglePropertyConfigHelper.UpdateConfigFile();
 		}
diff --git a/Engine/Source/Runtime/Core/Private/ProfilingDebugging/LoadTimeTracker.cpp b/Engine/Source/Runtime/Core/Private/ProfilingDebugging/LoadTimeTracker.cpp
--- a/Engine/Source/Runtime/Core/Private/ProfilingDebugging/LoadTimeTracker.cpp
+++ b/Engine/Source/Runtime/Core/Private/ProfilingDebugging/LoadTimeTracker.cpp
@@ -10,7 +10,10 @@
 #include "CorePrivatePCH.h"
 #include "LoadTimeTracker.h"
 
-#if ENABLE_LOADTIME_TRACKING
+FLoadTimeTracker::FLoadTimeTracker()
+{
+	ResetRawLoadTimes();
+}
 
 void FLoadTimeTracker::ReportScopeTime(double ScopeTime, const FName ScopeLabel)
 {
@@ -20,11 +23,11 @@ void FLoadTimeTracker::ReportScopeTime(double ScopeTime, const FName ScopeLabel)
 }
 
 
-void FLoadTimeTracker::DumpLoadTimes()
+void FLoadTimeTracker::DumpHighLevelLoadTimes() const
 {
 	double TotalTime = 0.0;
 	UE_LOG(LogLoad, Log, TEXT("------------- Load times -------------"));
-	for(auto Itr = TimeInfo.CreateIterator(); Itr; ++Itr)
+	for(auto Itr = TimeInfo.CreateConstIterator(); Itr; ++Itr)
 	{
 		const FString KeyName = Itr.Key().ToString();
 		const TArray<double>& LoadTimes = Itr.Value();
@@ -52,7 +55,7 @@ void FLoadTimeTracker::DumpLoadTimes()
 	UE_LOG(LogLoad, Log, TEXT("Total Load times: %f"), TotalTime);
 }
 
-void FLoadTimeTracker::ResetLoadTimes()
+void FLoadTimeTracker::ResetHighLevelLoadTimes()
 {
 	static bool bActuallyReset = !FParse::Param(FCommandLine::Get(), TEXT("NoLoadTrackClear"));
 	if(bActuallyReset)
@@ -61,11 +64,124 @@ void FLoadTimeTracker::ResetLoadTimes()
 	}
 }
 
+void FLoadTimeTracker::DumpRawLoadTimes() const
+{
+#if ENABLE_LOADTIME_RAW_TIMINGS
+	UE_LOG(LogStreaming, Display, TEXT("-------------------------------------------------"));
+	UE_LOG(LogStreaming, Display, TEXT("Async Loading Stats"));
+	UE_LOG(LogStreaming, Display, TEXT("-------------------------------------------------"));
+	UE_LOG(LogStreaming, Display, TEXT("AsyncLoadingTime: %f"), AsyncLoadingTime);
+	UE_LOG(LogStreaming, Display, TEXT("CreateAsyncPackagesFromQueueTime: %f"), CreateAsyncPackagesFromQueueTime);
+	UE_LOG(LogStreaming, Display, TEXT("ProcessAsyncLoadingTime: %f"), ProcessAsyncLoadingTime);
+	UE_LOG(LogStreaming, Display, TEXT("ProcessLoadedPackagesTime: %f"), ProcessLoadedPackagesTime);
+	//UE_LOG(LogStreaming, Display, TEXT("SerializeTaggedPropertiesTime: %f"), SerializeTaggedPropertiesTime);
+	UE_LOG(LogStreaming, Display, TEXT("CreateLinkerTime: %f"), CreateLinkerTime);
+	UE_LOG(LogStreaming, Display, TEXT("FinishLinkerTime: %f"), FinishLinkerTime);
+	UE_LOG(LogStreaming, Display, TEXT("CreateImportsTime: %f"), CreateImportsTime);
+	UE_LOG(LogStreaming, Display, TEXT("CreateExportsTime: %f"), CreateExportsTime);
+	UE_LOG(LogStreaming, Display, TEXT("PreLoadObjectsTime: %f"), PreLoadObjectsTime);
+	UE_LOG(LogStreaming, Display, TEXT("PostLoadObjectsTime: %f"), PostLoadObjectsTime);
+	UE_LOG(LogStreaming, Display, TEXT("PostLoadDeferredObjectsTime: %f"), PostLoadDeferredObjectsTime);
+	UE_LOG(LogStreaming, Display, TEXT("FinishObjectsTime: %f"), FinishObjectsTime);
+	UE_LOG(LogStreaming, Display, TEXT("MaterialPostLoad: %f"), MaterialPostLoad);
+	UE_LOG(LogStreaming, Display, TEXT("MaterialInstancePostLoad: %f"), MaterialInstancePostLoad);
+	UE_LOG(LogStreaming, Display, TEXT("SerializeInlineShaderMaps: %f"), SerializeInlineShaderMaps);
+	UE_LOG(LogStreaming, Display, TEXT("MaterialSerializeTime: %f"), MaterialSerializeTime);
+	UE_LOG(LogStreaming, Display, TEXT("MaterialInstanceSerializeTime: %f"), MaterialInstanceSerializeTime);
+	UE_LOG(LogStreaming, Display, TEXT(""));
+	UE_LOG(LogStreaming, Display, TEXT("LinkerLoad_CreateLoader: %f"), LinkerLoad_CreateLoader);
+	UE_LOG(LogStreaming, Display, TEXT("LinkerLoad_SerializePackageFileSummary: %f"), LinkerLoad_SerializePackageFileSummary);
+	UE_LOG(LogStreaming, Display, TEXT("LinkerLoad_SerializeNameMap: %f"), LinkerLoad_SerializeNameMap);
+	UE_LOG(LogStreaming, Display, TEXT("\tProcessingEntries: %f"), LinkerLoad_SerializeNameMap_ProcessingEntries);
+	UE_LOG(LogStreaming, Display, TEXT("LinkerLoad_SerializeGatherableTextDataMap: %f"), LinkerLoad_SerializeGatherableTextDataMap);
+	UE_LOG(LogStreaming, Display, TEXT("LinkerLoad_SerializeImportMap: %f"), LinkerLoad_SerializeImportMap);
+	UE_LOG(LogStreaming, Display, TEXT("LinkerLoad_SerializeExportMap: %f"), LinkerLoad_SerializeExportMap);
+	UE_LOG(LogStreaming, Display, TEXT("LinkerLoad_StartTextureAllocation: %f"), LinkerLoad_StartTextureAllocation);
+	UE_LOG(LogStreaming, Display, TEXT("LinkerLoad_FixupImportMap: %f"), LinkerLoad_FixupImportMap);
+	UE_LOG(LogStreaming, Display, TEXT("LinkerLoad_RemapImports: %f"), LinkerLoad_RemapImports);
+	UE_LOG(LogStreaming, Display, TEXT("LinkerLoad_FixupExportMap: %f"), LinkerLoad_FixupExportMap);
+	UE_LOG(LogStreaming, Display, TEXT("LinkerLoad_SerializeDependsMap: %f"), LinkerLoad_SerializeDependsMap);
+	UE_LOG(LogStreaming, Display, TEXT("LinkerLoad_CreateExportHash: %f"), LinkerLoad_CreateExportHash);
+	UE_LOG(LogStreaming, Display, TEXT("LinkerLoad_FindExistingExports: %f"), LinkerLoad_FindExistingExports);
+	UE_LOG(LogStreaming, Display, TEXT("LinkerLoad_FinalizeCreation: %f"), LinkerLoad_FinalizeCreation);
+
+	UE_LOG(LogStreaming, Display, TEXT("Package_FinishLinker: %f"), Package_FinishLinker);
+	UE_LOG(LogStreaming, Display, TEXT("Package_LoadImports: %f"), Package_LoadImports);
+	UE_LOG(LogStreaming, Display, TEXT("Package_CreateImports: %f"), Package_CreateImports);
+	UE_LOG(LogStreaming, Display, TEXT("Package_CreateLinker: %f"), Package_CreateLinker);
+	UE_LOG(LogStreaming, Display, TEXT("Package_FinishTextureAllocations: %f"), Package_FinishTextureAllocations);
+	UE_LOG(LogStreaming, Display, TEXT("Package_CreateExports: %f"), Package_CreateExports);
+	UE_LOG(LogStreaming, Display, TEXT("Package_PreLoadObjects: %f"), Package_PreLoadObjects);
+	UE_LOG(LogStreaming, Display, TEXT("Package_PostLoadObjects: %f"), Package_PostLoadObjects);
+	UE_LOG(LogStreaming, Display, TEXT("Package_Tick: %f"), Package_Tick);
+	UE_LOG(LogStreaming, Display, TEXT("Package_CreateAsyncPackagesFromQueue: %f"), Package_CreateAsyncPackagesFromQueue);
+
+	UE_LOG(LogStreaming, Display, TEXT("TickAsyncLoading_ProcessLoadedPackages: %f"), TickAsyncLoading_ProcessLoadedPackages);
+
+
+	UE_LOG(LogStreaming, Display, TEXT("-------------------------------------------------"));
+
+#endif
+
+}
+
+void FLoadTimeTracker::ResetRawLoadTimes()
+{
+#if ENABLE_LOADTIME_RAW_TIMINGS
+	CreateAsyncPackagesFromQueueTime = 0.0;
+	ProcessAsyncLoadingTime = 0.0;
+	ProcessLoadedPackagesTime = 0.0;
+	SerializeTaggedPropertiesTime = 0.0;
+	CreateLinkerTime = 0.0;
+	FinishLinkerTime = 0.0;
+	CreateImportsTime = 0.0;
+	CreateExportsTime = 0.0;
+	PreLoadObjectsTime = 0.0;
+	PostLoadObjectsTime = 0.0;
+	PostLoadDeferredObjectsTime = 0.0;
+	FinishObjectsTime = 0.0;
+	MaterialPostLoad = 0.0;
+	MaterialInstancePostLoad = 0.0;
+	SerializeInlineShaderMaps = 0.0;
+	MaterialSerializeTime = 0.0;
+	MaterialInstanceSerializeTime = 0.0;
+	AsyncLoadingTime = 0.0;
+	LinkerLoad_CreateLoader = 0.0;
+	LinkerLoad_SerializePackageFileSummary = 0.0;
+	LinkerLoad_SerializeNameMap = 0.0;
+	LinkerLoad_SerializeGatherableTextDataMap = 0.0;
+	LinkerLoad_SerializeImportMap = 0.0;
+	LinkerLoad_SerializeExportMap = 0.0;
+	LinkerLoad_StartTextureAllocation = 0.0;
+	LinkerLoad_FixupImportMap = 0.0;
+	LinkerLoad_RemapImports = 0.0;
+	LinkerLoad_FixupExportMap = 0.0;
+	LinkerLoad_SerializeDependsMap = 0.0;
+	LinkerLoad_CreateExportHash = 0.0;
+	LinkerLoad_FindExistingExports = 0.0;
+	LinkerLoad_FinalizeCreation = 0.0;
+
+	Package_FinishLinker = 0.0;
+	Package_LoadImports = 0.0;
+	Package_CreateImports = 0.0;
+	Package_CreateLinker = 0.0;
+	Package_FinishTextureAllocations = 0.0;
+	Package_CreateExports = 0.0;
+	Package_PreLoadObjects = 0.0;
+	Package_PostLoadObjects = 0.0;
+	Package_Tick = 0.0;
+	Package_CreateAsyncPackagesFromQueue = 0.0;
+
+	TickAsyncLoading_ProcessLoadedPackages = 0.0;
+
+	LinkerLoad_SerializeNameMap_ProcessingEntries = 0.0;
+
+#endif
+
+}
+
 static FAutoConsoleCommand LoadTimerDumpCmd(
 	TEXT("LoadTimes.DumpTracking"),
 	TEXT("Dump high level load times being tracked"),
-	FConsoleCommandDelegate::CreateStatic(&FLoadTimeTracker::DumpLoadTimesStatic)
+	FConsoleCommandDelegate::CreateStatic(&FLoadTimeTracker::DumpHighLevelLoadTimesStatic)
 	);
-
-#endif
-
diff --git a/Engine/Source/Runtime/Core/Public/CoreGlobals.h b/Engine/Source/Runtime/Core/Public/CoreGlobals.h
--- a/Engine/Source/Runtime/Core/Public/CoreGlobals.h
+++ b/Engine/Source/Runtime/Core/Public/CoreGlobals.h
@@ -213,9 +213,6 @@ extern CORE_API int32 GPlayInEditorID;
 /** Whether or not PIE was attempting to play from PlayerStart */
 extern CORE_API bool GIsPIEUsingPlayerStart;
 
-/** Proxy class that allows verification on FApp::IsGame() accesses. */
-extern CORE_API bool IsInGameThread();
-
 /** true if the runtime needs textures to be powers of two */
 extern CORE_API bool GPlatformNeedsPowerOfTwoTextures;
 
diff --git a/Engine/Source/Runtime/Core/Public/Delegates/DelegateCombinations.h b/Engine/Source/Runtime/Core/Public/Delegates/DelegateCombinations.h
--- a/Engine/Source/Runtime/Core/Public/Delegates/DelegateCombinations.h
+++ b/Engine/Source/Runtime/Core/Public/Delegates/DelegateCombinations.h
@@ -3671,3 +3671,410 @@
 #undef FUNC_PARAM_PARMS_PASSIN
 #undef FUNC_PARAM_INITIALIZER_LIST
 #undef FUNC_IS_VOID
+
+#define FUNC_SUFFIX NineParams
+#define FUNC_RETVAL_TYPEDEF 
+#define FUNC_TEMPLATE_DECL RetValType, Param1Type, Param2Type, Param3Type, Param4Type, Param5Type, Param6Type, Param7Type, Param8Type, Param9Type
+#define FUNC_TEMPLATE_DECL_TYPENAME typename RetValType, typename Param1Type, typename Param2Type, typename Param3Type, typename Param4Type, typename Param5Type, typename Param6Type, typename Param7Type, typename Param8Type, typename Param9Type
+#define FUNC_TEMPLATE_DECL_NO_SHADOW typename RetValTypeNoShadow, typename Param1TypeNoShadow, typename Param2TypeNoShadow, typename Param3TypeNoShadow, typename Param4TypeNoShadow, typename Param5TypeNoShadow, typename Param6TypeNoShadow, typename Param7TypeNoShadow, typename Param8TypeNoShadow, typename Param9TypeNoShadow
+#define FUNC_TEMPLATE_ARGS RetValType, Param1Type, Param2Type, Param3Type, Param4Type, Param5Type, Param6Type, Param7Type, Param8Type, Param9Type
+#define FUNC_HAS_PARAMS 1
+#define FUNC_PARAM_LIST Param1Type InParam1, Param2Type InParam2, Param3Type InParam3, Param4Type InParam4, Param5Type InParam5, Param6Type InParam6, Param7Type InParam7, Param8Type InParam8, Param9Type InParam9
+#define FUNC_PARAM_MEMBERS Param1Type Param1; Param2Type Param2; Param3Type Param3; Param4Type Param4; Param5Type Param5; Param6Type Param6; Param7Type Param7; Param8Type Param8; Param9Type Param9;
+#define FUNC_PARAM_PASSTHRU InParam1, InParam2, InParam3, InParam4, InParam5, InParam6, InParam7, InParam8, InParam9
+#define FUNC_PARAM_PARMS_PASSIN Parms.Param1, Parms.Param2, Parms.Param3, Parms.Param4, Parms.Param5, Parms.Param6, Parms.Param7, Parms.Param8, Params.Param9
+#define FUNC_PARAM_INITIALIZER_LIST Param1( InParam1 ), Param2( InParam2 ), Param3( InParam3 ), Param4( InParam4 ), Param5( InParam5 ), Param6( InParam6 ), Param7( InParam7 ), Param8( InParam8 ), Param9( InParam9 )
+#define FUNC_IS_VOID 1
+#include "DelegateInstanceInterfaceImpl.inl"
+
+#define FUNC_HAS_PAYLOAD 0
+#define FUNC_PAYLOAD_SUFFIX NineParams
+#define FUNC_PAYLOAD_TEMPLATE_DECL RetValType, Param1Type, Param2Type, Param3Type, Param4Type, Param5Type, Param6Type, Param7Type, Param8Type, Param9Type
+#define FUNC_PAYLOAD_TEMPLATE_DECL_TYPENAME typename RetValType, typename Param1Type, typename Param2Type, typename Param3Type, typename Param4Type, typename Param5Type, typename Param6Type, typename Param7Type, typename Param8Type, typename Param9Type
+#define FUNC_PAYLOAD_TEMPLATE_DECL_NO_SHADOW typename RetValTypeNoShadow, typename Param1TypeNoShadow, typename Param2TypeNoShadow, typename Param3TypeNoShadow, typename Param4TypeNoShadow, typename Param5TypeNoShadow, typename Param6TypeNoShadow, typename Param7TypeNoShadow, typename Param8TypeNoShadow, typename Param9TypeNoShadow
+#define FUNC_PAYLOAD_TEMPLATE_ARGS RetValType, Param1Type, Param2Type, Param3Type, Param4Type, Param5Type, Param6Type, Param7Type, Param8Type, Param9Type
+#define FUNC_PAYLOAD_MEMBERS 
+#define FUNC_PAYLOAD_LIST
+#define FUNC_PAYLOAD_PASSTHRU
+#define FUNC_PAYLOAD_PASSIN
+#define FUNC_PAYLOAD_INITIALIZER_LIST
+#define FUNC_IS_CONST 0
+#define FUNC_CONST_SUFFIX 
+#include "DelegateInstancesImpl.inl"
+#undef FUNC_IS_CONST
+#undef FUNC_CONST_SUFFIX
+#define FUNC_IS_CONST 1
+#define FUNC_CONST_SUFFIX _Const
+#include "DelegateInstancesImpl.inl"
+#undef FUNC_IS_CONST
+#undef FUNC_CONST_SUFFIX
+
+#undef FUNC_HAS_PAYLOAD
+#undef FUNC_PAYLOAD_SUFFIX
+#undef FUNC_PAYLOAD_TEMPLATE_DECL
+#undef FUNC_PAYLOAD_TEMPLATE_DECL_TYPENAME
+#undef FUNC_PAYLOAD_TEMPLATE_DECL_NO_SHADOW
+#undef FUNC_PAYLOAD_TEMPLATE_ARGS
+#undef FUNC_PAYLOAD_MEMBERS
+#undef FUNC_PAYLOAD_LIST
+#undef FUNC_PAYLOAD_PASSTHRU
+#undef FUNC_PAYLOAD_PASSIN
+#undef FUNC_PAYLOAD_INITIALIZER_LIST
+
+#define FUNC_HAS_PAYLOAD 1
+#define FUNC_PAYLOAD_SUFFIX NineParams_OneVar
+#define FUNC_PAYLOAD_TEMPLATE_DECL RetValType, Param1Type, Param2Type, Param3Type, Param4Type, Param5Type, Param6Type, Param7Type, Param8Type, Param9Type, Var1Type
+#define FUNC_PAYLOAD_TEMPLATE_DECL_TYPENAME typename RetValType, typename Param1Type, typename Param2Type, typename Param3Type, typename Param4Type, typename Param5Type, typename Param6Type, typename Param7Type, typename Param8Type, typename Param9Type, typename Var1Type
+#define FUNC_PAYLOAD_TEMPLATE_DECL_NO_SHADOW typename RetValTypeNoShadow, typename Param1TypeNoShadow, typename Param2TypeNoShadow, typename Param3TypeNoShadow, typename Param4TypeNoShadow, typename Param5TypeNoShadow, typename Param6TypeNoShadow, typename Param7TypeNoShadow, typename Param8TypeNoShadow, typename Param9TypeNoShadow,typename Var1TypeNoShadow
+#define FUNC_PAYLOAD_TEMPLATE_ARGS RetValType, Param1Type, Param2Type, Param3Type, Param4Type, Param5Type, Param6Type, Param7Type, Param8Type, Param9Type, Var1Type
+#define FUNC_PAYLOAD_MEMBERS Var1Type Var1; 
+#define FUNC_PAYLOAD_LIST Var1Type InVar1
+#define FUNC_PAYLOAD_PASSTHRU InVar1
+#define FUNC_PAYLOAD_PASSIN Var1
+#define FUNC_PAYLOAD_INITIALIZER_LIST Var1( InVar1 )
+#define FUNC_IS_CONST 0
+#define FUNC_CONST_SUFFIX 
+#include "DelegateInstancesImpl.inl"
+#undef FUNC_IS_CONST
+#undef FUNC_CONST_SUFFIX
+#define FUNC_IS_CONST 1
+#define FUNC_CONST_SUFFIX _Const
+#include "DelegateInstancesImpl.inl"
+#undef FUNC_IS_CONST
+#undef FUNC_CONST_SUFFIX
+
+#undef FUNC_HAS_PAYLOAD
+#undef FUNC_PAYLOAD_SUFFIX
+#undef FUNC_PAYLOAD_TEMPLATE_DECL
+#undef FUNC_PAYLOAD_TEMPLATE_DECL_TYPENAME
+#undef FUNC_PAYLOAD_TEMPLATE_DECL_NO_SHADOW
+#undef FUNC_PAYLOAD_TEMPLATE_ARGS
+#undef FUNC_PAYLOAD_MEMBERS
+#undef FUNC_PAYLOAD_LIST
+#undef FUNC_PAYLOAD_PASSTHRU
+#undef FUNC_PAYLOAD_PASSIN
+#undef FUNC_PAYLOAD_INITIALIZER_LIST
+
+#define FUNC_HAS_PAYLOAD 1
+#define FUNC_PAYLOAD_SUFFIX NineParams_TwoVars
+#define FUNC_PAYLOAD_TEMPLATE_DECL RetValType, Param1Type, Param2Type, Param3Type, Param4Type, Param5Type, Param6Type, Param7Type, Param8Type, Param9Type, Var1Type, Var2Type
+#define FUNC_PAYLOAD_TEMPLATE_DECL_TYPENAME typename RetValType, typename Param1Type, typename Param2Type, typename Param3Type, typename Param4Type, typename Param5Type, typename Param6Type, typename Param7Type, typename Param8Type, typename Param9Type, typename Var1Type, typename Var2Type
+#define FUNC_PAYLOAD_TEMPLATE_DECL_NO_SHADOW typename RetValTypeNoShadow, typename Param1TypeNoShadow, typename Param2TypeNoShadow, typename Param3TypeNoShadow, typename Param4TypeNoShadow, typename Param5TypeNoShadow, typename Param6TypeNoShadow, typename Param7TypeNoShadow, typename Param8TypeNoShadow, typename Param9TypeNoShadow, typename Var1TypeNoShadow, typename Var2TypeNoShadow
+#define FUNC_PAYLOAD_TEMPLATE_ARGS RetValType, Param1Type, Param2Type, Param3Type, Param4Type, Param5Type, Param6Type, Param7Type, Param8Type, Param9Type, Var1Type, Var2Type
+#define FUNC_PAYLOAD_MEMBERS Var1Type Var1; Var2Type Var2; 
+#define FUNC_PAYLOAD_LIST Var1Type InVar1, Var2Type InVar2
+#define FUNC_PAYLOAD_PASSTHRU InVar1, InVar2
+#define FUNC_PAYLOAD_PASSIN Var1, Var2
+#define FUNC_PAYLOAD_INITIALIZER_LIST Var1( InVar1 ), Var2( InVar2 )
+#define FUNC_IS_CONST 0
+#define FUNC_CONST_SUFFIX 
+#include "DelegateInstancesImpl.inl"
+#undef FUNC_IS_CONST
+#undef FUNC_CONST_SUFFIX
+#define FUNC_IS_CONST 1
+#define FUNC_CONST_SUFFIX _Const
+#include "DelegateInstancesImpl.inl"
+#undef FUNC_IS_CONST
+#undef FUNC_CONST_SUFFIX
+
+#undef FUNC_HAS_PAYLOAD
+#undef FUNC_PAYLOAD_SUFFIX
+#undef FUNC_PAYLOAD_TEMPLATE_DECL
+#undef FUNC_PAYLOAD_TEMPLATE_DECL_TYPENAME
+#undef FUNC_PAYLOAD_TEMPLATE_DECL_NO_SHADOW
+#undef FUNC_PAYLOAD_TEMPLATE_ARGS
+#undef FUNC_PAYLOAD_MEMBERS
+#undef FUNC_PAYLOAD_LIST
+#undef FUNC_PAYLOAD_PASSTHRU
+#undef FUNC_PAYLOAD_PASSIN
+#undef FUNC_PAYLOAD_INITIALIZER_LIST
+
+#define FUNC_HAS_PAYLOAD 1
+#define FUNC_PAYLOAD_SUFFIX NineParams_ThreeVars
+#define FUNC_PAYLOAD_TEMPLATE_DECL RetValType, Param1Type, Param2Type, Param3Type, Param4Type, Param5Type, Param6Type, Param7Type, Param8Type, Param9Type, Var1Type, Var2Type, Var3Type
+#define FUNC_PAYLOAD_TEMPLATE_DECL_TYPENAME typename RetValType, typename Param1Type, typename Param2Type, typename Param3Type, typename Param4Type, typename Param5Type, typename Param6Type, typename Param7Type, typename Param8Type, typename Param9Type, typename Var1Type, typename Var2Type, typename Var3Type
+#define FUNC_PAYLOAD_TEMPLATE_DECL_NO_SHADOW typename RetValTypeNoShadow, typename Param1TypeNoShadow, typename Param2TypeNoShadow, typename Param3TypeNoShadow, typename Param4TypeNoShadow, typename Param5TypeNoShadow, typename Param6TypeNoShadow, typename Param7TypeNoShadow, typename Param8TypeNoShadow, typename Param9TypeNoShadow, typename Var1TypeNoShadow, typename Var2TypeNoShadow, typename Var3TypeNoShadow
+#define FUNC_PAYLOAD_TEMPLATE_ARGS RetValType, Param1Type, Param2Type, Param3Type, Param4Type, Param5Type, Param6Type, Param7Type, Param8Type, Param9Type, Var1Type, Var2Type, Var3Type
+#define FUNC_PAYLOAD_MEMBERS Var1Type Var1; Var2Type Var2; Var3Type Var3; 
+#define FUNC_PAYLOAD_LIST Var1Type InVar1, Var2Type InVar2, Var3Type InVar3
+#define FUNC_PAYLOAD_PASSTHRU InVar1, InVar2, InVar3
+#define FUNC_PAYLOAD_PASSIN Var1, Var2, Var3
+#define FUNC_PAYLOAD_INITIALIZER_LIST Var1( InVar1 ), Var2( InVar2 ), Var3( InVar3 )
+#define FUNC_IS_CONST 0
+#define FUNC_CONST_SUFFIX 
+#include "DelegateInstancesImpl.inl"
+#undef FUNC_IS_CONST
+#undef FUNC_CONST_SUFFIX
+#define FUNC_IS_CONST 1
+#define FUNC_CONST_SUFFIX _Const
+#include "DelegateInstancesImpl.inl"
+#undef FUNC_IS_CONST
+#undef FUNC_CONST_SUFFIX
+
+#undef FUNC_HAS_PAYLOAD
+#undef FUNC_PAYLOAD_SUFFIX
+#undef FUNC_PAYLOAD_TEMPLATE_DECL
+#undef FUNC_PAYLOAD_TEMPLATE_DECL_TYPENAME
+#undef FUNC_PAYLOAD_TEMPLATE_DECL_NO_SHADOW
+#undef FUNC_PAYLOAD_TEMPLATE_ARGS
+#undef FUNC_PAYLOAD_MEMBERS
+#undef FUNC_PAYLOAD_LIST
+#undef FUNC_PAYLOAD_PASSTHRU
+#undef FUNC_PAYLOAD_PASSIN
+#undef FUNC_PAYLOAD_INITIALIZER_LIST
+
+#define FUNC_HAS_PAYLOAD 1
+#define FUNC_PAYLOAD_SUFFIX NineParams_FourVars
+#define FUNC_PAYLOAD_TEMPLATE_DECL RetValType, Param1Type, Param2Type, Param3Type, Param4Type, Param5Type, Param6Type, Param7Type, Param8Type, Param9Type, Var1Type, Var2Type, Var3Type, Var4Type
+#define FUNC_PAYLOAD_TEMPLATE_DECL_TYPENAME typename RetValType, typename Param1Type, typename Param2Type, typename Param3Type, typename Param4Type, typename Param5Type, typename Param6Type, typename Param7Type, typename Param8Type, typename Param9Type, typename Var1Type, typename Var2Type, typename Var3Type, typename Var4Type
+#define FUNC_PAYLOAD_TEMPLATE_DECL_NO_SHADOW typename RetValTypeNoShadow, typename Param1TypeNoShadow, typename Param2TypeNoShadow, typename Param3TypeNoShadow, typename Param4TypeNoShadow, typename Param5TypeNoShadow, typename Param6TypeNoShadow, typename Param7TypeNoShadow, typename Param8TypeNoShadow, typename Param9TypeNoShadow, typename Var1TypeNoShadow, typename Var2TypeNoShadow, typename Var3TypeNoShadow, typename Var4TypeNoShadow
+#define FUNC_PAYLOAD_TEMPLATE_ARGS RetValType, Param1Type, Param2Type, Param3Type, Param4Type, Param5Type, Param6Type, Param7Type, Param8Type, Param9Type, Var1Type, Var2Type, Var3Type, Var4Type
+#define FUNC_PAYLOAD_MEMBERS Var1Type Var1; Var2Type Var2; Var3Type Var3; Var4Type Var4; 
+#define FUNC_PAYLOAD_LIST Var1Type InVar1, Var2Type InVar2, Var3Type InVar3, Var4Type InVar4
+#define FUNC_PAYLOAD_PASSTHRU InVar1, InVar2, InVar3, InVar4
+#define FUNC_PAYLOAD_PASSIN Var1, Var2, Var3, Var4
+#define FUNC_PAYLOAD_INITIALIZER_LIST Var1( InVar1 ), Var2( InVar2 ), Var3( InVar3 ), Var4( InVar4 )
+#define FUNC_IS_CONST 0
+#define FUNC_CONST_SUFFIX 
+#include "DelegateInstancesImpl.inl"
+#undef FUNC_IS_CONST
+#undef FUNC_CONST_SUFFIX
+#define FUNC_IS_CONST 1
+#define FUNC_CONST_SUFFIX _Const
+#include "DelegateInstancesImpl.inl"
+#undef FUNC_IS_CONST
+#undef FUNC_CONST_SUFFIX
+
+#undef FUNC_HAS_PAYLOAD
+#undef FUNC_PAYLOAD_SUFFIX
+#undef FUNC_PAYLOAD_TEMPLATE_DECL
+#undef FUNC_PAYLOAD_TEMPLATE_DECL_TYPENAME
+#undef FUNC_PAYLOAD_TEMPLATE_DECL_NO_SHADOW
+#undef FUNC_PAYLOAD_TEMPLATE_ARGS
+#undef FUNC_PAYLOAD_MEMBERS
+#undef FUNC_PAYLOAD_LIST
+#undef FUNC_PAYLOAD_PASSTHRU
+#undef FUNC_PAYLOAD_PASSIN
+#undef FUNC_PAYLOAD_INITIALIZER_LIST
+
+#include "DelegateSignatureImpl.inl"
+
+#define DECLARE_DELEGATE_NineParams( DelegateName, Param1Type, Param2Type, Param3Type, Param4Type, Param5Type, Param6Type, Param7Type, Param8Type, Param9Type ) FUNC_DECLARE_DELEGATE( NineParams, DelegateName, void, Param1Type, Param2Type, Param3Type, Param4Type, Param5Type, Param6Type, Param7Type, Param8Type, Param9Type )
+#define DECLARE_MULTICAST_DELEGATE_NineParams( DelegateName, Param1Type, Param2Type, Param3Type, Param4Type, Param5Type, Param6Type, Param7Type, Param8Type, Param9Type ) FUNC_DECLARE_MULTICAST_DELEGATE( NineParams, DelegateName, void, Param1Type, Param2Type, Param3Type, Param4Type, Param5Type, Param6Type, Param7Type, Param8Type, Param9Type )
+#define DECLARE_EVENT_NineParams( OwningType, EventName, Param1Type, Param2Type, Param3Type, Param4Type, Param5Type, Param6Type, Param7Type, Param8Type, Param9Type ) FUNC_DECLARE_EVENT( OwningType, EventName, NineParams, void, Param1Type, Param2Type, Param3Type, Param4Type, Param5Type, Param6Type, Param7Type, Param8Type, Param9Type )
+#define DECLARE_DYNAMIC_DELEGATE_NineParams( DelegateName, Param1Type, Param1Name, Param2Type, Param2Name, Param3Type, Param3Name, Param4Type, Param4Name, Param5Type, Param5Name, Param6Type, Param6Name, Param7Type, Param7Name, Param8Type, Param8Name, Param9Type, Param9Name ) BODY_MACRO_COMBINE(CURRENT_FILE_ID,_,__LINE__,_DELEGATE) FUNC_DECLARE_DYNAMIC_DELEGATE( FWeakObjectPtr, NineParams, DelegateName, DelegateName##_DelegateWrapper, FUNC_CONCAT( Param1Type InParam1, Param2Type InParam2, Param3Type InParam3, Param4Type InParam4, Param5Type InParam5, Param6Type InParam6, Param7Type InParam7, Param8Type InParam8, Param9Type InParam9 ), FUNC_CONCAT( *this, InParam1, InParam2, InParam3, InParam4, InParam5, InParam6, InParam7, InParam8, InParam9 ), void, Param1Type, Param2Type, Param3Type, Param4Type, Param5Type, Param6Type, Param7Type, Param8Type, Param9Type )
+#define DECLARE_DYNAMIC_MULTICAST_DELEGATE_NineParams( DelegateName, Param1Type, Param1Name, Param2Type, Param2Name, Param3Type, Param3Name, Param4Type, Param4Name, Param5Type, Param5Name, Param6Type, Param6Name, Param7Type, Param7Name, Param8Type, Param8Name, Param9Type, Param9Name ) BODY_MACRO_COMBINE(CURRENT_FILE_ID,_,__LINE__,_DELEGATE) FUNC_DECLARE_DYNAMIC_MULTICAST_DELEGATE( FWeakObjectPtr, NineParams, DelegateName, DelegateName##_DelegateWrapper, FUNC_CONCAT( Param1Type InParam1, Param2Type InParam2, Param3Type InParam3, Param4Type InParam4, Param5Type InParam5, Param6Type InParam6, Param7Type InParam7, Param8Type InParam8, Param9Type InParam9 ), FUNC_CONCAT( *this, InParam1, InParam2, InParam3, InParam4, InParam5, InParam6, InParam7, InParam8, InParam9 ), void, Param1Type, Param2Type, Param3Type, Param4Type, Param5Type, Param6Type, Param7Type, Param8Type, Param9Type )
+#undef FUNC_SUFFIX
+#undef FUNC_TEMPLATE_DECL
+#undef FUNC_TEMPLATE_DECL_TYPENAME
+#undef FUNC_TEMPLATE_DECL_NO_SHADOW
+#undef FUNC_TEMPLATE_ARGS
+#undef FUNC_HAS_PARAMS
+#undef FUNC_PARAM_LIST
+#undef FUNC_PARAM_MEMBERS
+#undef FUNC_PARAM_PASSTHRU
+#undef FUNC_PARAM_PARMS_PASSIN
+#undef FUNC_PARAM_INITIALIZER_LIST
+#undef FUNC_IS_VOID
+
+#define FUNC_SUFFIX RetVal_NineParams
+#define FUNC_RETVAL_TYPEDEF  typedef RetValType RetValType;
+#define FUNC_TEMPLATE_DECL RetValType, Param1Type, Param2Type, Param3Type, Param4Type, Param5Type, Param6Type, Param7Type, Param8Type, Param9Type
+#define FUNC_TEMPLATE_DECL_TYPENAME typename RetValType, typename Param1Type, typename Param2Type, typename Param3Type, typename Param4Type, typename Param5Type, typename Param6Type, typename Param7Type, typename Param8Type, typename Param9Type
+#define FUNC_TEMPLATE_DECL_NO_SHADOW typename RetValTypeNoShadow, typename Param1TypeNoShadow, typename Param2TypeNoShadow, typename Param3TypeNoShadow, typename Param4TypeNoShadow, typename Param5TypeNoShadow, typename Param6TypeNoShadow, typename Param7TypeNoShadow, typename Param8TypeNoShadow, typename Param9TypeNoShadow
+#define FUNC_TEMPLATE_ARGS RetValType, Param1Type, Param2Type, Param3Type, Param4Type, Param5Type, Param6Type, Param7Type, Param8Type, Param9Type
+#define FUNC_HAS_PARAMS 1
+#define FUNC_PARAM_LIST Param1Type InParam1, Param2Type InParam2, Param3Type InParam3, Param4Type InParam4, Param5Type InParam5, Param6Type InParam6, Param7Type InParam7, Param8Type InParam8, Param9Type InParam9
+#define FUNC_PARAM_MEMBERS Param1Type Param1; Param2Type Param2; Param3Type Param3; Param4Type Param4; Param5Type Param5; Param6Type Param6; Param7Type Param7; Param8Type Param8; Param9Type Param9;
+#define FUNC_PARAM_PASSTHRU InParam1, InParam2, InParam3, InParam4, InParam5, InParam6, InParam7, InParam8, InParam9
+#define FUNC_PARAM_PARMS_PASSIN Parms.Param1, Parms.Param2, Parms.Param3, Parms.Param4, Parms.Param5, Parms.Param6, Parms.Param7, Parms.Param8, Parms.Param9
+#define FUNC_PARAM_INITIALIZER_LIST Param1( InParam1 ), Param2( InParam2 ), Param3( InParam3 ), Param4( InParam4 ), Param5( InParam5 ), Param6( InParam6 ), Param7( InParam7 ), Param8( InParam8 ), Param9( InParam9 )
+#define FUNC_IS_VOID 0
+#include "DelegateInstanceInterfaceImpl.inl"
+
+#define FUNC_HAS_PAYLOAD 0
+#define FUNC_PAYLOAD_SUFFIX RetVal_NineParams
+#define FUNC_PAYLOAD_TEMPLATE_DECL RetValType, Param1Type, Param2Type, Param3Type, Param4Type, Param5Type, Param6Type, Param7Type, Param8Type, Param9Type
+#define FUNC_PAYLOAD_TEMPLATE_DECL_TYPENAME typename RetValType, typename Param1Type, typename Param2Type, typename Param3Type, typename Param4Type, typename Param5Type, typename Param6Type, typename Param7Type, typename Param8Type, typename Param9Type
+#define FUNC_PAYLOAD_TEMPLATE_DECL_NO_SHADOW typename RetValTypeNoShadow, typename Param1TypeNoShadow, typename Param2TypeNoShadow, typename Param3TypeNoShadow, typename Param4TypeNoShadow, typename Param5TypeNoShadow, typename Param6TypeNoShadow, typename Param7TypeNoShadow, typename Param8TypeNoShadow, typename Param9TypeNoShadow
+#define FUNC_PAYLOAD_TEMPLATE_ARGS RetValType, Param1Type, Param2Type, Param3Type, Param4Type, Param5Type, Param6Type, Param7Type, Param8Type, Param9Type
+#define FUNC_PAYLOAD_MEMBERS 
+#define FUNC_PAYLOAD_LIST
+#define FUNC_PAYLOAD_PASSTHRU
+#define FUNC_PAYLOAD_PASSIN
+#define FUNC_PAYLOAD_INITIALIZER_LIST
+#define FUNC_IS_CONST 0
+#define FUNC_CONST_SUFFIX 
+#include "DelegateInstancesImpl.inl"
+#undef FUNC_IS_CONST
+#undef FUNC_CONST_SUFFIX
+#define FUNC_IS_CONST 1
+#define FUNC_CONST_SUFFIX _Const
+#include "DelegateInstancesImpl.inl"
+#undef FUNC_IS_CONST
+#undef FUNC_CONST_SUFFIX
+
+#undef FUNC_HAS_PAYLOAD
+#undef FUNC_PAYLOAD_SUFFIX
+#undef FUNC_PAYLOAD_TEMPLATE_DECL
+#undef FUNC_PAYLOAD_TEMPLATE_DECL_TYPENAME
+#undef FUNC_PAYLOAD_TEMPLATE_DECL_NO_SHADOW
+#undef FUNC_PAYLOAD_TEMPLATE_ARGS
+#undef FUNC_PAYLOAD_MEMBERS
+#undef FUNC_PAYLOAD_LIST
+#undef FUNC_PAYLOAD_PASSTHRU
+#undef FUNC_PAYLOAD_PASSIN
+#undef FUNC_PAYLOAD_INITIALIZER_LIST
+
+#define FUNC_HAS_PAYLOAD 1
+#define FUNC_PAYLOAD_SUFFIX RetVal_NineParams_OneVar
+#define FUNC_PAYLOAD_TEMPLATE_DECL RetValType, Param1Type, Param2Type, Param3Type, Param4Type, Param5Type, Param6Type, Param7Type, Param8Type, Param9Type, Var1Type
+#define FUNC_PAYLOAD_TEMPLATE_DECL_TYPENAME typename RetValType, typename Param1Type, typename Param2Type, typename Param3Type, typename Param4Type, typename Param5Type, typename Param6Type, typename Param7Type, typename Param8Type, typename Param9Type, typename Var1Type
+#define FUNC_PAYLOAD_TEMPLATE_DECL_NO_SHADOW typename RetValTypeNoShadow, typename Param1TypeNoShadow, typename Param2TypeNoShadow, typename Param3TypeNoShadow, typename Param4TypeNoShadow, typename Param5TypeNoShadow, typename Param6TypeNoShadow, typename Param7TypeNoShadow, typename Param8TypeNoShadow, typename Param9TypeNoShadow, typename Var1TypeNoShadow
+#define FUNC_PAYLOAD_TEMPLATE_ARGS RetValType, Param1Type, Param2Type, Param3Type, Param4Type, Param5Type, Param6Type, Param7Type, Param8Type, Param9Type, Var1Type
+#define FUNC_PAYLOAD_MEMBERS Var1Type Var1; 
+#define FUNC_PAYLOAD_LIST Var1Type InVar1
+#define FUNC_PAYLOAD_PASSTHRU InVar1
+#define FUNC_PAYLOAD_PASSIN Var1
+#define FUNC_PAYLOAD_INITIALIZER_LIST Var1( InVar1 )
+#define FUNC_IS_CONST 0
+#define FUNC_CONST_SUFFIX 
+#include "DelegateInstancesImpl.inl"
+#undef FUNC_IS_CONST
+#undef FUNC_CONST_SUFFIX
+#define FUNC_IS_CONST 1
+#define FUNC_CONST_SUFFIX _Const
+#include "DelegateInstancesImpl.inl"
+#undef FUNC_IS_CONST
+#undef FUNC_CONST_SUFFIX
+
+#undef FUNC_HAS_PAYLOAD
+#undef FUNC_PAYLOAD_SUFFIX
+#undef FUNC_PAYLOAD_TEMPLATE_DECL
+#undef FUNC_PAYLOAD_TEMPLATE_DECL_TYPENAME
+#undef FUNC_PAYLOAD_TEMPLATE_DECL_NO_SHADOW
+#undef FUNC_PAYLOAD_TEMPLATE_ARGS
+#undef FUNC_PAYLOAD_MEMBERS
+#undef FUNC_PAYLOAD_LIST
+#undef FUNC_PAYLOAD_PASSTHRU
+#undef FUNC_PAYLOAD_PASSIN
+#undef FUNC_PAYLOAD_INITIALIZER_LIST
+
+#define FUNC_HAS_PAYLOAD 1
+#define FUNC_PAYLOAD_SUFFIX RetVal_NineParams_TwoVars
+#define FUNC_PAYLOAD_TEMPLATE_DECL RetValType, Param1Type, Param2Type, Param3Type, Param4Type, Param5Type, Param6Type, Param7Type, Param8Type, Param9Type, Var1Type, Var2Type
+#define FUNC_PAYLOAD_TEMPLATE_DECL_TYPENAME typename RetValType, typename Param1Type, typename Param2Type, typename Param3Type, typename Param4Type, typename Param5Type, typename Param6Type, typename Param7Type, typename Param8Type, typename Param9Type, typename Var1Type, typename Var2Type
+#define FUNC_PAYLOAD_TEMPLATE_DECL_NO_SHADOW typename RetValTypeNoShadow, typename Param1TypeNoShadow, typename Param2TypeNoShadow, typename Param3TypeNoShadow, typename Param4TypeNoShadow, typename Param5TypeNoShadow, typename Param6TypeNoShadow, typename Param7TypeNoShadow, typename Param8TypeNoShadow, typename Param9TypeNoShadow, typename Var1TypeNoShadow, typename Var2TypeNoShadow
+#define FUNC_PAYLOAD_TEMPLATE_ARGS RetValType, Param1Type, Param2Type, Param3Type, Param4Type, Param5Type, Param6Type, Param7Type, Param8Type, Param9Type, Var1Type, Var2Type
+#define FUNC_PAYLOAD_MEMBERS Var1Type Var1; Var2Type Var2; 
+#define FUNC_PAYLOAD_LIST Var1Type InVar1, Var2Type InVar2
+#define FUNC_PAYLOAD_PASSTHRU InVar1, InVar2
+#define FUNC_PAYLOAD_PASSIN Var1, Var2
+#define FUNC_PAYLOAD_INITIALIZER_LIST Var1( InVar1 ), Var2( InVar2 )
+#define FUNC_IS_CONST 0
+#define FUNC_CONST_SUFFIX 
+#include "DelegateInstancesImpl.inl"
+#undef FUNC_IS_CONST
+#undef FUNC_CONST_SUFFIX
+#define FUNC_IS_CONST 1
+#define FUNC_CONST_SUFFIX _Const
+#include "DelegateInstancesImpl.inl"
+#undef FUNC_IS_CONST
+#undef FUNC_CONST_SUFFIX
+
+#undef FUNC_HAS_PAYLOAD
+#undef FUNC_PAYLOAD_SUFFIX
+#undef FUNC_PAYLOAD_TEMPLATE_DECL
+#undef FUNC_PAYLOAD_TEMPLATE_DECL_TYPENAME
+#undef FUNC_PAYLOAD_TEMPLATE_DECL_NO_SHADOW
+#undef FUNC_PAYLOAD_TEMPLATE_ARGS
+#undef FUNC_PAYLOAD_MEMBERS
+#undef FUNC_PAYLOAD_LIST
+#undef FUNC_PAYLOAD_PASSTHRU
+#undef FUNC_PAYLOAD_PASSIN
+#undef FUNC_PAYLOAD_INITIALIZER_LIST
+
+#define FUNC_HAS_PAYLOAD 1
+#define FUNC_PAYLOAD_SUFFIX RetVal_NineParams_ThreeVars
+#define FUNC_PAYLOAD_TEMPLATE_DECL RetValType, Param1Type, Param2Type, Param3Type, Param4Type, Param5Type, Param6Type, Param7Type, Param8Type, Param9Type, Var1Type, Var2Type, Var3Type
+#define FUNC_PAYLOAD_TEMPLATE_DECL_TYPENAME typename RetValType, typename Param1Type, typename Param2Type, typename Param3Type, typename Param4Type, typename Param5Type, typename Param6Type, typename Param7Type, typename Param8Type, typename Param9Type, typename Var1Type, typename Var2Type, typename Var3Type
+#define FUNC_PAYLOAD_TEMPLATE_DECL_NO_SHADOW typename RetValTypeNoShadow, typename Param1TypeNoShadow, typename Param2TypeNoShadow, typename Param3TypeNoShadow, typename Param4TypeNoShadow, typename Param5TypeNoShadow, typename Param6TypeNoShadow, typename Param7TypeNoShadow, typename Param8TypeNoShadow, typename Param9TypeNoShadow, typename Var1TypeNoShadow, typename Var2TypeNoShadow, typename Var3TypeNoShadow
+#define FUNC_PAYLOAD_TEMPLATE_ARGS RetValType, Param1Type, Param2Type, Param3Type, Param4Type, Param5Type, Param6Type, Param7Type, Param8Type, Param9Type, Var1Type, Var2Type, Var3Type
+#define FUNC_PAYLOAD_MEMBERS Var1Type Var1; Var2Type Var2; Var3Type Var3; 
+#define FUNC_PAYLOAD_LIST Var1Type InVar1, Var2Type InVar2, Var3Type InVar3
+#define FUNC_PAYLOAD_PASSTHRU InVar1, InVar2, InVar3
+#define FUNC_PAYLOAD_PASSIN Var1, Var2, Var3
+#define FUNC_PAYLOAD_INITIALIZER_LIST Var1( InVar1 ), Var2( InVar2 ), Var3( InVar3 )
+#define FUNC_IS_CONST 0
+#define FUNC_CONST_SUFFIX 
+#include "DelegateInstancesImpl.inl"
+#undef FUNC_IS_CONST
+#undef FUNC_CONST_SUFFIX
+#define FUNC_IS_CONST 1
+#define FUNC_CONST_SUFFIX _Const
+#include "DelegateInstancesImpl.inl"
+#undef FUNC_IS_CONST
+#undef FUNC_CONST_SUFFIX
+
+#undef FUNC_HAS_PAYLOAD
+#undef FUNC_PAYLOAD_SUFFIX
+#undef FUNC_PAYLOAD_TEMPLATE_DECL
+#undef FUNC_PAYLOAD_TEMPLATE_DECL_TYPENAME
+#undef FUNC_PAYLOAD_TEMPLATE_DECL_NO_SHADOW
+#undef FUNC_PAYLOAD_TEMPLATE_ARGS
+#undef FUNC_PAYLOAD_MEMBERS
+#undef FUNC_PAYLOAD_LIST
+#undef FUNC_PAYLOAD_PASSTHRU
+#undef FUNC_PAYLOAD_PASSIN
+#undef FUNC_PAYLOAD_INITIALIZER_LIST
+
+#define FUNC_HAS_PAYLOAD 1
+#define FUNC_PAYLOAD_SUFFIX RetVal_NineParams_FourVars
+#define FUNC_PAYLOAD_TEMPLATE_DECL RetValType, Param1Type, Param2Type, Param3Type, Param4Type, Param5Type, Param6Type, Param7Type, Param8Type, Param9Type, Var1Type, Var2Type, Var3Type, Var4Type
+#define FUNC_PAYLOAD_TEMPLATE_DECL_TYPENAME typename RetValType, typename Param1Type, typename Param2Type, typename Param3Type, typename Param4Type, typename Param5Type, typename Param6Type, typename Param7Type, typename Param8Type, typename Param9Type, typename Var1Type, typename Var2Type, typename Var3Type, typename Var4Type
+#define FUNC_PAYLOAD_TEMPLATE_DECL_NO_SHADOW typename RetValTypeNoShadow, typename Param1TypeNoShadow, typename Param2TypeNoShadow, typename Param3TypeNoShadow, typename Param4TypeNoShadow, typename Param5TypeNoShadow, typename Param6TypeNoShadow, typename Param7TypeNoShadow, typename Param8TypeNoShadow, type Param9TypeNoShadow, typename Var1TypeNoShadow, typename Var2TypeNoShadow, typename Var3TypeNoShadow, typename Var4TypeNoShadow
+#define FUNC_PAYLOAD_TEMPLATE_ARGS RetValType, Param1Type, Param2Type, Param3Type, Param4Type, Param5Type, Param6Type, Param7Type, Param8Type, Param9Type, Var1Type, Var2Type, Var3Type, Var4Type
+#define FUNC_PAYLOAD_MEMBERS Var1Type Var1; Var2Type Var2; Var3Type Var3; Var4Type Var4; 
+#define FUNC_PAYLOAD_LIST Var1Type InVar1, Var2Type InVar2, Var3Type InVar3, Var4Type InVar4
+#define FUNC_PAYLOAD_PASSTHRU InVar1, InVar2, InVar3, InVar4
+#define FUNC_PAYLOAD_PASSIN Var1, Var2, Var3, Var4
+#define FUNC_PAYLOAD_INITIALIZER_LIST Var1( InVar1 ), Var2( InVar2 ), Var3( InVar3 ), Var4( InVar4 )
+#define FUNC_IS_CONST 0
+#define FUNC_CONST_SUFFIX 
+#include "DelegateInstancesImpl.inl"
+#undef FUNC_IS_CONST
+#undef FUNC_CONST_SUFFIX
+#define FUNC_IS_CONST 1
+#define FUNC_CONST_SUFFIX _Const
+#include "DelegateInstancesImpl.inl"
+#undef FUNC_IS_CONST
+#undef FUNC_CONST_SUFFIX
+
+#undef FUNC_HAS_PAYLOAD
+#undef FUNC_PAYLOAD_SUFFIX
+#undef FUNC_PAYLOAD_TEMPLATE_DECL
+#undef FUNC_PAYLOAD_TEMPLATE_DECL_TYPENAME
+#undef FUNC_PAYLOAD_TEMPLATE_DECL_NO_SHADOW
+#undef FUNC_PAYLOAD_TEMPLATE_ARGS
+#undef FUNC_PAYLOAD_MEMBERS
+#undef FUNC_PAYLOAD_LIST
+#undef FUNC_PAYLOAD_PASSTHRU
+#undef FUNC_PAYLOAD_PASSIN
+#undef FUNC_PAYLOAD_INITIALIZER_LIST
+
+#include "DelegateSignatureImpl.inl"
+
+#define DECLARE_DELEGATE_RetVal_NineParams( RetValType, DelegateName, Param1Type, Param2Type, Param3Type, Param4Type, Param5Type, Param6Type, Param7Type, Param8Type, Param9Type ) FUNC_DECLARE_DELEGATE( RetVal_NineParams, DelegateName, RetValType, Param1Type, Param2Type, Param3Type, Param4Type, Param5Type, Param6Type, Param7Type, Param8Type, Param9Type )
+#define DECLARE_DYNAMIC_DELEGATE_RetVal_NineParams( RetValType, DelegateName, Param1Type, Param1Name, Param2Type, Param2Name, Param3Type, Param3Name, Param4Type, Param4Name, Param5Type, Param5Name, Param6Type, Param6Name, Param7Type, Param7Name, Param8Type, Param8Name, Param9Type, Param9Name ) BODY_MACRO_COMBINE(CURRENT_FILE_ID,_,__LINE__,_DELEGATE) FUNC_DECLARE_DYNAMIC_DELEGATE_RETVAL( FWeakObjectPtr, RetVal_NineParams, DelegateName, DelegateName##_DelegateWrapper, RetValType, FUNC_CONCAT( Param1Type InParam1, Param2Type InParam2, Param3Type InParam3, Param4Type InParam4, Param5Type InParam5, Param6Type InParam6, Param7Type InParam7, Param8Type InParam8, Param9Type InParam9 ), FUNC_CONCAT( *this, InParam1, InParam2, InParam3, InParam4, InParam5, InParam6, InParam7, InParam8, InParam9 ), RetValType, Param1Type, Param2Type, Param3Type, Param4Type, Param5Type, Param6Type, Param7Type, Param8Type, Param9Type )
+#undef FUNC_SUFFIX
+#undef FUNC_TEMPLATE_DECL
+#undef FUNC_TEMPLATE_DECL_TYPENAME
+#undef FUNC_TEMPLATE_DECL_NO_SHADOW
+#undef FUNC_TEMPLATE_ARGS
+#undef FUNC_HAS_PARAMS
+#undef FUNC_PARAM_LIST
+#undef FUNC_PARAM_MEMBERS
+#undef FUNC_PARAM_PASSTHRU
+#undef FUNC_PARAM_PARMS_PASSIN
+#undef FUNC_PARAM_INITIALIZER_LIST
+#undef FUNC_IS_VOID
\ No newline at end of file
diff --git a/Engine/Source/Runtime/Core/Public/HAL/ThreadingBase.h b/Engine/Source/Runtime/Core/Public/HAL/ThreadingBase.h
--- a/Engine/Source/Runtime/Core/Public/HAL/ThreadingBase.h
+++ b/Engine/Source/Runtime/Core/Public/HAL/ThreadingBase.h
@@ -1207,7 +1207,16 @@ private:
 
 
 /** @return True if called from the game thread. */
-extern CORE_API bool IsInGameThread();
+FORCEINLINE bool IsInGameThread()
+{
+	if(GIsGameThreadIdInitialized)
+	{
+		const uint32 CurrentThreadId = FPlatformTLS::GetCurrentThreadId();
+		return CurrentThreadId == GGameThreadId || CurrentThreadId == GSlateLoadingThreadId;
+	}
+
+	return true;
+}
 
 /** @return True if called from the slate thread, and not merely a thread calling slate functions. */
 extern CORE_API bool IsInSlateThread();
diff --git a/Engine/Source/Runtime/Core/Public/Math/IntVector.h b/Engine/Source/Runtime/Core/Public/Math/IntVector.h
--- a/Engine/Source/Runtime/Core/Public/Math/IntVector.h
+++ b/Engine/Source/Runtime/Core/Public/Math/IntVector.h
@@ -194,9 +194,9 @@ public:
 	float GetMin() const;
 
 	/**
-	 * Gets the distance of this point from (0,0).
+	 * Gets the distance of this point from (0,0,0).
 	 *
-	 * @return The distance of this point from (0,0).
+	 * @return The distance of this point from (0,0,0).
 	 */
 	int32 Size() const;
 
@@ -403,7 +403,7 @@ FORCEINLINE float FIntVector::GetMin() const
 
 FORCEINLINE int32 FIntVector::Num()
 {
-	return 2;
+	return 3;
 }
 
 
diff --git a/Engine/Source/Runtime/Core/Public/Math/Quat.h b/Engine/Source/Runtime/Core/Public/Math/Quat.h
--- a/Engine/Source/Runtime/Core/Public/Math/Quat.h
+++ b/Engine/Source/Runtime/Core/Public/Math/Quat.h
@@ -130,8 +130,8 @@ public:
 	 * Checks whether another Quaternion is equal to this, within specified tolerance.
 	 *
 	 * @param Q The other Quaternion.
-	 * @param Tolerance Error Tolerance.
-	 * @return true if two Quaternion are equal, within specified tolerance, otherwise false.
+	 * @param Tolerance Error tolerance for comparison with other Quaternion.
+	 * @return true if two Quaternions are equal, within specified tolerance, otherwise false.
 	 */
 	FORCEINLINE bool Equals(const FQuat& Q, float Tolerance=KINDA_SMALL_NUMBER) const;
 
@@ -139,9 +139,10 @@ public:
 	 * Checks whether this Quaternion is an Identity Quaternion.
 	 * Assumes Quaternion tested is normalized.
 	 *
+	 * @param Tolerance Error tolerance for comparison with Identity Quaternion.
 	 * @return true if Quaternion is a normalized Identity Quaternion.
 	 */
-	FORCEINLINE bool IsIdentity() const;
+	FORCEINLINE bool IsIdentity(float Tolerance=SMALL_NUMBER) const;
 
 	/**
 	 * Subtracts another quaternion from this.
@@ -680,7 +681,7 @@ FORCEINLINE FQuat::FQuat( const FQuat& Q )
 
 FORCEINLINE FString FQuat::ToString() const
 {
-	return FString::Printf(TEXT("X=%.6f Y=%.6f Z=%.6f W=%.6f"), X, Y, Z, W);
+	return FString::Printf(TEXT("X=%.9f Y=%.9f Z=%.9f W=%.9f"), X, Y, Z, W);
 }
 
 
@@ -753,9 +754,9 @@ FORCEINLINE bool FQuat::Equals(const FQuat& Q, float Tolerance) const
 #endif // PLATFORM_ENABLE_VECTORINTRINSICS
 }
 
-FORCEINLINE bool FQuat::IsIdentity() const
+FORCEINLINE bool FQuat::IsIdentity(float Tolerance) const
 {
-	return ((W * W) > ((1.f - DELTA) * (1.f - DELTA)));
+	return Equals(FQuat::Identity, Tolerance);
 }
 
 FORCEINLINE FQuat FQuat::operator-=(const FQuat& Q)
diff --git a/Engine/Source/Runtime/Core/Public/Math/Transform.h b/Engine/Source/Runtime/Core/Public/Math/Transform.h
--- a/Engine/Source/Runtime/Core/Public/Math/Transform.h
+++ b/Engine/Source/Runtime/Core/Public/Math/Transform.h
@@ -569,7 +569,7 @@ public:
 	FORCEINLINE FVector		GetScaledAxis(EAxis::Type InAxis) const;
 	FORCEINLINE FVector		GetUnitAxis(EAxis::Type InAxis) const;
 	FORCEINLINE void		Mirror(EAxis::Type MirrorAxis, EAxis::Type FlipAxis);
-	static FORCEINLINE FVector		GetSafeScaleReciprocal(const FVector& InScale);
+	FORCEINLINE FVector		GetSafeScaleReciprocal(const FVector& InScale, float Tolerance=SMALL_NUMBER) const;
 
 	// temp function for easy conversion
 	FORCEINLINE FVector GetLocation() const
@@ -1435,15 +1435,10 @@ inline float FTransform::GetMinimumAxisScale() const
 // anymore because you should be instead of showing gigantic infinite mesh
 // also returning BIG_NUMBER causes sequential NaN issues by multiplying 
 // so we hardcode as 0
-FORCEINLINE FVector FTransform::GetSafeScaleReciprocal(const FVector& InScale)
+FORCEINLINE FVector FTransform::GetSafeScaleReciprocal(const FVector& InScale, float Tolerance) const
 {
 	FVector SafeReciprocalScale;
-	// mathematically if you have 0 scale, it should be infinite, 
-	// however, in practice if you have 0 scale, and relative transform doesn't make much sense 
-	// anymore because you should be instead of showing gigantic infinite mesh
-	// also returning BIG_NUMBER causes sequential NaN issues by multiplying 
-	// so we hardcode as 0
-	if (InScale.X == 0)
+	if (FMath::Abs(InScale.X) <= Tolerance)
 	{
 		SafeReciprocalScale.X = 0.f;
 	}
@@ -1452,7 +1447,7 @@ FORCEINLINE FVector FTransform::GetSafeScaleReciprocal(const FVector& InScale)
 		SafeReciprocalScale.X = 1/InScale.X;
 	}
 
-	if (InScale.Y == 0)
+	if (FMath::Abs(InScale.Y) <= Tolerance)
 	{
 		SafeReciprocalScale.Y = 0.f;
 	}
@@ -1461,7 +1456,7 @@ FORCEINLINE FVector FTransform::GetSafeScaleReciprocal(const FVector& InScale)
 		SafeReciprocalScale.Y = 1/InScale.Y;
 	}
 
-	if (InScale.Z == 0)
+	if (FMath::Abs(InScale.Z) <= Tolerance)
 	{
 		SafeReciprocalScale.Z = 0.f;
 	}
@@ -1471,4 +1466,4 @@ FORCEINLINE FVector FTransform::GetSafeScaleReciprocal(const FVector& InScale)
 	}
 
 	return SafeReciprocalScale;
-}
\ No newline at end of file
+}
diff --git a/Engine/Source/Runtime/Core/Public/Math/TransformVectorized.h b/Engine/Source/Runtime/Core/Public/Math/TransformVectorized.h
--- a/Engine/Source/Runtime/Core/Public/Math/TransformVectorized.h
+++ b/Engine/Source/Runtime/Core/Public/Math/TransformVectorized.h
@@ -567,7 +567,7 @@ public:
 	FORCEINLINE FVector		GetScaledAxis(EAxis::Type InAxis) const;
 	FORCEINLINE FVector		GetUnitAxis(EAxis::Type InAxis) const;
 	FORCEINLINE void		Mirror(EAxis::Type MirrorAxis, EAxis::Type FlipAxis);
-	FORCEINLINE FVector		GetSafeScaleReciprocal(const FVector& InScale, float Tolerance=0.0f) const;
+	FORCEINLINE FVector		GetSafeScaleReciprocal(const FVector& InScale, float Tolerance=SMALL_NUMBER) const;
 
 
 	FORCEINLINE FVector GetLocation() const
@@ -1330,7 +1330,7 @@ private:
 	 * also returning BIG_NUMBER causes sequential NaN issues by multiplying 
 	 * so we hardcode as 0
 	 */
-	static FORCEINLINE VectorRegister		GetSafeScaleReciprocal(const VectorRegister& InScale, const ScalarRegister& Tolerance = ScalarRegister(VectorZero()))
+	static FORCEINLINE VectorRegister		GetSafeScaleReciprocal(const VectorRegister& InScale, const ScalarRegister& Tolerance = ScalarRegister(GlobalVectorConstants::SmallNumber))
 	{		
 		// SafeReciprocalScale.X = (InScale.X == 0) ? 0.f : 1/InScale.X; // same for YZW
 		VectorRegister SafeReciprocalScale;
diff --git a/Engine/Source/Runtime/Core/Public/Misc/ConfigCacheIni.h b/Engine/Source/Runtime/Core/Public/Misc/ConfigCacheIni.h
--- a/Engine/Source/Runtime/Core/Public/Misc/ConfigCacheIni.h
+++ b/Engine/Source/Runtime/Core/Public/Misc/ConfigCacheIni.h
@@ -11,15 +11,103 @@ CORE_API DECLARE_LOG_CATEGORY_EXTERN(LogConfig, Warning, All);
 // Server builds should be tweakable even in Shipping
 #define ALLOW_INI_OVERRIDE_FROM_COMMANDLINE			(UE_SERVER || !(UE_BUILD_SHIPPING))
 
-typedef TMultiMap<FName,FString> FConfigSectionMap;
+struct FConfigValue
+{
+public:
+	FConfigValue() { }
+	FConfigValue(const TCHAR* InValue)
+		: SavedValue(InValue)
+	{
+		ExpandValue();
+	}
+
+	FConfigValue(FString InValue)
+		: SavedValue(MoveTemp(InValue))
+	{
+		ExpandValue();
+	}
+
+	// Returns the ini setting with any macros expanded out
+	const FString& GetValue() const { return (ExpandedValue.Len() > 0 ? ExpandedValue : SavedValue); }
+
+	// Returns the original ini setting without macro expansion
+	const FString& GetSavedValue() const { return SavedValue; }
+
+	DEPRECATED(4.12, "Please switch to explicitly doing a GetValue() or GetSavedValue()")
+	operator const FString& () const { return GetValue(); }
+
+	DEPRECATED(4.12, "Please switch to explicitly doing a GetValue() or GetSavedValue()")
+	const TCHAR* operator*() const { return *GetValue(); }
+
+	bool operator==(const FConfigValue& Other) const { return (SavedValue.Compare(Other.SavedValue) == 0); }
+	bool operator!=(const FConfigValue& Other) const { return !(FConfigValue::operator==(Other)); }
+
+	friend FArchive& operator<<(FArchive& Ar, FConfigValue& ConfigSection)
+	{
+		Ar << ConfigSection.SavedValue;
+
+		if (Ar.IsLoading())
+		{
+			ConfigSection.ExpandValue();
+		}
+
+		return Ar;
+	}
+
+
+private:
+
+	CORE_API void ExpandValue();
+
+	FString SavedValue;
+	FString ExpandedValue;
+};
+
+typedef TMultiMap<FName,FConfigValue> FConfigSectionMap;
 
 // One section in a config file.
 class FConfigSection : public FConfigSectionMap
 {
 public:
-	bool HasQuotes( const FString& Test ) const;
+	static bool HasQuotes( const FString& Test );
 	bool operator==( const FConfigSection& Other ) const;
 	bool operator!=( const FConfigSection& Other ) const;
+
+	void ReplaceOrAdd(const FName Key, FString Value)
+	{
+		if (FConfigValue* ConfigValue = FConfigSectionMap::Find(Key))
+		{
+			*ConfigValue = FConfigValue(MoveTemp(Value));
+		}
+		else
+		{
+			Add(Key, MoveTemp(Value));
+		}
+	}
+
+	template<typename Allocator> 
+	void MultiFind(const FName Key, TArray<FConfigValue, Allocator>& OutValues, const bool bMaintainOrder = false) const
+	{
+		FConfigSectionMap::MultiFind(Key, OutValues, bMaintainOrder);
+	}
+
+	template<typename Allocator> 
+	void MultiFind(const FName Key, TArray<FString, Allocator>& OutValues, const bool bMaintainOrder = false) const
+	{
+		for (const TPair<FName, FConfigValue>& Pair : Pairs)
+		{
+			if (Pair.Key == Key)
+			{
+				OutValues.Add(Pair.Value.GetValue());
+			}
+		}
+
+		if (bMaintainOrder)
+		{
+			Algo::Reverse(OutValues);
+		}
+	}
+
 };
 
 /**
@@ -203,7 +291,7 @@ private:
 	 * @param SectionName - The section name the array property is being written to
 	 * @param PropertyName - The property name of the array
 	 */
-	void ProcessPropertyAndWriteForDefaults(const TArray< FString >& InCompletePropertyToProcess, FString& OutText, const FString& SectionName, const FString& PropertyName);
+	void ProcessPropertyAndWriteForDefaults(const TArray<FConfigValue>& InCompletePropertyToProcess, FString& OutText, const FString& SectionName, const FString& PropertyName);
 
 };
 
diff --git a/Engine/Source/Runtime/Core/Public/ProfilingDebugging/LoadTimeTracker.h b/Engine/Source/Runtime/Core/Public/ProfilingDebugging/LoadTimeTracker.h
--- a/Engine/Source/Runtime/Core/Public/ProfilingDebugging/LoadTimeTracker.h
+++ b/Engine/Source/Runtime/Core/Public/ProfilingDebugging/LoadTimeTracker.h
@@ -5,13 +5,11 @@
 */
 
 #pragma once
-
 #define ENABLE_LOADTIME_TRACKING 1
+#define ENABLE_LOADTIME_RAW_TIMINGS 0
 
 #include "ScopedTimers.h"
 
-#if ENABLE_LOADTIME_TRACKING
-
 /** High level load time tracker utility (such as initial engine startup or game specific timings) */
 class CORE_API FLoadTimeTracker
 {
@@ -26,12 +24,11 @@ public:
 	void ReportScopeTime(double ScopeTime, const FName ScopeLabel);
 
 	/** Prints out total time and individual times */
-	void DumpLoadTimes();
+	void DumpHighLevelLoadTimes() const;
 
-
-	static void DumpLoadTimesStatic()
+	static void DumpHighLevelLoadTimesStatic()
 	{
-		Get().DumpLoadTimes();
+		Get().DumpHighLevelLoadTimes();
 	}
 
 	const TMap<FName, TArray<double>>& GetData() const
@@ -39,46 +36,81 @@ public:
 		return TimeInfo;
 	}
 
-	void ResetLoadTimes();
+	void ResetHighLevelLoadTimes();
+
+	/** Prints out raw load times for individual timers */
+	void DumpRawLoadTimes() const;
+
+	void ResetRawLoadTimes();
+
+#if ENABLE_LOADTIME_RAW_TIMINGS
+
+	/** Raw Timers */
+	double CreateAsyncPackagesFromQueueTime;
+	double ProcessAsyncLoadingTime;
+	double ProcessLoadedPackagesTime;
+	double SerializeTaggedPropertiesTime;
+	double CreateLinkerTime;
+	double FinishLinkerTime;
+	double CreateImportsTime;
+	double CreateExportsTime;
+	double PreLoadObjectsTime;
+	double PostLoadObjectsTime;
+	double PostLoadDeferredObjectsTime;
+	double FinishObjectsTime;
+	double MaterialPostLoad;
+	double MaterialInstancePostLoad;
+	double SerializeInlineShaderMaps;
+	double MaterialSerializeTime;
+	double MaterialInstanceSerializeTime;
+	double AsyncLoadingTime;
+
+	double LinkerLoad_CreateLoader;
+	double LinkerLoad_SerializePackageFileSummary;
+	double LinkerLoad_SerializeNameMap;
+	double LinkerLoad_SerializeGatherableTextDataMap;
+	double LinkerLoad_SerializeImportMap;
+	double LinkerLoad_SerializeExportMap;
+	double LinkerLoad_StartTextureAllocation;
+	double LinkerLoad_FixupImportMap;
+	double LinkerLoad_RemapImports;
+	double LinkerLoad_FixupExportMap;
+	double LinkerLoad_SerializeDependsMap;
+	double LinkerLoad_CreateExportHash;
+	double LinkerLoad_FindExistingExports;
+	double LinkerLoad_FinalizeCreation;
+
+	double Package_FinishLinker;
+	double Package_LoadImports;
+	double Package_CreateImports;
+	double Package_CreateLinker;
+	double Package_FinishTextureAllocations;
+	double Package_CreateExports;
+	double Package_PreLoadObjects;
+	double Package_PostLoadObjects;
+	double Package_Tick;
+	double Package_CreateAsyncPackagesFromQueue;
+
+	double TickAsyncLoading_ProcessLoadedPackages;
+
+
+	double LinkerLoad_SerializeNameMap_ProcessingEntries;
+#endif
 
 private:
 	TMap<FName, TArray<double>> TimeInfo;
-
 private:
-	FLoadTimeTracker()
-	{
-	}
+	FLoadTimeTracker();
 };
 
-/** Scoped duration timer that automatically reports to FLoadTimeTracker */
-class CORE_API FScopedLoadTimer : public FDurationTimer
-{
-public:
-	explicit FScopedLoadTimer(const FName InLabel) 
-		: FDurationTimer(TimerData)
-		, TimerData(0)
-		, Label(InLabel)
-	{
-		Start();
-	}
-
-	/** Dtor, updating seconds with time delta. */
-	~FScopedLoadTimer()
-	{
-		Stop();
-		FLoadTimeTracker::Get().ReportScopeTime(GetAccumulatedTime(), Label);
-	}
-private:
-	double TimerData;
-	const FName Label;
-};
-
-
-#define SCOPED_LOADTIMER(TimerName) FScopedLoadTimer TimerName##_ScopedTimer_ (FName(TEXT(#TimerName)))
+#if ENABLE_LOADTIME_TRACKING
 #define ACCUM_LOADTIME(TimerName, Time) FLoadTimeTracker::Get().ReportScopeTime(Time, FName(TimerName));
-#define SET_LOADTIME(TimerName, Time) FLoadTimeTracker::Get().SetTime(Time, FName(TimerName));
 #else
-#define SCOPED_LOADTIMER(TimerName)
 #define ACCUM_LOADTIME(TimerName, Time)
-#define SET_LOADTIME(TimerName, Time)
+#endif
+
+#if ENABLE_LOADTIME_RAW_TIMINGS
+#define SCOPED_LOADTIMER(TimerName) FScopedDurationTimer DurationTimer_##TimerName(FLoadTimeTracker::Get().TimerName);
+#else
+#define SCOPED_LOADTIMER(TimerName)
 #endif
\ No newline at end of file
diff --git a/Engine/Source/Runtime/CoreUObject/Private/Misc/StartupPackages.cpp b/Engine/Source/Runtime/CoreUObject/Private/Misc/StartupPackages.cpp
--- a/Engine/Source/Runtime/CoreUObject/Private/Misc/StartupPackages.cpp
+++ b/Engine/Source/Runtime/CoreUObject/Private/Misc/StartupPackages.cpp
@@ -23,7 +23,7 @@ void FStartupPackages::GetStartupPackageNames(TArray<FString>& PackageNames, con
 				if (It.Key() == TEXT("Package"))
 				{
 					// add this package to the list to be fully loaded later
-					PackageNames.Add(*(It.Value()));
+					PackageNames.Add(It.Value().GetValue());
 				}
 			}
 		}
diff --git a/Engine/Source/Runtime/CoreUObject/Private/Serialization/AsyncLoading.cpp b/Engine/Source/Runtime/CoreUObject/Private/Serialization/AsyncLoading.cpp
--- a/Engine/Source/Runtime/CoreUObject/Private/Serialization/AsyncLoading.cpp
+++ b/Engine/Source/Runtime/CoreUObject/Private/Serialization/AsyncLoading.cpp
@@ -14,6 +14,7 @@
 #include "ExclusiveLoadPackageTimeTracker.h"
 #include "AssetRegistryInterface.h"
 #include "BlueprintSupport.h"
+#include "LoadTimeTracker.h"
 #include "HAL/ThreadHeartBeat.h"
 #include "HAL/ExceptionHandling.h"
 
@@ -545,6 +546,7 @@ void FAsyncLoadingThread::ProcessAsyncPackageRequest(FAsyncPackageDesc* InReques
 int32 FAsyncLoadingThread::CreateAsyncPackagesFromQueue()
 {
 	SCOPE_CYCLE_COUNTER(STAT_FAsyncPackage_CreateAsyncPackagesFromQueue);
+	SCOPED_LOADTIMER(CreateAsyncPackagesFromQueueTime);
 
 	FAsyncLoadingTickScope InAsyncLoadingTick;
 
@@ -637,6 +639,7 @@ void FAsyncLoadingThread::InsertPackage(FAsyncPackage* Package, EAsyncPackageIns
 EAsyncPackageState::Type FAsyncLoadingThread::ProcessAsyncLoading(int32& OutPackagesProcessed, bool bUseTimeLimit /*= false*/, bool bUseFullTimeLimit /*= false*/, float TimeLimit /*= 0.0f*/)
 {
 	SCOPE_CYCLE_COUNTER(STAT_FAsyncLoadingThread_ProcessAsyncLoading);
+	SCOPED_LOADTIMER(AsyncLoadingTime);
 	
 	// If we're not multithreaded and flushing async loading, update the thread heartbeat
 	const bool bNeedsHeartbeatTick = !bUseTimeLimit && !FAsyncLoadingThread::IsMultithreaded();
@@ -648,6 +651,7 @@ EAsyncPackageState::Type FAsyncLoadingThread::ProcessAsyncLoading(int32& OutPack
 	// like e.g. when called from FlushAsyncLoading.
 	for (int32 PackageIndex = 0; LoadingState != EAsyncPackageState::TimeOut && PackageIndex < AsyncPackages.Num(); ++PackageIndex)
 	{
+		SCOPED_LOADTIMER(ProcessAsyncLoadingTime);
 		OutPackagesProcessed++;
 
 		// Package to be loaded.
@@ -708,8 +712,11 @@ EAsyncPackageState::Type FAsyncLoadingThread::ProcessAsyncLoading(int32& OutPack
 
 EAsyncPackageState::Type FAsyncLoadingThread::ProcessLoadedPackages(bool bUseTimeLimit, bool bUseFullTimeLimit, float TimeLimit, int32 WaitForRequestID /*= INDEX_NONE*/)
 {
+	SCOPED_LOADTIMER(TickAsyncLoading_ProcessLoadedPackages);
+
 	EAsyncPackageState::Type Result = EAsyncPackageState::Complete;
 
+
 	double TickStartTime = FPlatformTime::Seconds();
 	{
 #if THREADSAFE_UOBJECTS
@@ -722,6 +729,8 @@ EAsyncPackageState::Type FAsyncLoadingThread::ProcessLoadedPackages(bool bUseTim
 		
 	for (int32 PackageIndex = 0; PackageIndex < LoadedPackagesToProcess.Num() && !IsAsyncLoadingSuspended(); ++PackageIndex)
 	{
+		SCOPED_LOADTIMER(ProcessLoadedPackagesTime);
+
 		if (PackageIndex % 20 == 0 && IsTimeLimitExceeded(TickStartTime, bUseTimeLimit, TimeLimit, TEXT("ProcessLoadedPackages")))
 		{
 			break;
@@ -1226,6 +1235,7 @@ void FAsyncPackage::EndAsyncLoad()
 EAsyncPackageState::Type FAsyncPackage::Tick(bool InbUseTimeLimit, bool InbUseFullTimeLimit, float& InOutTimeLimit)
 {
 	SCOPE_CYCLE_COUNTER(STAT_FAsyncPackage_Tick);
+	SCOPED_LOADTIMER(Package_Tick);
 
 	// Whether we should execute the next step.
 	EAsyncPackageState::Type LoadingState = EAsyncPackageState::Complete;
@@ -1273,42 +1283,49 @@ EAsyncPackageState::Type FAsyncPackage::Tick(bool InbUseTimeLimit, bool InbUseFu
 		// Create raw linker. Needs to be async created via ticking before it can be used.
 		if (LoadingState == EAsyncPackageState::Complete)
 		{
+			SCOPED_LOADTIMER(Package_CreateLinker);
 			LoadingState = CreateLinker();
 		}
 
 		// Async create linker.
 		if (LoadingState == EAsyncPackageState::Complete)
 		{
+			SCOPED_LOADTIMER(Package_FinishLinker);
 			LoadingState = FinishLinker();
 		}
 
 		// Load imports from linker import table asynchronously.
 		if (LoadingState == EAsyncPackageState::Complete)
 		{
+			SCOPED_LOADTIMER(Package_LoadImports);
 			LoadingState = LoadImports();
 		}
 
 		// Create imports from linker import table.
 		if (LoadingState == EAsyncPackageState::Complete)
 		{
+			SCOPED_LOADTIMER(Package_CreateImports);
 			LoadingState = CreateImports();
 		}
 
 		// Finish all async texture allocations.
 		if (LoadingState == EAsyncPackageState::Complete)
 		{
+			SCOPED_LOADTIMER(Package_FinishTextureAllocations);
 			LoadingState = FinishTextureAllocations();
 		}
 
 		// Create exports from linker export table and also preload them.
 		if (LoadingState == EAsyncPackageState::Complete)
 		{
+			SCOPED_LOADTIMER(Package_CreateExports);
 			LoadingState = CreateExports();
 		}
 
 		// Call Preload on the linker for all loaded objects which causes actual serialization.
 		if (LoadingState == EAsyncPackageState::Complete)
 		{
+			SCOPED_LOADTIMER(Package_PreLoadObjects);
 			LoadingState = PreLoadObjects();
 		}
 
@@ -1316,6 +1333,7 @@ EAsyncPackageState::Type FAsyncPackage::Tick(bool InbUseTimeLimit, bool InbUseFu
 		// another iteration of the PreLoad loop.
 		if (LoadingState == EAsyncPackageState::Complete)
 		{
+			SCOPED_LOADTIMER(Package_PostLoadObjects);
 			LoadingState = PostLoadObjects();
 		}
 
@@ -1356,6 +1374,7 @@ EAsyncPackageState::Type FAsyncPackage::Tick(bool InbUseTimeLimit, bool InbUseFu
  */
 EAsyncPackageState::Type FAsyncPackage::CreateLinker()
 {
+	SCOPED_LOADTIMER(CreateLinkerTime);
 	if (Linker == nullptr)
 	{
 		SCOPE_CYCLE_COUNTER(STAT_FAsyncPackage_CreateLinker);
@@ -1463,6 +1482,7 @@ EAsyncPackageState::Type FAsyncPackage::CreateLinker()
  */
 EAsyncPackageState::Type FAsyncPackage::FinishLinker()
 {
+	SCOPED_LOADTIMER(FinishLinkerTime);
 	EAsyncPackageState::Type Result = EAsyncPackageState::Complete;
 	if (Linker && !Linker->HasFinishedInitialization())
 	{
@@ -1727,6 +1747,7 @@ void FAsyncPackage::ImportFullyLoadedCallback(const FName& InPackageName, UPacka
  */
 EAsyncPackageState::Type FAsyncPackage::CreateImports()
 {
+	SCOPED_LOADTIMER(CreateImportsTime);
 	SCOPE_CYCLE_COUNTER(STAT_FAsyncPackage_CreateImports);
 
 	// GC can't run in here
@@ -1785,6 +1806,7 @@ EAsyncPackageState::Type FAsyncPackage::FinishTextureAllocations()
  */
 EAsyncPackageState::Type FAsyncPackage::CreateExports()
 {
+	SCOPED_LOADTIMER(CreateExportsTime);
 	SCOPE_CYCLE_COUNTER(STAT_FAsyncPackage_CreateExports);
 
 	// GC can't run in here
@@ -1854,6 +1876,7 @@ void FAsyncPackage::FreeReferencedImports()
  */
 EAsyncPackageState::Type FAsyncPackage::PreLoadObjects()
 {
+	SCOPED_LOADTIMER(PreLoadObjectsTime);
 	SCOPE_CYCLE_COUNTER(STAT_FAsyncPackage_PreLoadObjects);
 
 	// GC can't run in here
@@ -1885,6 +1908,8 @@ EAsyncPackageState::Type FAsyncPackage::PostLoadObjects()
 {
 	SCOPE_CYCLE_COUNTER(STAT_FAsyncPackage_PostLoadObjects);
 	
+	SCOPED_LOADTIMER(PostLoadObjectsTime);
+	
 	// GC can't run in here
 	FGCScopeGuard GCGuard;
 
@@ -1928,6 +1953,7 @@ void CreateClustersFromPackage(FLinkerLoad* PackageLinker);
 EAsyncPackageState::Type FAsyncPackage::PostLoadDeferredObjects(double InTickStartTime, bool bInUseTimeLimit, float& InOutTimeLimit)
 {
 	SCOPE_CYCLE_COUNTER(STAT_FAsyncPackage_PostLoadObjectsGameThread);
+	SCOPED_LOADTIMER(PostLoadDeferredObjectsTime);
 
 	EAsyncPackageState::Type Result = EAsyncPackageState::Complete;
 	TGuardValue<bool> GuardIsRoutingPostLoad(FUObjectThreadContext::Get().IsRoutingPostLoad, true);
@@ -2023,6 +2049,8 @@ EAsyncPackageState::Type FAsyncPackage::PostLoadDeferredObjects(double InTickSta
  */
 EAsyncPackageState::Type FAsyncPackage::FinishObjects()
 {
+	SCOPED_LOADTIMER(FinishObjectsTime);
+
 	SCOPE_CYCLE_COUNTER(STAT_FAsyncPackage_FinishObjects);
 	LastObjectWorkWasPerformedOn	= nullptr;
 	LastTypeOfWorkPerformed			= TEXT("finishing all objects");
diff --git a/Engine/Source/Runtime/CoreUObject/Private/UObject/Class.cpp b/Engine/Source/Runtime/CoreUObject/Private/UObject/Class.cpp
--- a/Engine/Source/Runtime/CoreUObject/Private/UObject/Class.cpp
+++ b/Engine/Source/Runtime/CoreUObject/Private/UObject/Class.cpp
@@ -10,6 +10,7 @@
 #include "LinkerPlaceholderClass.h"
 #include "LinkerPlaceholderFunction.h"
 #include "StructScriptLoader.h"
+#include "LoadTimeTracker.h"
 
 // This flag enables some expensive class tree validation that is meant to catch mutations of 
 // the class tree outside of SetSuperStruct. It has been disabled because loading blueprints 
@@ -849,9 +850,9 @@ void UStruct::InitTaggedPropertyRedirectsMap()
 				FName OldPropertyName = NAME_None;
 				FName NewPropertyName = NAME_None;
 
-				FParse::Value( *It.Value(), TEXT("ClassName="), ClassName );
-				FParse::Value( *It.Value(), TEXT("OldPropertyName="), OldPropertyName );
-				FParse::Value( *It.Value(), TEXT("NewPropertyName="), NewPropertyName );
+				FParse::Value( *It.Value().GetValue(), TEXT("ClassName="), ClassName );
+				FParse::Value( *It.Value().GetValue(), TEXT("OldPropertyName="), OldPropertyName );
+				FParse::Value( *It.Value().GetValue(), TEXT("NewPropertyName="), NewPropertyName );
 
 				check(ClassName != NAME_None && OldPropertyName != NAME_None && NewPropertyName != NAME_None );
 				TaggedPropertyRedirects.FindOrAdd(ClassName).Add(OldPropertyName, NewPropertyName);
@@ -866,6 +867,8 @@ void UStruct::InitTaggedPropertyRedirectsMap()
 
 void UStruct::SerializeTaggedProperties(FArchive& Ar, uint8* Data, UStruct* DefaultsStruct, uint8* Defaults, const UObject* BreakRecursionIfFullyLoad) const
 {
+	//SCOPED_LOADTIMER(SerializeTaggedPropertiesTime);
+	
 	check(Ar.IsLoading() || Ar.IsSaving());
 
 	UClass* DefaultsClass = dynamic_cast<UClass*>(DefaultsStruct);
diff --git a/Engine/Source/Runtime/CoreUObject/Private/UObject/Enum.cpp b/Engine/Source/Runtime/CoreUObject/Private/UObject/Enum.cpp
--- a/Engine/Source/Runtime/CoreUObject/Private/UObject/Enum.cpp
+++ b/Engine/Source/Runtime/CoreUObject/Private/UObject/Enum.cpp
@@ -734,7 +734,7 @@ void UEnum::InitEnumRedirectsMap()
 		{
 			if( It.Key() == TEXT("EnumRedirects") )
 			{
-				FString ConfigValue = It.Value();
+				const FString& ConfigValue = It.Value().GetValue();
 				FName EnumName = NAME_None;
 				FName OldEnumEntry = NAME_None;
 				FName NewEnumEntry = NAME_None;
