Turning off infinite stall detection when the debugger is attached. Extended min stall duration.

diff --git a/Engine/Source/Runtime/Core/Private/HAL/ThreadHeartBeat.cpp b/Engine/Source/Runtime/Core/Private/HAL/ThreadHeartBeat.cpp
--- a/Engine/Source/Runtime/Core/Private/HAL/ThreadHeartBeat.cpp
+++ b/Engine/Source/Runtime/Core/Private/HAL/ThreadHeartBeat.cpp
@@ -87,9 +87,9 @@ uint32 FThreadHeartBeat::CheckHeartBeat()
 {
 	// Editor and debug builds run too slow to measure them correctly
 #if !WITH_EDITORONLY_DATA && !IS_PROGRAM && !UE_BUILD_DEBUG
-	if (!GIsRequestingExit)
+	if (!GIsRequestingExit && !FPlatformMisc::IsDebuggerPresent())
 	{
-		const double HangDuration = 10.0;
+		const double HangDuration = 25.0;
 		const double CurrentTime = FPlatformTime::Seconds();
 		FScopeLock HeartBeatLock(&HeartBeatCritical);
 		for (TPair<uint32, double>& LastHeartBeat : ThreadHeartBeat)
