Merging binned 2 fix

diff --git a/Engine/Source/Runtime/Core/Private/HAL/MallocBinned2.cpp b/Engine/Source/Runtime/Core/Private/HAL/MallocBinned2.cpp
--- a/Engine/Source/Runtime/Core/Private/HAL/MallocBinned2.cpp
+++ b/Engine/Source/Runtime/Core/Private/HAL/MallocBinned2.cpp
@@ -83,21 +83,82 @@ FMallocBinned2::FPoolTable::FPoolTable()
 
 struct FMallocBinned2::FPoolInfo
 {
- public:	uint16      Taken;          // Number of allocated elements in this pool, when counts down to zero can free the entire pool
+	enum class ECanary : uint16
+	{
+		Unassigned = 0x3941,
+		FirstFreeBlockIsOSAllocSize = 0x17ea,
+		FirstFreeBlockIsPtr = 0xf317
+	};
+
+ public:	uint16      Taken;          // Number of allocated elements in this pool, when counts down to zero can free the entire pool	
+public:	ECanary		Canary;	// See ECanary
 private:	uint32      AllocSize;      // Number of bytes allocated
  public:	FFreeBlock* FirstFreeBlock; // Pointer to first free memory in this pool or the OS Allocation Size in bytes if this allocation is not binned
  public:	FPoolInfo*  Next;           // Pointer to next pool
  public:	FPoolInfo** PtrToPrevNext;  // Pointer to whichever pointer points to this pool
 
 public:
+	FPoolInfo()
+		: Taken(0)
+		, Canary(ECanary::Unassigned)
+		, AllocSize(0)
+		, FirstFreeBlock(nullptr)
+		, Next(nullptr)
+		, PtrToPrevNext(nullptr)
+	{
+	}
+	void CheckCanary(ECanary ShouldBe) const
+	{
+		if (Canary != ShouldBe)
+		{
+			UE_LOG(LogMemory, Fatal, TEXT("MallocBinned2 Corruption Canary was 0x%x, should be 0x%x"), int32(Canary), int32(ShouldBe));
+		}
+	}
+	void SetCanary(ECanary ShouldBe, bool bPreexisting, bool bGuarnteedToBeNew)
+	{
+		if (bPreexisting)
+		{
+			if (bGuarnteedToBeNew)
+			{
+				UE_LOG(LogMemory, Fatal, TEXT("MallocBinned2 Corruption Canary was 0x%x, should be 0x%x. This block is both preexisting and guaranteed to be new; which makes no sense."), int32(Canary), int32(ShouldBe));
+			}
+			if (ShouldBe == ECanary::Unassigned)
+			{
+				if (Canary != ECanary::FirstFreeBlockIsOSAllocSize && Canary != ECanary::FirstFreeBlockIsPtr)
+				{
+					UE_LOG(LogMemory, Fatal, TEXT("MallocBinned2 Corruption Canary was 0x%x, will be 0x%x because this block should be preexisting and in use."), int32(Canary), int32(ShouldBe));
+				}
+			}
+			else if (Canary != ShouldBe)
+			{
+				UE_LOG(LogMemory, Fatal, TEXT("MallocBinned2 Corruption Canary was 0x%x, should be 0x%x because this block should be preexisting."), int32(Canary), int32(ShouldBe));
+			}
+		}
+		else
+		{
+			if (bGuarnteedToBeNew)
+			{
+				if (Canary != ECanary::Unassigned)
+				{
+					UE_LOG(LogMemory, Fatal, TEXT("MallocBinned2 Corruption Canary was 0x%x, will be 0x%x. This block is guaranteed to be new yet is it already assigned."), int32(Canary), int32(ShouldBe));
+				}
+			}
+			else if (Canary != ShouldBe && Canary != ECanary::Unassigned)
+			{
+				UE_LOG(LogMemory, Fatal, TEXT("MallocBinned2 Corruption Canary was 0x%x, will be 0x%x does not have an expected value."), int32(Canary), int32(ShouldBe));
+			}
+		}
+		Canary = ShouldBe;
+	}
 	bool HasFreeRegularBlock() const
 	{
+		CheckCanary(ECanary::FirstFreeBlockIsPtr);
 		return FirstFreeBlock && FirstFreeBlock->GetNumFreeRegularBlocks() != 0;
 	}
 
 	void* AllocateRegularBlock()
 	{
-		checkSlow(HasFreeRegularBlock());
+		check(HasFreeRegularBlock());
 		++Taken;
 		void* Result = FirstFreeBlock->AllocateRegularBlock();
 		ExhaustPoolIfNecessary();
@@ -111,13 +172,15 @@ public:
 
 	UPTRINT GetOsAllocatedBytes() const
 	{
+		CheckCanary(ECanary::FirstFreeBlockIsOSAllocSize);
 		return (UPTRINT)FirstFreeBlock;
 	}
 
 	void SetOSAllocationSizes(uint32 InRequestedBytes, UPTRINT InAllocatedBytes)
 	{
-		checkSlow(InRequestedBytes != 0);                // Shouldn't be pooling zero byte allocations
-		checkSlow(InAllocatedBytes >= InRequestedBytes); // We must be allocating at least as much as we requested
+		CheckCanary(ECanary::FirstFreeBlockIsOSAllocSize);
+		check(InRequestedBytes != 0);                // Shouldn't be pooling zero byte allocations
+		check(InAllocatedBytes >= InRequestedBytes); // We must be allocating at least as much as we requested
 
 		AllocSize      = InRequestedBytes;
 		FirstFreeBlock = (FFreeBlock*)InAllocatedBytes;
@@ -150,15 +213,11 @@ private:
 		{
 			FirstFreeBlock = (FFreeBlock*)FirstFreeBlock->NextFreeBlock;
 		}
-//		checkSlow(!FirstFreeBlock || FirstFreeBlock->GetNumFreeRegularBlocks() != 0 || FirstFreeBlock->HasFreeShortBlock());
+		check(!FirstFreeBlock || FirstFreeBlock->GetNumFreeRegularBlocks() != 0);
 	}
 };
 
-template <>
-struct TIsZeroConstructType<FMallocBinned2::FPoolInfo>
-{
-	enum { Value = true };
-};
+
 
 /** Hash table struct for retrieving allocation book keeping information */
 struct FMallocBinned2::PoolHashBucket
@@ -207,7 +266,7 @@ struct FMallocBinned2::Private
 	/**
 	 * Gets the FPoolInfo for a memory address. If no valid info exists one is created.
 	 */
-	static FPoolInfo* GetOrCreatePoolInfo(FMallocBinned2& Allocator, void* InPtr)
+	static FPoolInfo* GetOrCreatePoolInfo(FMallocBinned2& Allocator, void* InPtr, FMallocBinned2::FPoolInfo::ECanary Kind, bool bPreexisting)
 	{
 		/** 
 		 * Creates an array of FPoolInfo structures for tracking allocations.
@@ -226,9 +285,10 @@ struct FMallocBinned2::Private
 			return (FPoolInfo*)Result;
 		};
 
-		UPTRINT BucketIndex;
+		uint32 BucketIndex;
+		UPTRINT BucketIndexCollision;
 		uint32  PoolIndex;
-		Allocator.PtrToPoolMapping.GetHashBucketAndPoolIndices(InPtr, BucketIndex, PoolIndex);
+		Allocator.PtrToPoolMapping.GetHashBucketAndPoolIndices(InPtr, BucketIndex, BucketIndexCollision, PoolIndex);
 
 		PoolHashBucket* FirstBucket = &Allocator.HashBuckets[BucketIndex];
 		PoolHashBucket* Collision   = FirstBucket;
@@ -236,14 +296,15 @@ struct FMallocBinned2::Private
 		{
 			if (!Collision->FirstPool)
 			{
-				Collision->BucketIndex = BucketIndex;
+				Collision->BucketIndex = BucketIndexCollision;
 				Collision->FirstPool   = CreatePoolArray(Allocator.NumPoolsPerPage);
-
+				Collision->FirstPool[PoolIndex].SetCanary(Kind, bPreexisting, true);
 				return &Collision->FirstPool[PoolIndex];
 			}
 
-			if (Collision->BucketIndex == BucketIndex)
+			if (Collision->BucketIndex == BucketIndexCollision)
 			{
+				Collision->FirstPool[PoolIndex].SetCanary(Kind, bPreexisting, false);
 				return &Collision->FirstPool[PoolIndex];
 			}
 
@@ -276,9 +337,14 @@ struct FMallocBinned2::Private
 		if (!NewBucket->FirstPool)
 		{
 			NewBucket->FirstPool = CreatePoolArray(Allocator.NumPoolsPerPage);
+			NewBucket->FirstPool[PoolIndex].SetCanary(Kind, bPreexisting, true);
+		}
+		else
+		{
+			NewBucket->FirstPool[PoolIndex].SetCanary(Kind, bPreexisting, false);
 		}
 
-		NewBucket->BucketIndex = BucketIndex;
+		NewBucket->BucketIndex = BucketIndexCollision;
 
 		FirstBucket->Link(NewBucket);
 
@@ -287,15 +353,17 @@ struct FMallocBinned2::Private
 
 	static FPoolInfo* FindPoolInfo(FMallocBinned2& Allocator, void* InPtr)
 	{
-		UPTRINT BucketIndex;
+		uint32 BucketIndex;
+		UPTRINT BucketIndexCollision;
 		uint32  PoolIndex;
-		Allocator.PtrToPoolMapping.GetHashBucketAndPoolIndices(InPtr, BucketIndex, PoolIndex);
+		Allocator.PtrToPoolMapping.GetHashBucketAndPoolIndices(InPtr, BucketIndex, BucketIndexCollision, PoolIndex);
+
 
 		PoolHashBucket* FirstBucket = &Allocator.HashBuckets[BucketIndex];
 		PoolHashBucket* Collision   = FirstBucket;
 		do
 		{
-			if (Collision->BucketIndex == BucketIndex)
+			if (Collision->BucketIndex == BucketIndexCollision)
 			{
 				return &Collision->FirstPool[PoolIndex];
 			}
@@ -375,6 +443,11 @@ struct FMallocBinned2::Private
 			{
 				FBundleNode* NextNode = Node->NextNodeInCurrentBundle;
 				FPoolInfo*   NodePool = FindPoolInfo(Allocator, Node);
+				if (!NodePool)
+				{
+					UE_LOG(LogMemory, Fatal, TEXT("FMallocBinned2 Attempt to free an unrecognized small block %p"), Node);
+				}
+				NodePool->CheckCanary(FPoolInfo::ECanary::FirstFreeBlockIsPtr);
 
 				// If this pool was exhausted, move to available list.
 				if (!NodePool->FirstFreeBlock)
@@ -390,9 +463,10 @@ struct FMallocBinned2::Private
 				NodePool->FirstFreeBlock   = Free;
 
 				// Free this pool.
-				checkSlow(NodePool->Taken >= 1);
+				check(NodePool->Taken >= 1);
 				if (--NodePool->Taken == 0)
 				{
+					NodePool->SetCanary(FPoolInfo::ECanary::Unassigned, true, false);
 					FFreeBlock* BasePtrOfNode = GetPoolHeaderFromPointer(Node);
 
 					// Free the OS memory.
@@ -417,13 +491,13 @@ FORCEINLINE bool FMallocBinned2::FPoolList::IsEmpty() const
 
 FORCEINLINE FMallocBinned2::FPoolInfo& FMallocBinned2::FPoolList::GetFrontPool()
 {
-	checkSlow(!IsEmpty());
+	check(!IsEmpty());
 	return *Front;
 }
 
 FORCEINLINE const FMallocBinned2::FPoolInfo& FMallocBinned2::FPoolList::GetFrontPool() const
 {
-	checkSlow(!IsEmpty());
+	check(!IsEmpty());
 	return *Front;
 }
 
@@ -445,7 +519,7 @@ FMallocBinned2::FPoolInfo& FMallocBinned2::FPoolList::PushNewPoolToFront(FMalloc
 	}
 	check(IsAligned(Free, LocalPageSize));
 	// Create pool
-	FPoolInfo* Result = Private::GetOrCreatePoolInfo(Allocator, Free);
+	FPoolInfo* Result = Private::GetOrCreatePoolInfo(Allocator, Free, FPoolInfo::ECanary::FirstFreeBlockIsPtr, false);
 	Result->Link(Front);
 	Result->Taken          = 0;
 	Result->FirstFreeBlock = Free;
@@ -593,8 +667,8 @@ void* FMallocBinned2::MallocExternal(SIZE_T Size, uint32 Alignment)
 	Alignment = FMath::Max<uint32>(Alignment, BINNED2_MINIMUM_ALIGNMENT);
 	Size = Align(FMath::Max((SIZE_T)1, Size), Alignment);
 
-	checkSlow(FMath::IsPowerOfTwo(Alignment));
-	checkSlow(Alignment <= PageSize);
+	check(FMath::IsPowerOfTwo(Alignment));
+	check(Alignment <= PageSize);
 
 	FScopeLock Lock(&Mutex);
 
@@ -608,7 +682,8 @@ void* FMallocBinned2::MallocExternal(SIZE_T Size, uint32 Alignment)
 	check(IsAligned(Result, PageSize) && IsOSAllocation(Result));
 
 	// Create pool.
-	FPoolInfo* Pool = Private::GetOrCreatePoolInfo(*this, Result);
+	FPoolInfo* Pool = Private::GetOrCreatePoolInfo(*this, Result, FPoolInfo::ECanary::FirstFreeBlockIsOSAllocSize, false);
+	check(Size > 0 && Size <= AlignedSize && AlignedSize >= OsAllocationGranularity);
 	Pool->SetOSAllocationSizes(Size, AlignedSize);
 
 	return Result;
@@ -623,12 +698,12 @@ void* FMallocBinned2::ReallocExternal(void* Ptr, SIZE_T NewSize, uint32 Alignmen
 		return nullptr;
 	}
 	static_assert(DEFAULT_ALIGNMENT <= BINNED2_MINIMUM_ALIGNMENT, "DEFAULT_ALIGNMENT is assumed to be zero"); // used below
-	checkSlow(FMath::IsPowerOfTwo(Alignment));
-	checkSlow(Alignment <= PageSize);
+	check(FMath::IsPowerOfTwo(Alignment));
+	check(Alignment <= PageSize);
 
 	if (!IsOSAllocation(Ptr))
 	{
-		checkSlow(Ptr); // null is 64k aligned so we should not be here
+		check(Ptr); // null is 64k aligned so we should not be here
 		// Reallocate to a smaller/bigger pool if necessary
 		FFreeBlock* Free = GetPoolHeaderFromPointer(Ptr);
 		Free->CanaryTest();
@@ -657,14 +732,20 @@ void* FMallocBinned2::ReallocExternal(void* Ptr, SIZE_T NewSize, uint32 Alignmen
 
 	// Allocated from OS.
 	FPoolInfo* Pool = Private::FindPoolInfo(*this, Ptr);
+	if (!Pool)
+	{
+		UE_LOG(LogMemory, Fatal, TEXT("FMallocBinned2 Attempt to realloc an unrecognized block %p"), Ptr);
+	}
 	UPTRINT PoolOsBytes = Pool->GetOsAllocatedBytes();
+	uint32 PoolOSRequestedBytes = Pool->GetOSRequestedBytes();
+	checkf(PoolOSRequestedBytes <= PoolOsBytes, TEXT("FMallocBinned2::ReallocExternal %d %d"), int32(PoolOSRequestedBytes), int32(PoolOsBytes));
 	if (NewSize > PoolOsBytes || // can't fit in the old block
 		(NewSize <= BINNED2_MAX_SMALL_POOL_SIZE && Alignment <= BINNED2_MINIMUM_ALIGNMENT) || // can switch to the small block allocator
 		Align(NewSize, OsAllocationGranularity) < PoolOsBytes) // we can get some pages back
 	{
 		// Grow or shrink.
 		void* Result = FMallocBinned2::MallocExternal(NewSize, Alignment);
-		FMemory::Memcpy(Result, Ptr, FMath::Min<SIZE_T>(NewSize, Pool->GetOSRequestedBytes()));
+		FMemory::Memcpy(Result, Ptr, FMath::Min<SIZE_T>(NewSize, PoolOSRequestedBytes));
 		FMallocBinned2::FreeExternal(Ptr);
 		return Result;
 	}
@@ -678,7 +759,7 @@ void FMallocBinned2::FreeExternal(void* Ptr)
 {
 	if (!IsOSAllocation(Ptr))
 	{
-		checkSlow(Ptr); // null is 64k aligned so we should not be here
+		check(Ptr); // null is 64k aligned so we should not be here
 		FFreeBlock* BasePtr = GetPoolHeaderFromPointer(Ptr);
 		BasePtr->CanaryTest();
 		uint32 BlockSize = BasePtr->BlockSize;
@@ -711,9 +792,12 @@ void FMallocBinned2::FreeExternal(void* Ptr)
 		{
 			UE_LOG(LogMemory, Fatal, TEXT("FMallocBinned2 Attempt to free an unrecognized block %p"), Ptr);
 		}
-
+		UPTRINT PoolOsBytes = Pool->GetOsAllocatedBytes();
+		uint32 PoolOSRequestedBytes = Pool->GetOSRequestedBytes();
+		checkf(PoolOSRequestedBytes <= PoolOsBytes, TEXT("FMallocBinned2::FreeExternal %d %d"), int32(PoolOSRequestedBytes), int32(PoolOsBytes));
+		Pool->SetCanary(FPoolInfo::ECanary::Unassigned, true, false);
 		// Free an OS allocation.
-		CachedOSPageAllocator.Free(Ptr, Pool->GetOsAllocatedBytes());
+		CachedOSPageAllocator.Free(Ptr, PoolOsBytes);
 	}
 }
 
@@ -721,7 +805,7 @@ bool FMallocBinned2::GetAllocationSizeExternal(void* Ptr, SIZE_T& SizeOut)
 {
 	if (!IsOSAllocation(Ptr))
 	{
-		checkSlow(Ptr); // null is 64k aligned so we should not be here
+		check(Ptr); // null is 64k aligned so we should not be here
 		const FFreeBlock* Free = GetPoolHeaderFromPointer(Ptr);
 		Free->CanaryTest();
 		uint32 BlockSize = Free->BlockSize;
@@ -734,7 +818,14 @@ bool FMallocBinned2::GetAllocationSizeExternal(void* Ptr, SIZE_T& SizeOut)
 	}
 	FScopeLock Lock(&Mutex);
 	FPoolInfo* Pool = Private::FindPoolInfo(*this, Ptr);
-	SizeOut = Pool->GetOSRequestedBytes();
+	if (!Pool)
+	{
+		UE_LOG(LogMemory, Fatal, TEXT("FMallocBinned2 Attempt to GetAllocationSizeExternal an unrecognized block %p"), Ptr);
+	}
+	UPTRINT PoolOsBytes = Pool->GetOsAllocatedBytes();
+	uint32 PoolOSRequestedBytes = Pool->GetOSRequestedBytes();
+	checkf(PoolOSRequestedBytes <= PoolOsBytes, TEXT("FMallocBinned2::GetAllocationSizeExternal %d %d"), int32(PoolOSRequestedBytes), int32(PoolOsBytes));
+	SizeOut = PoolOSRequestedBytes;
 	return true;
 }
 
diff --git a/Engine/Source/Runtime/Core/Private/HAL/OutputDevices.cpp b/Engine/Source/Runtime/Core/Private/HAL/OutputDevices.cpp
--- a/Engine/Source/Runtime/Core/Private/HAL/OutputDevices.cpp
+++ b/Engine/Source/Runtime/Core/Private/HAL/OutputDevices.cpp
@@ -805,7 +805,7 @@ void FOutputDeviceRedirector::FlushThreadedLogs()
 
 void FOutputDeviceRedirector::PanicFlushThreadedLogs()
 {
-	SCOPE_CYCLE_COUNTER(STAT_FlushThreadedLogs);
+//	SCOPE_CYCLE_COUNTER(STAT_FlushThreadedLogs);
 	// Acquire a lock on SynchronizationObject and call the unsynchronized worker function.
 	FScopeLock ScopeLock( &SynchronizationObject );
 	
diff --git a/Engine/Source/Runtime/Core/Public/HAL/MallocBinned2.h b/Engine/Source/Runtime/Core/Public/HAL/MallocBinned2.h
--- a/Engine/Source/Runtime/Core/Public/HAL/MallocBinned2.h
+++ b/Engine/Source/Runtime/Core/Public/HAL/MallocBinned2.h
@@ -167,10 +167,10 @@ class CORE_API FMallocBinned2 final : public FMalloc
 			MaxHashBuckets        = AddressLimit >> HashKeyShift;
 		}
 
-		FORCEINLINE void GetHashBucketAndPoolIndices(const void* InPtr, UPTRINT& OutBucketIndex, uint32& OutPoolIndex) const
+		FORCEINLINE void GetHashBucketAndPoolIndices(const void* InPtr, uint32& OutBucketIndex, UPTRINT& OutBucketCollision, uint32& OutPoolIndex) const
 		{
-			OutBucketIndex = (UPTRINT)InPtr >> HashKeyShift;
-			OutBucketIndex &= (MaxHashBuckets - 1);
+			OutBucketCollision = (UPTRINT)InPtr >> HashKeyShift;
+			OutBucketIndex = uint32(OutBucketCollision & (MaxHashBuckets - 1));
 			OutPoolIndex   = ((UPTRINT)InPtr >> PtrToPoolPageBitShift) & PoolMask;
 		}
 
