Thread heartbeat fixes/improvements (fixing hang false positives in hang detection):

- Make sure SearchPackageOnDisk sends a heartbeat
- Start measuring heartbeat after the initial init
- Hang duration can be controller with ini settings

diff --git a/Engine/Source/Runtime/Core/Private/HAL/ThreadHeartBeat.cpp b/Engine/Source/Runtime/Core/Private/HAL/ThreadHeartBeat.cpp
--- a/Engine/Source/Runtime/Core/Private/HAL/ThreadHeartBeat.cpp
+++ b/Engine/Source/Runtime/Core/Private/HAL/ThreadHeartBeat.cpp
@@ -4,11 +4,18 @@
 
 FThreadHeartBeat::FThreadHeartBeat()
 : Thread(nullptr)
+, bReadyToCheckHeartbeat(false)
+, HangDuration(25.0)
 {
 	// We don't care about programs for now so no point in spawning the extra thread
 #if !IS_PROGRAM
 	Thread = FRunnableThread::Create(this, TEXT("FHeartBeatThread"), 0, TPri_BelowNormal);
 #endif
+
+	if (GConfig)
+	{
+		GConfig->GetDouble(TEXT("Core.System"), TEXT("HangDuration"), HangDuration, GEngineIni);
+	}
 }
 
 FThreadHeartBeat::~FThreadHeartBeat()
@@ -62,7 +69,7 @@ uint32 FThreadHeartBeat::Run()
 				StackTrimmed += StackLines[LineIndex];
 				StackTrimmed += TEXT("\n");
 			}
-			UE_LOG(LogCore, Fatal, TEXT("Infinite stall detected on %s:\n%s"), *ThreadName, *StackTrimmed);
+			UE_LOG(LogCore, Fatal, TEXT("Hang detected on %s:\n%s"), *ThreadName, *StackTrimmed);
 		}
 		FPlatformProcess::SleepNoStats(0.5f);
 	}
@@ -75,6 +82,11 @@ void FThreadHeartBeat::Stop()
 	StopTaskCounter.Increment();
 }
 
+void FThreadHeartBeat::Start()
+{
+	bReadyToCheckHeartbeat = true;
+}
+
 void FThreadHeartBeat::HeartBeat()
 {
 	uint32 ThreadId = FPlatformTLS::GetCurrentThreadId();
@@ -87,7 +99,7 @@ uint32 FThreadHeartBeat::CheckHeartBeat()
 {
 	// Editor and debug builds run too slow to measure them correctly
 #if !WITH_EDITORONLY_DATA && !IS_PROGRAM && !UE_BUILD_DEBUG
-	if (!GIsRequestingExit && !FPlatformMisc::IsDebuggerPresent())
+	if (HangDuration > 0.0 && bReadyToCheckHeartbeat && !GIsRequestingExit && !FPlatformMisc::IsDebuggerPresent())
 	{
 		const double HangDuration = 25.0;
 		const double CurrentTime = FPlatformTime::Seconds();
diff --git a/Engine/Source/Runtime/Core/Public/HAL/ThreadHeartBeat.h b/Engine/Source/Runtime/Core/Public/HAL/ThreadHeartBeat.h
--- a/Engine/Source/Runtime/Core/Public/HAL/ThreadHeartBeat.h
+++ b/Engine/Source/Runtime/Core/Public/HAL/ThreadHeartBeat.h
@@ -15,6 +15,10 @@ class CORE_API FThreadHeartBeat : public FRunnable
 	FCriticalSection HeartBeatCritical;
 	/** Keeps track of the last heartbeat time for threads */
 	TMap<uint32, double> ThreadHeartBeat;
+	/** True if heartbeat should be measured */
+	FThreadSafeBool bReadyToCheckHeartbeat;
+	/** Max time the thread is allowed to not send the heartbeat*/
+	double HangDuration;
 
 	FThreadHeartBeat();
 	virtual ~FThreadHeartBeat();
@@ -30,6 +34,8 @@ public:
 	/** Gets the heartbeat singleton */
 	static FThreadHeartBeat& Get();
 
+	/** Begin measuring heartbeat */
+	void Start();
 	/** Called from a thread once per frame to update the heartbeat time */
 	void HeartBeat();
 	/** Called by a supervising thread to check the threads' health */
diff --git a/Engine/Source/Runtime/CoreUObject/Private/Misc/PackageName.cpp b/Engine/Source/Runtime/CoreUObject/Private/Misc/PackageName.cpp
--- a/Engine/Source/Runtime/CoreUObject/Private/Misc/PackageName.cpp
+++ b/Engine/Source/Runtime/CoreUObject/Private/Misc/PackageName.cpp
@@ -7,6 +7,7 @@
 #include "CoreUObjectPrivate.h"
 #include "Projects.h"
 #include "PackageLocalizationManager.h"
+#include "HAL/ThreadHeartBeat.h"
 
 DEFINE_LOG_CATEGORY_STATIC(LogPackageName, Log, All);
 
@@ -728,6 +729,12 @@ bool FPackageName::DoesPackageExist(const FString& LongPackageName, const FGuid*
 bool FPackageName::SearchForPackageOnDisk(const FString& PackageName, FString* OutLongPackageName, FString* OutFilename)
 {
 	DECLARE_SCOPE_CYCLE_COUNTER(TEXT("FPackageName::SearchForPackageOnDisk"), STAT_PackageName_SearchForPackageOnDisk, STATGROUP_LoadTime);
+	
+	// This function may take a long time to complete so make sure the heartbeat is reset
+	if (IsInGameThread())
+	{
+		FThreadHeartBeat::Get().HeartBeat();
+	}
 
 	bool bResult = false;
 	double StartTime = FPlatformTime::Seconds();
@@ -771,6 +778,12 @@ bool FPackageName::SearchForPackageOnDisk(const FString& PackageName, FString* O
 			// Search directly on disk. Very slow!
 			IFileManager::Get().FindFilesRecursive(Results, *Paths[PathIndex], *PackageWildcard, true, false);
 
+			// Send a heartbeat now too
+			if (IsInGameThread())
+			{
+				FThreadHeartBeat::Get().HeartBeat();
+			}
+
 			for (int32 FileIndex = 0; FileIndex < Results.Num(); ++FileIndex)
 			{			
 				FString Filename(Results[FileIndex]);
diff --git a/Engine/Source/Runtime/Launch/Private/LaunchEngineLoop.cpp b/Engine/Source/Runtime/Launch/Private/LaunchEngineLoop.cpp
--- a/Engine/Source/Runtime/Launch/Private/LaunchEngineLoop.cpp
+++ b/Engine/Source/Runtime/Launch/Private/LaunchEngineLoop.cpp
@@ -2289,6 +2289,9 @@ int32 FEngineLoop::Init()
 
 	FCoreDelegates::StarvedGameLoop.BindStatic(&GameLoopIsStarved);
 
+	// Ready to measure thread heartbeat
+	FThreadHeartBeat::Get().Start();
+
 	return 0;
 }
 
