Fixed a regression where it was possible to create a log file even when NO_LOGGING was defined to 1 when the game crashed by bypassing UE_LOG macros and writing directly to output device in cooked builds with memory logging enabled.

diff --git a/Engine/Source/Runtime/Core/Private/HAL/OutputDeviceMemory.cpp b/Engine/Source/Runtime/Core/Private/HAL/OutputDeviceMemory.cpp
--- a/Engine/Source/Runtime/Core/Private/HAL/OutputDeviceMemory.cpp
+++ b/Engine/Source/Runtime/Core/Private/HAL/OutputDeviceMemory.cpp
@@ -7,7 +7,7 @@
 #include "CorePrivatePCH.h"
 #include "HAL/OutputDeviceMemory.h"
 
-#define DUMP_LOG_ON_EXIT (PLATFORM_DESKTOP && (!UE_BUILD_SHIPPING || USE_LOGGING_IN_SHIPPING))
+#define DUMP_LOG_ON_EXIT (!NO_LOGGING && PLATFORM_DESKTOP && (!UE_BUILD_SHIPPING || USE_LOGGING_IN_SHIPPING))
 
 FOutputDeviceMemory::FOutputDeviceMemory(int32 InBufferSize /*= 1024 * 1024*/)
 : BufferStartPos(0)
@@ -47,7 +47,9 @@ void FOutputDeviceMemory::Flush()
 
 void FOutputDeviceMemory::Serialize(const TCHAR* Data, ELogVerbosity::Type Verbosity, const class FName& Category, const double Time)
 {
+#if !NO_LOGGING
 	FormatAndSerialize(Data, Verbosity, Category, Time);
+#endif
 }
 
 void FOutputDeviceMemory::Serialize(const TCHAR* Data, ELogVerbosity::Type Verbosity, const class FName& Category)
diff --git a/Engine/Source/Runtime/Core/Private/Windows/WindowsPlatformCrashContext.cpp b/Engine/Source/Runtime/Core/Private/Windows/WindowsPlatformCrashContext.cpp
--- a/Engine/Source/Runtime/Core/Private/Windows/WindowsPlatformCrashContext.cpp
+++ b/Engine/Source/Runtime/Core/Private/Windows/WindowsPlatformCrashContext.cpp
@@ -321,6 +321,7 @@ int32 ReportCrashUsingCrashReportClient(EXCEPTION_POINTERS* ExceptionInfo, const
 				const FString LogFileName = FPlatformOutputDevices::GetAbsoluteLogFilename();
 				// If we have a memory only log, make sure it's dumped to file before we attach it to the report
 				bool bHasLogFile = !FPlatformOutputDevices::GetLog()->IsMemoryOnly();
+#if !NO_LOGGING
 				if (!bHasLogFile)
 				{
 					FArchive* LogFile = IFileManager::Get().CreateFileWriter(*LogFileName, FILEWRITE_AllowRead);
@@ -332,6 +333,7 @@ int32 ReportCrashUsingCrashReportClient(EXCEPTION_POINTERS* ExceptionInfo, const
 						bHasLogFile = true;
 					}
 				}
+#endif
 				if (bHasLogFile)
 				{
 					WerReportAddFile(ReportHandle, *LogFileName, WerFileTypeOther, WER_FILE_ANONYMOUS_DATA);
