Copying //UE4/Dev-Sequencer to //UE4/Main (Source: //UE4/Dev-Sequencer @ 2945541)

diff --git a/Engine/Source/Runtime/CinematicCamera/Private/CineCameraComponent.cpp b/Engine/Source/Runtime/CinematicCamera/Private/CineCameraComponent.cpp
--- a/Engine/Source/Runtime/CinematicCamera/Private/CineCameraComponent.cpp
+++ b/Engine/Source/Runtime/CinematicCamera/Private/CineCameraComponent.cpp
@@ -20,7 +20,6 @@ UCineCameraComponent::UCineCameraComponent()
 	LensSettings.MinFStop = 2.f;
 	LensSettings.MaxFStop = 2.f;
 	LensSettings.MinimumFocusDistance = 15.f;
-	LensSettings.MaxReproductionRatio = 1 / 11.f;
 
 	bConstrainAspectRatio = true;
 
@@ -171,10 +170,6 @@ void UCineCameraComponent::RecalcDerivedData()
 #endif
 }
 
-static const FName NAME_GetDesiredFocusDistance(TEXT("GetDesiredFocusDistance"));
-static const float SpotFocusTraceDist = 1000000.f;
-static const float FocusDistInfinity = 1000000.f;
-
 float UCineCameraComponent::GetDesiredFocusDistance(FMinimalViewInfo& DesiredView) const
 {
 	float DesiredFocusDistance = 0.f;
@@ -186,22 +181,6 @@ float UCineCameraComponent::GetDesiredFocusDistance(FMinimalViewInfo& DesiredVie
 		DesiredFocusDistance = FocusSettings.ManualFocusDistance;
 		break;
 
-	case ECameraFocusMethod::Spot:
-		{
-			// AutoFocus, do a trace
-			AActor* const OwningActor = GetOwner();
-
-			// trace to get depth at center of screen
-			FCollisionQueryParams TraceParams(NAME_GetDesiredFocusDistance, true, OwningActor);
-			FHitResult Hit;
-
-			FVector const TraceStart = DesiredView.Location;
-			FVector const TraceEnd = TraceStart + DesiredView.Rotation.Vector() * SpotFocusTraceDist;
-			bool const bHit = GetWorld()->LineTraceSingleByChannel(Hit, TraceStart, TraceEnd, ECC_WorldStatic, TraceParams);
-			DesiredFocusDistance = bHit ? (TraceStart - Hit.ImpactPoint).Size() : FocusDistInfinity;
-		}
-		break;
-
 	case ECameraFocusMethod::Tracking:
 		{
 			AActor const* const TrackedActor = FocusSettings.TrackingFocusSettings.ActorToTrack;
@@ -218,56 +197,6 @@ float UCineCameraComponent::GetDesiredFocusDistance(FMinimalViewInfo& DesiredVie
 			}
 
 			DesiredFocusDistance = (FocusPoint - DesiredView.Location).Size();
-
-// 			// try to get the position of the tracked object
-// 			UObject* const TrackedObj = FocusSettings.TrackingFocusSettings.ObjectToTrack;
-// 
-// 			// infinity is the fallback in case none of the other code paths produces a good result
-// 			DesiredFocusDistance = FocusDistInfinity;
-// 			if (TrackedObj)
-// 			{
-// 				// we have a socket/bone, see if the object is a scenecomponent
-// 				USceneComponent* const TrackedSceneComp = Cast<USceneComponent>(TrackedObj);
-// 				if (TrackedSceneComp)
-// 				{
-// 					if (FocusSettings.TrackingFocusSettings.BoneOrSocketName != NAME_None)
-// 					{
-// 						FVector const SocketLoc = TrackedSceneComp->GetSocketLocation(FocusSettings.TrackingFocusSettings.BoneOrSocketName);
-// 						DesiredFocusDistance = (SocketLoc - DesiredView.Location).Size();
-// 					}
-// 					else
-// 					{
-// 						DesiredFocusDistance = (TrackedSceneComp->GetComponentLocation() - DesiredView.Location).Size();
-// 					}
-// 				}
-// 				else
-// 				{
-// 					AActor* const TrackedActor = Cast<AActor>(TrackedObj);
-// 					if (TrackedActor)
-// 					{
-// 						if (FocusSettings.TrackingFocusSettings.BoneOrSocketName != NAME_None)
-// 						{
-// 							// user asked for an actor and a bone or socket, try to find a suitable component
-// 							TArray<USceneComponent*> OutComponents;
-// 							TrackedActor->GetComponents(OutComponents);
-// 							for (USceneComponent const* Comp : OutComponents)
-// 							{
-// 								if (Comp->DoesSocketExist(FocusSettings.TrackingFocusSettings.BoneOrSocketName))
-// 								{
-// 									FVector const SocketLoc = Comp->GetSocketLocation(FocusSettings.TrackingFocusSettings.BoneOrSocketName);
-// 									DesiredFocusDistance = (SocketLoc - DesiredView.Location).Size();
-// 									break;
-// 								}
-// 							}
-// 
-// 						}
-// 						else
-// 						{
-// 							DesiredFocusDistance = (TrackedActor->GetActorLocation() - DesiredView.Location).Size();
-// 						}
-// 					}
-// 				}
-// 			}
 		}
 		break;
 	}
@@ -367,6 +296,7 @@ void UCineCameraComponent::CreateDebugFocusPlane()
 			DebugFocusPlaneComponent->CastShadow = false;
 			DebugFocusPlaneComponent->PostPhysicsComponentTick.bCanEverTick = false;
 			DebugFocusPlaneComponent->CreationMethod = CreationMethod;
+			DebugFocusPlaneComponent->bSelectable = false;
 
 			DebugFocusPlaneComponent->RelativeScale3D = FVector(10000.f, 10000.f, 1.f);
 			DebugFocusPlaneComponent->RelativeRotation = FRotator(90.f, 0.f, 0.f);
diff --git a/Engine/Source/Runtime/CinematicCamera/Public/CameraRig_Crane.h b/Engine/Source/Runtime/CinematicCamera/Public/CameraRig_Crane.h
--- a/Engine/Source/Runtime/CinematicCamera/Public/CameraRig_Crane.h
+++ b/Engine/Source/Runtime/CinematicCamera/Public/CameraRig_Crane.h
@@ -9,7 +9,7 @@
 class USceneComponent;
 
 /** 
- * 
+ * A simple rig for to simulate crane-like camera movements.
  */
 UCLASS(Blueprintable)
 class CINEMATICCAMERA_API ACameraRig_Crane : public AActor
@@ -24,12 +24,15 @@ public:
 	virtual void Tick(float DeltaTime) override;
 	virtual bool ShouldTickIfViewportsOnly() const override;
 
+	/** Controls the pitch of the crane arm (but does not affect the rotation of the mount). */
 	UPROPERTY(Interp, EditAnywhere, BlueprintReadWrite, Category = "Crane Controls", meta = (UIMin = "-360", UIMax = "360", Units = deg))
 	float CranePitch;
 	
+	/** Controls the yaw of the crane arm. */
 	UPROPERTY(Interp, EditAnywhere, BlueprintReadWrite, Category = "Crane Controls", meta = (UIMin = "-360", UIMax = "360", Units = deg))
 	float CraneYaw;
 	
+	/** Controls the length of the crane arm. */
 	UPROPERTY(Interp, EditAnywhere, BlueprintReadWrite, Category = "Crane Controls", meta = (ClampMin=0, Units = cm))
 	float CraneArmLength;
 
@@ -39,7 +42,6 @@ public:
 #endif
 	
 private:
-
 	void UpdatePreviewMeshes();
 	void UpdateCraneComponents();
 
@@ -47,15 +49,19 @@ private:
 	UPROPERTY(EditDefaultsOnly, Category = "Crane Components")
 	USceneComponent* TransformComponent;
 
+	/** Component to control Yaw. */
 	UPROPERTY(EditDefaultsOnly, Category = "Crane Components")
 	USceneComponent* CraneYawControl;
 
+	/** Component to control Pitch. */
 	UPROPERTY(EditDefaultsOnly, Category = "Crane Components")
 	USceneComponent* CranePitchControl;
 
+	/** Component to define the attach point for cameras. */
 	UPROPERTY(EditDefaultsOnly, Category = "Crane Components")
 	USceneComponent* CraneCameraMount;
 
+	/** Preview meshes for visualization */
 	UPROPERTY()
 	UStaticMeshComponent* PreviewMesh_CraneArm;
 
diff --git a/Engine/Source/Runtime/CinematicCamera/Public/CameraRig_Rail.h b/Engine/Source/Runtime/CinematicCamera/Public/CameraRig_Rail.h
--- a/Engine/Source/Runtime/CinematicCamera/Public/CameraRig_Rail.h
+++ b/Engine/Source/Runtime/CinematicCamera/Public/CameraRig_Rail.h
@@ -23,13 +23,17 @@ public:
 	virtual void Tick(float DeltaTime) override;
 	virtual bool ShouldTickIfViewportsOnly() const override;
 
-	UPROPERTY(Interp, EditAnywhere, BlueprintReadWrite, Category = "Rail Controls")
+	/** Defines current position of the mount point along the rail, in terms of normalized distance from the beginning of the rail. */
+	UPROPERTY(Interp, EditAnywhere, BlueprintReadWrite, Category = "Rail Controls", meta=(ClampMin="0.0", ClampMax = "1.0"))
 	float CurrentPositionOnRail;
 
 #if WITH_EDITOR
 	virtual class USceneComponent* GetDefaultAttachComponent() const override;
 	virtual void PostEditChangeProperty(FPropertyChangedEvent& PropertyChangedEvent) override;
 #endif
+
+	/** Returns the spline component that defines the rail path */
+	USplineComponent* GetRailSplineComponent() { return RailSplineComponent; }
 	
 private:
 	/** Makes sure all components are arranged properly. Call when something changes that might affect components. */
@@ -41,12 +45,16 @@ private:
 	UPROPERTY(EditDefaultsOnly, Category = "Rail Components")
  	USceneComponent* TransformComponent;
 
+	/** Spline component to define the rail path. */
 	UPROPERTY(EditDefaultsOnly, Category = "Rail Components")
 	USplineComponent* RailSplineComponent;
 
+	/** Component to define the attach point for cameras. Moves along the rail. */
 	UPROPERTY(EditDefaultsOnly, Category = "Rail Components")
 	USceneComponent* RailCameraMount;
 
+	
+	/** Preview meshes for visualization */
 	UPROPERTY()
 	USplineMeshComponent* PreviewMesh_Rail;
 
diff --git a/Engine/Source/Runtime/CinematicCamera/Public/CineCameraActor.h b/Engine/Source/Runtime/CinematicCamera/Public/CineCameraActor.h
--- a/Engine/Source/Runtime/CinematicCamera/Public/CineCameraActor.h
+++ b/Engine/Source/Runtime/CinematicCamera/Public/CineCameraActor.h
@@ -4,6 +4,7 @@
 
 #include "CineCameraActor.generated.h"
 
+/** Settings to control the camera's lookat feature */
 USTRUCT()
 struct FCameraLookatTrackingSettings
 {
@@ -13,6 +14,7 @@ struct FCameraLookatTrackingSettings
 	UPROPERTY(/*Interp, */EditAnywhere, BlueprintReadWrite, Category = "LookAt")
 	uint8 bEnableLookAtTracking : 1;
 
+	/** True to draw a debug representation of the lookat location */
 	UPROPERTY(Transient, EditAnywhere, BlueprintReadWrite, Category = "LookAt")
 	uint8 bDrawDebugLookAtTrackingPosition : 1;
 
@@ -23,6 +25,7 @@ struct FCameraLookatTrackingSettings
 	/** Last known lookat tracking rotation (used during interpolation) */
 	FRotator LastLookatTrackingRotation;
 
+	/** If set, camera will track this actor's location */
 	UPROPERTY(Interp, EditAnywhere, BlueprintReadWrite, Category = "LookAt")
 	AActor* ActorToTrack;
 
@@ -55,7 +58,6 @@ public:
 #endif
 
 protected:
-
 	/** Set to true to skip any interpolations on the next update. Resets to false automatically. */
 	uint8 bResetInterplation : 1;
 
diff --git a/Engine/Source/Runtime/CinematicCamera/Public/CineCameraComponent.h b/Engine/Source/Runtime/CinematicCamera/Public/CineCameraComponent.h
--- a/Engine/Source/Runtime/CinematicCamera/Public/CineCameraComponent.h
+++ b/Engine/Source/Runtime/CinematicCamera/Public/CineCameraComponent.h
@@ -23,11 +23,13 @@ struct FCameraFilmbackSettings
 	float SensorAspectRatio;
 };
 
+/** A named bundle of filmback settings used to implement filmback presets */
 USTRUCT()
 struct FNamedFilmbackPreset
 {
 	GENERATED_USTRUCT_BODY()
 
+	/** Name for the preset. */
 	UPROPERTY()
 	FString Name;
 
@@ -35,49 +37,51 @@ struct FNamedFilmbackPreset
 	FCameraFilmbackSettings FilmbackSettings;
 };
 
-/** #note, this struct has a details customization in CameraLensSettingsCustomization.cpp/h */
+/** 
+ * #note, this struct has a details customization in CameraLensSettingsCustomization.cpp/h
+ */
 USTRUCT()
 struct FCameraLensSettings
 {
 	GENERATED_USTRUCT_BODY()
 
+	/** Minimum focal length for this lens */
 	UPROPERTY(EditAnywhere, BlueprintReadWrite, Category = "Lens", meta = (ForceUnits = mm))
 	float MinFocalLength;
 
+	/** Maximum focal length for this lens */
 	UPROPERTY(EditAnywhere, BlueprintReadWrite, Category = "Lens", meta = (ForceUnits = mm))
 	float MaxFocalLength;
 
-	// #todo, split into MinFStopAtMinFocalLength, MinFStopAtMaxFocalLength for variable-min-fstop lenses
-
+	/** Minimum aperture for this lens (e.g. 2.8 for an f/2.8 lens) */
 	UPROPERTY(EditAnywhere, BlueprintReadWrite, Category = "Lens")
 	float MinFStop;
 
+	/** Minimum aperture for this lens (e.g. 2.8 for an f/2.8 lens) */
 	UPROPERTY(EditAnywhere, BlueprintReadWrite, Category = "Lens")
 	float MaxFStop;
 
+	/** Shortest distance this lens can focus on. */
 	UPROPERTY(EditAnywhere, BlueprintReadWrite, Category = "Lens", meta = (ForceUnits = mm))
 	float MinimumFocusDistance;
 
-	UPROPERTY(EditAnywhere, BlueprintReadWrite, Category = "Lens")
-	float MaxReproductionRatio;
-
 	bool operator==(const FCameraLensSettings& Other) const
 	{
 		return (MinFocalLength == Other.MinFocalLength)
 			&& (MaxFocalLength == Other.MaxFocalLength)
 			&& (MinFStop == Other.MinFStop)
 			&& (MaxFStop == Other.MaxFStop)
-			&& (MinimumFocusDistance == Other.MinimumFocusDistance)
-			&& (MaxReproductionRatio == Other.MaxReproductionRatio);
+			&& (MinimumFocusDistance == Other.MinimumFocusDistance);
 	}
 };
 
-
+/** A named bundle of lens settings used to implement lens presets. */
 USTRUCT()
 struct FNamedLensPreset
 {
 	GENERATED_USTRUCT_BODY()
 
+	/** Name for the preset. */
 	UPROPERTY()
 	FString Name;
 
@@ -85,77 +89,59 @@ struct FNamedLensPreset
 	FCameraLensSettings LensSettings;
 };
 
+/** Supported methods for focusing the camera. */
 UENUM()
 enum class ECameraFocusMethod : uint8
 {
-	None,		/** Disable DoF entirely. */
-	Manual,		/** Allows for specifying or animating exact focus distances. */
-	Spot UMETA(Hidden), /** Raycasts into the scene, focuses on the hit location. */		// #todo, make this work
-	Tracking,	/** Locks focus to specific object. */
-};
-
-USTRUCT()
-struct FCameraSpotFocusSettings
-{
-	GENERATED_USTRUCT_BODY()
-
-	UPROPERTY(EditAnywhere, BlueprintReadWrite, Category="Spot Focus")
-	TArray<FVector2D> NormalizedSpotCoordinates;
-	
-	FCameraSpotFocusSettings()
-	{
-		NormalizedSpotCoordinates.Add(FVector2D(0.5f, 0.5f));
-	}
+	None,					/** Disables DoF entirely. */
+	Manual,					/** Allows for specifying or animating exact focus distances. */
+	Tracking,				/** Locks focus to specific object. */
 };
 
+/** Settings to control tracking-focus mode. */
 USTRUCT()
 struct FCameraTrackingFocusSettings
 {
 	GENERATED_USTRUCT_BODY()
 
+	/** Focus distance will be tied to this actor's location. */
 	UPROPERTY(EditAnywhere, BlueprintReadWrite, Category = "Tracking Focus")
 	AActor* ActorToTrack;
 
-	// #todo, ability to pick a component and socket/bone?
-
-// 	UPROPERTY(EditAnywhere, BlueprintReadWrite, Category = "Tracking Focus")
-// 	USceneComponent* ComponentToTrack;
-// 	
-// 	UPROPERTY(EditAnywhere, BlueprintReadWrite, Category = "Tracking Focus")
-// 	FName BoneOrSocketName;
-
-	/** Offset from actor position to look at. Relative to actor if tracking an actor, relative to world otherwise. */
+	/** Offset from actor position to track. Relative to actor if tracking an actor, relative to world otherwise. */
 	UPROPERTY(Interp, EditAnywhere, BlueprintReadWrite, Category = "Tracking Focus")
 	FVector RelativeOffset;
 
+	/** True to draw a debug representation of the tracked position. */
 	UPROPERTY(Transient, EditAnywhere, BlueprintReadWrite, Category = "Tracking Focus")
 	uint8 bDrawDebugTrackingFocusPoint : 1;
 
 	FCameraTrackingFocusSettings()
-		: ActorToTrack(nullptr)
+		: ActorToTrack(nullptr),
+		bDrawDebugTrackingFocusPoint(false)
 	{}
 };
 
-
+/** Settings to control camera focus */
 USTRUCT()
 struct FCameraFocusSettings
 {
 	GENERATED_USTRUCT_BODY()
 
+	/** Which method to use to handle camera focus */
 	UPROPERTY(EditAnywhere, BlueprintReadWrite, Category = "Focus Method")
 	TEnumAsByte<ECameraFocusMethod> FocusMethod;
 	
-	UPROPERTY(Interp, EditAnywhere, BlueprintReadWrite, Category = "Manual Focus Settings")
+	/** Manually-controlled focus distance (manual focus mode only) */
+	UPROPERTY(Interp, EditAnywhere, BlueprintReadWrite, Category = "Manual Focus Settings", meta=(Units=cm))
 	float ManualFocusDistance;
 
-	UPROPERTY(EditAnywhere, BlueprintReadWrite, Category = "Spot Focus Settings")
-	FCameraSpotFocusSettings SpotFocusSettings;
-
+	/** Settings to control tracking focus (tracking focus mode only) */
 	UPROPERTY(EditAnywhere, BlueprintReadWrite, Category = "Tracking Focus Settings")
 	FCameraTrackingFocusSettings TrackingFocusSettings;
 
 #if WITH_EDITORONLY_DATA
-	/** True to draw a translucent plane at the current focus depth, for easy tweaking when not looking through the camera. */
+	/** True to draw a translucent plane at the current focus depth, for easy tweaking. */
 	UPROPERTY(Transient, EditAnywhere, Category = "Focus Settings")
 	uint8 bDrawDebugFocusPlane : 1;
 
@@ -164,13 +150,15 @@ struct FCameraFocusSettings
 	FColor DebugFocusPlaneColor;
 #endif 
 
+	/** True to use interpolation to smooth out changes in focus distance, false for focus distance changes to be instantaneous. */
 	UPROPERTY(EditAnywhere, BlueprintReadWrite, Category = "Focus Settings")
 	uint8 bSmoothFocusChanges : 1;
 	
+	/** Controls interpolation speed when smoothing focus distance changes. Ignored if bSmoothFocusChanges is false. */
 	UPROPERTY(EditAnywhere, BlueprintReadWrite, Category = "Focus Settings", meta = (EditCondition = "bSmoothFocusChanges"))
 	float FocusSmoothingInterpSpeed;
 
-	/** Additional focus depth offset, used for manually tweaking if yoour chosen focus method is needs adjustment */
+	/** Additional focus depth offset, used for manually tweaking if your chosen focus method needs adjustment */
 	UPROPERTY(Interp, EditAnywhere, BlueprintReadWrite, Category = "Focus Settings")
 	float FocusOffset;
 
@@ -178,10 +166,10 @@ struct FCameraFocusSettings
 		FocusMethod(ECameraFocusMethod::Manual),
 		ManualFocusDistance(100000.f),
 #if WITH_EDITORONLY_DATA
-		bDrawDebugFocusPlane(0),
+		bDrawDebugFocusPlane(false),
 		DebugFocusPlaneColor(102, 26, 204, 153),		// purple
 #endif
-		bSmoothFocusChanges(0),
+		bSmoothFocusChanges(false),
 		FocusSmoothingInterpSpeed(8.f),
 		FocusOffset(0.f)
 	{}
@@ -208,18 +196,23 @@ public:
 
 	virtual void GetCameraView(float DeltaTime, FMinimalViewInfo& DesiredView) override;
 
+	/** Controls the filmback of the camera. */
 	UPROPERTY(EditAnywhere, BlueprintReadWrite, Category = "Current Camera Settings")
 	FCameraFilmbackSettings FilmbackSettings;
 
+	/** Controls the camera's lens. */
 	UPROPERTY(EditAnywhere, BlueprintReadWrite, Category = "Current Camera Settings")
 	FCameraLensSettings LensSettings;
 
+	/** Controls the camera's focus. */
 	UPROPERTY(EditAnywhere, BlueprintReadWrite, Category = "Current Camera Settings")
 	FCameraFocusSettings FocusSettings;
 
+	/** Current focal length of the camera (i.e. controls FoV, zoom) */
 	UPROPERTY(Interp, EditAnywhere, BlueprintReadWrite, Category = "Current Camera Settings")
 	float CurrentFocalLength;
 
+	/** Current aperture, in terms of f-stop (e.g. 2.8 for f/2.8) */
 	UPROPERTY(Interp, EditAnywhere, BlueprintReadWrite, Category = "Current Camera Settings")
 	float CurrentAperture;
 	
@@ -231,32 +224,24 @@ public:
 	/** Read-only. Control this value with CurrentFocalLength (and filmback settings). */
 	UPROPERTY(VisibleAnywhere, BlueprintReadOnly, Category = "Current Camera Settings")
 	float CurrentHorizontalFOV;
-
-	UPROPERTY(transient)
-	UStaticMesh* DebugFocusPlaneMesh;
-
-	UPROPERTY(transient)
-	UMaterial* DebugFocusPlaneMaterial;
-
-	// The camera mesh to show visually where the camera is placed
-	UPROPERTY(transient)
-	UStaticMeshComponent* DebugFocusPlaneComponent;
-
-	UPROPERTY(transient)
-	UMaterialInstanceDynamic* DebugFocusPlaneMID;
 #endif
 	
-	/** Most recent calculated focus distance. Used for interpolation. */
-	float LastFocusDistance;
-
+	/** Returns the horizonal FOV of the camera with current settings. */
 	float GetHorizontalFieldOfView() const;
+	/** Returns the vertical FOV of the camera with current settings. */
 	float GetVerticalFieldOfView() const;
 
+	/** Returns a list of available filmback presets. */
 	static TArray<FNamedFilmbackPreset> const& GetFilmbackPresets();
+	
+	/** Returns a list of available lens presets. */
 	static TArray<FNamedLensPreset> const& GetLensPresets();
 
 protected:
 
+	/** Most recent calculated focus distance. Used for interpolation. */
+	float LastFocusDistance;
+
 	/** Set to true to skip any interpolations on the next update. Resets to false automatically. */
 	uint8 bResetInterpolation : 1;
 
@@ -269,19 +254,40 @@ protected:
 #endif
 
 #if WITH_EDITORONLY_DATA
+
+	/** Mesh used for debug focus plane visualization */
+	UPROPERTY(transient)
+	UStaticMesh* DebugFocusPlaneMesh;
+
+	/** Material used for debug focus plane visualization */
+	UPROPERTY(transient)
+	UMaterial* DebugFocusPlaneMaterial;
+
+	/** Component for the debug focus plane visualization */
+	UPROPERTY(transient)
+	UStaticMeshComponent* DebugFocusPlaneComponent;
+
+	/** Dynamic material instance for the debug focus plane visualization */
+	UPROPERTY(transient)
+	UMaterialInstanceDynamic* DebugFocusPlaneMID;
+
 	virtual void PostEditChangeProperty(FPropertyChangedEvent& PropertyChangedEvent) override;
 	virtual void ResetProxyMeshTransform() override;
 #endif
 
+	/** List of available filmback presets */
 	UPROPERTY(config)
 	TArray<FNamedFilmbackPreset> FilmbackPresets;
 
+	/** List of available lens presets */
 	UPROPERTY(config)
 	TArray<FNamedLensPreset> LensPresets;
 
+	/** Name of the default filmback preset */
 	UPROPERTY(config)
 	FString DefaultFilmbackPresetName;
 
+	/** Name of the default lens preset */
 	UPROPERTY(config)
 	FString DefaultLensPresetName;
 	
diff --git a/Engine/Source/Runtime/Core/Private/Containers/Ticker.cpp b/Engine/Source/Runtime/Core/Private/Containers/Ticker.cpp
--- a/Engine/Source/Runtime/Core/Private/Containers/Ticker.cpp
+++ b/Engine/Source/Runtime/Core/Private/Containers/Ticker.cpp
@@ -26,7 +26,7 @@ FDelegateHandle FTicker::AddTicker(const FTickerDelegate& InDelegate, float InDe
 void FTicker::RemoveTicker(FDelegateHandle Handle)
 {
 	// must remove the handle from both arrays because we could be in the middle of a tick, 
-	// and may be removing ourselves or soemthing else we've already considered this Tick
+	// and may be removing ourselves or something else we've already considered this Tick
 	auto CompareHandle = [=](const FElement& Element){ return Element.Delegate.GetHandle() == Handle; };
 	// We can remove elements safely even during tick.
 	Elements.RemoveAllSwap(CompareHandle);
diff --git a/Engine/Source/Runtime/Core/Private/Misc/Timespan.cpp b/Engine/Source/Runtime/Core/Private/Misc/Timespan.cpp
--- a/Engine/Source/Runtime/Core/Private/Misc/Timespan.cpp
+++ b/Engine/Source/Runtime/Core/Private/Misc/Timespan.cpp
@@ -6,7 +6,7 @@
 /* FTimespan interface
  *****************************************************************************/
 
-bool FTimespan::ExportTextItem( FString& ValueStr, FTimespan const& DefaultValue, UObject* Parent, int32 PortFlags, UObject* ExportRootScope ) const
+bool FTimespan::ExportTextItem(FString& ValueStr, FTimespan const& DefaultValue, UObject* Parent, int32 PortFlags, UObject* ExportRootScope) const
 {
 	if (0 != (PortFlags & EPropertyPortFlags::PPF_ExportCpp))
 	{
@@ -20,14 +20,14 @@ bool FTimespan::ExportTextItem( FString& ValueStr, FTimespan const& DefaultValue
 }
 
 
-bool FTimespan::ImportTextItem( const TCHAR*& Buffer, int32 PortFlags, UObject* Parent, FOutputDevice* ErrorText )
+bool FTimespan::ImportTextItem(const TCHAR*& Buffer, int32 PortFlags, UObject* Parent, FOutputDevice* ErrorText)
 {
 	// @todo gmp: implement FTimespan::ImportTextItem
 	return false;
 }
 
 
-bool FTimespan::Serialize( FArchive& Ar )
+bool FTimespan::Serialize(FArchive& Ar)
 {
 	Ar << *this;
 
@@ -46,7 +46,7 @@ FString FTimespan::ToString() const
 }
 
 
-FString FTimespan::ToString( const TCHAR* Format ) const
+FString FTimespan::ToString(const TCHAR* Format) const
 {
 	FString Result;
 
@@ -89,35 +89,47 @@ FString FTimespan::ToString( const TCHAR* Format ) const
 /* FTimespan static interface
  *****************************************************************************/
 
-FTimespan FTimespan::FromDays( double Days )
+FTimespan FTimespan::FromDays(double Days)
 {
 	check((Days >= MinValue().GetTotalDays()) && (Days <= MaxValue().GetTotalDays()));
 
 	return FTimespan(Days * ETimespan::TicksPerDay);
 }
 
-FTimespan FTimespan::FromHours( double Hours )
+
+FTimespan FTimespan::FromHours(double Hours)
 {
 	check((Hours >= MinValue().GetTotalHours()) && (Hours <= MaxValue().GetTotalHours()));
 
 	return FTimespan(Hours * ETimespan::TicksPerHour);
 }
 
-FTimespan FTimespan::FromMilliseconds( double Milliseconds )
+
+FTimespan FTimespan::FromMicroseconds(double Microseconds)
+{
+	check((Microseconds >= MinValue().GetTotalMicroseconds()) && (Microseconds <= MaxValue().GetTotalMicroseconds()));
+
+	return FTimespan(Microseconds * ETimespan::TicksPerMicrosecond);
+}
+
+
+FTimespan FTimespan::FromMilliseconds(double Milliseconds)
 {
 	check((Milliseconds >= MinValue().GetTotalMilliseconds()) && (Milliseconds <= MaxValue().GetTotalMilliseconds()));
 
 	return FTimespan(Milliseconds * ETimespan::TicksPerMillisecond);
 }
 
-FTimespan FTimespan::FromMinutes( double Minutes )
+
+FTimespan FTimespan::FromMinutes(double Minutes)
 {
 	check((Minutes >= MinValue().GetTotalMinutes()) && (Minutes <= MaxValue().GetTotalMinutes()));
 
 	return FTimespan(Minutes * ETimespan::TicksPerMinute);
 }
 
-FTimespan FTimespan::FromSeconds( double Seconds )
+
+FTimespan FTimespan::FromSeconds(double Seconds)
 {
 	check((Seconds >= MinValue().GetTotalSeconds()) && (Seconds <= MaxValue().GetTotalSeconds()));
 
@@ -125,13 +137,12 @@ FTimespan FTimespan::FromSeconds( double Seconds )
 }
 
 
-bool FTimespan::Parse( const FString& TimespanString, FTimespan& OutTimespan )
+bool FTimespan::Parse(const FString& TimespanString, FTimespan& OutTimespan)
 {
 	// @todo gmp: implement stricter FTimespan parsing; this implementation is too forgiving.
 	FString TokenString = TimespanString.Replace(TEXT("."), TEXT(":"));
 
 	bool Negative = TokenString.StartsWith(TEXT("-"));
-
 	TokenString.ReplaceInline(TEXT("-"), TEXT(":"), ESearchCase::CaseSensitive);
 
 	TArray<FString> Tokens;
@@ -142,26 +153,26 @@ bool FTimespan::Parse( const FString& TimespanString, FTimespan& OutTimespan )
 		Tokens.Insert(TEXT("0"), 0);
 	}
 
-	if (Tokens.Num() == 5)
+	if (Tokens.Num() != 5)
 	{
-		OutTimespan.Assign(FCString::Atoi(*Tokens[0]), FCString::Atoi(*Tokens[1]), FCString::Atoi(*Tokens[2]), FCString::Atoi(*Tokens[3]), FCString::Atoi(*Tokens[4]));
+		return false;
+	}
 
-		if (Negative)
-		{
-			OutTimespan.Ticks *= -1;
-		}
+	OutTimespan.Assign(FCString::Atoi(*Tokens[0]), FCString::Atoi(*Tokens[1]), FCString::Atoi(*Tokens[2]), FCString::Atoi(*Tokens[3]), FCString::Atoi(*Tokens[4]), 0);
 
-		return true;
+	if (Negative)
+	{
+		OutTimespan.Ticks *= -1;
 	}
 
-	return false;
+	return true;
 }
 
 
 /* FTimespan friend functions
  *****************************************************************************/
 
-FArchive& operator<<( FArchive& Ar, FTimespan& Timespan )
+FArchive& operator<<(FArchive& Ar, FTimespan& Timespan)
 {
 	return Ar << Timespan.Ticks;
 }
@@ -170,10 +181,10 @@ FArchive& operator<<( FArchive& Ar, FTimespan& Timespan )
 /* FTimespan implementation
  *****************************************************************************/
 
-void FTimespan::Assign( int32 Days, int32 Hours, int32 Minutes, int32 Seconds, int32 Milliseconds )
+void FTimespan::Assign(int32 Days, int32 Hours, int32 Minutes, int32 Seconds, int32 Milliseconds, int32 Microseconds)
 {
-	int64 totalms = 1000 * (60 * 60 * 24 * (int64)Days + 60 * 60 * (int64)Hours + 60 * (int64)Minutes + (int64)Seconds) + (int64)Milliseconds;
-	check((totalms >= MinValue().GetTotalMilliseconds()) && (totalms <= MaxValue().GetTotalMilliseconds()));
+	int64 TotalTicks = ETimespan::TicksPerMicrosecond * (1000 * (1000 * (60 * 60 * 24 * (int64)Days + 60 * 60 * (int64)Hours + 60 * (int64)Minutes + (int64)Seconds) + (int64)Milliseconds) + (int64)Microseconds);
+	check((TotalTicks >= MinValue().GetTicks()) && (TotalTicks <= MaxValue().GetTicks()));
 
-	Ticks = totalms * ETimespan::TicksPerMillisecond;
+	Ticks = TotalTicks;
 }
diff --git a/Engine/Source/Runtime/Core/Private/Tests/Misc/TimespanTest.cpp b/Engine/Source/Runtime/Core/Private/Tests/Misc/TimespanTest.cpp
--- a/Engine/Source/Runtime/Core/Private/Tests/Misc/TimespanTest.cpp
+++ b/Engine/Source/Runtime/Core/Private/Tests/Misc/TimespanTest.cpp
@@ -10,54 +10,70 @@ IMPLEMENT_SIMPLE_AUTOMATION_TEST(FTimespanTest, "System.Core.Misc.Timespan", EAu
 bool FTimespanTest::RunTest( const FString& Parameters )
 {
 	// constructors must create equal objects
-	FTimespan ts1_1 = FTimespan(3, 2, 1);
-	FTimespan ts1_2 = FTimespan(0, 3, 2, 1);
-	FTimespan ts1_3 = FTimespan(0, 3, 2, 1, 0);
+	{
+		FTimespan ts1_1 = FTimespan(3, 2, 1);
+		FTimespan ts1_2 = FTimespan(0, 3, 2, 1);
+		FTimespan ts1_3 = FTimespan(0, 3, 2, 1, 0);
 
-	TestEqual(TEXT("Constructors must create equal objects (Hours/Minutes/Seconds vs. Days/Hours/Minutes/Seconds)"), ts1_1, ts1_2);
-	TestEqual(TEXT("Constructors must create equal objects (Hours/Minutes/Seconds vs. Days/Hours/Minutes/Seconds/Milliseconds)"), ts1_1, ts1_3);
+		TestEqual(TEXT("Constructors must create equal objects (Hours/Minutes/Seconds vs. Days/Hours/Minutes/Seconds)"), ts1_1, ts1_2);
+		TestEqual(TEXT("Constructors must create equal objects (Hours/Minutes/Seconds vs. Days/Hours/Minutes/Seconds/Milliseconds)"), ts1_1, ts1_3);
+	}
 
 	// component getters must return correct values
-	FTimespan ts2_1 = FTimespan(1, 2, 3, 4, 5);
+	{
+		FTimespan ts2_1 = FTimespan(1, 2, 3, 4, 5, 6);
 
-	TestEqual(TEXT("Component getters must return correct values (Days)"), ts2_1.GetDays(), 1);
-	TestEqual(TEXT("Component getters must return correct values (Hours)"), ts2_1.GetHours(), 2);
-	TestEqual(TEXT("Component getters must return correct values (Minutes)"), ts2_1.GetMinutes(), 3);
-	TestEqual(TEXT("Component getters must return correct values (Seconds)"), ts2_1.GetSeconds(), 4);
-	TestEqual(TEXT("Component getters must return correct values (Milliseconds)"), ts2_1.GetMilliseconds(), 5);
+		TestEqual(TEXT("Component getters must return correct values (Days)"), ts2_1.GetDays(), 1);
+		TestEqual(TEXT("Component getters must return correct values (Hours)"), ts2_1.GetHours(), 2);
+		TestEqual(TEXT("Component getters must return correct values (Minutes)"), ts2_1.GetMinutes(), 3);
+		TestEqual(TEXT("Component getters must return correct values (Seconds)"), ts2_1.GetSeconds(), 4);
+		TestEqual(TEXT("Component getters must return correct values (Milliseconds)"), ts2_1.GetMilliseconds(), 5);
+		TestEqual(TEXT("Component getters must return correct values (Microseconds)"), ts2_1.GetMicroseconds(), 6);
+	}
 
 	// durations of positive and negative time spans must match
-	FTimespan ts3_1 = FTimespan(1, 2, 3, 4, 5);
-	FTimespan ts3_2 = FTimespan(-1, -2, -3, -4, -5);
+	{
+		FTimespan ts3_1 = FTimespan(1, 2, 3, 4, 5, 6);
+		FTimespan ts3_2 = FTimespan(-1, -2, -3, -4, -5, -6);
 
-	TestEqual(TEXT("Durations of positive and negative time spans must match"), ts3_1.GetDuration(), ts3_2.GetDuration());
+		TestEqual(TEXT("Durations of positive and negative time spans must match"), ts3_1.GetDuration(), ts3_2.GetDuration());
+	}
 
 	// static constructors must create correct values
-	TestEqual(TEXT("Static constructors must create correct values (FromDays)"), FTimespan::FromDays(123).GetTotalDays(), 123.0);
-	TestEqual(TEXT("Static constructors must create correct values (FromHours)"), FTimespan::FromHours(123).GetTotalHours(), 123.0);
-	TestEqual(TEXT("Static constructors must create correct values (FromMinutes)"), FTimespan::FromMinutes(123).GetTotalMinutes(), 123.0);
-	TestEqual(TEXT("Static constructors must create correct values (FromSeconds)"), FTimespan::FromSeconds(123).GetTotalSeconds(), 123.0);
-	TestEqual(TEXT("Static constructors must create correct values (FromMilliseconds)"), FTimespan::FromMilliseconds(123).GetTotalMilliseconds(), 123.0);
+	{
+		TestEqual(TEXT("Static constructors must create correct values (FromDays)"), FTimespan::FromDays(123).GetTotalDays(), 123.0);
+		TestEqual(TEXT("Static constructors must create correct values (FromHours)"), FTimespan::FromHours(123).GetTotalHours(), 123.0);
+		TestEqual(TEXT("Static constructors must create correct values (FromMinutes)"), FTimespan::FromMinutes(123).GetTotalMinutes(), 123.0);
+		TestEqual(TEXT("Static constructors must create correct values (FromSeconds)"), FTimespan::FromSeconds(123).GetTotalSeconds(), 123.0);
+		TestEqual(TEXT("Static constructors must create correct values (FromMilliseconds)"), FTimespan::FromMilliseconds(123).GetTotalMilliseconds(), 123.0);
+		TestEqual(TEXT("Static constructors must create correct values (FromMicroseconds)"), FTimespan::FromMicroseconds(123).GetTotalMicroseconds(), 123.0);
+	}
 
 	// string conversions must return correct strings
-	FTimespan ts5_1 = FTimespan(1, 2, 3, 4, 5);
+	{
+		FTimespan ts5_1 = FTimespan(1, 2, 3, 4, 5);
 
-	TestEqual<FString>(TEXT("String conversion (Default)"), ts5_1.ToString(), TEXT("1.02:03:04.005"));
-	TestEqual<FString>(TEXT("String conversion (%n%d.%h:%m:%s.%f)"), ts5_1.ToString(TEXT("%n%d.%h:%m:%s.%f")), TEXT("1.02:03:04.005"));
+		TestEqual<FString>(TEXT("String conversion (Default)"), ts5_1.ToString(), TEXT("1.02:03:04.005"));
+		TestEqual<FString>(TEXT("String conversion (%n%d.%h:%m:%s.%f)"), ts5_1.ToString(TEXT("%n%d.%h:%m:%s.%f")), TEXT("1.02:03:04.005"));
+	}
 
 	// parsing valid strings must succeed
-	FTimespan ts6_1 = FTimespan(1, 2, 3, 4, 5);
-	FTimespan ts6_2;
+	{
+		FTimespan ts6_1 = FTimespan(1, 2, 3, 4, 5);
+		FTimespan ts6_2;
 
-	TestTrue(TEXT("Parsing valid strings must succeed (1.02:03:04.005)"), FTimespan::Parse(TEXT("1.02:03:04.005"), ts6_2));
-	TestEqual(TEXT("Parsing valid strings must result in correct values (1.02:03:04.005)"), ts6_2, ts6_1);
+		TestTrue(TEXT("Parsing valid strings must succeed (1.02:03:04.005)"), FTimespan::Parse(TEXT("1.02:03:04.005"), ts6_2));
+		TestEqual(TEXT("Parsing valid strings must result in correct values (1.02:03:04.005)"), ts6_2, ts6_1);
+	}
 
 	// parsing invalid strings must fail
-	FTimespan ts7_1;
+	{
+		FTimespan ts7_1;
 
-	//TestFalse(TEXT("Parsing invalid strings must fail (1,02:03:04.005)"), FTimespan::Parse(TEXT("1,02:03:04.005"), ts7_1));
-	//TestFalse(TEXT("Parsing invalid strings must fail (1.1.02:03:04:005)"), FTimespan::Parse(TEXT("1.1.02:03:04:005"), ts7_1));
-	//TestFalse(TEXT("Parsing invalid strings must fail (04:005)"), FTimespan::Parse(TEXT("04:005"), ts7_1));
+		//TestFalse(TEXT("Parsing invalid strings must fail (1,02:03:04.005)"), FTimespan::Parse(TEXT("1,02:03:04.005"), ts7_1));
+		//TestFalse(TEXT("Parsing invalid strings must fail (1.1.02:03:04:005)"), FTimespan::Parse(TEXT("1.1.02:03:04:005"), ts7_1));
+		//TestFalse(TEXT("Parsing invalid strings must fail (04:005)"), FTimespan::Parse(TEXT("04:005"), ts7_1));
+	}
 
 	return true;
 }
diff --git a/Engine/Source/Runtime/Core/Public/Containers/Array.h b/Engine/Source/Runtime/Core/Public/Containers/Array.h
--- a/Engine/Source/Runtime/Core/Public/Containers/Array.h
+++ b/Engine/Source/Runtime/Core/Public/Containers/Array.h
@@ -20,6 +20,7 @@
 
 #define AGRESSIVE_ARRAY_FORCEINLINE
 
+
 /**
  * Generic iterator which can operate on types that expose the following:
  * - A type called ElementType representing the contained type.
@@ -127,10 +128,12 @@ public:
 	FORCEINLINE friend bool operator!=(const TIndexedContainerIterator& Lhs, const TIndexedContainerIterator& Rhs) { return &Lhs.Container != &Rhs.Container || Lhs.Index != Rhs.Index; }
 
 private:
+
 	ContainerType& Container;
 	IndexType      Index;
 };
 
+
 /** operator + */
 template <typename ContainerType, typename ElementType, typename IndexType>
 FORCEINLINE TIndexedContainerIterator<ContainerType, ElementType, IndexType> operator+(int32 Offset, TIndexedContainerIterator<ContainerType, ElementType, IndexType> RHS)
@@ -138,8 +141,8 @@ FORCEINLINE TIndexedContainerIterator<ContainerType, ElementType, IndexType> ope
 	return RHS + Offset;
 }
 
-#if TARRAY_RANGED_FOR_CHECKS
 
+#if TARRAY_RANGED_FOR_CHECKS
 	/**
 	 * Pointer-like iterator type for ranged-for loops which checks that the
 	 * container hasn't been resized during iteration.
@@ -192,9 +195,9 @@ FORCEINLINE TIndexedContainerIterator<ContainerType, ElementType, IndexType> ope
 			return Lhs.Ptr != Rhs.Ptr;
 		}
 	};
-
 #endif
 
+
 template <typename ElementType, typename IteratorType>
 struct TDereferencingIterator
 {
@@ -223,13 +226,16 @@ private:
 	}
 };
 
+
 /**
  * Base dynamic array.
  * An untyped data array; mirrors a TArray's members, but doesn't need an exact C++ type for its elements.
  **/
-class FScriptArray : protected FHeapAllocator::ForAnyElementType
+class FScriptArray
+	: protected FHeapAllocator::ForAnyElementType
 {
 public:
+
 	FORCEINLINE void* GetData()
 	{
 		return this->GetAllocation();
@@ -418,6 +424,7 @@ public:
 	void operator=(const FScriptArray&) { check(false); }
 };
 
+
 template<> struct TIsZeroConstructType<FScriptArray> { enum { Value = true }; };
 
 
@@ -438,6 +445,7 @@ public:
 	FORCEINLINE bool operator()( const ElementType& A, const ElementType& B ) const { return Predicate( B, A ); }
 };
 
+
 /**
  * Partially specialized version of the above.
  */
@@ -464,6 +472,7 @@ public:
 	}
 };
 
+
 namespace UE4Array_Private
 {
 	template <typename FromArrayType, typename ToArrayType>
@@ -487,6 +496,7 @@ namespace UE4Array_Private
 	};
 }
 
+
 /**
  * Templated dynamic array
  *
@@ -505,6 +515,7 @@ class TArray
 	friend class TArray;
 
 public:
+
 	typedef InElementType ElementType;
 	typedef InAllocator   Allocator;
 
@@ -582,6 +593,7 @@ public:
 	}
 
 private:
+
 	/**
 	 * Moves or copies array. Depends on the array type traits.
 	 *
@@ -705,9 +717,7 @@ public:
 		return *this;
 	}
 
-	/**
-	 * Destructor.
-	 */
+	/** Destructor. */
 	AGRESSIVE_ARRAY_FORCEINLINE ~TArray()
 	{
 		DestructItems(GetData(), ArrayNum);
@@ -762,6 +772,8 @@ public:
 
 	/**
 	 * Returns the amount of slack in this array in elements.
+	 *
+	 * @see Num, Shrink
 	 */
 	FORCEINLINE int32 GetSlack() const
 	{
@@ -797,7 +809,6 @@ public:
 	 * Tests if index is valid, i.e. than or equal to zero, and less than the number of elements in the array.
 	 *
 	 * @param Index Index to test.
-	 *
 	 * @returns True if index is valid. False otherwise.
 	 */
 	FORCEINLINE bool IsValidIndex(int32 Index) const
@@ -809,6 +820,7 @@ public:
 	 * Returns number of elements in array.
 	 *
 	 * @returns Number of elements in array.
+	 * @see GetSlack
 	 */
 	FORCEINLINE int32 Num() const
 	{
@@ -819,6 +831,7 @@ public:
 	 * Returns maximum number of elements in array.
 	 *
 	 * @returns Maximum number of elements in array.
+	 * @see GetSlack
 	 */
 	FORCEINLINE int32 Max() const
 	{
@@ -852,9 +865,7 @@ public:
 	/**
 	 * Pops element from the array.
 	 *
-	 * @param bAllowShrinking If this call allows shrinking of the array during
-	 *                        element remove.
-	 *
+	 * @param bAllowShrinking If this call allows shrinking of the array during element remove.
 	 * @returns Popped element.
 	 */
 	FORCEINLINE ElementType Pop(bool bAllowShrinking = true)
@@ -881,6 +892,7 @@ public:
 	 * Const ref version of the above.
 	 *
 	 * @param Item Item to push.
+	 * @see Pop, Top
 	 */
 	FORCEINLINE void Push(const ElementType& Item)
 	{
@@ -891,6 +903,7 @@ public:
 	 * Returns the top element, i.e. the last one.
 	 *
 	 * @returns Reference to the top element.
+	 * @see Pop, Push
 	 */
 	FORCEINLINE ElementType& Top()
 	{
@@ -903,6 +916,7 @@ public:
 	 * Const version of the above.
 	 *
 	 * @returns Reference to the top element.
+	 * @see Pop, Push
 	 */
 	FORCEINLINE const ElementType& Top() const
 	{
@@ -912,9 +926,7 @@ public:
 	/**
 	 * Returns n-th last element from the array.
 	 *
-	 * @param IndexFromTheEnd (Optional) Index from the end of array.
-	 *                        Default is 0.
-	 *
+	 * @param IndexFromTheEnd (Optional) Index from the end of array (default = 0).
 	 * @returns Reference to n-th last element from the array.
 	 */
 	FORCEINLINE ElementType& Last(int32 IndexFromTheEnd = 0)
@@ -928,9 +940,7 @@ public:
 	 *
 	 * Const version of the above.
 	 *
-	 * @param IndexFromTheEnd (Optional) Index from the end of array.
-	 *                        Default is 0.
-	 *
+	 * @param IndexFromTheEnd (Optional) Index from the end of array (default = 0).
 	 * @returns Reference to n-th last element from the array.
 	 */
 	FORCEINLINE const ElementType& Last(int32 IndexFromTheEnd = 0) const
@@ -940,8 +950,9 @@ public:
 	}
 
 	/**
-	 * Shrinks the array's used memory to smallest possible to store elements
-	 * currently in it.
+	 * Shrinks the array's used memory to smallest possible to store elements currently in it.
+	 *
+	 * @see Slack
 	 */
 	FORCEINLINE void Shrink()
 	{
@@ -956,9 +967,9 @@ public:
 	 * Finds element within the array.
 	 *
 	 * @param Item Item to look for.
-	 * @param Index Output parameter. Found index.
-	 *
+	 * @param Index Will contain the found index.
 	 * @returns True if found. False otherwise.
+	 * @see FindLast, FindLastByPredicate
 	 */
 	FORCEINLINE bool Find(const ElementType& Item, int32& Index) const
 	{
@@ -970,8 +981,8 @@ public:
 	 * Finds element within the array.
 	 *
 	 * @param Item Item to look for.
-	 *
 	 * @returns Index of the found element. INDEX_NONE otherwise.
+	 * @see FindLast, FindLastByPredicate
 	 */
 	AGRESSIVE_ARRAY_FORCEINLINE int32 Find(const ElementType& Item) const
 	{
@@ -991,8 +1002,8 @@ public:
 	 *
 	 * @param Item Item to look for.
 	 * @param Index Output parameter. Found index.
-	 *
 	 * @returns True if found. False otherwise.
+	 * @see Find, FindLastByPredicate
 	 */
 	FORCEINLINE bool FindLast(const ElementType& Item, int32& Index) const
 	{
@@ -1004,7 +1015,6 @@ public:
 	 * Finds element within the array starting from the end.
 	 *
 	 * @param Item Item to look for.
-	 *
 	 * @returns Index of the found element. INDEX_NONE otherwise.
 	 */
 	int32 FindLast(const ElementType& Item) const
@@ -1025,7 +1035,6 @@ public:
 	 *
 	 * @param Pred Predicate taking array element and returns true if element matches search criteria, false otherwise.
 	 * @param StartIndex Index of element from which to start searching.
-	 *
 	 * @returns Index of the found element. INDEX_NONE otherwise.
 	 */
 	template <typename Predicate>
@@ -1044,12 +1053,11 @@ public:
 	}
 
 	/**
-	* Finds element within the array starting from the end. Uses predicate to match element.
-	*
-	* @param Pred Predicate taking array element and returns true if element matches search criteria, false otherwise.
-	*
-	* @returns Index of the found element. INDEX_NONE otherwise.
-	*/
+	 * Finds element within the array starting from the end. Uses predicate to match element.
+	 *
+	 * @param Pred Predicate taking array element and returns true if element matches search criteria, false otherwise.
+	 * @returns Index of the found element. INDEX_NONE otherwise.
+	 */
 	template <typename Predicate>
 	FORCEINLINE int32 FindLastByPredicate(Predicate Pred) const
 	{
@@ -1061,9 +1069,7 @@ public:
 	 * the comparison).
 	 *
 	 * @param Key The key to search by.
-	 *
-	 * @returns Index to the first matching element, or INDEX_NONE if none is
-	 *          found.
+	 * @returns Index to the first matching element, or INDEX_NONE if none is found.
 	 */
 	template <typename KeyType>
 	AGRESSIVE_ARRAY_FORCEINLINE int32 IndexOfByKey(const KeyType& Key) const
@@ -1083,9 +1089,7 @@ public:
 	 * Finds an item by predicate.
 	 *
 	 * @param Pred The predicate to match.
-	 *
-	 * @returns Index to the first matching element, or INDEX_NONE if none is
-	 *          found.
+	 * @returns Index to the first matching element, or INDEX_NONE if none is found.
 	 */
 	template <typename Predicate>
 	AGRESSIVE_ARRAY_FORCEINLINE int32 IndexOfByPredicate(Predicate Pred) const
@@ -1106,8 +1110,8 @@ public:
 	 * the comparison).
 	 *
 	 * @param Key The key to search by.
-	 *
 	 * @returns Pointer to the first matching element, or nullptr if none is found.
+	 * @see Find
 	 */
 	template <typename KeyType>
 	FORCEINLINE const ElementType* FindByKey(const KeyType& Key) const
@@ -1120,8 +1124,8 @@ public:
 	 * the comparison). Time Complexity: O(n), starts iteration from the beginning so better performance if Key is in the front
 	 *
 	 * @param Key The key to search by.
-	 *
 	 * @returns Pointer to the first matching element, or nullptr if none is found.
+	 * @see Find
 	 */
 	template <typename KeyType>
 	AGRESSIVE_ARRAY_FORCEINLINE ElementType* FindByKey(const KeyType& Key)
@@ -1141,9 +1145,8 @@ public:
 	 * Finds an element which matches a predicate functor.
 	 *
 	 * @param Pred The functor to apply to each element.
-	 *
-	 * @returns Pointer to the first element for which the predicate returns
-	 *          true, or nullptr if none is found.
+	 * @returns Pointer to the first element for which the predicate returns true, or nullptr if none is found.
+	 * @see FilterByPredicate, ContainsByPredicate
 	 */
 	template <typename Predicate>
 	FORCEINLINE const ElementType* FindByPredicate(Predicate Pred) const
@@ -1154,10 +1157,8 @@ public:
 	/**
 	 * Finds an element which matches a predicate functor.
 	 *
-	 * @param Pred The functor to apply to each element.
-	 *
-	 * @return Pointer to the first element for which the predicate returns
-	 *         true, or nullptr if none is found.
+	 * @param Pred The functor to apply to each element. true, or nullptr if none is found.
+	 * @see FilterByPredicate, ContainsByPredicate
 	 */
 	template <typename Predicate>
 	AGRESSIVE_ARRAY_FORCEINLINE ElementType* FindByPredicate(Predicate Pred)
@@ -1177,9 +1178,9 @@ public:
 	 * Filters the elements in the array based on a predicate functor.
 	 *
 	 * @param Pred The functor to apply to each element.
-	 *
 	 * @returns TArray with the same type as this object which contains
 	 *          the subset of elements for which the functor returns true.
+	 * @see FindByPredicate, ContainsByPredicate
 	 */
 	template <typename Predicate>
 	TArray<ElementType> FilterByPredicate(Predicate Pred) const
@@ -1199,6 +1200,7 @@ public:
 	 * Checks if this array contains the element.
 	 *
 	 * @returns	True if found. False otherwise.
+	 * @see ContainsByPredicate, FilterByPredicate, FindByPredicate
 	 */
 	template <typename ComparisonType>
 	AGRESSIVE_ARRAY_FORCEINLINE bool Contains(const ComparisonType& Item) const
@@ -1217,8 +1219,8 @@ public:
 	 * Checks if this array contains element for which the predicate is true.
 	 *
 	 * @param Predicate to use
-	 *
 	 * @returns	True if found. False otherwise.
+	 * @see Contains, Find
 	 */
 	template <typename Predicate>
 	FORCEINLINE bool ContainsByPredicate(Predicate Pred) const
@@ -1230,7 +1232,6 @@ public:
 	 * Equality operator.
 	 *
 	 * @param OtherArray Array to compare.
-	 *
 	 * @returns True if this array is the same as OtherArray. False otherwise.
 	 */
 	AGRESSIVE_ARRAY_FORCEINLINE bool operator==(const TArray& OtherArray) const
@@ -1244,7 +1245,6 @@ public:
 	 * Inequality operator.
 	 *
 	 * @param OtherArray Array to compare.
-	 *
 	 * @returns True if this array is NOT the same as OtherArray. False otherwise.
 	 */
 	FORCEINLINE bool operator!=(const TArray& OtherArray) const
@@ -1257,7 +1257,6 @@ public:
 	 *
 	 * @param Ar Archive to serialize the array with.
 	 * @param A Array to serialize.
-	 *
 	 * @returns Passing the given archive.
 	 */
 	friend FArchive& operator<<(FArchive& Ar, TArray& A)
@@ -1380,7 +1379,6 @@ public:
 	 * require a constructor to function properly.
 	 *
 	 * @param Count Number of elements to add.
-	 *
 	 * @returns Number of elements in array before addition.
 	 */
 	FORCEINLINE int32 AddUninitialized(int32 Count = 1)
@@ -1431,6 +1429,7 @@ public:
 	 *
 	 * @param Index Tells where to insert the new elements.
 	 * @param Count Number of elements to add.
+	 * @see Insert, InsertUninitialized
 	 */
 	void InsertZeroed(int32 Index, int32 Count = 1)
 	{
@@ -1443,8 +1442,7 @@ public:
 	 *
 	 * @param Items Array of elements to insert.
 	 * @param InIndex Tells where to insert the new elements.
-	 *
-	 * @returns Location at which the insert was done.
+	 * @returns Location at which the item was inserted.
 	 */
 	int32 Insert(const TArray<ElementType>& Items, const int32 InIndex)
 	{
@@ -1462,11 +1460,11 @@ public:
 	/**
 	 * Inserts a raw array of elements at a particular index in the TArray.
 	 *
-	 * @param Ptr   A pointer to an array of elements to add.
+	 * @param Ptr A pointer to an array of elements to add.
 	 * @param Count The number of elements to insert from Ptr.
 	 * @param Index The index to insert the elements at.
-	 *
 	 * @return The index of the first element inserted.
+	 * @see Add, Remove
 	 */
 	int32 Insert(const ElementType* Ptr, int32 Count, int32 Index)
 	{
@@ -1484,6 +1482,7 @@ public:
 	 * aren't going to be invalidated by possible reallocation.
 	 *
 	 * @param Addr The address to check.
+	 * @see Add, Remove
 	 */
 	FORCEINLINE void CheckAddress(const ElementType* Addr) const
 	{
@@ -1496,8 +1495,8 @@ public:
 	 *
 	 * @param Item The element to insert.
 	 * @param Index Tells where to insert the new elements.
-	 *
 	 * @returns Location at which the insert was done.
+	 * @see Add, Remove
 	 */
 	int32 Insert(ElementType&& Item, int32 Index)
 	{
@@ -1515,8 +1514,8 @@ public:
 	 *
 	 * @param Item The element to insert.
 	 * @param Index Tells where to insert the new elements.
-	 *
 	 * @returns Location at which the insert was done.
+	 * @see Add, Remove
 	 */
 	int32 Insert(const ElementType& Item, int32 Index)
 	{
@@ -1535,8 +1534,7 @@ public:
 	 *
 	 * @param Index Location in array of the element to remove.
 	 * @param Count (Optional) Number of elements to remove. Default is 1.
-	 * @param bAllowShrinking (Optional) Tells if this call can shrink array if
-	 *                        suitable after remove. Default is true.
+	 * @param bAllowShrinking (Optional) Tells if this call can shrink array if suitable after remove. Default is true.
 	 */
 	void RemoveAt(int32 Index, int32 Count = 1, bool bAllowShrinking = true)
 	{
@@ -1640,6 +1638,7 @@ public:
 
 		checkSlow(Slack >= 0);
 		ArrayNum = 0;
+
 		if (ArrayMax != Slack)
 		{
 			ResizeTo(Slack);
@@ -1716,6 +1715,7 @@ public:
 	 * Allocator changing version.
 	 *
 	 * @param Source The array to append.
+	 * @see Add, Insert
 	 */
 	template <typename OtherElementType, typename OtherAllocator>
 	AGRESSIVE_ARRAY_FORCEINLINE void Append(const TArray<OtherElementType, OtherAllocator>& Source)
@@ -1732,7 +1732,6 @@ public:
 
 		// Allocate memory for the new elements.
 		Reserve(ArrayNum + SourceCount);
-
 		ConstructItems<ElementType>(GetData() + ArrayNum, Source.GetData(), SourceCount);
 
 		ArrayNum += SourceCount;
@@ -1742,6 +1741,7 @@ public:
 	 * Appends the specified array to this array.
 	 *
 	 * @param Source The array to append.
+	 * @see Add, Insert
 	 */
 	template <typename OtherElementType, typename OtherAllocator>
 	AGRESSIVE_ARRAY_FORCEINLINE void Append(TArray<OtherElementType, OtherAllocator>&& Source)
@@ -1758,7 +1758,6 @@ public:
 
 		// Allocate memory for the new elements.
 		Reserve(ArrayNum + SourceCount);
-
 		RelocateConstructItems<ElementType>(GetData() + ArrayNum, Source.GetData(), SourceCount);
 		Source.ArrayNum = 0;
 
@@ -1770,6 +1769,7 @@ public:
 	 *
 	 * @param Ptr   A pointer to an array of elements to add.
 	 * @param Count The number of elements to insert from Ptr.
+	 * @see Add, Insert
 	 */
 	AGRESSIVE_ARRAY_FORCEINLINE void Append(const ElementType* Ptr, int32 Count)
 	{
@@ -1885,16 +1885,18 @@ public:
 	 *
 	 * Move semantics version.
 	 *
-	 * @param Item	The item to add
-	 * @return		Index to the new item
+	 * @param Item The item to add
+	 * @return Index to the new item
+	 * @see AddDefaulted, AddUnique, AddZeroed, Append, Insert
 	 */
 	FORCEINLINE int32 Add(ElementType&& Item) { CheckAddress(&Item); return Emplace(MoveTemp(Item)); }
 
 	/**
 	 * Adds a new item to the end of the array, possibly reallocating the whole array to fit.
 	 *
-	 * @param Item	The item to add
-	 * @return		Index to the new item
+	 * @param Item The item to add
+	 * @return Index to the new item
+	 * @see AddDefaulted, AddUnique, AddZeroed, Append, Insert
 	 */
 	FORCEINLINE int32 Add(const ElementType& Item) { CheckAddress(&Item); return Emplace(Item); }
 
@@ -1907,8 +1909,8 @@ public:
 	 * a constructor to function properly.
 	 *
 	 * @param  Count  The number of new items to add.
-	 *
 	 * @return Index to the first of the new items.
+	 * @see Add, AddDefaulted, AddUnique, Append, Insert
 	 */
 	AGRESSIVE_ARRAY_FORCEINLINE int32 AddZeroed(int32 Count = 1)
 	{
@@ -1922,8 +1924,8 @@ public:
 	 * array to fit. The new items will be default-constructed.
 	 *
 	 * @param  Count  The number of new items to add.
-	 *
 	 * @return Index to the first of the new items.
+	 * @see Add, AddZeroed, AddUnique, Append, Insert
 	 */
 	AGRESSIVE_ARRAY_FORCEINLINE int32 AddDefaulted(int32 Count = 1)
 	{
@@ -1933,11 +1935,11 @@ public:
 	}
 
 private:
+
 	/**
 	 * Adds unique element to array if it doesn't exist.
 	 *
 	 * @param Args Item to add.
-	 *
 	 * @returns Index of the element in the array.
 	 */
 	template <typename ArgsType>
@@ -1953,14 +1955,15 @@ private:
 	}
 
 public:
+
 	/**
 	 * Adds unique element to array if it doesn't exist.
 	 *
 	 * Move semantics version.
 	 *
 	 * @param Args Item to add.
-	 *
 	 * @returns Index of the element in the array.
+	 * @see Add, AddDefaulted, AddZeroed, Append, Insert
 	 */
 	FORCEINLINE int32 AddUnique(ElementType&& Item) { return AddUniqueImpl(MoveTemp(Item)); }
 
@@ -1968,16 +1971,16 @@ public:
 	 * Adds unique element to array if it doesn't exist.
 	 *
 	 * @param Args Item to add.
-	 *
 	 * @returns Index of the element in the array.
+	 * @see Add, AddDefaulted, AddZeroed, Append, Insert
 	 */
 	FORCEINLINE int32 AddUnique(const ElementType& Item) { return AddUniqueImpl(Item); }
 
 	/**
 	 * Reserves memory such that the array can contain at least Number elements.
 	 *
-	 * @param Number The number of elements that the array should be able to
-	 *               contain after allocation.
+	 * @param Number The number of elements that the array should be able to contain after allocation.
+	 * @see Shrink
 	 */
 	FORCEINLINE void Reserve(int32 Number)
 	{
@@ -1991,8 +1994,7 @@ public:
 	 * Sets the size of the array, filling it with the given element.
 	 *
 	 * @param Element The element to fill array with.
-	 * @param Number The number of elements that the array should be able to
-	 *               contain after allocation.
+	 * @param Number The number of elements that the array should be able to contain after allocation.
 	 */
 	AGRESSIVE_ARRAY_FORCEINLINE void Init(const ElementType& Element, int32 Number)
 	{
@@ -2008,9 +2010,8 @@ public:
 	 * maintaining order but not indices.
 	 *
 	 * @param Item The item to remove.
-	 *
-	 * @returns The number of items removed. For RemoveSingleItem, this is
-	 *          always either 0 or 1.
+	 * @returns The number of items removed. For RemoveSingleItem, this is always either 0 or 1.
+	 * @see Add, Insert, Remove, RemoveAll, RemoveAllSwap
 	 */
 	int32 RemoveSingle(const ElementType& Item)
 	{
@@ -2039,8 +2040,8 @@ public:
 	 * order but not indices.
 	 *
 	 * @param Item Item to remove from array.
-	 *
 	 * @returns Number of removed elements.
+	 * @see Add, Insert, RemoveAll, RemoveAllSwap, RemoveSingle, RemoveSwap
 	 */
 	int32 Remove(const ElementType& Item)
 	{
@@ -2050,14 +2051,13 @@ public:
 		return RemoveAll([&Item](ElementType& Element) { return Element == Item; });
 	}
 
-
 	/**
 	 * Remove all instances that match the predicate, maintaining order but not indices
 	 * Optimized to work with runs of matches/non-matches
 	 *
 	 * @param Predicate Predicate class instance
-	 *
 	 * @returns Number of removed elements.
+	 * @see Add, Insert, RemoveAllSwap, RemoveSingle, RemoveSwap
 	 */
 	template <class PREDICATE_CLASS>
 	int32 RemoveAll(const PREDICATE_CLASS& Predicate)
@@ -2105,6 +2105,7 @@ public:
 	 * Remove all instances that match the predicate
 	 *
 	 * @param Predicate Predicate class instance
+	 * @see Remove, RemoveSingle, RemoveSingleSwap, RemoveSwap
 	 */
 	template <class PREDICATE_CLASS>
 	void RemoveAllSwap(const PREDICATE_CLASS& Predicate, bool bAllowShrinking = true)
@@ -2129,6 +2130,7 @@ public:
 	 * @param Item The item to remove
 	 *
 	 * @returns The number of items removed. For RemoveSingleItem, this is always either 0 or 1.
+	 * @see Add, Insert, Remove, RemoveAll, RemoveAllSwap, RemoveSwap
 	 */
 	int32 RemoveSingleSwap(const ElementType& Item, bool bAllowShrinking = true)
 	{
@@ -2152,6 +2154,7 @@ public:
 	 * but does not preserve the order.
 	 *
 	 * @returns Number of elements removed.
+	 * @see Add, Insert, Remove, RemoveAll, RemoveAllSwap
 	 */
 	int32 RemoveSwap(const ElementType& Item)
 	{
@@ -2213,7 +2216,6 @@ public:
 	 *             element hasn't been found.
 	 * @param StartIndex (Optional) Index in array at which the function should
 	 *             start to look for element.
-	 *
 	 * @returns True if element was found. False otherwise.
 	 */
 	template<typename SearchType>
@@ -2271,6 +2273,7 @@ public:
 	#endif
 
 private:
+
 	/**
 	 * DO NOT USE DIRECTLY
 	 * STL-like iterators to enable range-based for loop support.
@@ -2288,6 +2291,7 @@ private:
 	#endif
 
 public:
+
 	/**
 	 * Sorts the array assuming < operator is defined for the item type.
 	 */
@@ -2464,7 +2468,6 @@ public:
 	 *
 	 * @param InItem Item to be added.
 	 * @param Predicate Predicate class instance.
-	 *
 	 * @return The index of the new element.
 	 */
 	template <class PREDICATE_CLASS>
@@ -2487,7 +2490,6 @@ public:
 	 * template type.
 	 *
 	 * @param InItem Item to be added.
-	 *
 	 * @return The index of the new element.
 	 */
 	int32 HeapPush(const ElementType& InItem)
@@ -2566,9 +2568,10 @@ public:
 	}
 
 	/** 
-	 * Removes the top element from the heap. Assumes < operator is defined for
-	 * the template type.
-	 * @param bAllowShrinking (Optional) Tells if this call can shrink the array allocation if suitable after the discard. Default is true.
+	 * Removes the top element from the heap. Assumes < operator is defined for the template type.
+	 *
+	 * @param bAllowShrinking (Optional) Tells if this call can shrink the array
+	 *		allocation if suitable after the discard. Default is true.
 	 */
 	void HeapPopDiscard(bool bAllowShrinking = true)
 	{
@@ -2602,7 +2605,8 @@ public:
 	 *
 	 * @param Index Position at which to remove item.
 	 * @param Predicate Predicate class instance.
-	 * @param bAllowShrinking (Optional) Tells if this call can shrink the array allocation if suitable after the remove. Default is true.
+	 * @param bAllowShrinking (Optional) Tells if this call can shrink the array allocation
+	 *		if suitable after the remove (default = true).
 	 */
 	template <class PREDICATE_CLASS>
 	void HeapRemoveAt(int32 Index, const PREDICATE_CLASS& Predicate, bool bAllowShrinking = true)
@@ -2619,11 +2623,11 @@ public:
 	}
 
 	/**
-	 * Removes an element from the heap. Assumes < operator is defined for the
-	 * template type.
+	 * Removes an element from the heap. Assumes < operator is defined for the template type.
 	 *
 	 * @param Index Position at which to remove item.
-	 * @param bAllowShrinking (Optional) Tells if this call can shrink the array allocation if suitable after the remove. Default is true.
+	 * @param bAllowShrinking (Optional) Tells if this call can shrink the array allocation
+	 *		if suitable after the remove (default = true).
 	 */
 	void HeapRemoveAt(int32 Index, bool bAllowShrinking = true)
 	{
@@ -2680,7 +2684,6 @@ private:
 	 * Gets the index of the left child of node at Index.
 	 *
 	 * @param Index Node for which the left child index is to be returned.
-	 *
 	 * @returns Index of the left child.
 	 */
 	FORCEINLINE int32 HeapGetLeftChildIndex(int32 Index) const
@@ -2692,8 +2695,7 @@ private:
 	 * Checks if node located at Index is a leaf or not.
 	 *
 	 * @param Index Node index.
-	 *
-	 * @returns True if node is a leaf. False otherwise.
+	 * @returns true if node is a leaf, false otherwise.
 	 */
 	FORCEINLINE bool HeapIsLeaf(int32 Index, int32 Count) const
 	{
@@ -2704,7 +2706,6 @@ private:
 	 * Gets the parent index for node at Index.
 	 *
 	 * @param Index node index.
-	 *
 	 * @returns Parent index.
 	 */
 	FORCEINLINE int32 HeapGetParentIndex(int32 Index) const
@@ -2713,6 +2714,7 @@ private:
 	}
 
 private:
+
 	/**
 	 * Fixes a possible violation of order property between node at Index and a child.
 	 *
@@ -2774,18 +2776,21 @@ private:
 	}
 };
 
+
 template <typename InElementType, typename Allocator>
 struct TIsZeroConstructType<TArray<InElementType, Allocator>>
 {
 	enum { Value = TAllocatorTraits<Allocator>::IsZeroConstruct };
 };
 
+
 template <typename InElementType, typename Allocator>
 struct TContainerTraits<TArray<InElementType, Allocator> > : public TContainerTraitsBase<TArray<InElementType, Allocator> >
 {
 	enum { MoveWillEmptyContainer = TAllocatorTraits<Allocator>::SupportsMove };
 };
 
+
 /**
  * Traits class which determines whether or not a type is a TArray.
  */
@@ -2796,6 +2801,7 @@ template <typename InElementType, typename Allocator> struct TIsTArray<const
 template <typename InElementType, typename Allocator> struct TIsTArray<      volatile TArray<InElementType, Allocator>> { enum { Value = true }; };
 template <typename InElementType, typename Allocator> struct TIsTArray<const volatile TArray<InElementType, Allocator>> { enum { Value = true }; };
 
+
 //
 // Array operator news.
 //
@@ -2822,7 +2828,8 @@ template <typename T,typename Allocator> void* operator new( size_t Size, TArray
  * - Any item that is added to the array is moved to the top.
  */
 template<typename T, typename Allocator = FDefaultAllocator>
-class TMRUArray : public TArray<T, Allocator>
+class TMRUArray
+	: public TArray<T, Allocator>
 {
 public:
 	typedef TArray<T, Allocator> Super;
@@ -2862,7 +2869,6 @@ public:
 	 * Assign operator.
 	 *
 	 * @param Other Other array to assign with.
-	 *
 	 * @returns Reference to this object.
 	 */
 	FORCEINLINE TMRUArray& operator=(const TMRUArray& Other)
@@ -2885,7 +2891,6 @@ public:
 	 * Move assign operator.
 	 *
 	 * @param Other Other array to assign with.
-	 *
 	 * @returns Reference to this object.
 	 */
 	FORCEINLINE TMRUArray& operator=(TMRUArray&& Other)
@@ -2900,7 +2905,6 @@ public:
 	 * limit.
 	 *
 	 * @param Item Item to add.
-	 *
 	 * @returns Always 0.
 	 */
 	int32 Add(const T& Item)
@@ -2916,7 +2920,6 @@ public:
 	 * add more than the limit.
 	 *
 	 * @param Count (Optional) A number of elements to add. Default is 0.
-	 *
 	 * @returns Always 0.
 	 */
 	int32 AddZeroed(int32 Count = 1)
@@ -2932,7 +2935,6 @@ public:
 	 * the limit. If the item existed it will be removed before addition.
 	 *
 	 * @param Item Element to add.
-	 *
 	 * @returns Always 0.
 	 */
 	int32 AddUnique(const T& Item)
@@ -2966,12 +2968,14 @@ public:
 	}
 };
 
+
 template<typename T, typename Allocator>
 struct TContainerTraits<TMRUArray<T, Allocator> > : public TContainerTraitsBase<TMRUArray<T, Allocator> >
 {
 	enum { MoveWillEmptyContainer = TContainerTraitsBase<typename TMRUArray<T, Allocator>::Super>::MoveWillEmptyContainer };
 };
 
+
 /*-----------------------------------------------------------------------------
 	Indirect array.
 	Same as a TArray above, but stores pointers to the elements, to allow
@@ -2984,9 +2988,7 @@ public:
 	typedef T                        ElementType;
 	typedef TArray<void*, Allocator> InternalArrayType;
 
-	/**
-	 * Default constructors.
-	 */
+	/** Default constructors. */
 #if PLATFORM_COMPILER_HAS_DEFAULTED_FUNCTIONS
 
 	TIndirectArray() = default;
@@ -3033,9 +3035,7 @@ public:
 		return *this;
 	}
 
-	/**
-	 * Destructor.
-	 */
+	/** Destructor. */
 	~TIndirectArray()
 	{
 		Empty();
@@ -3085,7 +3085,6 @@ public:
 	 * Bracket array access operator.
 	 *
 	 * @param Index Position of element to return.
-	 *
 	 * @returns Reference to element in array at given position.
 	 */
 	FORCEINLINE T& operator[](int32 Index)
@@ -3099,7 +3098,6 @@ public:
 	 * Const version.
 	 *
 	 * @param Index Position of element to return.
-	 *
 	 * @returns Reference to element in array at given position.
 	 */
 	FORCEINLINE const T& operator[](int32 Index) const
@@ -3110,9 +3108,7 @@ public:
 	/**
 	 * Returns n-th last element from the array.
 	 *
-	 * @param IndexFromTheEnd (Optional) Index from the end of array.
-	 *                        Default is 0.
-	 *
+	 * @param IndexFromTheEnd (Optional) Index from the end of array (default = 0).
 	 * @returns Reference to n-th last element from the array.
 	 */
 	FORCEINLINE ElementType& Last(int32 IndexFromTheEnd = 0)
@@ -3125,9 +3121,7 @@ public:
 	 *
 	 * Const version.
 	 *
-	 * @param IndexFromTheEnd (Optional) Index from the end of array.
-	 *                        Default is 0.
-	 *
+	 * @param IndexFromTheEnd (Optional) Index from the end of array (default = 0).
 	 * @returns Reference to n-th last element from the array.
 	 */
 	FORCEINLINE const ElementType& Last(int32 IndexFromTheEnd = 0) const
@@ -3149,7 +3143,7 @@ public:
 	 * items.
 	 *
 	 * @param NewSize (Optional) The expected usage size after calling this
-	 *                function. Default is 0.
+	 *		function. Default is 0.
 	 */
 	void Reset(int32 NewSize = 0)
 	{
@@ -3161,8 +3155,8 @@ public:
 	 * Special serialize function passing the owning UObject along as required
 	 * by FUnytpedBulkData serialization.
 	 *
-	 * @param	Ar		Archive to serialize with.
-	 * @param	Owner	UObject this structure is serialized within.
+	 * @param Ar Archive to serialize with.
+	 * @param Owner UObject this structure is serialized within.
 	 */
 	void Serialize(FArchive& Ar, UObject* Owner)
 	{
@@ -3197,9 +3191,8 @@ public:
 	/**
 	 * Serialization operator for TIndirectArray.
 	 *
-	 * @param	Ar		Archive to serialize with.
-	 * @param	A		Array to serialize.
-	 *
+	 * @param Ar Archive to serialize with.
+	 * @param A Array to serialize.
 	 * @returns Passing down serializing archive.
 	 */
 	friend FArchive& operator<<(FArchive& Ar, TIndirectArray& A)
@@ -3326,7 +3319,6 @@ public:
 	 * whole array to fit.
 	 *
 	 * @param Item The item to add.
-	 *
 	 * @returns Index to the new item.
 	 */
 	FORCEINLINE int32 Add(T* Item)
@@ -3361,7 +3353,6 @@ public:
 	 * elements in array.
 	 *
 	 * @param Index Index to test.
-	 *
 	 * @returns True if index is valid. False otherwise.
 	 */
 	FORCEINLINE bool IsValidIndex(int32 Index) const
@@ -3404,6 +3395,7 @@ public:
 	}
 
 private:
+
 	/**
 	 * Calls destructor and frees memory on every element in the array.
 	 */
@@ -3433,12 +3425,15 @@ private:
 	InternalArrayType Array;
 };
 
+
 template<typename T, typename Allocator>
-struct TContainerTraits<TIndirectArray<T, Allocator> > : public TContainerTraitsBase<TIndirectArray<T, Allocator> >
+struct TContainerTraits<TIndirectArray<T, Allocator> >
+	: public TContainerTraitsBase<TIndirectArray<T, Allocator> >
 {
 	enum { MoveWillEmptyContainer = TContainerTraitsBase<typename TIndirectArray<T, Allocator>::InternalArrayType>::MoveWillEmptyContainer };
 };
 
+
 template <typename T,typename Allocator> void* operator new( size_t Size, TIndirectArray<T,Allocator>& Array )
 {
 	check(Size == sizeof(T));
@@ -3446,6 +3441,7 @@ template <typename T,typename Allocator> void* operator new( size_t Size, TIndir
 	return &Array[Index];
 }
 
+
 template <typename T,typename Allocator> void* operator new( size_t Size, TIndirectArray<T,Allocator>& Array, int32 Index )
 {
 	check(Size == sizeof(T));
@@ -3453,6 +3449,7 @@ template <typename T,typename Allocator> void* operator new( size_t Size, TIndir
 	return &Array[Index];
 }
 
+
 /*-----------------------------------------------------------------------------
 	Transactional array.
 -----------------------------------------------------------------------------*/
@@ -3460,9 +3457,11 @@ template <typename T,typename Allocator> void* operator new( size_t Size, TIndir
 // NOTE: Right now, you can't use a custom allocation policy with transactional arrays. If
 // you need to do it, you will have to fix up FTransaction::FObjectRecord to use the correct TArray<Allocator>.
 template< typename T >
-class TTransArray : public TArray<T>
+class TTransArray
+	: public TArray<T>
 {
 public:
+
 	typedef TArray<T> Super;
 
 	// Constructors.
@@ -3471,6 +3470,7 @@ public:
 	{
 		checkSlow(Owner);
 	}
+
 	TTransArray( UObject* InOwner, const Super& Other )
 	:	Super( Other )
 	,	Owner( InOwner )
@@ -3525,6 +3525,7 @@ public:
 		}
 		return Index;
 	}
+
 	void InsertUninitialized( int32 Index, int32 Count=1 )
 	{
 		Super::InsertUninitialized( Index, Count );
@@ -3533,6 +3534,7 @@ public:
 			GUndo->SaveArray( Owner, (FScriptArray*)this, Index, Count, 1, sizeof(T), DefaultConstructItem, SerializeItem, DestructItem );
 		}
 	}
+
 	void RemoveAt( int32 Index, int32 Count=1 )
 	{
 		if( GUndo )
@@ -3541,6 +3543,7 @@ public:
 		}
 		Super::RemoveAt( Index, Count );
 	}
+
 	void Empty( int32 Slack=0 )
 	{
 		if( GUndo )
@@ -3567,12 +3570,14 @@ public:
 		new(*this) T(Item);
 		return this->Num() - 1;
 	}
+
 	int32 AddZeroed( int32 n=1 )
 	{
 		const int32 Index = AddUninitialized(n);
 		FMemory::Memzero(this->GetData() + Index, n*sizeof(T));
 		return Index;
 	}
+
 	int32 AddUnique( const T& Item )
 	{
 		for( int32 Index=0; Index<this->ArrayNum; Index++ )
@@ -3584,6 +3589,7 @@ public:
 		}
 		return Add( Item );
 	}
+
 	int32 Remove( const T& Item )
 	{
 		this->CheckAddress(&Item);
@@ -3604,10 +3610,12 @@ public:
 	{
 		return Owner;
 	}
+
 	void SetOwner( UObject* NewOwner )
 	{
 		Owner = NewOwner;
 	}
+
 	void ModifyItem( int32 Index )
 	{
 		if( GUndo )
@@ -3615,6 +3623,7 @@ public:
 			GUndo->SaveArray( Owner, (FScriptArray*)this, Index, 1, 0, sizeof(T), DefaultConstructItem, SerializeItem, DestructItem );
 		}
 	}
+
 	void ModifyAllItems()
 	{
 		if( GUndo )
@@ -3622,13 +3631,16 @@ public:
 			GUndo->SaveArray( Owner, (FScriptArray*)this, 0, this->Num(), 0, sizeof(T), DefaultConstructItem, SerializeItem, DestructItem );
 		}
 	}
+
 	friend FArchive& operator<<( FArchive& Ar, TTransArray& A )
 	{
 		Ar << A.Owner;
 		Ar << (Super&)A;
 		return Ar;
 	}
+
 protected:
+
 	static void DefaultConstructItem( void* TPtr )
 	{
 		new (TPtr) T;
@@ -3644,12 +3656,14 @@ protected:
 	UObject* Owner;
 };
 
+
 template<typename T>
 struct TContainerTraits<TTransArray<T> > : public TContainerTraitsBase<TTransArray<T> >
 {
 	enum { MoveWillEmptyContainer = TContainerTraitsBase<typename TTransArray<T>::Super>::MoveWillEmptyContainer };
 };
 
+
 //
 // Transactional array operator news.
 //
@@ -3659,6 +3673,8 @@ template <typename T> void* operator new( size_t Size, TTransArray<T>& Array )
 	const int32 Index = Array.AddUninitialized();
 	return &Array[Index];
 }
+
+
 template <typename T> void* operator new( size_t Size, TTransArray<T>& Array, int32 Index )
 {
 	check(Size == sizeof(T));
diff --git a/Engine/Source/Runtime/Core/Public/Math/IntPoint.h b/Engine/Source/Runtime/Core/Public/Math/IntPoint.h
--- a/Engine/Source/Runtime/Core/Public/Math/IntPoint.h
+++ b/Engine/Source/Runtime/Core/Public/Math/IntPoint.h
@@ -21,10 +21,10 @@ struct FIntPoint
 
 public:
 
-	/** An int point with zeroed values. */
+	/** An integer point with zeroed values. */
 	CORE_API static const FIntPoint ZeroValue;
 
-	/** An int point with INDEX_NONE values. */
+	/** An integer point with INDEX_NONE values. */
 	CORE_API static const FIntPoint NoneValue;
 
 public:
@@ -33,144 +33,144 @@ public:
 	FIntPoint();
 
 	/**
-	 * Creates and initializes a new instance with the specified coordinates.
+	 * Create and initialize a new instance with the specified coordinates.
 	 *
 	 * @param InX The x-coordinate.
 	 * @param InY The y-coordinate.
 	 */
-	FIntPoint( int32 InX, int32 InY );
+	FIntPoint(int32 InX, int32 InY);
 
 	/**
-	 * Creates and initializes a new instance to zero.
+	 * Create and initialize a new instance to zero.
 	 *
 	 * @param EForceInit Force init enum
 	 */
-	explicit FORCEINLINE FIntPoint( EForceInit );
+	explicit FORCEINLINE FIntPoint(EForceInit);
 
 public:
 
 	/**
-	 * Gets specific component of a point.
+	 * Get specific component of a point.
 	 *
 	 * @param PointIndex Index of point component.
 	 * @return const reference to component.
 	 */
-	const int32& operator()( int32 PointIndex ) const;
+	const int32& operator()(int32 PointIndex) const;
 
 	/**
-	 * Gets specific component of a point.
+	 * Get specific component of a point.
 	 *
 	 * @param PointIndex Index of point component
 	 * @return reference to component.
 	 */
-	int32& operator()( int32 PointIndex );
+	int32& operator()(int32 PointIndex);
 
 	/**
-	 * Compares points for equality.
+	 * Compare two points for equality.
 	 *
 	 * @param Other The other int point being compared.
 	 * @return true if the points are equal, false otherwise.
 	 */
-	bool operator==( const FIntPoint& Other ) const;
+	bool operator==(const FIntPoint& Other) const;
 
 	/**
-	 * Compares points for inequality.
+	 * Compare two points for inequality.
 	 *
 	 * @param Other The other int point being compared.
 	 * @return true if the points are not equal, false otherwise.
 	 */
-	bool operator!=( const FIntPoint& Other ) const;
+	bool operator!=(const FIntPoint& Other) const;
 
 	/**
-	 * Scales this point.
+	 * Scale this point.
 	 *
 	 * @param Scale What to multiply the point by.
 	 * @return Reference to this point after multiplication.
 	 */
-	FIntPoint& operator*=( int32 Scale );
+	FIntPoint& operator*=(int32 Scale);
 
 	/**
-	 * Divides this point.
+	 * Divide this point by a scalar.
 	 *
 	 * @param Divisor What to divide the point by.
 	 * @return Reference to this point after division.
 	 */
-	FIntPoint& operator/=( int32 Divisor );
+	FIntPoint& operator/=(int32 Divisor);
 
 	/**
-	 * Adds to this point.
+	 * Add another point component-wise to this point.
 	 *
 	 * @param Other The point to add to this point.
 	 * @return Reference to this point after addition.
 	 */
-	FIntPoint& operator+=( const FIntPoint& Other );
+	FIntPoint& operator+=(const FIntPoint& Other);
 	
 	/**
-	 * Subtracts from this point.
+	 * Subtract another point component-wise from this point.
 	 *
 	 * @param Other The point to subtract from this point.
 	 * @return Reference to this point after subtraction.
 	 */
-	FIntPoint& operator-=( const FIntPoint& Other );
+	FIntPoint& operator-=(const FIntPoint& Other);
 
 	/**
-	 * Divides this point.
+	 * Divide this point component-wise by another point.
 	 *
 	 * @param Other The point to divide with.
 	 * @return Reference to this point after division.
 	 */
-	FIntPoint& operator/=( const FIntPoint& Other );
+	FIntPoint& operator/=(const FIntPoint& Other);
 
 	/**
-	 * Assigns another point to this one.
+	 * Assign another point to this one.
 	 *
 	 * @param Other The point to assign this point from.
 	 * @return Reference to this point after assignment.
 	 */
-	FIntPoint& operator=( const FIntPoint& Other );
+	FIntPoint& operator=(const FIntPoint& Other);
 
 	/**
-	 * Gets the result of scaling on this point.
+	 * Get the result of scaling on this point.
 	 *
 	 * @param Scale What to multiply the point by.
 	 * @return A new scaled int point.
 	 */
-	FIntPoint operator*( int32 Scale ) const;
+	FIntPoint operator*(int32 Scale) const;
 
 	/**
-	 * Gets the result of division on this point.
+	 * Get the result of division on this point.
 	 *
 	 * @param Divisor What to divide the point by.
 	 * @return A new divided int point.
 	 */
-	FIntPoint operator/( int32 Divisor ) const;
+	FIntPoint operator/(int32 Divisor) const;
 
 	/**
-	 * Gets the result of addition on this point.
+	 * Get the result of addition on this point.
 	 *
 	 * @param Other The other point to add to this.
 	 * @return A new combined int point.
 	 */
-	FIntPoint operator+( const FIntPoint& Other ) const;
+	FIntPoint operator+(const FIntPoint& Other) const;
 
 	/**
-	 * Gets the result of subtraction from this point.
+	 * Get the result of subtraction from this point.
 	 *
 	 * @param Other The other point to subtract from this.
 	 * @return A new subtracted int point.
 	 */
-	FIntPoint operator-( const FIntPoint& Other ) const;
+	FIntPoint operator-(const FIntPoint& Other) const;
 
 	/**
-	 * Gets the result of division on this point.
+	 * Get the result of division on this point.
 	 *
 	 * @param Other The other point to subtract from this.
 	 * @return A new subtracted int point.
 	 */
-	FIntPoint operator/( const FIntPoint& Other ) const;
+	FIntPoint operator/(const FIntPoint& Other) const;
 
 	/**
-	* Gets specific component of the point.
+	* Get specific component of the point.
 	*
 	* @param Index the index of point component
 	* @return reference to component.
@@ -178,39 +178,47 @@ public:
 	int32& operator[](int32 Index);
 
 	/**
-	* Gets specific component of the point.
+	* Get specific component of the point.
 	*
 	* @param Index the index of point component
 	* @return copy of component value.
 	*/
 	int32 operator[](int32 Index) const;
 
-	/** Gets the component-wise min of two vectors. */
+public:
+
+	/**
+	 * Get the component-wise min of two points.
+	 *
+	 * @see ComponentMax, GetMax
+	 */
 	FORCEINLINE FIntPoint ComponentMin(const FIntPoint& Other) const;
 
-	/** Gets the component-wise max of two vectors. */
+	/**
+	 * Get the component-wise max of two points.
+	 *
+	 * @see ComponentMin, GetMin
+	 */
 	FORCEINLINE FIntPoint ComponentMax(const FIntPoint& Other) const;
 
-public:
-
 	/**
-	 * Gets the maximum value in the point.
+	 * Get the larger of the point's two components.
 	 *
-	 * @return The maximum value in the point.
+	 * @return The maximum component of the point.
 	 * @see GetMin, Size, SizeSquared
 	 */
 	int32 GetMax() const;
 
 	/**
-	 * Gets the minimum value in the point.
+	 * Get the smaller of the point's two components.
 	 *
-	 * @return The minimum value in the point.
+	 * @return The minimum component of the point.
 	 * @see GetMax, Size, SizeSquared
 	 */
 	int32 GetMin() const;
 
 	/**
-	 * Gets the distance of this point from (0,0).
+	 * Get the distance of this point from (0,0).
 	 *
 	 * @return The distance of this point from (0,0).
 	 * @see GetMax, GetMin, SizeSquared
@@ -218,7 +226,7 @@ public:
 	int32 Size() const;
 
 	/**
-	 * Gets the squared distance of this point from (0,0).
+	 * Get the squared distance of this point from (0,0).
 	 *
 	 * @return The squared distance of this point from (0,0).
 	 * @see GetMax, GetMin, Size
@@ -242,8 +250,8 @@ public:
 	 * @return A new divided int point.
 	 * @see DivideAndRoundDown
 	 */
-	static FIntPoint DivideAndRoundUp( FIntPoint lhs, int32 Divisor );
-	static FIntPoint DivideAndRoundUp( FIntPoint lhs, FIntPoint Divisor );
+	static FIntPoint DivideAndRoundUp(FIntPoint lhs, int32 Divisor);
+	static FIntPoint DivideAndRoundUp(FIntPoint lhs, FIntPoint Divisor);
 
 	/**
 	 * Divide an int point and round down the result.
@@ -253,10 +261,11 @@ public:
 	 * @return A new divided int point.
 	 * @see DivideAndRoundUp
 	 */
-	static FIntPoint DivideAndRoundDown( FIntPoint lhs, int32 Divisor );
+	static FIntPoint DivideAndRoundDown(FIntPoint lhs, int32 Divisor);
 
 	/**
-	 * Gets number of components point has.
+	 * Get number of components point has.
+	 *
 	 * @return number of components point has.
 	 */
 	static int32 Num();
@@ -264,18 +273,24 @@ public:
 public:
 
 	/**
-	 * Serializes the point.
+	 * Serialize the point.
 	 *
 	 * @param Ar The archive to serialize into.
 	 * @param Point The point to serialize.
 	 * @return Reference to the Archive after serialization.
 	 */
-	friend FArchive& operator<<( FArchive& Ar, FIntPoint& Point )
+	friend FArchive& operator<<(FArchive& Ar, FIntPoint& Point)
 	{
 		return Ar << Point.X << Point.Y;
 	}
 
-	bool Serialize( FArchive& Ar )
+	/**
+	 * Serialize the point.
+	 *
+	 * @param Ar The archive to serialize into.
+	 * @return true on success, false otherwise.
+	 */
+	bool Serialize(FArchive& Ar)
 	{
 		Ar << *this;
 		return true;
@@ -289,25 +304,25 @@ public:
 FORCEINLINE FIntPoint::FIntPoint() { }
 
 
-FORCEINLINE FIntPoint::FIntPoint( int32 InX, int32 InY )
+FORCEINLINE FIntPoint::FIntPoint(int32 InX, int32 InY)
 	: X(InX)
 	, Y(InY)
 { }
 
 
-FORCEINLINE FIntPoint::FIntPoint( EForceInit )
+FORCEINLINE FIntPoint::FIntPoint(EForceInit)
 	: X(0)
 	, Y(0)
 { }
 
 
-FORCEINLINE const int32& FIntPoint::operator()( int32 PointIndex ) const
+FORCEINLINE const int32& FIntPoint::operator()(int32 PointIndex) const
 {
 	return (&X)[PointIndex];
 }
 
 
-FORCEINLINE int32& FIntPoint::operator()( int32 PointIndex )
+FORCEINLINE int32& FIntPoint::operator()(int32 PointIndex)
 {
 	return (&X)[PointIndex];
 }
@@ -319,19 +334,19 @@ FORCEINLINE int32 FIntPoint::Num()
 }
 
 
-FORCEINLINE bool FIntPoint::operator==( const FIntPoint& Other ) const
+FORCEINLINE bool FIntPoint::operator==(const FIntPoint& Other) const
 {
 	return X==Other.X && Y==Other.Y;
 }
 
 
-FORCEINLINE bool FIntPoint::operator!=( const FIntPoint& Other ) const
+FORCEINLINE bool FIntPoint::operator!=(const FIntPoint& Other) const
 {
 	return (X != Other.X) || (Y != Other.Y);
 }
 
 
-FORCEINLINE FIntPoint& FIntPoint::operator*=( int32 Scale )
+FORCEINLINE FIntPoint& FIntPoint::operator*=(int32 Scale)
 {
 	X *= Scale;
 	Y *= Scale;
@@ -340,7 +355,7 @@ FORCEINLINE FIntPoint& FIntPoint::operator*=( int32 Scale )
 }
 
 
-FORCEINLINE FIntPoint& FIntPoint::operator/=( int32 Divisor )
+FORCEINLINE FIntPoint& FIntPoint::operator/=(int32 Divisor)
 {
 	X /= Divisor;
 	Y /= Divisor;
@@ -349,7 +364,7 @@ FORCEINLINE FIntPoint& FIntPoint::operator/=( int32 Divisor )
 }
 
 
-FORCEINLINE FIntPoint& FIntPoint::operator+=( const FIntPoint& Other )
+FORCEINLINE FIntPoint& FIntPoint::operator+=(const FIntPoint& Other)
 {
 	X += Other.X;
 	Y += Other.Y;
@@ -358,7 +373,7 @@ FORCEINLINE FIntPoint& FIntPoint::operator+=( const FIntPoint& Other )
 }
 
 
-FORCEINLINE FIntPoint& FIntPoint::operator-=( const FIntPoint& Other )
+FORCEINLINE FIntPoint& FIntPoint::operator-=(const FIntPoint& Other)
 {
 	X -= Other.X;
 	Y -= Other.Y;
@@ -367,7 +382,7 @@ FORCEINLINE FIntPoint& FIntPoint::operator-=( const FIntPoint& Other )
 }
 
 
-FORCEINLINE FIntPoint& FIntPoint::operator/=( const FIntPoint& Other )
+FORCEINLINE FIntPoint& FIntPoint::operator/=(const FIntPoint& Other)
 {
 	X /= Other.X;
 	Y /= Other.Y;
@@ -376,7 +391,7 @@ FORCEINLINE FIntPoint& FIntPoint::operator/=( const FIntPoint& Other )
 }
 
 
-FORCEINLINE FIntPoint& FIntPoint::operator=( const FIntPoint& Other )
+FORCEINLINE FIntPoint& FIntPoint::operator=(const FIntPoint& Other)
 {
 	X = Other.X;
 	Y = Other.Y;
@@ -385,13 +400,13 @@ FORCEINLINE FIntPoint& FIntPoint::operator=( const FIntPoint& Other )
 }
 
 
-FORCEINLINE FIntPoint FIntPoint::operator*( int32 Scale ) const
+FORCEINLINE FIntPoint FIntPoint::operator*(int32 Scale) const
 {
 	return FIntPoint(*this) *= Scale;
 }
 
 
-FORCEINLINE FIntPoint FIntPoint::operator/( int32 Divisor ) const
+FORCEINLINE FIntPoint FIntPoint::operator/(int32 Divisor) const
 {
 	return FIntPoint(*this) /= Divisor;
 }
@@ -422,35 +437,35 @@ FORCEINLINE FIntPoint FIntPoint::ComponentMax(const FIntPoint& Other) const
 	return FIntPoint(FMath::Max(X, Other.X), FMath::Max(Y, Other.Y));
 }
 
-FORCEINLINE FIntPoint FIntPoint::DivideAndRoundUp( FIntPoint lhs, int32 Divisor )
+FORCEINLINE FIntPoint FIntPoint::DivideAndRoundUp(FIntPoint lhs, int32 Divisor)
 {
 	return FIntPoint(FMath::DivideAndRoundUp(lhs.X, Divisor), FMath::DivideAndRoundUp(lhs.Y, Divisor));
 }
 
-FORCEINLINE FIntPoint FIntPoint::DivideAndRoundUp( FIntPoint lhs, FIntPoint Divisor )
+FORCEINLINE FIntPoint FIntPoint::DivideAndRoundUp(FIntPoint lhs, FIntPoint Divisor)
 {
 	return FIntPoint(FMath::DivideAndRoundUp(lhs.X, Divisor.X), FMath::DivideAndRoundUp(lhs.Y, Divisor.Y));
 }	
 
-FORCEINLINE FIntPoint FIntPoint::DivideAndRoundDown( FIntPoint lhs, int32 Divisor )
+FORCEINLINE FIntPoint FIntPoint::DivideAndRoundDown(FIntPoint lhs, int32 Divisor)
 {
 	return FIntPoint(FMath::DivideAndRoundDown(lhs.X, Divisor), FMath::DivideAndRoundDown(lhs.Y, Divisor));
 }	
 
 
-FORCEINLINE FIntPoint FIntPoint::operator+( const FIntPoint& Other ) const
+FORCEINLINE FIntPoint FIntPoint::operator+(const FIntPoint& Other) const
 {
 	return FIntPoint(*this) += Other;
 }
 
 
-FORCEINLINE FIntPoint FIntPoint::operator-( const FIntPoint& Other ) const
+FORCEINLINE FIntPoint FIntPoint::operator-(const FIntPoint& Other) const
 {
 	return FIntPoint(*this) -= Other;
 }
 
 
-FORCEINLINE FIntPoint FIntPoint::operator/( const FIntPoint& Other ) const
+FORCEINLINE FIntPoint FIntPoint::operator/(const FIntPoint& Other) const
 {
 	return FIntPoint(*this) /= Other;
 }
@@ -477,7 +492,7 @@ FORCEINLINE int32 FIntPoint::Size() const
 {
 	int64 X64 = (int64)X;
 	int64 Y64 = (int64)Y;
-	return int32(FMath::Sqrt( float(X64*X64 + Y64*Y64)));
+	return int32(FMath::Sqrt(float(X64 * X64 + Y64 * Y64)));
 }
 
 FORCEINLINE int32 FIntPoint::SizeSquared() const
diff --git a/Engine/Source/Runtime/Core/Public/Misc/Timespan.h b/Engine/Source/Runtime/Core/Public/Misc/Timespan.h
--- a/Engine/Source/Runtime/Core/Public/Misc/Timespan.h
+++ b/Engine/Source/Runtime/Core/Public/Misc/Timespan.h
@@ -21,6 +21,9 @@ namespace ETimespan
 	/** The number of timespan ticks per hour. */
 	const int64 TicksPerHour = 36000000000;
 
+	/** The number of timespan ticks per microsecond. */
+	const int64 TicksPerMicrosecond = 10;
+
 	/** The number of timespan ticks per millisecond. */
 	const int64 TicksPerMillisecond = 10000;
 
@@ -61,7 +64,7 @@ public:
 	 *
 	 * @param Ticks The number of ticks.
 	 */
-	FTimespan( int64 InTicks )
+	FTimespan(int64 InTicks)
 		: Ticks(InTicks)
 	{ }
 
@@ -72,9 +75,9 @@ public:
 	 * @param Minutes The minutes component.
 	 * @param Seconds The seconds component.
 	 */
-	FTimespan( int32 Hours, int32 Minutes, int32 Seconds )
+	FTimespan(int32 Hours, int32 Minutes, int32 Seconds)
 	{
-		Assign(0, Hours, Minutes, Seconds, 0);
+		Assign(0, Hours, Minutes, Seconds, 0, 0);
 	}
 
 	/**
@@ -85,9 +88,9 @@ public:
 	 * @param Minutes The minutes component.
 	 * @param Seconds The seconds component.
 	 */
-	FTimespan( int32 Days, int32 Hours, int32 Minutes, int32 Seconds )
+	FTimespan(int32 Days, int32 Hours, int32 Minutes, int32 Seconds)
 	{
-		Assign(Days, Hours, Minutes, Seconds, 0);
+		Assign(Days, Hours, Minutes, Seconds, 0, 0);
 	}
 
 	/**
@@ -98,10 +101,11 @@ public:
 	 * @param Minutes The minutes component.
 	 * @param Seconds The seconds component.
 	 * @param Milliseconds The milliseconds component.
+	 * @param Microseconds The microseconds component (default = 0).
 	 */
-	FTimespan( int32 Days, int32 Hours, int32 Minutes, int32 Seconds, int32 Milliseconds )
+	FTimespan(int32 Days, int32 Hours, int32 Minutes, int32 Seconds, int32 Milliseconds, int32 Microseconds = 0)
 	{
-		Assign(Days, Hours, Minutes, Seconds, Milliseconds);
+		Assign(Days, Hours, Minutes, Seconds, Milliseconds, Microseconds);
 	}
 
 public:
@@ -111,7 +115,7 @@ public:
 	 *
 	 * @return A time span whose value is the sum of this time span and the given time span.
 	 */
-	FTimespan operator+( const FTimespan& Other ) const
+	FTimespan operator+(const FTimespan& Other) const
 	{
 		return FTimespan(Ticks + Other.Ticks);
 	}
@@ -121,7 +125,7 @@ public:
 	 *
 	 * @return This time span.
 	 */
-	FTimespan& operator+=( const FTimespan& Other )
+	FTimespan& operator+=(const FTimespan& Other)
 	{
 		Ticks += Other.Ticks;
 
@@ -146,7 +150,7 @@ public:
 	 * @param Other The time span to compare with.
 	 * @return A time span whose value is the difference of this time span and the given time span.
 	 */
-	FTimespan operator-( const FTimespan& Other ) const
+	FTimespan operator-(const FTimespan& Other) const
 	{
 		return FTimespan(Ticks - Other.Ticks);
 	}
@@ -157,7 +161,7 @@ public:
 	 * @param Other The time span to subtract.
 	 * @return This time span.
 	 */
-	FTimespan& operator-=( const FTimespan& Other )
+	FTimespan& operator-=(const FTimespan& Other)
 	{
 		Ticks -= Other.Ticks;
 
@@ -170,7 +174,7 @@ public:
 	 * @param Scalar The scalar to multiply with.
 	 * @return A time span whose value is the product of this time span and the given scalar.
 	 */
-	FTimespan operator*( float Scalar ) const
+	FTimespan operator*(float Scalar) const
 	{
 		return FTimespan((int64)(Ticks * Scalar));
 	}
@@ -181,7 +185,7 @@ public:
 	 * @param Scalar The scalar to multiply with.
 	 * @return This time span.
 	 */
-	FTimespan& operator*=( float Scalar )
+	FTimespan& operator*=(float Scalar)
 	{
 		Ticks = (int64)(Ticks * Scalar);
 
@@ -194,7 +198,7 @@ public:
 	 * @param Other The time span to compare with.
 	 * @return true if the time spans are equal, false otherwise.
 	 */
-	bool operator==( const FTimespan& Other ) const
+	bool operator==(const FTimespan& Other) const
 	{
 		return (Ticks == Other.Ticks);
 	}
@@ -205,7 +209,7 @@ public:
 	 * @param Other The time span to compare with.
 	 * @return true if the time spans are not equal, false otherwise.
 	 */
-	bool operator!=( const FTimespan& Other ) const
+	bool operator!=(const FTimespan& Other) const
 	{
 		return (Ticks != Other.Ticks);
 	}
@@ -216,7 +220,7 @@ public:
 	 * @param Other The time span to compare with.
 	 * @return true if this time span is greater, false otherwise.
 	 */
-	bool operator>( const FTimespan& Other ) const
+	bool operator>(const FTimespan& Other) const
 	{
 		return (Ticks > Other.Ticks);
 	}
@@ -227,7 +231,7 @@ public:
 	 * @param Other The time span to compare with.
 	 * @return true if this time span is greater or equal, false otherwise.
 	 */
-	bool operator>=( const FTimespan& Other ) const
+	bool operator>=(const FTimespan& Other) const
 	{
 		return (Ticks >= Other.Ticks);
 	}
@@ -238,7 +242,7 @@ public:
 	 * @param Other The time span to compare with.
 	 * @return true if this time span is less, false otherwise.
 	 */
-	bool operator<( const FTimespan& Other ) const
+	bool operator<(const FTimespan& Other) const
 	{
 		return (Ticks < Other.Ticks);
 	}
@@ -249,7 +253,7 @@ public:
 	 * @param Other The time span to compare with.
 	 * @return true if this time span is less or equal, false otherwise.
 	 */
-	bool operator<=( const FTimespan& Other ) const
+	bool operator<=(const FTimespan& Other) const
 	{
 		return (Ticks <= Other.Ticks);
 	}
@@ -267,7 +271,7 @@ public:
 	 * @return true on success, false otherwise.
 	 * @see ImportTextItem
 	 */
-	CORE_API bool ExportTextItem( FString& ValueStr, FTimespan const& DefaultValue, UObject* Parent, int32 PortFlags, UObject* ExportRootScope ) const;
+	CORE_API bool ExportTextItem(FString& ValueStr, FTimespan const& DefaultValue, UObject* Parent, int32 PortFlags, UObject* ExportRootScope) const;
 
 	/**
 	 * Gets the days component of this time span.
@@ -303,6 +307,17 @@ public:
 		return (int32)((Ticks / ETimespan::TicksPerHour) % 24);
 	}
 
+	/**
+	 * Gets the microseconds component of this time span.
+	 *
+	 * @return Microseconds component.
+	 * @see GetTotalMicroseconds
+	 */
+	int32 GetMicroseconds() const
+	{
+		return (int32)((Ticks / ETimespan::TicksPerMicrosecond) % 1000);
+	}
+
 	/**
 	 * Gets the milliseconds component of this time span.
 	 *
@@ -368,6 +383,17 @@ public:
 		return ((double)Ticks / ETimespan::TicksPerHour);
 	}
 
+	/**
+	 * Gets the total number of microseconds represented by this time span.
+	 *
+	 * @return Number of microseconds.
+	 * @see GetMicroseconds
+	 */
+	double GetTotalMicroseconds() const
+	{
+		return ((double)Ticks / ETimespan::TicksPerMicrosecond);
+	}
+
 	/**
 	 * Gets the total number of milliseconds represented by this time span.
 	 *
@@ -411,7 +437,7 @@ public:
 	 * @return true on success, false otherwise.
 	 * @see ExportTextItem
 	 */
-	CORE_API bool ImportTextItem( const TCHAR*& Buffer, int32 PortFlags, UObject* Parent, FOutputDevice* ErrorText );
+	CORE_API bool ImportTextItem(const TCHAR*& Buffer, int32 PortFlags, UObject* Parent, FOutputDevice* ErrorText);
 
 
 	/**
@@ -420,7 +446,7 @@ public:
 	 * @param Ar The archive to serialize from or into.
 	 * @return true on success, false otherwise.
 	 */
-	CORE_API bool Serialize( FArchive& Ar );
+	CORE_API bool Serialize(FArchive& Ar);
 
 	/**
 	 * Returns the string representation of this time span using a default format.
@@ -454,7 +480,7 @@ public:
 	 * @return String representation.
 	 * @see Parse
 	 */
-	CORE_API FString ToString( const TCHAR* Format ) const;
+	CORE_API FString ToString(const TCHAR* Format) const;
 
 public:
 
@@ -463,45 +489,54 @@ public:
 	 *
 	 * @param Days The number of days.
 	 * @return Time span.
-	 * @see FromHours, FromMilliseconds, FromMinutes, FromSeconds
+	 * @see FromHours, FromMicroseconds, FromMilliseconds, FromMinutes, FromSeconds
 	 */
-	static CORE_API FTimespan FromDays( double Days );
+	static CORE_API FTimespan FromDays(double Days);
 
 	/**
 	 * Creates a time span that represents the specified number of hours.
 	 *
 	 * @param Hours The number of hours.
 	 * @return Time span.
-	 * @see FromDays, FromMilliseconds, FromMinutes, FromSeconds
+	 * @see FromDays, FromMicroseconds, FromMilliseconds, FromMinutes, FromSeconds
+	 */
+	static CORE_API FTimespan FromHours(double Hours);
+
+	/**
+	 * Creates a time span that represents the specified number of microseconds.
+	 *
+	 * @param Microseconds The number of microseconds.
+	 * @return Time span.
+	 * @see FromDays, FromHours, FromMinutes, FromSeconds, FromMilliseconds
 	 */
-	static CORE_API FTimespan FromHours( double Hours );
+	static CORE_API FTimespan FromMicroseconds(double Microseconds);
 
 	/**
 	 * Creates a time span that represents the specified number of milliseconds.
 	 *
 	 * @param Milliseconds The number of milliseconds.
 	 * @return Time span.
-	 * @see FromDays, FromHours, FromMinutes, FromSeconds
+	 * @see FromDays, FromHours, FromMicroseconds, FromMinutes, FromSeconds
 	 */
-	static CORE_API FTimespan FromMilliseconds( double Milliseconds );
+	static CORE_API FTimespan FromMilliseconds(double Milliseconds);
 
 	/**
 	 * Creates a time span that represents the specified number of minutes.
 	 *
 	 * @param Minutes The number of minutes.
 	 * @return Time span.
-	 * @see FromDays, FromHours, FromMilliseconds, FromSeconds
+	 * @see FromDays, FromHours, FromMicroseconds, FromMilliseconds, FromSeconds
 	 */
-	static CORE_API FTimespan FromMinutes( double Minutes );
+	static CORE_API FTimespan FromMinutes(double Minutes);
 
 	/**
 	 * Creates a time span that represents the specified number of seconds.
 	 *
 	 * @param Seconds The number of seconds.
 	 * @return Time span.
-	 * @see FromDays, FromHours, FromMilliseconds, FromMinutes
+	 * @see FromDays, FromHours, FromMicroseconds, FromMilliseconds, FromMinutes
 	 */
-	static CORE_API FTimespan FromSeconds( double Seconds );
+	static CORE_API FTimespan FromSeconds(double Seconds);
 
 	/**
 	 * Returns the maximum time span value.
@@ -540,7 +575,7 @@ public:
 	 * @return true if the string was converted successfully, false otherwise.
 	 * @see ToString
 	 */
-	static CORE_API bool Parse( const FString& TimespanString, FTimespan& OutTimespan );
+	static CORE_API bool Parse(const FString& TimespanString, FTimespan& OutTimespan);
 
 	/**
 	 * Returns the zero time span value.
@@ -568,7 +603,7 @@ public:
 	 *
 	 * @todo gmp: Figure out better include order in Core.h so this can be inlined.
 	 */
-	friend CORE_API FArchive& operator<<( FArchive& Ar, FTimespan& Timespan );
+	friend CORE_API FArchive& operator<<(FArchive& Ar, FTimespan& Timespan);
 
 	/**
 	 * Gets the hash for the specified time span.
@@ -576,7 +611,7 @@ public:
 	 * @param Timespan The timespan to get the hash for.
 	 * @return Hash value.
 	 */
-	friend uint32 GetTypeHash( const FTimespan& Timespan )
+	friend uint32 GetTypeHash(const FTimespan& Timespan)
 	{
 		return GetTypeHash(Timespan.Ticks);
 	}
@@ -591,8 +626,9 @@ protected:
 	 * @param Minutes The minutes component.
 	 * @param Seconds The seconds component.
 	 * @param Milliseconds The milliseconds component.
+	 * @param Microseconds The microseconds component.
 	 */
-	void CORE_API Assign( int32 Days, int32 Hours, int32 Minutes, int32 Seconds, int32 Milliseconds );
+	void CORE_API Assign(int32 Days, int32 Hours, int32 Minutes, int32 Seconds, int32 Milliseconds, int32 Microseconds);
 
 private:
 
@@ -607,7 +643,7 @@ private:
  * @param Scalar The scalar to pre-multiply with.
  * @param Timespan The time span to multiply.
  */
-inline FTimespan operator*( float Scalar, const FTimespan& Timespan )
+inline FTimespan operator*(float Scalar, const FTimespan& Timespan)
 {
 	return Timespan.operator*(Scalar);
 }
diff --git a/Engine/Source/Runtime/Engine/Classes/Animation/AnimInstance.h b/Engine/Source/Runtime/Engine/Classes/Animation/AnimInstance.h
--- a/Engine/Source/Runtime/Engine/Classes/Animation/AnimInstance.h
+++ b/Engine/Source/Runtime/Engine/Classes/Animation/AnimInstance.h
@@ -594,7 +594,10 @@ public:
 
 	/** Return true if it's playing the slot animation */
 	UFUNCTION(BlueprintCallable, Category="Animation")
-	bool IsPlayingSlotAnimation(UAnimSequenceBase* Asset, FName SlotNodeName ); 
+	bool IsPlayingSlotAnimation(UAnimSequenceBase* Asset, FName SlotNodeName); 
+
+	/** Return true if this instance playing the slot animation, also returning the montage it is playing on */
+	bool IsPlayingSlotAnimation(UAnimSequenceBase* Asset, FName SlotNodeName, UAnimMontage*& OutMontage);
 
 	/*********************************************************************************************
 	 * AnimMontage
@@ -1072,6 +1075,10 @@ public:
 
 	// Debug output for this anim instance 
 	void DisplayDebug(UCanvas* Canvas, const FDebugDisplayInfo& DebugDisplay, float& YL, float& YPos);
+
+	/** Reset any dynamics running simulation-style updates (e.g. on teleport, time skip etc.) */
+	void ResetDynamics();
+
 public:
 
 	/** Access the required bones array */
diff --git a/Engine/Source/Runtime/Engine/Classes/Animation/AnimNodeBase.h b/Engine/Source/Runtime/Engine/Classes/Animation/AnimNodeBase.h
--- a/Engine/Source/Runtime/Engine/Classes/Animation/AnimNodeBase.h
+++ b/Engine/Source/Runtime/Engine/Classes/Animation/AnimNodeBase.h
@@ -566,6 +566,15 @@ struct ENGINE_API FAnimNode_Base
 
 	/** Override this to perform game-thread work prior to non-game thread Update() being called */
 	virtual void PreUpdate(const UAnimInstance* InAnimInstance) {}
+
+	/**
+	 * For nodes that implement some kind of simulation, return true here so ResetDynamics() gets called
+	 * when things like teleports, time skips etc. occur that might require special handling
+	 */
+	virtual bool NeedsDynamicReset() const { return false; }
+
+	/** Override this to perform game-thread work prior to non-game thread Update() being called */
+	virtual void ResetDynamics() {}
 	// End of interface to implement
 
 	virtual ~FAnimNode_Base() {}
diff --git a/Engine/Source/Runtime/Engine/Classes/Camera/CameraAnimInst.h b/Engine/Source/Runtime/Engine/Classes/Camera/CameraAnimInst.h
--- a/Engine/Source/Runtime/Engine/Classes/Camera/CameraAnimInst.h
+++ b/Engine/Source/Runtime/Engine/Classes/Camera/CameraAnimInst.h
@@ -144,9 +144,6 @@ public:
 	UFUNCTION(BlueprintCallable, Category = CameraAnimInst)
 	void SetScale(float NewDuration);
 
-	/** Sets the camera actor that will be modified when this anim is played. (Usually a temporary actor whose data is subsequently harvested). */
-	void SetCameraActor(class AActor* Actor);
-
 	/** Takes the given view and applies the camera anim transform and fov changes to it. Does not affect PostProcess. */
 	void ApplyToView(FMinimalViewInfo& InOutPOV) const;
 
