UE4 - Fix hang in FTaskGraphInterface::BroadcastSlow_OnlyUseForSpecialPurposes caused by back to back calls.

diff --git a/Engine/Source/Runtime/Core/Private/Async/TaskGraph.cpp b/Engine/Source/Runtime/Core/Private/Async/TaskGraph.cpp
--- a/Engine/Source/Runtime/Core/Private/Async/TaskGraph.cpp
+++ b/Engine/Source/Runtime/Core/Private/Async/TaskGraph.cpp
@@ -2671,11 +2671,13 @@ public:
 		Function(CurrentThread);
 		if (StallForTaskThread)
 		{
-			StallForTaskThread->Decrement();
-			// we wait for the others to finish here so that we do all task threads
-			while (StallForTaskThread->GetValue())
+			if (StallForTaskThread->Decrement())
 			{
-				FPlatformProcess::SleepNoStats(.0002f);
+				// we wait for the others to finish here so that we do all task threads
+				while (StallForTaskThread->GetValue())
+				{
+					FPlatformProcess::SleepNoStats(.0002f);
+				}
 			}
 		}
 	}
@@ -2697,18 +2699,18 @@ void FTaskGraphInterface::BroadcastSlow_OnlyUseForSpecialPurposes(bool bDoBackgr
 		return;
 	}
 
-	// these are static because we don't actually wait for the tasks to clear
-	static FThreadSafeCounter StallForTaskThread;
+	FThreadSafeCounter StallForTaskThread;
 
 	int32 Workers = FTaskGraphInterface::Get().GetNumWorkerThreads();
 	StallForTaskThread.Reset();
 	StallForTaskThread.Add(Workers * (1 + (bDoBackgroundThreads && ENamedThreads::bHasBackgroundThreads) + !!(ENamedThreads::bHasHighPriorityThreads)));
 
+	FGraphEventArray TaskThreadTasks;
 	{
 
 		for (int32 Index = 0; Index < Workers; Index++)
 		{
-			TGraphTask<FBroadcastTask>::CreateTask().ConstructAndDispatchWhenReady(Callback, ENamedThreads::AnyNormalThreadHiPriTask, &StallForTaskThread);
+			TaskThreadTasks.Add(TGraphTask<FBroadcastTask>::CreateTask().ConstructAndDispatchWhenReady(Callback, ENamedThreads::AnyNormalThreadHiPriTask, &StallForTaskThread));
 		}
 
 	}
@@ -2716,14 +2718,14 @@ void FTaskGraphInterface::BroadcastSlow_OnlyUseForSpecialPurposes(bool bDoBackgr
 	{
 		for (int32 Index = 0; Index < Workers; Index++)
 		{
-			TGraphTask<FBroadcastTask>::CreateTask().ConstructAndDispatchWhenReady(Callback, ENamedThreads::AnyHiPriThreadHiPriTask, &StallForTaskThread);
+			TaskThreadTasks.Add(TGraphTask<FBroadcastTask>::CreateTask().ConstructAndDispatchWhenReady(Callback, ENamedThreads::AnyHiPriThreadHiPriTask, &StallForTaskThread));
 		}
 	}
 	if (bDoBackgroundThreads && ENamedThreads::bHasBackgroundThreads)
 	{
 		for (int32 Index = 0; Index < Workers; Index++)
 		{
-			TGraphTask<FBroadcastTask>::CreateTask().ConstructAndDispatchWhenReady(Callback, ENamedThreads::AnyBackgroundHiPriTask, &StallForTaskThread);
+			TaskThreadTasks.Add(TGraphTask<FBroadcastTask>::CreateTask().ConstructAndDispatchWhenReady(Callback, ENamedThreads::AnyBackgroundHiPriTask, &StallForTaskThread));
 		}
 	}
 	FGraphEventArray Tasks;
@@ -2752,9 +2754,10 @@ void FTaskGraphInterface::BroadcastSlow_OnlyUseForSpecialPurposes(bool bDoBackgr
 		else if (FPlatformTime::Seconds() - StartTime > 3.0)
 		{
 			StartTime = FPlatformTime::Seconds();
-			UE_LOG(LogTaskGraph, Error, TEXT("Broadcast failed after three seconds"));
+			UE_LOG(LogTaskGraph, Error, TEXT("FTaskGraphInterface::BroadcastSlow_OnlyUseForSpecialPurposes Broadcast failed after three seconds"));
 		}
 	}
+	FTaskGraphInterface::Get().WaitUntilTasksComplete(TaskThreadTasks, ENamedThreads::GameThread_Local);
 }
 
 
