Fixed -windowed using the wrong window size

This was being clamped to the client size of the window, rather than the screen size of the window. This was also preventing the resolution being increased from the initial size.

We now consider the max size of a window to be its client size, and we use the maximum of the display area or requested resolution as the maximum window size (as we sometimes create 4K windows to capture).

diff --git a/Engine/Source/Runtime/Core/Private/Windows/WindowsApplication.cpp b/Engine/Source/Runtime/Core/Private/Windows/WindowsApplication.cpp
--- a/Engine/Source/Runtime/Core/Private/Windows/WindowsApplication.cpp
+++ b/Engine/Source/Runtime/Core/Private/Windows/WindowsApplication.cpp
@@ -1087,10 +1087,27 @@ int32 FWindowsApplication::ProcessMessage( HWND hwnd, uint32 msg, WPARAM wParam,
 			{
 				MINMAXINFO* MinMaxInfo = (MINMAXINFO*)lParam;
 				FWindowSizeLimits SizeLimits = MessageHandler->GetSizeLimitsForWindow(CurrentNativeEventWindow);
+
+				// We need to inflate the max values if using an OS window border
+				int32 BorderWidth = 0;
+				int32 BorderHeight = 0;
+				if (CurrentNativeEventWindow->GetDefinition().HasOSWindowBorder)
+				{
+					const DWORD WindowStyle = ::GetWindowLong(hwnd, GWL_STYLE);
+					const DWORD WindowExStyle = ::GetWindowLong(hwnd, GWL_EXSTYLE);
+
+					// This adjusts a zero rect to give us the size of the border
+					RECT BorderRect = { 0, 0, 0, 0 };
+					::AdjustWindowRectEx(&BorderRect, WindowStyle, false, WindowExStyle);
+
+					BorderWidth = BorderRect.right - BorderRect.left;
+					BorderHeight = BorderRect.bottom - BorderRect.top;
+				}
+
 				MinMaxInfo->ptMinTrackSize.x = FMath::RoundToInt( SizeLimits.GetMinWidth().Get(MinMaxInfo->ptMinTrackSize.x) );
 				MinMaxInfo->ptMinTrackSize.y = FMath::RoundToInt( SizeLimits.GetMinHeight().Get(MinMaxInfo->ptMinTrackSize.y) );
-				MinMaxInfo->ptMaxTrackSize.x = FMath::RoundToInt( SizeLimits.GetMaxWidth().Get(MinMaxInfo->ptMaxTrackSize.x) );
-				MinMaxInfo->ptMaxTrackSize.y = FMath::RoundToInt( SizeLimits.GetMaxHeight().Get(MinMaxInfo->ptMaxTrackSize.y) );
+				MinMaxInfo->ptMaxTrackSize.x = FMath::RoundToInt( SizeLimits.GetMaxWidth().Get(MinMaxInfo->ptMaxTrackSize.x - BorderWidth) ) + BorderWidth;
+				MinMaxInfo->ptMaxTrackSize.y = FMath::RoundToInt( SizeLimits.GetMaxHeight().Get(MinMaxInfo->ptMaxTrackSize.y - BorderHeight) ) + BorderHeight;
 				return 0;
 			}
 			break;
diff --git a/Engine/Source/Runtime/Engine/Private/GameEngine.cpp b/Engine/Source/Runtime/Engine/Private/GameEngine.cpp
--- a/Engine/Source/Runtime/Engine/Private/GameEngine.cpp
+++ b/Engine/Source/Runtime/Engine/Private/GameEngine.cpp
@@ -329,13 +329,35 @@ TSharedRef<SWindow> UGameEngine::CreateGameWindow()
 		AutoCenterType = EAutoCenter::None;
 	}
 
+	// Give the window the max width/height of either the requested resolution, or your available desktop resolution
+	// We need to do this as we request some 4K windows when rendering sequences, and the OS may try and clamp that
+	// window to your available desktop resolution
+	TOptional<float> MaxWindowWidth;
+	TOptional<float> MaxWindowHeight;
+	if (WindowMode == EWindowMode::Windowed)
+	{
+		// Get available desktop area
+		FDisplayMetrics DisplayMetrics;
+		if (FSlateApplication::IsInitialized())
+		{
+			FSlateApplication::Get().GetInitialDisplayMetrics(DisplayMetrics);
+		}
+		else
+		{
+			FDisplayMetrics::GetDisplayMetrics(DisplayMetrics);
+		}
+
+		MaxWindowWidth = FMath::Max(DisplayMetrics.VirtualDisplayRect.Right - DisplayMetrics.VirtualDisplayRect.Left, ResX);
+		MaxWindowHeight = FMath::Max(DisplayMetrics.VirtualDisplayRect.Bottom - DisplayMetrics.VirtualDisplayRect.Top, ResY);
+	}
+
 	TSharedRef<SWindow> Window = SNew(SWindow)
 	.ClientSize(FVector2D( ResX, ResY ))
 	.Title(WindowTitle)
 	.AutoCenter(AutoCenterType)
 	.ScreenPosition(FVector2D(WinX, WinY))
-	.MaxWidth(ResX)
-	.MaxHeight(ResY)
+	.MaxWidth(MaxWindowWidth)
+	.MaxHeight(MaxWindowHeight)
 	.FocusWhenFirstShown(true)
 	.SaneWindowPlacement(AutoCenterType == EAutoCenter::None)
 	.UseOSWindowBorder(true);
