Backing out some changes from CL# 2976697

These caused an issue with Slate hit-testing. The more correct fix here is to make the Slate Windows OS layer treat window positions as relative to the top-left of the window client area, rather than relative to the top-left of the window itself (which includes the OS border). This now matches what other platforms do.

To this end, FWindowsWindow::Initialize, FWindowsWindow::MoveWindowTo, and FWindowsWindow::ReshapeWindow all now consider the given window position to be relative to the window client area, and will consistently adjust it to relative to the window before moving/creating the OS window. This only impacts windows with OS borders (aka, non-fullscreen and non-Slate drawn windows).

diff --git a/Engine/Source/Runtime/Core/Private/Windows/WindowsApplication.cpp b/Engine/Source/Runtime/Core/Private/Windows/WindowsApplication.cpp
--- a/Engine/Source/Runtime/Core/Private/Windows/WindowsApplication.cpp
+++ b/Engine/Source/Runtime/Core/Private/Windows/WindowsApplication.cpp
@@ -902,21 +902,15 @@ int32 FWindowsApplication::ProcessMessage( HWND hwnd, uint32 msg, WPARAM wParam,
 
 		case WM_MOVE:
 			{
-				// client area position (needed to test for minimized windows)
-				int32 NewX = (int)(short)(LOWORD(lParam));
-				int32 NewY = (int)(short)(HIWORD(lParam));
+				// client area position
+				const int32 NewX = (int)(short)(LOWORD(lParam));
+				const int32 NewY = (int)(short)(HIWORD(lParam));
+				FIntPoint NewPosition(NewX,NewY);
 
 				// Only cache the screen position if its not minimized
-				if (FWindowsApplication::MinimizedWindowPosition.X != NewX && FWindowsApplication::MinimizedWindowPosition.Y != NewY)
+				if ( FWindowsApplication::MinimizedWindowPosition != NewPosition )
 				{
-					// Slate expects screen space, but WM_MOVE is given in client space - get the real window rect
-					RECT WndRect;
-					::GetWindowRect(hwnd, &WndRect);
-
-					NewX = WndRect.left;
-					NewY = WndRect.top;
-
-					MessageHandler->OnMovedWindow(CurrentNativeEventWindow, NewX, NewY);
+					MessageHandler->OnMovedWindow( CurrentNativeEventWindow, NewX, NewY );
 
 					return 0;
 				}
@@ -1110,10 +1104,11 @@ int32 FWindowsApplication::ProcessMessage( HWND hwnd, uint32 msg, WPARAM wParam,
 					BorderHeight = BorderRect.bottom - BorderRect.top;
 				}
 
+				// We always apply BorderWidth and BorderHeight since Slate always works with client area window sizes
 				MinMaxInfo->ptMinTrackSize.x = FMath::RoundToInt( SizeLimits.GetMinWidth().Get(MinMaxInfo->ptMinTrackSize.x) );
 				MinMaxInfo->ptMinTrackSize.y = FMath::RoundToInt( SizeLimits.GetMinHeight().Get(MinMaxInfo->ptMinTrackSize.y) );
-				MinMaxInfo->ptMaxTrackSize.x = FMath::RoundToInt( SizeLimits.GetMaxWidth().Get(MinMaxInfo->ptMaxTrackSize.x - BorderWidth) ) + BorderWidth;
-				MinMaxInfo->ptMaxTrackSize.y = FMath::RoundToInt( SizeLimits.GetMaxHeight().Get(MinMaxInfo->ptMaxTrackSize.y - BorderHeight) ) + BorderHeight;
+				MinMaxInfo->ptMaxTrackSize.x = FMath::RoundToInt( SizeLimits.GetMaxWidth().Get(MinMaxInfo->ptMaxTrackSize.x) ) + BorderWidth;
+				MinMaxInfo->ptMaxTrackSize.y = FMath::RoundToInt( SizeLimits.GetMaxHeight().Get(MinMaxInfo->ptMaxTrackSize.y) ) + BorderHeight;
 				return 0;
 			}
 			break;
diff --git a/Engine/Source/Runtime/Core/Private/Windows/WindowsWindow.cpp b/Engine/Source/Runtime/Core/Private/Windows/WindowsWindow.cpp
--- a/Engine/Source/Runtime/Core/Private/Windows/WindowsWindow.cpp
+++ b/Engine/Source/Runtime/Core/Private/Windows/WindowsWindow.cpp
@@ -53,10 +53,12 @@ void FWindowsWindow::Initialize( FWindowsApplication* const Application, const T
 	const float WidthInitial = Definition->WidthDesiredOnScreen;
 	const float HeightInitial = Definition->HeightDesiredOnScreen;
 
-	int32 X = FMath::TruncToInt( XInitialRect );
-	int32 Y = FMath::TruncToInt( YInitialRect );
+	int32 ClientX = FMath::TruncToInt( XInitialRect );
+	int32 ClientY = FMath::TruncToInt( YInitialRect );
 	int32 ClientWidth = FMath::TruncToInt( WidthInitial );
 	int32 ClientHeight = FMath::TruncToInt( HeightInitial );
+	int32 WindowX = ClientX;
+	int32 WindowY = ClientY;
 	int32 WindowWidth = ClientWidth;
 	int32 WindowHeight = ClientHeight;
 	const bool bApplicationSupportsPerPixelBlending =
@@ -138,16 +140,19 @@ void FWindowsWindow::Initialize( FWindowsApplication* const Application, const T
 			WindowStyle |= WS_POPUP | WS_BORDER;
 		}
 
-		// Note SizeX and SizeY should be the size of the client area.  We need to get the actual window size by adjusting the client size to account for standard windows border around the window
-		RECT WindowRect = { 0, 0, ClientWidth, ClientWidth };
-		::AdjustWindowRectEx(&WindowRect,WindowStyle,0,WindowExStyle);
+		// X,Y, Width, Height defines the top-left pixel of the client area on the screen
+		// This adjusts a zero rect to give us the size of the border
+		RECT BorderRect = { 0, 0, 0, 0 };
+		::AdjustWindowRectEx(&BorderRect, WindowStyle, false, WindowExStyle);
 
-		X += WindowRect.left;
-		Y += WindowRect.top;
-		WindowWidth = WindowRect.right - WindowRect.left;
-		WindowHeight = WindowRect.bottom - WindowRect.top;
-	}
+		// Border rect size is negative - see MoveWindowTo
+		WindowX += BorderRect.left;
+		WindowY += BorderRect.top;
 
+		// Inflate the window size by the OS border
+		WindowWidth += BorderRect.right - BorderRect.left;
+		WindowHeight += BorderRect.bottom - BorderRect.top;
+	}
 
 	// Creating the Window
 	HWnd = CreateWindowEx(
@@ -155,7 +160,8 @@ void FWindowsWindow::Initialize( FWindowsApplication* const Application, const T
 		AppWindowClass,
 		*Definition->Title,
 		WindowStyle,
-		X, Y, WindowWidth, WindowHeight,
+		WindowX, WindowY, 
+		WindowWidth, WindowHeight,
 		( InParent.IsValid() ) ? static_cast<HWND>( InParent->HWnd ) : NULL,
 		NULL, InHInstance, NULL);
 
@@ -166,7 +172,7 @@ void FWindowsWindow::Initialize( FWindowsApplication* const Application, const T
 	// in the initial creation of the window. Slate should only pass client area dimensions.
 	// Reshape window may resize the window if the non-client area is encroaching on our
 	// desired client area space.
-	ReshapeWindow( X, Y, ClientWidth, ClientHeight );
+	ReshapeWindow( ClientX, ClientY, ClientWidth, ClientHeight );
 
 	if( HWnd == NULL )
 	{
@@ -342,16 +348,20 @@ void FWindowsWindow::ReshapeWindow( int32 NewX, int32 NewY, int32 NewWidth, int3
 	WindowInfo.cbSize = sizeof( WindowInfo );
 	::GetWindowInfo( HWnd, &WindowInfo );
 
-	// X,Y, Width, Height defines the pixel left of the client area on the screen
-	RECT ClientRect = { 0, 0, NewWidth, NewHeight };
-	RECT WindowRect = { 0, 0, NewWidth, NewHeight };
+	// X,Y, Width, Height defines the top-left pixel of the client area on the screen
 	if( Definition->HasOSWindowBorder )
 	{
-		// Note SizeX and SizeY should be the size of the client area.  We need to get the actual window size by adjusting the client size to account for standard windows border around the window
-		::AdjustWindowRectEx(&WindowRect,WindowInfo.dwStyle,0,WindowInfo.dwExStyle);
+		// This adjusts a zero rect to give us the size of the border
+		RECT BorderRect = { 0, 0, 0, 0 };
+		::AdjustWindowRectEx(&BorderRect, WindowInfo.dwStyle, false, WindowInfo.dwExStyle);
 
-		NewWidth = WindowRect.right - WindowRect.left;
-		NewHeight = WindowRect.bottom - WindowRect.top;
+		// Border rect size is negative - see MoveWindowTo
+		NewX += BorderRect.left;
+		NewY += BorderRect.top;
+
+		// Inflate the window size by the OS border
+		NewWidth += BorderRect.right - BorderRect.left;
+		NewHeight += BorderRect.bottom - BorderRect.top;
 	}
 
 	// the window position is the requested position
@@ -413,9 +423,24 @@ bool FWindowsWindow::GetFullScreenInfo( int32& X, int32& Y, int32& Width, int32&
 	return true;
 }
 
-/** Native windows should implement MoveWindowTo by relocating the platform-specific window to (X,Y). */
+/** Native windows should implement MoveWindowTo by relocating the client area of the platform-specific window to (X,Y). */
 void FWindowsWindow::MoveWindowTo( int32 X, int32 Y )
 {
+	// Slate gives the window position as relative to the client area of a window, so we may need to compensate for the OS border
+	if (Definition->HasOSWindowBorder)
+	{
+		const LONG WindowStyle = ::GetWindowLong(HWnd, GWL_STYLE);
+		const LONG WindowExStyle = ::GetWindowLong(HWnd, GWL_EXSTYLE);
+
+		// This adjusts a zero rect to give us the size of the border
+		RECT BorderRect = { 0, 0, 0, 0 };
+		::AdjustWindowRectEx(&BorderRect, WindowStyle, false, WindowExStyle);
+
+		// Border rect size is negative
+		X += BorderRect.left;
+		Y += BorderRect.top;
+	}
+
 	::SetWindowPos(HWnd, nullptr, X, Y, 0, 0, SWP_NOACTIVATE | SWP_NOSIZE | SWP_NOZORDER);
 }
 
diff --git a/Engine/Source/Runtime/SlateCore/Private/Widgets/SWindow.cpp b/Engine/Source/Runtime/SlateCore/Private/Widgets/SWindow.cpp
--- a/Engine/Source/Runtime/SlateCore/Private/Widgets/SWindow.cpp
+++ b/Engine/Source/Runtime/SlateCore/Private/Widgets/SWindow.cpp
@@ -335,6 +335,10 @@ void SWindow::Construct(const FArguments& InArgs)
 		const FVector2D DisplayTopLeft( AutoCenterRect.Left, AutoCenterRect.Top );
 		const FVector2D DisplaySize( AutoCenterRect.Right - AutoCenterRect.Left, AutoCenterRect.Bottom - AutoCenterRect.Top );
 		WindowPosition = DisplayTopLeft + ( DisplaySize - WindowSize ) * 0.5f;
+
+		// Don't allow the window to center to outside of the work area
+		WindowPosition.X = FMath::Max(WindowPosition.X, (float)PrimaryDisplayRect.Left);
+		WindowPosition.Y = FMath::Max(WindowPosition.Y, (float)PrimaryDisplayRect.Top);
 	}
 
 #if PLATFORM_HTML5 
