Fixing FeedbackContextAnsi not printing Display verbosity (programs not outputting anything to stdout).

diff --git a/Engine/Source/Runtime/Core/Public/HAL/FeedbackContextAnsi.h b/Engine/Source/Runtime/Core/Public/HAL/FeedbackContextAnsi.h
--- a/Engine/Source/Runtime/Core/Public/HAL/FeedbackContextAnsi.h
+++ b/Engine/Source/Runtime/Core/Public/HAL/FeedbackContextAnsi.h
@@ -38,7 +38,7 @@ public:
 	{}
 	void Serialize( const TCHAR* V, ELogVerbosity::Type Verbosity, const class FName& Category ) override
 	{
-		if( Verbosity==ELogVerbosity::Error || Verbosity==ELogVerbosity::Warning )
+		if (Verbosity == ELogVerbosity::Error || Verbosity == ELogVerbosity::Warning || Verbosity == ELogVerbosity::Display)
 		{
 			if( TreatWarningsAsErrors && Verbosity==ELogVerbosity::Warning )
 			{
@@ -55,11 +55,11 @@ public:
 			// Only store off the message if running a commandlet.
 			if ( IsRunningCommandlet() )
 			{
-				if(Verbosity == ELogVerbosity::Error)
+				if (Verbosity == ELogVerbosity::Error)
 				{
 					AddError(Format);
 				}
-				else
+				else if (Verbosity == ELogVerbosity::Warning)
 				{
 					AddWarning(Format);
 				}
