Copying //UE4/Dev-Sequencer to //UE4/Dev-Main (Source: //UE4/Dev-Sequencer @ 3003857)

==========================
MAJOR FEATURES + CHANGES
==========================

Change 2945755 on 2016/04/15 by [REDACTED_PERSON]

	Sequencer - Fix issues with level visibility.
	+ Don't mark sub-levels as dirty when the track evaluates.
	+ Fix an issue where sequencer gets into a refresh loop because drawing thumbnails causes levels to be added which was rebuilding the tree, which was redrawing thumbnails.
	+ Null check for when an objects world is null but the track is still evaluating.
	+ Remove UnrealEd references.

Change 2947197 on 2016/04/18 by [REDACTED_PERSON]

	Sequencer: Expose settings sequencer settings in the Editor Preferences page. Note, UMG and Niagara have separate sequencer settings pages.

Change 2948468 on 2016/04/19 by [REDACTED_PERSON]

	Sequencer: Fix particles not firing on loop.

Change 2948590 on 2016/04/19 by [REDACTED_PERSON]

	Sequencer: Fix spawnables not getting default tracks.

Change 2955993 on 2016/04/26 by [REDACTED_PERSON]

	Sequencer: Refresh instances when done recording. This fixes a bug where spawned recorded actors aren't visible when done recording.

Change 2958567 on 2016/04/27 by [REDACTED_PERSON]

	RHI: Made SetReferencedTexture public, so that the referenced texture can be set

Change 2958718 on 2016/04/28 by [REDACTED_PERSON]

	Sequencer: Folder colors. Right click on a folder and choose "Set Color"

Change 2960172 on 2016/04/28 by [REDACTED_PERSON]

	Slate: Slate Remote Server (for the iOS touch input app) is now disabled by default, so we don't open up the socket unless desired by the user

Change 2960411 on 2016/04/28 by [REDACTED_PERSON]

	Sequencer: Don't remove label if it's not being used.

Change 2960414 on 2016/04/28 by [REDACTED_PERSON]

	Matinee: Don't automatically turn frustums on/off when entering and exiting Matinee.

Change 2962784 on 2016/05/02 by [REDACTED_PERSON]

	Sequencer: Add master sequence

Change 2964399 on 2016/05/03 by [REDACTED_PERSON]

	Sequencer: Added ability to apply cook-time optimization to tracks and objects
	  - For now, if a spawnable has a spawn track that is disabled, or will never spawn, the entire spawnable object will be removed from a cooked package.
	  - Possessables also afford the same optimization, although none is currently implemented
	  - We could, in future, also remove any tracks that are completely disabled
	  - Deprecated UMovieSceneBoolSection::DefaultValue in favor of the default stored on FIntegralCurve

Change 2967549 on 2016/05/05 by [REDACTED_PERSON]

	Sequencer: Fix crash converting possessable to spawnable when the possessable doesn't exist.

Change 2967670 on 2016/05/05 by [REDACTED_PERSON]

	Sequencer: Set ui min/max for sequencer settings

Change 2978969 on 2016/05/16 by [REDACTED_PERSON]

	Sequencer: Restore state when focusing on a shot level sequence. This fixes issues where tracks in the movie scene that are active before switching to the new movie scene need to return to
their initial state. For example, setting a fade track in the master sequence and switching into a shot should disable the effects of the fade track in the master sequence.

Change 2983237 on 2016/05/19 by [REDACTED_PERSON]

	Protocol settings for movie captures are now set up correctly when a capture type is specified on the command line

	Thanks to original github author, [REDACTED_PERSON]

	#pr
	#2257

Change 2991115 on 2016/05/26 by [REDACTED_PERSON]

	Sequencer: Added {shot} and {shot_frame} format args for movie captures

	  - Additionally, rendering out movie scenes as videos will now generate a new video for each unique filename it encounters. This allows us to render out a video per shot by using {shot} as
the output format.
	  - Frame numbers are now zero-padded as per the sequencer setting.

Change 2991920 on 2016/05/26 by [REDACTED_PERSON]

	Sequencer: Fix movie scene getting dirtied unnecessarily when the fixed frame interval changes.

Change 2992387 on 2016/05/26 by [REDACTED_PERSON]

	Sequencer: Fix crash when getting the color key properties of a collapsed key that doesn't have all channels keyed.

Change 2993553 on 2016/05/27 by [REDACTED_PERSON]

	Sequencer: Added the ability to add burn-ins to level sequences

	  - A default burn-in is provided which hosts a great level of flexibility
	  - 6 regions (L/C/R + T/B) on a 30% black border allow positioning of a range of frame statistics such as shot name ({ShotName}), frame numbers ({MasterFrame}, {ShotFrame}), and other
information.
	  - Watermark is provided by default (currently no tiling is exposed)
	  - Users can use the default built in UMG widget as a guideline for their own custom implementations.

Change 2993554 on 2016/05/27 by [REDACTED_PERSON]

	Sequencer: Default level sequence burn ins

	  - Also made a font asset out of our fixed width font shipped with the engine

Change 2993856 on 2016/05/30 by [REDACTED_PERSON]

	Sequencer: Import/Export EDL

	 - Added a new option in the render movie dialog to export an Edit Decision List (EDL) in cmx and rv formats if there is a shot track. The default is true.
	 - Added "Import EDL" to shot track right click menu which imports a cmx EDL and conforms the shot order and cut information to it.
	 - Added "Export EDL" to shot track right click menu which exports EDLs in cmx and rv formats.
	 - Added "Render Shot" to shot right click menu which loads up the render movie dialog with the start and end frames of the selected shot.

Change 2994761 on 2016/05/31 by [REDACTED_PERSON]

	Sequence Recorder: Add a setting to allow recording of actors that are spawned by sequencer itself.

Change 2995648 on 2016/06/01 by HaarmPieter.Duiker

	Sequencer EXR output gamut controls

Change 2996241 on 2016/06/01 by [REDACTED_PERSON]

	Sequencer - Add a small epsilon when "force fixed frame interval" is enabled, to make sure we're in the start of the next frame.

Change 2996244 on 2016/06/01 by [REDACTED_PERSON]

	Sequencer - Set the tick prerequisite for all components, not just the root.

Change 2997865 on 2016/06/02 by [REDACTED_PERSON]

	Sequencer: Fixed Crash in Sequencer play rate track when setting negative play rate (UE-31431)

Change 2999631 on 2016/06/03 by [REDACTED_PERSON]

	Sequencer - At runtime, make sure to stop playing skeletal animations to prevent them from being double updated each frame, once by sequencer, and then again by tick.

Change 3000820 on 2016/06/03 by [REDACTED_PERSON]

	Sequencer: Add hotkey (ctrl-T) to toggle between showing frame numbers and time.

Change 3001056 on 2016/06/05 by [REDACTED_PERSON]

	Sequencer: Fix fade color section crash by using an inline color picker in the details panel instead of a popup color picker.

Change 3001057 on 2016/06/05 by [REDACTED_PERSON]

	Movie Capture: Fix audio getting disabled after recording a movie.

Change 3001690 on 2016/06/06 by [REDACTED_PERSON]

	Sequencer: Fixed recording video sequences when not overwriting existing videos

Change 3001823 on 2016/06/06 by [REDACTED_PERSON]

	Sequencer: Fix filtered nodes in folders so that other unfiltered children aren't visible.

diff --git a/Engine/Config/BaseEngine.ini b/Engine/Config/BaseEngine.ini
--- a/Engine/Config/BaseEngine.ini
+++ b/Engine/Config/BaseEngine.ini
@@ -2375,7 +2375,8 @@ WorldNormal=(Material="/Engine/BufferVisualization/WorldNormal.WorldNormal", Nam
 AmbientOcclusion=(Material="/Engine/BufferVisualization/AmbientOcclusion.AmbientOcclusion", Name=LOCTEXT("BaseAmbientOcclusionMat", "Ambient Occlusion"))
 CustomDepthWorldUnits=(Material="/Engine/BufferVisualization/CustomDepthWorldUnits.CustomDepthWorldUnits", Name=LOCTEXT("BaseCustomDepthWorldUnitsMat", "Custom Depth World Units"))
 SceneDepthWorldUnits=(Material="/Engine/BufferVisualization/SceneDepthWorldUnits.SceneDepthWorldUnits", Name=LOCTEXT("BaseSceneDepthWorldUnitsMat", "Scene Depth World Units"))
-HDRColor=(Material="/Engine/BufferVisualization/HDRColor.HDRColor", Name=LOCTEXT("HDRColor", "Pre Tonemap HDR Color"))
+; HDRColor=(Material="/Engine/BufferVisualization/HDRColor.HDRColor", Name=LOCTEXT("HDRColor", "Pre Tonemap HDR Color"))
+PostTonemapHDRColor=(Material="/Engine/BufferVisualization/HDRColor.HDRColor", Name=LOCTEXT("PostTonemapHDRColor", "Post Tonemap HDR Color"))
 
 [DeviceProfileManager]
 DeviceProfileSelectionModule="ExampleDeviceProfileSelector"
diff --git a/Engine/Plugins/Slate/SlateRemote/Source/SlateRemote/Private/SlateRemoteClasses.cpp b/Engine/Plugins/Slate/SlateRemote/Source/SlateRemote/Private/SlateRemoteClasses.cpp
--- a/Engine/Plugins/Slate/SlateRemote/Source/SlateRemote/Private/SlateRemoteClasses.cpp
+++ b/Engine/Plugins/Slate/SlateRemote/Source/SlateRemote/Private/SlateRemoteClasses.cpp
@@ -5,7 +5,7 @@
 
 USlateRemoteSettings::USlateRemoteSettings( const FObjectInitializer& ObjectInitializer )
 	: Super(ObjectInitializer)
-	, EnableRemoteServer(true)
+	, EnableRemoteServer(false)
 	, EditorServerEndpoint(SLATE_REMOTE_SERVER_DEFAULT_EDITOR_ENDPOINT.ToString())
 	, GameServerEndpoint(SLATE_REMOTE_SERVER_DEFAULT_GAME_ENDPOINT.ToString())
 { }
diff --git a/Engine/Shaders/TonemapCommon.usf b/Engine/Shaders/TonemapCommon.usf
--- a/Engine/Shaders/TonemapCommon.usf
+++ b/Engine/Shaders/TonemapCommon.usf
@@ -271,9 +271,11 @@ float3x3 OuputGamutMappingMatrix( uint OutputGamut )
 
 	// Set gamut mapping matrix 
 	// 0 = sRGB - D65
-	// 1 = P3 (DCI) - D65
+	// 1 = P3 - D65
 	// 2 = Rec.2020 - D65
 	// 3 = ACES AP0 - D60
+	// 4 = ACES AP1 - D60
+
 	float3x3 GamutMappingMatrix = AP1_2_sRGB;
 	if( OutputGamut == 1 )
 		GamutMappingMatrix = AP1_2_DCI_D65;
@@ -281,6 +283,8 @@ float3x3 OuputGamutMappingMatrix( uint OutputGamut )
 		GamutMappingMatrix = AP1_2_Rec2020;
 	else if( OutputGamut == 3 )
 		GamutMappingMatrix = AP1_2_AP0_MAT;
+	else if( OutputGamut == 4 )
+		GamutMappingMatrix = float3x3(1,0,0,0,1,0,0,0,1);
 
 	return GamutMappingMatrix;
 }
diff --git a/Engine/Source/Editor/DetailCustomizations/Private/DetailCustomizations.cpp b/Engine/Source/Editor/DetailCustomizations/Private/DetailCustomizations.cpp
--- a/Engine/Source/Editor/DetailCustomizations/Private/DetailCustomizations.cpp
+++ b/Engine/Source/Editor/DetailCustomizations/Private/DetailCustomizations.cpp
@@ -103,6 +103,7 @@
 #include "CaptureTypeCustomization.h"
 #include "RenderPassesCustomization.h"
 #include "MovieSceneCaptureCustomization.h"
+#include "LevelSequenceBurnInOptionsCustomization.h"
 #include "TextCustomization.h"
 #include "AnimTrailNodeDetails.h"
 #include "MaterialProxySettingsCustomizations.h"
@@ -228,6 +229,8 @@ void FDetailCustomizationsModule::RegisterPropertyTypeCustomizations()
 	RegisterCustomPropertyTypeLayout("CameraFilmbackSettings", FOnGetPropertyTypeCustomizationInstance::CreateStatic(&FCameraFilmbackSettingsCustomization::MakeInstance));
 	RegisterCustomPropertyTypeLayout("CameraLensSettings", FOnGetPropertyTypeCustomizationInstance::CreateStatic(&FCameraLensSettingsCustomization::MakeInstance));
 	RegisterCustomPropertyTypeLayout("CameraFocusSettings", FOnGetPropertyTypeCustomizationInstance::CreateStatic(&FCameraFocusSettingsCustomization::MakeInstance));
+	RegisterCustomPropertyTypeLayout("LevelSequenceBurnInOptions", FOnGetPropertyTypeCustomizationInstance::CreateStatic(&FLevelSequenceBurnInOptionsCustomization::MakeInstance));
+	RegisterCustomPropertyTypeLayout("LevelSequenceBurnInInitSettings", FOnGetPropertyTypeCustomizationInstance::CreateStatic(&FLevelSequenceBurnInInitSettingsCustomization::MakeInstance));
 }
 
 
diff --git a/Engine/Source/Editor/Matinee/Private/Matinee.cpp b/Engine/Source/Editor/Matinee/Private/Matinee.cpp
--- a/Engine/Source/Editor/Matinee/Private/Matinee.cpp
+++ b/Engine/Source/Editor/Matinee/Private/Matinee.cpp
@@ -949,9 +949,6 @@ void FMatinee::InitMatinee(const EToolkitMode::Type Mode, const TSharedPtr< clas
 				//Ensure Realtime is turned on and store the original setting so we can restore it later.
 				LevelVC->SetRealtime(true, true);
 			}
-
-			// Turn on 'show camera frustums' flag
-			LevelVC->EngineShowFlags.SetCameraFrustums(true);
 		}
 	}
 
@@ -2066,9 +2063,6 @@ void FMatinee::OnClose()
 				//Specify true so RestoreRealtime will allow us to disable Realtime if it was original disabled
 				LevelVC->RestoreRealtime(true);
 			}
-
-			// Turn off 'show camera frustums' flag.
-			LevelVC->EngineShowFlags.SetCameraFrustums(false);
 		}
 	}
 
@@ -2978,7 +2972,7 @@ void FMatinee::OnMenuCreateMovie()
 
 	// Create a new movie scene capture object for a generic level capture
 	ULevelCapture* MovieSceneCapture = NewObject<ULevelCapture>(GetTransientPackage(), ULevelCapture::StaticClass(), NAME_None, RF_Transient);
-	MovieSceneCapture->LoadConfig();
+	MovieSceneCapture->LoadFromConfig();
 
 	// Ensure that this matinee is up and running before we start capturing
 	MovieSceneCapture->SetPrerequisiteActor(MatineeActor);
diff --git a/Engine/Source/Editor/MovieSceneCaptureDialog/Private/MovieSceneCaptureDialogModule.cpp b/Engine/Source/Editor/MovieSceneCaptureDialog/Private/MovieSceneCaptureDialogModule.cpp
--- a/Engine/Source/Editor/MovieSceneCaptureDialog/Private/MovieSceneCaptureDialogModule.cpp
+++ b/Engine/Source/Editor/MovieSceneCaptureDialog/Private/MovieSceneCaptureDialogModule.cpp
@@ -8,6 +8,7 @@
 #include "SlateBasics.h"
 #include "SlateExtras.h"
 #include "SceneViewport.h"
+#include "AudioDevice.h"
 
 #include "SDockTab.h"
 #include "JsonObjectConverter.h"
@@ -331,6 +332,13 @@ struct FInEditorCapture
 		UGameViewportClient::OnViewportCreated().AddRaw(this, &FInEditorCapture::OnStart);
 		FEditorDelegates::EndPIE.AddRaw(this, &FInEditorCapture::OnEndPIE);
 		
+		FAudioDevice* AudioDevice = GWorld->GetAudioDevice();
+		if (AudioDevice != nullptr)
+		{
+			TransientMasterVolume = AudioDevice->TransientMasterVolume;
+			AudioDevice->TransientMasterVolume = 0.0f;
+		}
+
 		GEditor->RequestPlaySession(true, nullptr, false);
 	}
 
@@ -443,6 +451,12 @@ struct FInEditorCapture
 
 		FObjectReader(GetMutableDefault<ULevelEditorPlaySettings>(), BackedUpPlaySettings);
 
+		FAudioDevice* AudioDevice = GWorld->GetAudioDevice();
+		if (AudioDevice != nullptr)
+		{
+			AudioDevice->TransientMasterVolume = TransientMasterVolume;
+		}
+
 		CaptureObject->Close();
 		CaptureObject->RemoveFromRoot();
 
@@ -463,6 +477,7 @@ struct FInEditorCapture
 
 	TFunction<void()> OnStarted;
 	bool bScreenMessagesWereEnabled;
+	float TransientMasterVolume;
 	int32 BackedUpStreamingPoolSize;
 	int32 BackedUpUseFixedPoolSize;
 	TArray<uint8> BackedUpPlaySettings;
@@ -569,11 +584,7 @@ class FMovieSceneCaptureDialogModule : public IMovieSceneCaptureDialogModule
 			return LOCTEXT("AlreadyCapturing", "There is already a movie scene capture process open. Please close it and try again.");
 		}
 
-		CaptureObject->SaveConfig();
-		if (CaptureObject->ProtocolSettings)
-		{
-			CaptureObject->ProtocolSettings->SaveConfig();
-		}
+		CaptureObject->SaveToConfig();
 
 		return CaptureObject->bUseSeparateProcess ? CaptureInNewProcess(CaptureObject, MapNameToLoad) : CaptureInEditor(CaptureObject, MapNameToLoad);
 	}
@@ -627,15 +638,9 @@ class FMovieSceneCaptureDialogModule : public IMovieSceneCaptureDialogModule
 			RootObject->SetField(TEXT("Type"), MakeShareable(new FJsonValueString(CaptureObject->GetClass()->GetPathName())));
 			RootObject->SetField(TEXT("Data"), MakeShareable(new FJsonValueObject(Object)));
 
-			if (CaptureObject->ProtocolSettings)
-			{
-				RootObject->SetField(TEXT("ProtocolType"), MakeShareable(new FJsonValueString(CaptureObject->ProtocolSettings->GetClass()->GetPathName())));
-				TSharedRef<FJsonObject> ProtocolDataObject = MakeShareable(new FJsonObject);
-				if (FJsonObjectConverter::UStructToJsonObject(CaptureObject->ProtocolSettings->GetClass(), CaptureObject->ProtocolSettings, ProtocolDataObject, 0, 0))
-				{
-					RootObject->SetField(TEXT("ProtocolData"), MakeShareable(new FJsonValueObject(ProtocolDataObject)));
-				}
-			}
+			TSharedRef<FJsonObject> AdditionalJson = MakeShareable(new FJsonObject);
+			CaptureObject->SerializeJson(*AdditionalJson);
+			RootObject->SetField(TEXT("AdditionalData"), MakeShareable(new FJsonValueObject(AdditionalJson)));
 
 			FString Json;
 			TSharedRef<TJsonWriter<> > JsonWriter = TJsonWriterFactory<>::Create(&Json, 0);
diff --git a/Engine/Source/Editor/MovieSceneTools/MovieSceneTools.Build.cs b/Engine/Source/Editor/MovieSceneTools/MovieSceneTools.Build.cs
--- a/Engine/Source/Editor/MovieSceneTools/MovieSceneTools.Build.cs
+++ b/Engine/Source/Editor/MovieSceneTools/MovieSceneTools.Build.cs
@@ -13,7 +13,8 @@ public class MovieSceneTools : ModuleRules
                 "Editor/MovieSceneTools/Private/TrackEditors",
 				"Editor/MovieSceneTools/Private/TrackEditors/PropertyTrackEditors",
                 "Editor/MovieSceneTools/Private/TrackEditorThumbnail",
-				"Editor/MovieSceneTools/Private/Sections"
+				"Editor/MovieSceneTools/Private/Sections",
+                "Editor/UnrealEd/Private",	//compatibility for FBX importer
             }
         );
 
@@ -31,8 +32,10 @@ public class MovieSceneTools : ModuleRules
 		PrivateDependencyModuleNames.AddRange(
 			new string[] {
                 "ActorPickerMode",
+                "DesktopPlatform",
 				"LevelSequence",
 				"MovieScene",
+                "MovieSceneCapture",
 				"MovieSceneTracks",
 				"BlueprintGraph",
                 "ContentBrowser",
@@ -55,7 +58,9 @@ public class MovieSceneTools : ModuleRules
                 "Settings",
 				"SceneOutliner",
 				"PropertyEditor",
-			}
+                "MainFrame",
+                "UnrealEd",
+            }
 		);
 
 		DynamicallyLoadedModuleNames.AddRange(
@@ -64,9 +69,16 @@ public class MovieSceneTools : ModuleRules
 				"AssetTools",
 				"SceneOutliner",
 				"PropertyEditor",
+                "MainFrame",
 			}
 		);
 
-        CircularlyReferencedDependentModules.Add("Sequencer");
+        CircularlyReferencedDependentModules.AddRange(
+            new string[] {
+                "Sequencer",
+            }
+        );
+
+        AddEngineThirdPartyPrivateStaticDependencies(Target, "FBX");
     }
 }
diff --git a/Engine/Source/Editor/MovieSceneTools/Private/MovieSceneToolHelpers.cpp b/Engine/Source/Editor/MovieSceneTools/Private/MovieSceneToolHelpers.cpp
--- a/Engine/Source/Editor/MovieSceneTools/Private/MovieSceneToolHelpers.cpp
+++ b/Engine/Source/Editor/MovieSceneTools/Private/MovieSceneToolHelpers.cpp
@@ -1,11 +1,25 @@
 // Copyright 1998-2016 Epic Games, Inc. All Rights Reserved.
 
 #include "MovieSceneToolsPrivatePCH.h"
+#include "MovieScene.h"
 #include "MovieSceneSection.h"
+#include "MovieSceneAudioTrack.h"
+#include "MovieScene3DTransformSection.h"
+#include "MovieScene3DTransformTrack.h"
 #include "MovieSceneCinematicShotSection.h"
+#include "MovieSceneCinematicShotTrack.h"
 #include "LevelSequence.h"
 #include "AssetRegistryModule.h"
-
+#include "CoreMisc.h"
+#include "DesktopPlatformModule.h"
+#include "IDesktopPlatform.h"
+#include "Editor/MainFrame/Public/MainFrame.h"
+#include "SlateBasics.h"
+#include "SlateExtras.h"
+#include "INotificationWidget.h"
+#include "MovieSceneCaptureHelpers.h"
+#include "FbxImporter.h"
+#include "MatineeImportTools.h"
 
 /* MovieSceneToolHelpers
  *****************************************************************************/
@@ -493,4 +507,235 @@ TSharedRef<SWidget> MovieSceneToolHelpers::MakeEnumComboBox(const UEnum* InEnum,
 		.CurrentValue(InCurrentValue)
 		.OnEnumSelectionChanged(InOnSelectionChanged)
 		.IntermediateValue(InIntermediateValue);
+}
+
+bool MovieSceneToolHelpers::ShowImportEDLDialog(UMovieScene* InMovieScene, float InFrameRate, FString InOpenDirectory)
+{
+	TArray<FString> OpenFilenames;
+	IDesktopPlatform* DesktopPlatform = FDesktopPlatformModule::Get();
+	bool bOpen = false;
+	if (DesktopPlatform)
+	{
+		void* ParentWindowWindowHandle = NULL;
+
+		IMainFrameModule& MainFrameModule = FModuleManager::LoadModuleChecked<IMainFrameModule>(TEXT("MainFrame"));
+		const TSharedPtr<SWindow>& MainFrameParentWindow = MainFrameModule.GetParentWindow();
+		if ( MainFrameParentWindow.IsValid() && MainFrameParentWindow->GetNativeWindow().IsValid() )
+		{
+			ParentWindowWindowHandle = MainFrameParentWindow->GetNativeWindow()->GetOSWindowHandle();
+		}
+
+		FString ExtensionStr;
+		ExtensionStr += TEXT("CMX 3600 EDL (*.edl)|*.edl|");
+
+		bOpen = DesktopPlatform->OpenFileDialog(
+			ParentWindowWindowHandle,
+			NSLOCTEXT("UnrealEd", "MovieSceneToolHelpersImportEDL", "Import EDL from...").ToString(), 
+			InOpenDirectory,
+			TEXT(""), 
+			*ExtensionStr,
+			EFileDialogFlags::None,
+			OpenFilenames
+			);
+	}
+	if (!bOpen)
+	{
+		return false;
+	}
+
+	if (!OpenFilenames.Num())
+	{
+		return false;
+	}
+
+	const FScopedTransaction Transaction( NSLOCTEXT( "MovieSceneTools", "ImportEDL", "Import EDL" ) );
+
+	return MovieSceneCaptureHelpers::ImportEDL(InMovieScene, InFrameRate, OpenFilenames[0]);
+}
+
+bool MovieSceneToolHelpers::ShowExportEDLDialog(const UMovieScene* InMovieScene, float InFrameRate, FString InSaveDirectory)
+{
+	FString SequenceName = InMovieScene->GetOuter()->GetName();
+	FString SequencePath = InMovieScene->GetOuter()->GetPathName();
+	SequencePath = FPaths::GetPath(SequencePath);
+
+	// Pop open a dialog to request the location of the edl
+	FString SaveDirectory = InSaveDirectory;
+	if (InSaveDirectory.IsEmpty())
+	{
+		IDesktopPlatform* DesktopPlatform = FDesktopPlatformModule::Get();
+		bool bSave = false;
+		if (DesktopPlatform)
+		{
+			void* ParentWindowWindowHandle = NULL;
+
+			IMainFrameModule& MainFrameModule = FModuleManager::LoadModuleChecked<IMainFrameModule>(TEXT("MainFrame"));
+			const TSharedPtr<SWindow>& MainFrameParentWindow = MainFrameModule.GetParentWindow();
+			if ( MainFrameParentWindow.IsValid() && MainFrameParentWindow->GetNativeWindow().IsValid() )
+			{
+				ParentWindowWindowHandle = MainFrameParentWindow->GetNativeWindow()->GetOSWindowHandle();
+			}
+
+			bSave = DesktopPlatform->OpenDirectoryDialog(
+				ParentWindowWindowHandle,
+				NSLOCTEXT("UnrealEd", "MovieSceneToolHelpersExportEDL", "Export EDL to...").ToString(), 
+				SequencePath,
+				SaveDirectory 
+				);
+		}
+		if (!bSave)
+		{
+			return false;
+		}
+	}
+
+	if (MovieSceneCaptureHelpers::ExportEDL(InMovieScene, InFrameRate, SaveDirectory))
+	{
+		FNotificationInfo NotificationInfo(NSLOCTEXT("MovieSceneTools", "EDLExportFinished", "EDL Export finished"));
+		NotificationInfo.ExpireDuration = 5.f;
+		NotificationInfo.bFireAndForget = true;
+		NotificationInfo.Hyperlink = FSimpleDelegate::CreateStatic( [](FString InDirectory) { FPlatformProcess::ExploreFolder( *InDirectory ); }, SaveDirectory);
+		NotificationInfo.HyperlinkText = NSLOCTEXT("MovieSceneTools", "OpenEDLExportFolder", "Open EDL Export Folder...");
+		FSlateNotificationManager::Get().AddNotification(NotificationInfo);
+		
+		return true;
+	}
+
+	return false;
+}
+
+bool MovieSceneToolHelpers::ImportFBX(UMovieScene* InMovieScene, FGuid InObjectBinding)
+{
+	/*
+	TArray<FString> OpenFilenames;
+	IDesktopPlatform* DesktopPlatform = FDesktopPlatformModule::Get();
+	bool bOpen = false;
+	if (DesktopPlatform)
+	{
+		void* ParentWindowWindowHandle = NULL;
+
+		IMainFrameModule& MainFrameModule = FModuleManager::LoadModuleChecked<IMainFrameModule>(TEXT("MainFrame"));
+		const TSharedPtr<SWindow>& MainFrameParentWindow = MainFrameModule.GetParentWindow();
+		if ( MainFrameParentWindow.IsValid() && MainFrameParentWindow->GetNativeWindow().IsValid() )
+		{
+			ParentWindowWindowHandle = MainFrameParentWindow->GetNativeWindow()->GetOSWindowHandle();
+		}
+
+		FString ExtensionStr;
+		ExtensionStr += TEXT("FBX (*.fbx)|*.fbx|");
+
+		bOpen = DesktopPlatform->OpenFileDialog(
+			ParentWindowWindowHandle,
+			NSLOCTEXT("UnrealEd", "MovieSceneToolHelpersImportFBX", "Import FBX from...").ToString(), 
+			FEditorDirectories::Get().GetLastDirectory(ELastDirectory::GENERIC_IMPORT),
+			TEXT(""), 
+			*ExtensionStr,
+			EFileDialogFlags::None,
+			OpenFilenames
+			);
+	}
+	if (!bOpen)
+	{
+		return false;
+	}
+
+	if (!OpenFilenames.Num())
+	{
+		return false;
+	}
+
+	UnFbx::FFbxImporter* FbxImporter = UnFbx::FFbxImporter::GetInstance();
+
+	const FString FileExtension = FPaths::GetExtension(OpenFilenames[0]);
+	if (!FbxImporter->ImportFromFile(*OpenFilenames[0], FileExtension, true))
+	{
+		// Log the error message and fail the import.
+		FbxImporter->ReleaseScene();
+		return false;
+	}
+
+	const FScopedTransaction Transaction( NSLOCTEXT( "MovieSceneTools", "ImportFBX", "Import FBX" ) );
+
+	//@todo Find transform tracks explicitly for now. This should be dynamic in the future
+	UMovieScene3DTransformTrack* TransformTrack = InMovieScene->FindTrack<UMovieScene3DTransformTrack>(InObjectBinding); 
+	if (!TransformTrack)
+	{
+		InMovieScene->Modify();
+		TransformTrack = InMovieScene->AddTrack<UMovieScene3DTransformTrack>(InObjectBinding);
+	}
+
+	bool bSectionAdded = false;
+	UMovieScene3DTransformSection* TransformSection = Cast<UMovieScene3DTransformSection>(TransformTrack->FindOrAddSection(0.f, bSectionAdded));
+	if (!TransformSection)
+	{
+		return false;
+	}
+
+	if (bSectionAdded)
+	{
+		TransformSection->SetIsInfinite(true);
+	}
+
+	float MinTime = FLT_MAX;
+	float MaxTime = -FLT_MAX;
+
+	UnFbx::FFbxCurvesAPI CurveAPI;
+	FbxImporter->PopulateAnimatedCurveData(CurveAPI);
+	TArray<FString> AnimatedNodeNames;
+	CurveAPI.GetAnimatedNodeNameArray(AnimatedNodeNames);
+	for (FString NodeName : AnimatedNodeNames)
+	{
+		FInterpCurveFloat Translation[3];
+		FInterpCurveFloat EulerRotation[3];
+
+		CurveAPI.GetConvertedTransformCurveData(NodeName, Translation[0], Translation[1], Translation[2], EulerRotation[0], EulerRotation[1], EulerRotation[2]);
+
+		for (int32 CurveIndex = 0; CurveIndex < 2; ++CurveIndex)
+		{
+			for (int32 ChannelIndex = 0; ChannelIndex < 3; ++ChannelIndex)
+			{
+				EAxis::Type ChannelAxis = EAxis::X;
+				if (ChannelIndex == 1)
+				{
+					ChannelAxis = EAxis::Y;
+				}
+				else if (ChannelIndex == 2)
+				{
+					ChannelAxis = EAxis::Z;
+				}
+	
+				FInterpCurveFloat* CurveFloat = nullptr;
+				FRichCurve* ChannelCurve = nullptr;
+				
+				if (CurveIndex == 0)
+				{
+					CurveFloat = &Translation[ChannelIndex];
+					ChannelCurve = &TransformSection->GetTranslationCurve(ChannelAxis);
+				}
+				else
+				{
+					CurveFloat = &EulerRotation[ChannelIndex];
+					ChannelCurve = &TransformSection->GetRotationCurve(ChannelAxis);
+				}
+
+				if (ChannelCurve != nullptr && CurveFloat != nullptr)
+				{
+					ChannelCurve->Reset();
+					for (int32 KeyIndex = 0; KeyIndex < CurveFloat->Points.Num(); ++KeyIndex)
+					{
+						MinTime = FMath::Min(MinTime, CurveFloat->Points[KeyIndex].InVal);
+						MaxTime = FMath::Max(MaxTime, CurveFloat->Points[KeyIndex].InVal);
+						FMatineeImportTools::SetOrAddKey(*ChannelCurve, CurveFloat->Points[KeyIndex].InVal, CurveFloat->Points[KeyIndex].OutVal, CurveFloat->Points[KeyIndex].ArriveTangent, CurveFloat->Points[KeyIndex].LeaveTangent, CurveFloat->Points[KeyIndex].InterpMode);
+					}
+				}
+			}
+		}
+	}
+		
+	TransformSection->SetStartTime(MinTime);
+	TransformSection->SetEndTime(MaxTime);
+	
+	FbxImporter->ReleaseScene();
+	*/
+	return true;
 }
\ No newline at end of file
diff --git a/Engine/Source/Editor/MovieSceneTools/Private/MovieSceneToolsModule.cpp b/Engine/Source/Editor/MovieSceneTools/Private/MovieSceneToolsModule.cpp
--- a/Engine/Source/Editor/MovieSceneTools/Private/MovieSceneToolsModule.cpp
+++ b/Engine/Source/Editor/MovieSceneTools/Private/MovieSceneToolsModule.cpp
@@ -9,6 +9,7 @@
 #include "ScopedTransaction.h"
 #include "MovieScene.h"
 #include "MovieSceneTrackEditor.h"
+#include "MovieSceneToolsProjectSettingsCustomization.h"
 
 #include "BoolPropertyTrackEditor.h"
 #include "BytePropertyTrackEditor.h"
@@ -42,6 +43,7 @@
 #include "ClipboardTypes.h"
 #include "Curves/CurveBase.h"
 #include "ISettingsModule.h"
+#include "PropertyEditorModule.h"
 
 #define LOCTEXT_NAMESPACE "FMovieSceneToolsModule"
 
@@ -98,6 +100,10 @@ public:
 		CameraShakeTrackCreateEditorHandle = SequencerModule.RegisterTrackEditor_Handle(FOnCreateTrackEditor::CreateStatic(&FCameraShakeTrackEditor::CreateTrackEditor));
 
 		RegisterClipboardConversions();
+
+		// register details customization
+		FPropertyEditorModule& PropertyModule = FModuleManager::LoadModuleChecked<FPropertyEditorModule>("PropertyEditor");
+		PropertyModule.RegisterCustomClassLayout("MovieSceneToolsProjectSettings", FOnGetDetailCustomizationInstance::CreateStatic(&FMovieSceneToolsProjectSettingsCustomization::MakeInstance));
 	}
 
 	virtual void ShutdownModule() override
@@ -141,6 +147,12 @@ public:
 		SequencerModule.UnRegisterTrackEditor_Handle( SpawnTrackCreateEditorHandle );
 		SequencerModule.UnRegisterTrackEditor_Handle( CameraAnimTrackCreateEditorHandle );
 		SequencerModule.UnRegisterTrackEditor_Handle( CameraShakeTrackCreateEditorHandle );
+
+		if (FModuleManager::Get().IsModuleLoaded("PropertyEditor"))
+		{	
+			FPropertyEditorModule& PropertyModule = FModuleManager::LoadModuleChecked<FPropertyEditorModule>("PropertyEditor");
+			PropertyModule.UnregisterCustomClassLayout("MovieSceneToolsProjectSettings");
+		}
 	}
 
 	void RegisterClipboardConversions()
diff --git a/Engine/Source/Editor/MovieSceneTools/Private/MovieSceneToolsProjectSettingsCustomization.cpp b/Engine/Source/Editor/MovieSceneTools/Private/MovieSceneToolsProjectSettingsCustomization.cpp
new file mode 100644
--- /dev/null
+++ b/Engine/Source/Editor/MovieSceneTools/Private/MovieSceneToolsProjectSettingsCustomization.cpp
@@ -0,0 +1,46 @@
+// Copyright 1998-2016 Epic Games, Inc. All Rights Reserved.
+
+#include "MovieSceneToolsPrivatePCH.h"
+#include "MovieSceneToolsProjectSettingsCustomization.h"
+
+#include "PropertyEditorModule.h"
+#include "DetailLayoutBuilder.h"
+#include "DetailCategoryBuilder.h"
+#include "DetailWidgetRow.h"
+#include "IDetailPropertyRow.h"
+#include "PropertyHandle.h"
+
+//#include "Editor/PropertyEditor/Public/IDetailsView.h"
+//#include "IDetailCustomization.h"
+//#include "DetailLayoutBuilder.h"
+//#include "PropertyHandle.h"
+
+
+TSharedRef<IDetailCustomization> FMovieSceneToolsProjectSettingsCustomization::MakeInstance()
+{
+	return MakeShareable(new FMovieSceneToolsProjectSettingsCustomization);
+}
+
+void FMovieSceneToolsProjectSettingsCustomization::CustomizeDetails(IDetailLayoutBuilder& DetailBuilder)
+{
+	TSharedRef<IPropertyHandle> TakeSeparatorProperty = DetailBuilder.GetProperty(GET_MEMBER_NAME_CHECKED(UMovieSceneToolsProjectSettings, TakeSeparator));
+	TakeSeparatorProperty->SetOnPropertyValueChanged(FSimpleDelegate::CreateRaw(this, &FMovieSceneToolsProjectSettingsCustomization::OnTakeSeparatorUpdated));
+}
+
+void FMovieSceneToolsProjectSettingsCustomization::OnTakeSeparatorUpdated()
+{
+	UMovieSceneToolsProjectSettings* ProjectSettings = GetMutableDefault<UMovieSceneToolsProjectSettings>();
+
+	FString TakeSeparator = ProjectSettings->TakeSeparator;
+
+	// Make sure the take separator is a valid single character
+	FText OutErrorMessage;
+	if (!FName(*TakeSeparator).IsValidXName(INVALID_OBJECTNAME_CHARACTERS INVALID_LONGPACKAGE_CHARACTERS, &OutErrorMessage))
+	{
+		ProjectSettings->TakeSeparator = TEXT("_");
+	}
+	else if (ProjectSettings->TakeSeparator.Len() > 1)
+	{
+		ProjectSettings->TakeSeparator = ProjectSettings->TakeSeparator.LeftChop(ProjectSettings->TakeSeparator.Len()-1);
+	}
+}
diff --git a/Engine/Source/Editor/MovieSceneTools/Private/Sections/CinematicShotSection.cpp b/Engine/Source/Editor/MovieSceneTools/Private/Sections/CinematicShotSection.cpp
--- a/Engine/Source/Editor/MovieSceneTools/Private/Sections/CinematicShotSection.cpp
+++ b/Engine/Source/Editor/MovieSceneTools/Private/Sections/CinematicShotSection.cpp
@@ -177,6 +177,13 @@ void FCinematicShotSection::BuildSectionContextMenu(FMenuBuilder& MenuBuilder, c
 				FUIAction(FExecuteAction::CreateSP(CinematicShotTrackEditor.Pin().ToSharedRef(), &FCinematicShotTrackEditor::DuplicateShot, &SectionObject))
 			);
 
+			MenuBuilder.AddMenuEntry(
+				LOCTEXT("RenderShot", "Render Shot"),
+				FText::Format(LOCTEXT("RenderShotTooltip", "Render shot movie"), SectionObject.GetShotDisplayName()),
+				FSlateIcon(),
+				FUIAction(FExecuteAction::CreateSP(CinematicShotTrackEditor.Pin().ToSharedRef(), &FCinematicShotTrackEditor::RenderShot, &SectionObject))
+			);
+
 			/*
 			//@todo
 			MenuBuilder.AddMenuEntry(
diff --git a/Engine/Source/Editor/MovieSceneTools/Private/TrackEditors/CinematicShotTrackEditor.cpp b/Engine/Source/Editor/MovieSceneTools/Private/TrackEditors/CinematicShotTrackEditor.cpp
--- a/Engine/Source/Editor/MovieSceneTools/Private/TrackEditors/CinematicShotTrackEditor.cpp
+++ b/Engine/Source/Editor/MovieSceneTools/Private/TrackEditors/CinematicShotTrackEditor.cpp
@@ -17,7 +17,10 @@
 #include "CinematicShotSection.h"
 #include "SequencerUtilities.h"
 #include "ContentBrowserModule.h"
-#include "MovieSceneToolHelpers.h"
+#include "MovieSceneCaptureSettings.h"
+#include "MovieSceneCapture.h"
+#include "AutomatedLevelSequenceCapture.h"
+#include "MovieSceneCaptureModule.h"
 #include "LevelSequence.h"
 #include "AssetToolsModule.h"
 
@@ -177,6 +180,23 @@ void FCinematicShotTrackEditor::Tick(float DeltaTime)
 	}
 }
 
+void FCinematicShotTrackEditor::BuildTrackContextMenu( FMenuBuilder& MenuBuilder, UMovieSceneTrack* Track )
+{
+	MenuBuilder.AddMenuEntry(
+		NSLOCTEXT( "Sequencer", "ImportEDL", "Import EDL..." ),
+		NSLOCTEXT( "Sequencer", "ImportEDLTooltip", "Import Edit Decision List (EDL) for non-linear editors." ),
+		FSlateIcon(),
+		FUIAction(
+		FExecuteAction::CreateRaw(this, &FCinematicShotTrackEditor::ImportEDL )));
+
+	MenuBuilder.AddMenuEntry(
+		NSLOCTEXT( "Sequencer", "ExportEDL", "Export EDL" ),
+		NSLOCTEXT( "Sequencer", "ExportEDLTooltip", "Export Edit Decision List (EDL) for non-linear editors." ),
+		FSlateIcon(),
+		FUIAction(
+		FExecuteAction::CreateRaw(this, &FCinematicShotTrackEditor::ExportEDL )));
+}
+
 const FSlateBrush* FCinematicShotTrackEditor::GetIconBrush() const
 {
 	return FEditorStyle::GetBrush("Sequencer.Tracks.CinematicShot");
@@ -322,6 +342,12 @@ void FCinematicShotTrackEditor::DuplicateShot(UMovieSceneCinematicShotSection* S
 	GetSequencer()->NotifyMovieSceneDataChanged();
 }
 
+void FCinematicShotTrackEditor::RenderShot(UMovieSceneCinematicShotSection* Section)
+{
+	GetSequencer()->RenderMovie(Section);
+
+}
+
 
 void FCinematicShotTrackEditor::RenameShot(UMovieSceneCinematicShotSection* Section)
 {
@@ -622,4 +648,73 @@ bool FCinematicShotTrackEditor::HandleSequenceAdded(float KeyTime, UMovieSceneSe
 	return true;
 }
 
+void FCinematicShotTrackEditor::ImportEDL()
+{
+	UMovieSceneSequence* FocusedSequence = GetSequencer()->GetFocusedMovieSceneSequence();
+	if (!FocusedSequence)
+	{
+		return;
+	}
+
+	UMovieScene* MovieScene = FocusedSequence->GetMovieScene();
+	if (!MovieScene)
+	{
+		return;
+	}
+
+	UAutomatedLevelSequenceCapture* MovieSceneCapture = Cast<UAutomatedLevelSequenceCapture>(IMovieSceneCaptureModule::Get().GetFirstActiveMovieSceneCapture());
+	if (!MovieSceneCapture)
+	{
+		MovieSceneCapture = NewObject<UAutomatedLevelSequenceCapture>(GetTransientPackage(), UAutomatedLevelSequenceCapture::StaticClass(), NAME_None, RF_Transient);
+		MovieSceneCapture->LoadFromConfig();
+	}
+
+	if (!MovieSceneCapture)
+	{
+		return;
+	}
+
+	const FMovieSceneCaptureSettings& Settings = MovieSceneCapture->GetSettings();
+	FString SaveDirectory = FPaths::ConvertRelativePathToFull(Settings.OutputDirectory.Path);
+	float FrameRate = Settings.FrameRate;
+
+	if (MovieSceneToolHelpers::ShowImportEDLDialog(MovieScene, FrameRate, SaveDirectory))
+	{
+		GetSequencer()->NotifyMovieSceneDataChanged();
+	}
+}
+
+void FCinematicShotTrackEditor::ExportEDL()
+{
+	UMovieSceneSequence* FocusedSequence = GetSequencer()->GetFocusedMovieSceneSequence();
+	if (!FocusedSequence)
+	{
+		return;
+	}
+
+	UMovieScene* MovieScene = FocusedSequence->GetMovieScene();
+	if (!MovieScene)
+	{
+		return;
+	}
+		
+	UAutomatedLevelSequenceCapture* MovieSceneCapture = Cast<UAutomatedLevelSequenceCapture>(IMovieSceneCaptureModule::Get().GetFirstActiveMovieSceneCapture());
+	if (!MovieSceneCapture)
+	{
+		MovieSceneCapture = NewObject<UAutomatedLevelSequenceCapture>(GetTransientPackage(), UAutomatedLevelSequenceCapture::StaticClass(), NAME_None, RF_Transient);
+		MovieSceneCapture->LoadFromConfig();
+	}
+
+	if (!MovieSceneCapture)
+	{
+		return;
+	}
+
+	const FMovieSceneCaptureSettings& Settings = MovieSceneCapture->GetSettings();
+	FString SaveDirectory = FPaths::ConvertRelativePathToFull(Settings.OutputDirectory.Path);
+	float FrameRate = Settings.FrameRate;
+
+	MovieSceneToolHelpers::ShowExportEDLDialog(MovieScene, FrameRate, SaveDirectory);
+}
+
 #undef LOCTEXT_NAMESPACE
diff --git a/Engine/Source/Editor/MovieSceneTools/Private/TrackEditors/CinematicShotTrackEditor.h b/Engine/Source/Editor/MovieSceneTools/Private/TrackEditors/CinematicShotTrackEditor.h
--- a/Engine/Source/Editor/MovieSceneTools/Private/TrackEditors/CinematicShotTrackEditor.h
+++ b/Engine/Source/Editor/MovieSceneTools/Private/TrackEditors/CinematicShotTrackEditor.h
@@ -48,6 +48,7 @@ public:
 	virtual bool HandleAssetAdded(UObject* Asset, const FGuid& TargetObjectGuid) override;
 	virtual bool SupportsType(TSubclassOf<UMovieSceneTrack> Type ) const override;
 	virtual void Tick(float DeltaTime) override;
+	virtual void BuildTrackContextMenu( FMenuBuilder& MenuBuilder, UMovieSceneTrack* Track ) override;
 	virtual const FSlateBrush* GetIconBrush() const override;
 
 
@@ -69,6 +70,13 @@ public:
 	 */
 	void DuplicateShot(UMovieSceneCinematicShotSection* Section);
 
+	/*
+	 * Render shot. 
+	 *
+	 * @param Section The section to render
+	 */
+	void RenderShot(UMovieSceneCinematicShotSection* Section);
+
 	/*
 	 * Rename shot. 
 	 *
@@ -91,7 +99,6 @@ public:
 	*/
 	void SwitchTake(UMovieSceneCinematicShotSection* Section, uint32 TakeNumber);
 
-
 private:
 
 	/*
@@ -152,6 +159,12 @@ private:
 	/** Callback for AnimatablePropertyChanged in HandleAssetAdded. */
 	bool HandleSequenceAdded(float KeyTime, UMovieSceneSequence* Sequence);
 
+	/** Callback for ImportEDL. */
+	void ImportEDL();
+	
+	/** Callback for ExportEDL. */
+	void ExportEDL();
+
 private:
 
 	/** The Thumbnail pool which draws all the viewport thumbnails for the shot track. */
diff --git a/Engine/Source/Editor/MovieSceneTools/Public/MovieSceneToolHelpers.h b/Engine/Source/Editor/MovieSceneTools/Public/MovieSceneToolHelpers.h
--- a/Engine/Source/Editor/MovieSceneTools/Public/MovieSceneToolHelpers.h
+++ b/Engine/Source/Editor/MovieSceneTools/Public/MovieSceneToolHelpers.h
@@ -3,6 +3,7 @@
 #pragma once
 
 class UMovieSceneSection;
+class UMovieScene;
 
 DECLARE_DELEGATE_TwoParams(FOnEnumSelectionChanged, int32 /*Selection*/, ESelectInfo::Type /*SelectionType*/);
 
@@ -94,4 +95,34 @@ public:
 	 * @return The new widget
 	 */
 	static TSharedRef<SWidget> MakeEnumComboBox(const UEnum* Enum, TAttribute<int32> CurrentValue, FOnEnumSelectionChanged OnSelectionChanged, TAttribute<TOptional<uint8>> InIntermediateValue);
+
+
+	/**
+	 * Show Import EDL Dialog
+	 *
+	 * @param InMovieScene The movie scene to import the edl into
+	 * @param InFrameRate The frame rate to import the EDL at
+	 * @param InOpenDirectory Optional directory path to open from. If none given, a dialog will pop up to prompt the user
+	 * @return Whether the import was successful
+	 */
+	static bool ShowImportEDLDialog(UMovieScene* InMovieScene, float InFrameRate, FString InOpenDirectory = TEXT(""));
+
+	/**
+	 * Show Export EDL Dialog
+	 *
+	 * @param InMovieScene The movie scene with the cinematic shot track and audio tracks to export
+	 * @param InFrameRate The frame rate to export the EDL at
+	 * @param InSaveDirectory Optional directory path to save to. If none given, a dialog will pop up to prompt the user
+	 * @return Whether the export was successful
+	 */
+	static bool ShowExportEDLDialog(const UMovieScene* InMovieScene, float InFrameRate, FString InSaveDirectory = TEXT(""));
+
+	/**
+	 * Import FBX
+	 *
+	 * @param InMovieScene The movie scene to import the fbx into
+	 * @param InObjectBinding The object to import fbx animation onto
+	 * @return Whether the import was successful
+	 */
+	static bool ImportFBX(UMovieScene* InMovieScene, FGuid InObjectBinding);
 };
diff --git a/Engine/Source/Editor/MovieSceneTools/Public/MovieSceneToolsProjectSettings.h b/Engine/Source/Editor/MovieSceneTools/Public/MovieSceneToolsProjectSettings.h
--- a/Engine/Source/Editor/MovieSceneTools/Public/MovieSceneToolsProjectSettings.h
+++ b/Engine/Source/Editor/MovieSceneTools/Public/MovieSceneToolsProjectSettings.h
@@ -49,7 +49,7 @@ public:
 	UPROPERTY(config, EditAnywhere, Category=Shots, meta = (UIMin = "1", UIMax = "10"))
 	uint32 FirstTakeNumber;
 
-	/** The separator between the shot number and the take number. */
+	/** A single character separator between the shot number and the take number. */
 	UPROPERTY(config, EditAnywhere, Category=Shots)
 	FString TakeSeparator;
 };
diff --git a/Engine/Source/Editor/SequenceRecorder/Private/ActorRecording.cpp b/Engine/Source/Editor/SequenceRecorder/Private/ActorRecording.cpp
--- a/Engine/Source/Editor/SequenceRecorder/Private/ActorRecording.cpp
+++ b/Engine/Source/Editor/SequenceRecorder/Private/ActorRecording.cpp
@@ -41,15 +41,15 @@ UActorRecording::UActorRecording(const FObjectInitializer& ObjectInitializer)
 
 bool UActorRecording::IsRelevantForRecording(AActor* Actor)
 {
+	const USequenceRecorderSettings* Settings = GetDefault<USequenceRecorderSettings>();
+
 	// don't record actors that sequencer has spawned itself!
-	if(Actor->ActorHasTag(SequencerActorTag))
+	if(!Settings->bRecordSequencerSpawnedActors && Actor->ActorHasTag(SequencerActorTag))
 	{
 		return false;
 	}
 
 	TInlineComponentArray<USceneComponent*> SceneComponents(Actor);
-	const USequenceRecorderSettings* Settings = GetDefault<USequenceRecorderSettings>();
-
 	for(USceneComponent* SceneComponent : SceneComponents)
 	{
 		for (TSubclassOf<USceneComponent> ComponentClass : Settings->ComponentClassesToRecord)
diff --git a/Engine/Source/Editor/SequenceRecorder/Public/SequenceRecorderSettings.h b/Engine/Source/Editor/SequenceRecorder/Public/SequenceRecorderSettings.h
--- a/Engine/Source/Editor/SequenceRecorder/Public/SequenceRecorderSettings.h
+++ b/Engine/Source/Editor/SequenceRecorder/Public/SequenceRecorderSettings.h
@@ -65,4 +65,8 @@ public:
 	/** Component classes we record by default. If an actor does not contain one of these classes it will be ignored. */
 	UPROPERTY(Config, EditAnywhere, AdvancedDisplay, Category = "Sequence Recording")
 	TArray<TSubclassOf<USceneComponent>> ComponentClassesToRecord;
+
+	/** Whether to record actors that are spawned by sequencer itself (this is usually disabled as results can be unexpected) */
+	UPROPERTY(Config, EditAnywhere, AdvancedDisplay, Category = "Sequence Recording")
+	bool bRecordSequencerSpawnedActors;
 };
\ No newline at end of file
diff --git a/Engine/Source/Editor/Sequencer/Private/DisplayNodes/SequencerDisplayNode.cpp b/Engine/Source/Editor/Sequencer/Private/DisplayNodes/SequencerDisplayNode.cpp
--- a/Engine/Source/Editor/Sequencer/Private/DisplayNodes/SequencerDisplayNode.cpp
+++ b/Engine/Source/Editor/Sequencer/Private/DisplayNodes/SequencerDisplayNode.cpp
@@ -348,6 +348,7 @@ TSharedRef<SWidget> FSequencerDisplayNode::GenerateContainerWidgetForOutliner(co
 {
 	auto NewWidget = SNew(SAnimationOutlinerTreeNode, SharedThis(this), InRow)
 	.IconBrush(this, &FSequencerDisplayNode::GetIconBrush)
+	.IconColor(this, &FSequencerDisplayNode::GetIconColor)
 	.IconOverlayBrush(this, &FSequencerDisplayNode::GetIconOverlayBrush)
 	.IconToolTipText(this, &FSequencerDisplayNode::GetIconToolTipText)
 	.CustomContent()
@@ -373,6 +374,11 @@ const FSlateBrush* FSequencerDisplayNode::GetIconOverlayBrush() const
 	return nullptr;
 }
 
+FSlateColor FSequencerDisplayNode::GetIconColor() const
+{
+	return FSlateColor( FLinearColor::White );
+}
+
 FText FSequencerDisplayNode::GetIconToolTipText() const
 {
 	return FText();
diff --git a/Engine/Source/Editor/Sequencer/Private/DisplayNodes/SequencerDisplayNode.h b/Engine/Source/Editor/Sequencer/Private/DisplayNodes/SequencerDisplayNode.h
--- a/Engine/Source/Editor/Sequencer/Private/DisplayNodes/SequencerDisplayNode.h
+++ b/Engine/Source/Editor/Sequencer/Private/DisplayNodes/SequencerDisplayNode.h
@@ -192,6 +192,13 @@ public:
 	 */
 	virtual const FSlateBrush* GetIconOverlayBrush() const;
 
+	/**
+	 * Gets the color for the icon brush
+	 * 
+	 * @return This node's representative color
+	 */
+	virtual FSlateColor GetIconColor() const;
+
 	/**
 	 * Get the tooltip text to display for this node's icon
 	 * 
diff --git a/Engine/Source/Editor/Sequencer/Private/DisplayNodes/SequencerFolderNode.cpp b/Engine/Source/Editor/Sequencer/Private/DisplayNodes/SequencerFolderNode.cpp
--- a/Engine/Source/Editor/Sequencer/Private/DisplayNodes/SequencerFolderNode.cpp
+++ b/Engine/Source/Editor/Sequencer/Private/DisplayNodes/SequencerFolderNode.cpp
@@ -7,6 +7,9 @@
 #include "MovieSceneSequence.h"
 #include "SequencerDisplayNodeDragDropOp.h"
 #include "ScopedTransaction.h"
+#include "SColorPicker.h"
+
+#define LOCTEXT_NAMESPACE "SequencerFolderNode"
 
 FSequencerFolderNode::FSequencerFolderNode( UMovieSceneFolder& InMovieSceneFolder, TSharedPtr<FSequencerDisplayNode> InParentNode, FSequencerNodeTree& InParentTree )
 	: FSequencerDisplayNode( InMovieSceneFolder.GetFolderName(), InParentNode, InParentTree )
@@ -82,6 +85,11 @@ const FSlateBrush* FSequencerFolderNode::GetIconBrush() const
 }
 
 
+FSlateColor FSequencerFolderNode::GetIconColor() const
+{
+	return FSlateColor(MovieSceneFolder.GetFolderColor());
+}
+
 bool FSequencerFolderNode::CanDrag() const
 {
 	return true;
@@ -256,6 +264,68 @@ void FSequencerFolderNode::Drop( const TArray<TSharedRef<FSequencerDisplayNode>>
 	ParentTree.GetSequencer().NotifyMovieSceneDataChanged();
 }
 
+void FSequencerFolderNode::BuildContextMenu(FMenuBuilder& MenuBuilder)
+{
+	FSequencerDisplayNode::BuildContextMenu(MenuBuilder);
+
+	TSharedRef<FSequencerFolderNode> ThisNode = SharedThis(this);
+
+	MenuBuilder.BeginSection("Folder", LOCTEXT("FolderContextMenuSectionName", "Folder"));
+	{
+		MenuBuilder.AddMenuEntry(
+			LOCTEXT("SetColor", "Set Color"),
+			LOCTEXT("SetColorTooltip", "Set the folder color"),
+			FSlateIcon(),
+			FUIAction(
+				FExecuteAction::CreateSP(this, &FSequencerFolderNode::SetFolderColor))
+		);
+	}
+	MenuBuilder.EndSection();
+}
+
+FColor InitialFolderColor;
+bool bFolderPickerWasCancelled;
+
+void FSequencerFolderNode::SetFolderColor()
+{
+	InitialFolderColor = MovieSceneFolder.GetFolderColor();
+	bFolderPickerWasCancelled = false;
+
+	FColorPickerArgs PickerArgs;
+	PickerArgs.bUseAlpha = false;
+	PickerArgs.DisplayGamma = TAttribute<float>::Create( TAttribute<float>::FGetter::CreateUObject(GEngine, &UEngine::GetDisplayGamma) );
+	PickerArgs.InitialColorOverride = InitialFolderColor.ReinterpretAsLinear();
+	PickerArgs.OnColorCommitted = FOnLinearColorValueChanged::CreateSP( this, &FSequencerFolderNode::OnColorPickerPicked);
+	PickerArgs.OnColorPickerWindowClosed = FOnWindowClosed::CreateSP( this, &FSequencerFolderNode::OnColorPickerClosed);
+	PickerArgs.OnColorPickerCancelled  = FOnColorPickerCancelled::CreateSP( this, &FSequencerFolderNode::OnColorPickerCancelled );
+
+	OpenColorPicker(PickerArgs);
+}
+
+void FSequencerFolderNode::OnColorPickerPicked(FLinearColor NewFolderColor)
+{			
+	MovieSceneFolder.SetFolderColor(NewFolderColor.ToFColor(false));
+}
+
+void FSequencerFolderNode::OnColorPickerClosed(const TSharedRef<SWindow>& Window)
+{	
+	if (!bFolderPickerWasCancelled)
+	{
+		const FScopedTransaction Transaction( NSLOCTEXT( "SequencerFolderNode", "SetFolderColor", "Set Folder Color" ) );
+		
+		FColor CurrentColor = MovieSceneFolder.GetFolderColor();
+		MovieSceneFolder.SetFolderColor(InitialFolderColor);
+		MovieSceneFolder.Modify();
+		MovieSceneFolder.SetFolderColor(CurrentColor);
+	}
+}
+
+void FSequencerFolderNode::OnColorPickerCancelled(FLinearColor NewFolderColor)
+{
+	bFolderPickerWasCancelled = true;
+
+	MovieSceneFolder.SetFolderColor(InitialFolderColor);
+}
 
 void FSequencerFolderNode::AddChildNode( TSharedRef<FSequencerDisplayNode> ChildNode )
 {
@@ -268,3 +338,4 @@ UMovieSceneFolder& FSequencerFolderNode::GetFolder() const
 	return MovieSceneFolder;
 }
 
+#undef LOCTEXT_NAMESPACE
diff --git a/Engine/Source/Editor/Sequencer/Private/DisplayNodes/SequencerFolderNode.h b/Engine/Source/Editor/Sequencer/Private/DisplayNodes/SequencerFolderNode.h
--- a/Engine/Source/Editor/Sequencer/Private/DisplayNodes/SequencerFolderNode.h
+++ b/Engine/Source/Editor/Sequencer/Private/DisplayNodes/SequencerFolderNode.h
@@ -18,9 +18,11 @@ public:
 	virtual FText GetDisplayName() const override;
 	virtual void SetDisplayName( const FText& NewDisplayName ) override;
 	virtual const FSlateBrush* GetIconBrush() const override;
+	virtual FSlateColor GetIconColor() const override;
 	virtual bool CanDrag() const override;
 	virtual TOptional<EItemDropZone> CanDrop( FSequencerDisplayNodeDragDropOp& DragDropOp, EItemDropZone ItemDropZone ) const override;
 	virtual void Drop( const TArray<TSharedRef<FSequencerDisplayNode>>& DraggedNodes, EItemDropZone ItemDropZone ) override;
+	virtual void BuildContextMenu(FMenuBuilder& MenuBuilder);
 
 	/** Adds a child node to this folder node. */
 	void AddChildNode( TSharedRef<FSequencerDisplayNode> ChildNode );
@@ -28,6 +30,18 @@ public:
 	/** Gets the folder data for this display node. */
 	UMovieSceneFolder& GetFolder() const;
 
+	/** Callback to set the folder color */
+	void SetFolderColor();
+
+	/** Callback for the color being picked from the color picker */
+	void OnColorPickerPicked(FLinearColor NewFolderColor);
+
+	/** Callback for the color picker being closed */
+	void OnColorPickerClosed(const TSharedRef<SWindow>& Window);
+
+	/** Callback for the color picker being cancelled */
+	void OnColorPickerCancelled(FLinearColor NewFolderColor);
+
 private:
 
 	/** The brush used to draw the icon when this folder is open .*/
diff --git a/Engine/Source/Editor/Sequencer/Private/DisplayNodes/SequencerObjectBindingNode.cpp b/Engine/Source/Editor/Sequencer/Private/DisplayNodes/SequencerObjectBindingNode.cpp
--- a/Engine/Source/Editor/Sequencer/Private/DisplayNodes/SequencerObjectBindingNode.cpp
+++ b/Engine/Source/Editor/Sequencer/Private/DisplayNodes/SequencerObjectBindingNode.cpp
@@ -119,6 +119,16 @@ void FSequencerObjectBindingNode::BuildContextMenu(FMenuBuilder& MenuBuilder)
 			MenuBuilder.AddMenuEntry( FSequencerCommands::Get().ConvertToSpawnable );
 		}
 
+		/*
+		MenuBuilder.AddMenuEntry(
+			LOCTEXT("Import FBX", "Import FBX..."),
+			LOCTEXT("ImportFBXTooltip", "Import FBX animation to this object"),
+			FSlateIcon(),
+			FUIAction(
+				FExecuteAction::CreateLambda([=]{ GetSequencer().ImportFBX(ObjectBinding); })
+			));
+			*/
+
 		MenuBuilder.BeginSection("Organize", LOCTEXT("OrganizeContextMenuSectionName", "Organize"));
 		{
 			MenuBuilder.AddSubMenu(
diff --git a/Engine/Source/Editor/Sequencer/Private/SAnimationOutlinerTreeNode.cpp b/Engine/Source/Editor/Sequencer/Private/SAnimationOutlinerTreeNode.cpp
--- a/Engine/Source/Editor/Sequencer/Private/SAnimationOutlinerTreeNode.cpp
+++ b/Engine/Source/Editor/Sequencer/Private/SAnimationOutlinerTreeNode.cpp
@@ -179,6 +179,7 @@ void SAnimationOutlinerTreeNode::Construct( const FArguments& InArgs, TSharedRef
 								[
 									SNew(SImage)
 									.Image(InArgs._IconBrush)
+									.ColorAndOpacity(InArgs._IconColor)
 								]
 
 								+ SOverlay::Slot()
diff --git a/Engine/Source/Editor/Sequencer/Private/SAnimationOutlinerTreeNode.h b/Engine/Source/Editor/Sequencer/Private/SAnimationOutlinerTreeNode.h
--- a/Engine/Source/Editor/Sequencer/Private/SAnimationOutlinerTreeNode.h
+++ b/Engine/Source/Editor/Sequencer/Private/SAnimationOutlinerTreeNode.h
@@ -22,6 +22,7 @@ public:
 	SLATE_BEGIN_ARGS(SAnimationOutlinerTreeNode){}
 		SLATE_ATTRIBUTE(const FSlateBrush*, IconBrush)
 		SLATE_ATTRIBUTE(const FSlateBrush*, IconOverlayBrush)
+		SLATE_ATTRIBUTE(FSlateColor, IconColor)
 		SLATE_ATTRIBUTE(FText, IconToolTipText)
 		SLATE_NAMED_SLOT(FArguments, CustomContent)
 	SLATE_END_ARGS()
diff --git a/Engine/Source/Editor/Sequencer/Private/SSequencer.cpp b/Engine/Source/Editor/Sequencer/Private/SSequencer.cpp
--- a/Engine/Source/Editor/Sequencer/Private/SSequencer.cpp
+++ b/Engine/Source/Editor/Sequencer/Private/SSequencer.cpp
@@ -1062,7 +1062,7 @@ TSharedRef<SWidget> SSequencer::MakeTimeRange(const TSharedRef<SWidget>& InnerCo
 	FTimeRangeArgs Args(
 		ShowRange,
 		TimeSliderController.ToSharedRef(),
-		TAttribute<EVisibility>(this, &SSequencer::GetTimeRangeVisibility),
+		EVisibility::Visible,
 		TAttribute<bool>(this, &SSequencer::ShowFrameNumbers),
 		GetZeroPadNumericTypeInterface()
 		);
@@ -1760,10 +1760,13 @@ void SSequencer::OnTimeSnapIntervalChanged()
 	TSharedPtr<FSequencer> Sequencer = SequencerPtr.Pin();
 	if ( Sequencer.IsValid() )
 	{
-		FScopedTransaction SetFixedFrameIntervalTransaction( NSLOCTEXT( "Sequencer", "SetFixedFrameInterval", "Set scene fixed frame interval" ) );
 		UMovieScene* MovieScene = Sequencer->GetFocusedMovieSceneSequence()->GetMovieScene();
-		MovieScene->Modify();
-		MovieScene->SetFixedFrameInterval( Settings->GetTimeSnapInterval() );
+		if (!FMath::IsNearlyEqual(MovieScene->GetFixedFrameInterval(), Settings->GetTimeSnapInterval()))
+		{
+			FScopedTransaction SetFixedFrameIntervalTransaction( NSLOCTEXT( "Sequencer", "SetFixedFrameInterval", "Set scene fixed frame interval" ) );
+			MovieScene->Modify();
+			MovieScene->SetFixedFrameInterval( Settings->GetTimeSnapInterval() );
+		}
 	}
 }
 
diff --git a/Engine/Source/Editor/Sequencer/Private/Sequencer.cpp b/Engine/Source/Editor/Sequencer/Private/Sequencer.cpp
--- a/Engine/Source/Editor/Sequencer/Private/Sequencer.cpp
+++ b/Engine/Source/Editor/Sequencer/Private/Sequencer.cpp
@@ -450,7 +450,8 @@ void FSequencer::PostTickRenderStateFixup()
 				{
 					for (UMovieSceneSection* Section : SubTrack->GetAllSections())
 					{
-						if (Section->IsInfinite() || Section->IsTimeWithinSection(Time))
+						// Subsections are inclusive start and exclusive end
+						if (Section->IsInfinite() || (Time >= Section->GetStartTime() && Time < Section->GetEndTime()))
 						{
 							UMovieSceneSubSection* SubSection = CastChecked<UMovieSceneSubSection>(Section);
 
@@ -539,6 +540,11 @@ TSharedRef<FMovieSceneSequenceInstance> FSequencer::GetFocusedMovieSceneSequence
 
 void FSequencer::FocusSequenceInstance(UMovieSceneSubSection& InSubSection)
 {
+	if (RootMovieSceneSequenceInstance.IsValid())
+	{
+		RootMovieSceneSequenceInstance->RestoreState(*this);
+	}
+
 	TSharedRef<FMovieSceneSequenceInstance> SequenceInstance = GetSequenceInstanceForSection(InSubSection);
 
 	// Check for infinite recursion
@@ -1314,6 +1320,13 @@ void FSequencer::SetGlobalTimeDirectly( float NewTime, ESnapTimeMode SnapTimeMod
 	// Update the position
 	ScrubPosition = NewTime;
 
+	UMovieScene* MovieScene = SequenceInstanceStack.Top()->GetSequence()->GetMovieScene();
+	if ( MovieScene->GetForceFixedFrameIntervalPlayback() && MovieScene->GetFixedFrameInterval() > 0 )
+	{
+		NewTime = UMovieScene::CalculateFixedFrameTime( NewTime, MovieScene->GetFixedFrameInterval() );
+		LastTime = UMovieScene::CalculateFixedFrameTime( LastTime, MovieScene->GetFixedFrameInterval() );
+	}
+
 	EMovieSceneUpdateData UpdateData(NewTime, LastTime);
 	UpdateData.bLooped = bLooped;
 	SequenceInstanceStack.Top()->Update(UpdateData, *this);
@@ -1417,6 +1430,37 @@ void FSequencer::SetPerspectiveViewportCameraCutEnabled(bool bEnabled)
 	bPerspectiveViewportCameraCutEnabled = bEnabled;
 }
 
+void FSequencer::RenderMovie(UMovieSceneSection* InSection) const
+{
+	RenderMovieInternal(InSection->GetStartTime(), InSection->GetEndTime());
+}
+
+void FSequencer::RenderMovieInternal(float InStartTime, float InEndTime) const
+{
+	FLevelEditorModule& LevelEditorModule = FModuleManager::GetModuleChecked<FLevelEditorModule>(TEXT("LevelEditor"));
+
+	// Create a new movie scene capture object for an automated level sequence, and open the tab
+	UAutomatedLevelSequenceCapture* MovieSceneCapture = NewObject<UAutomatedLevelSequenceCapture>(GetTransientPackage(), UAutomatedLevelSequenceCapture::StaticClass(), NAME_None, RF_Transient);
+	MovieSceneCapture->LoadFromConfig();
+
+	MovieSceneCapture->SetLevelSequenceAsset(GetCurrentAsset()->GetPathName());
+
+	if (CanShowFrameNumbers())
+	{
+		MovieSceneCapture->Settings.FrameRate = FMath::RoundToInt(1.f / Settings->GetTimeSnapInterval());
+		MovieSceneCapture->Settings.ZeroPadFrameNumbers = Settings->GetZeroPadFrames();
+		MovieSceneCapture->Settings.bUseRelativeFrameNumbers = false;
+
+		const int32 SequenceStartFrame = FMath::RoundToInt( InStartTime * MovieSceneCapture->Settings.FrameRate );
+		const int32 SequenceEndFrame = FMath::Max( SequenceStartFrame, FMath::RoundToInt( InEndTime * MovieSceneCapture->Settings.FrameRate ) );
+
+		MovieSceneCapture->StartFrame = SequenceStartFrame;
+		MovieSceneCapture->EndFrame = SequenceEndFrame;
+	}
+
+	IMovieSceneCaptureDialogModule::Get().OpenDialog(LevelEditorModule.GetLevelEditorTabManager().ToSharedRef(), MovieSceneCapture);
+}
+
 ISequencer::FOnActorAddedToSequencer& FSequencer::OnActorAddedToSequencer()
 {
 	return OnActorAddedToSequencerEvent;
@@ -3974,6 +4018,14 @@ void FSequencer::DoAssignActor(AActor*const* InActors, int32 NumActors, FGuid In
 	UpdateRuntimeInstances();
 }
 
+void FSequencer::ImportFBX(FGuid InObjectBinding)
+{
+	UMovieScene* MovieScene = GetFocusedMovieSceneSequence()->GetMovieScene();
+	if (MovieSceneToolHelpers::ImportFBX(MovieScene, InObjectBinding))
+	{
+		NotifyMovieSceneDataChanged();
+	}
+}
 
 void FSequencer::DeleteNode(TSharedRef<FSequencerDisplayNode> NodeToBeDeleted)
 {
@@ -5404,29 +5456,7 @@ void FSequencer::BindCommands()
 
 	SequencerCommandBindings->MapAction(
 		Commands.RenderMovie,
-		FExecuteAction::CreateLambda([this]{
-			FLevelEditorModule& LevelEditorModule = FModuleManager::GetModuleChecked<FLevelEditorModule>(TEXT("LevelEditor"));
-
-			// Create a new movie scene capture object for an automated level sequence, and open the tab
-			UAutomatedLevelSequenceCapture* MovieSceneCapture = NewObject<UAutomatedLevelSequenceCapture>(GetTransientPackage(), UAutomatedLevelSequenceCapture::StaticClass(), NAME_None, RF_Transient);
-			MovieSceneCapture->LoadConfig();
-
-			MovieSceneCapture->SetLevelSequenceAsset(GetCurrentAsset()->GetPathName());
-
-			if (CanShowFrameNumbers())
-			{
-				MovieSceneCapture->Settings.FrameRate = FMath::RoundToInt(1.f / Settings->GetTimeSnapInterval());
-				MovieSceneCapture->Settings.bUseRelativeFrameNumbers = false;
-
-				const int32 SequenceStartFrame = FMath::RoundToInt( GetPlaybackRange().GetLowerBoundValue() * MovieSceneCapture->Settings.FrameRate );
-				const int32 SequenceEndFrame = FMath::Max( SequenceStartFrame, FMath::RoundToInt( GetPlaybackRange().GetUpperBoundValue() * MovieSceneCapture->Settings.FrameRate ) );
-
-				MovieSceneCapture->StartFrame = SequenceStartFrame;
-				MovieSceneCapture->EndFrame = SequenceEndFrame;
-			}
-
-			IMovieSceneCaptureDialogModule::Get().OpenDialog(LevelEditorModule.GetLevelEditorTabManager().ToSharedRef(), MovieSceneCapture);
-		})
+		FExecuteAction::CreateLambda([this]{ RenderMovieInternal(GetPlaybackRange().GetLowerBoundValue(), GetPlaybackRange().GetUpperBoundValue()); })
 	);
 
 	SequencerCommandBindings->MapAction(
diff --git a/Engine/Source/Editor/Sequencer/Private/Sequencer.h b/Engine/Source/Editor/Sequencer/Private/Sequencer.h
--- a/Engine/Source/Editor/Sequencer/Private/Sequencer.h
+++ b/Engine/Source/Editor/Sequencer/Private/Sequencer.h
@@ -367,6 +367,9 @@ public:
 	void AssignActor(FMenuBuilder& MenuBuilder, FGuid ObjectBinding);
 	void DoAssignActor(AActor*const* InActors, int32 NumActors, FGuid ObjectBinding);
 
+	/** Called when a user executes the import fbx to track menu item */
+	void ImportFBX(FGuid ObjectBinding);	
+
 	/** Called when a user executes the delete node menu item */
 	void DeleteNode(TSharedRef<FSequencerDisplayNode> NodeToBeDeleted);
 	void DeleteSelectedNodes();
@@ -499,6 +502,7 @@ public:
 	virtual void SetGlobalTimeDirectly(float Time, ESnapTimeMode SnapTimeMode = ESnapTimeMode::STM_None, bool bLooped = false) override;
 	virtual void SetPerspectiveViewportPossessionEnabled(bool bEnabled) override;
 	virtual void SetPerspectiveViewportCameraCutEnabled(bool bEnabled) override;
+	virtual void RenderMovie(UMovieSceneSection* InSection) const override;
 	virtual void EnterSilentMode() override { ++SilentModeCount; }
 	virtual void ExitSilentMode() override { --SilentModeCount; ensure(SilentModeCount >= 0); }
 	virtual bool IsInSilentMode() const override { return SilentModeCount != 0; }
@@ -755,6 +759,9 @@ protected:
 	/** Internal conversion function that doesn't perform expensive reset/update tasks */
 	FMovieSceneSpawnable* ConvertToSpawnableInternal(FGuid PossessableGuid);
 
+	/** Internal function to render movie for a given start/end time */
+	void RenderMovieInternal(float InStartTime, float InEndTime) const;
+
 	/** Handles adding a new folder to the outliner tree. */
 	void OnAddFolder();
 
diff --git a/Engine/Source/Editor/Sequencer/Private/SequencerCommands.cpp b/Engine/Source/Editor/Sequencer/Private/SequencerCommands.cpp
--- a/Engine/Source/Editor/Sequencer/Private/SequencerCommands.cpp
+++ b/Engine/Source/Editor/Sequencer/Private/SequencerCommands.cpp
@@ -58,20 +58,20 @@ void FSequencerCommands::RegisterCommands()
 
 	UI_COMMAND( ToggleAutoScroll, "Auto Scroll", "Toggle auto-scroll: When enabled, automatically scrolls the sequencer view to keep the current time visible", EUserInterfaceActionType::ToggleButton, FInputChord(EModifierKey::Shift, EKeys::S) );
 
-	UI_COMMAND( ToggleShowFrameNumbers, "Show Frame Numbers", "Enables and disables showing frame numbers", EUserInterfaceActionType::ToggleButton, FInputChord() );
-	UI_COMMAND( ToggleShowGotoBox, "Go to Time...", "Go to a particular point on the time line", EUserInterfaceActionType::Button, FInputChord(EModifierKey::Control, EKeys::G) );
+	UI_COMMAND( ToggleShowFrameNumbers, "Show Frame Numbers", "Enables and disables showing frame numbers", EUserInterfaceActionType::ToggleButton, FInputChord(EModifierKey::Control, EKeys::T) );
+	UI_COMMAND( ToggleShowGotoBox, "Go to Time...", "Go to a particular point on the timeline", EUserInterfaceActionType::Button, FInputChord(EModifierKey::Control, EKeys::G) );
 	UI_COMMAND( ToggleShowRangeSlider, "Show Range Slider", "Enables and disables showing the time range slider", EUserInterfaceActionType::ToggleButton, FInputChord() );
 	UI_COMMAND( ToggleIsSnapEnabled, "Enable Snapping", "Enables and disables snapping", EUserInterfaceActionType::ToggleButton, FInputChord() );
 
-	UI_COMMAND( ToggleSnapKeyTimesToInterval, "Snap to the Interval", "Snap keys to the interval", EUserInterfaceActionType::ToggleButton, FInputChord() );
+	UI_COMMAND( ToggleSnapKeyTimesToInterval, "Snap to the Interval", "Snap keys to the time snapping interval", EUserInterfaceActionType::ToggleButton, FInputChord() );
 	UI_COMMAND( ToggleSnapKeyTimesToKeys, "Snap to Other Keys", "Snap keys to other keys in the section", EUserInterfaceActionType::ToggleButton, FInputChord() );
 
-	UI_COMMAND( ToggleSnapSectionTimesToInterval, "Snap to the Interval", "Snap sections to the interval", EUserInterfaceActionType::ToggleButton, FInputChord() );
+	UI_COMMAND( ToggleSnapSectionTimesToInterval, "Snap to the Interval", "Snap sections to the time snapping interval", EUserInterfaceActionType::ToggleButton, FInputChord() );
 	UI_COMMAND( ToggleSnapSectionTimesToSections, "Snap to Other Sections", "Snap sections to other sections", EUserInterfaceActionType::ToggleButton, FInputChord() );
 
-	UI_COMMAND( ToggleSnapPlayTimeToKeys, "Snap to Keys While Scrubbing", "Snap the time to keys while scrubbing", EUserInterfaceActionType::ToggleButton, FInputChord() );
-	UI_COMMAND( ToggleSnapPlayTimeToInterval, "Snap to the Interval While Scrubbing", "Snap the play time to the interval while scrubbing", EUserInterfaceActionType::ToggleButton, FInputChord() );
-	UI_COMMAND( ToggleSnapPlayTimeToDraggedKey, "Snap to the Dragged Key", "Snap the play time to the dragged key", EUserInterfaceActionType::ToggleButton, FInputChord() );
+	UI_COMMAND( ToggleSnapPlayTimeToKeys, "Snap to Keys While Scrubbing", "Snap the current time to keys of the selected track while scrubbing", EUserInterfaceActionType::ToggleButton, FInputChord() );
+	UI_COMMAND( ToggleSnapPlayTimeToInterval, "Snap to the Interval While Scrubbing", "Snap the current time to the time snapping interval while scrubbing", EUserInterfaceActionType::ToggleButton, FInputChord() );
+	UI_COMMAND( ToggleSnapPlayTimeToDraggedKey, "Snap to the Dragged Key", "Snap the current time to the dragged or pressed key", EUserInterfaceActionType::ToggleButton, FInputChord() );
 
 	UI_COMMAND( ToggleSnapCurveValueToInterval, "Snap Curve Key Values", "Snap curve keys to the value snapping interval", EUserInterfaceActionType::ToggleButton, FInputChord() );
 
@@ -83,9 +83,6 @@ void FSequencerCommands::RegisterCommands()
 	UI_COMMAND( ToggleShowCurveEditor, "Curve Editor", "Show the animation keys in a curve editor", EUserInterfaceActionType::ToggleButton, FInputChord() );
 	UI_COMMAND( ToggleLinkCurveEditorTimeRange, "Link Curve Editor Time Range", "Link the curve editor time range to the sequence", EUserInterfaceActionType::ToggleButton, FInputChord() );
 
-	UI_COMMAND( MoveTool, "Move", "Activates the move tool", EUserInterfaceActionType::ToggleButton, FInputChord( EKeys::V ) );
-	UI_COMMAND( MarqueeTool, "Marquee", "Activates the marquee selection tool", EUserInterfaceActionType::ToggleButton, FInputChord( EKeys::M ) );
-
 	UI_COMMAND( RenderMovie, "Render Movie", "Render this movie to a video, or image frame sequence", EUserInterfaceActionType::Button, FInputChord() );
 	UI_COMMAND( CreateCamera, "Create Camera", "Create a new camera and set it as the current camera cut", EUserInterfaceActionType::Button, FInputChord() );
 
diff --git a/Engine/Source/Editor/Sequencer/Private/SequencerLabelManager.cpp b/Engine/Source/Editor/Sequencer/Private/SequencerLabelManager.cpp
--- a/Engine/Source/Editor/Sequencer/Private/SequencerLabelManager.cpp
+++ b/Engine/Source/Editor/Sequencer/Private/SequencerLabelManager.cpp
@@ -51,6 +51,9 @@ void FSequencerLabelManager::RemoveObjectLabel(const FGuid& ObjectId, const FStr
 		if (ObjectId.IsValid())
 		{
 			ObjectsToLabels.FindOrAdd(ObjectId.ToString()).Strings.Remove(Label);
+
+			// Add to dummy object id so that this label will stick around until intentionally removed by user
+			ObjectsToLabels.FindOrAdd(TEXT("UnusedLabelsDummy")).Strings.Add(Label);
 		}
 		else
 		{
diff --git a/Engine/Source/Editor/Sequencer/Private/SequencerNodeTree.cpp b/Engine/Source/Editor/Sequencer/Private/SequencerNodeTree.cpp
--- a/Engine/Source/Editor/Sequencer/Private/SequencerNodeTree.cpp
+++ b/Engine/Source/Editor/Sequencer/Private/SequencerNodeTree.cpp
@@ -509,27 +509,24 @@ static bool FilterNodesRecursive( FSequencer& Sequencer, const TSharedRef<FSeque
 	bool bObjectHasLabels = false;
 	for (const FString& String : FilterStrings)
 	{
-		if (String.StartsWith(TEXT("label:")))
-		{				
-			if (!StartNode->GetParent().IsValid() && (String.Len() > 6))
+		if (String.StartsWith(TEXT("label:")) && String.Len() > 6)
+		{
+			if (StartNode->GetType() == ESequencerNode::Object)
 			{
-				if (StartNode->GetType() == ESequencerNode::Object)
-				{
-					bObjectHasLabels = true;
-					auto ObjectBindingNode = StaticCastSharedRef<FSequencerObjectBindingNode>(StartNode);
-					auto Labels = Sequencer.GetLabelManager().GetObjectLabels(ObjectBindingNode->GetObjectBinding());
+				bObjectHasLabels = true;
+				auto ObjectBindingNode = StaticCastSharedRef<FSequencerObjectBindingNode>(StartNode);
+				auto Labels = Sequencer.GetLabelManager().GetObjectLabels(ObjectBindingNode->GetObjectBinding());
 
-					if (Labels != nullptr && Labels->Strings.Contains(String.RightChop(6)))
-					{
-						bMatchedLabel = true;
-						break;
-					}
-				}
-				else
+				if (Labels != nullptr && Labels->Strings.Contains(String.RightChop(6)))
 				{
-					return false;
+					bMatchedLabel = true;
+					break;
 				}
 			}
+			else if (!StartNode->GetParent().IsValid())
+			{
+				return false;
+			}
 		}
 	}
 
@@ -562,18 +559,21 @@ static bool FilterNodesRecursive( FSequencer& Sequencer, const TSharedRef<FSeque
 	}
 
 	// check each child node to determine if it is filtered
-	const TArray<TSharedRef<FSequencerDisplayNode>>& ChildNodes = StartNode->GetChildNodes();
-
-	for (const auto& Node : ChildNodes)
+	if (StartNode->GetType() != ESequencerNode::Folder)
 	{
-		// Mark the parent as filtered if any child node was filtered
-		bPassedTextFilter |= FilterNodesRecursive(Sequencer, Node, FilterStrings, OutFilteredNodes);
+		const TArray<TSharedRef<FSequencerDisplayNode>>& ChildNodes = StartNode->GetChildNodes();
 
-		if (bPassedTextFilter && !bInFilter)
+		for (const auto& Node : ChildNodes)
 		{
-			AddFilteredNode(StartNode, OutFilteredNodes);
+			// Mark the parent as filtered if any child node was filtered
+			bPassedTextFilter |= FilterNodesRecursive(Sequencer, Node, FilterStrings, OutFilteredNodes);
 
-			bInFilter = true;
+			if (bPassedTextFilter && !bInFilter)
+			{
+				AddFilteredNode(StartNode, OutFilteredNodes);
+
+				bInFilter = true;
+			}
 		}
 	}
 
diff --git a/Engine/Source/Editor/Sequencer/Public/ISequencer.h b/Engine/Source/Editor/Sequencer/Public/ISequencer.h
--- a/Engine/Source/Editor/Sequencer/Public/ISequencer.h
+++ b/Engine/Source/Editor/Sequencer/Public/ISequencer.h
@@ -245,6 +245,13 @@ public:
 	 */ 
 	virtual bool IsPerspectiveViewportCameraCutEnabled() const { return true; }
 
+	/*
+	 * Render movie for a section.
+	 * 
+	 * @param InSection The given section to render.
+	 */
+	virtual void RenderMovie(UMovieSceneSection* InSection) const = 0;
+
 	/*
 	 * Puts sequencer in a silent state (whereby it will not redraw viewports, or attempt to update external state besides the sequence itself)
 	 */
diff --git a/Engine/Source/Editor/Sequencer/Public/SequencerSettings.h b/Engine/Source/Editor/Sequencer/Public/SequencerSettings.h
--- a/Engine/Source/Editor/Sequencer/Public/SequencerSettings.h
+++ b/Engine/Source/Editor/Sequencer/Public/SequencerSettings.h
@@ -244,102 +244,134 @@ public:
 
 protected:
 
+	/** Enable or disable autokeying. */
 	UPROPERTY( config, EditAnywhere, Category=Keyframing )
 	TEnumAsByte<EAutoKeyMode> AutoKeyMode;
 
+	/** Enable or disable keying all channels when any are keyed. */
 	UPROPERTY( config, EditAnywhere, Category=Keyframing )
 	bool bKeyAllEnabled;
 
+	/** Enable or disable only keyframing properties marked with the 'Interp' keyword. */
 	UPROPERTY( config, EditAnywhere, Category=Keyframing )
 	bool bKeyInterpPropertiesOnly;
 
+	/** The interpolation type for newly created keyframes */
 	UPROPERTY( config, EditAnywhere, Category=Keyframing )
 	TEnumAsByte<EMovieSceneKeyInterpolation> KeyInterpolation;
 
+	/** The default location of a spawnable when it is first dragged into the viewport from the content browser. */
 	UPROPERTY( config, EditAnywhere, Category=General )
 	TEnumAsByte<ESequencerSpawnPosition> SpawnPosition;
 
+	/** Enable or disable creating of spawnable cameras whenever cameras are created. */
 	UPROPERTY( config, EditAnywhere, Category=General )
 	bool bCreateSpawnableCameras;
 
+	/** Show frame numbers or time in the timeline. */
 	UPROPERTY( config, EditAnywhere, Category=Timeline )
 	bool bShowFrameNumbers;
 
+	/** Show the in/out range in the timeline with respect to the start/end range. */
 	UPROPERTY( config, EditAnywhere, Category=Timeline )
 	bool bShowRangeSlider;
 
+	/** Enable or disable snapping in the timeline. */
 	UPROPERTY( config, EditAnywhere, Category=Snapping )
 	bool bIsSnapEnabled;
 
+	/** The time snapping interval in the timeline. */
 	UPROPERTY( config, EditAnywhere, Category=Snapping )
 	float TimeSnapInterval;
 
+	/** Enable or disable snapping keys to the time snapping interval. */
 	UPROPERTY( config, EditAnywhere, Category=Snapping )
 	bool bSnapKeyTimesToInterval;
 
+	/** Enable or disable snapping keys to other keys. */
 	UPROPERTY( config, EditAnywhere, Category=Snapping )
 	bool bSnapKeyTimesToKeys;
-
+	
+	/** Enable or disable snapping sections to the time snapping interval. */
 	UPROPERTY( config, EditAnywhere, Category=Snapping )
 	bool bSnapSectionTimesToInterval;
 
+	/** Enable or disable snapping sections to other sections. */
 	UPROPERTY( config, EditAnywhere, Category=Snapping )
 	bool bSnapSectionTimesToSections;
 
+	/** Enable or disable snapping the current time to keys of the selected track while scrubbing. */
 	UPROPERTY( config, EditAnywhere, Category=Snapping )
 	bool bSnapPlayTimeToKeys;
 
+	/** Enable or disable snapping the current time to the time snapping interval while scrubbing. */
 	UPROPERTY( config, EditAnywhere, Category=Snapping )
 	bool bSnapPlayTimeToInterval;
 
+	/** Enable or disable snapping the current time to the dragged or pressed key. */
 	UPROPERTY( config, EditAnywhere, Category=Snapping )
 	bool bSnapPlayTimeToDraggedKey;
 
-	UPROPERTY( config, EditAnywhere, Category=Snapping )
+	/** The curve value interval to snap to. */
 	float CurveValueSnapInterval;
 
+	/** Enable or disable snapping the curve value to the curve value interval. */
 	UPROPERTY( config, EditAnywhere, Category=Snapping )
 	bool bSnapCurveValueToInterval;
 
+	/** Enable or disable the label browser. */
 	UPROPERTY( config, EditAnywhere, Category=General )
 	bool bLabelBrowserVisible;
 
+	/** Whether to zoom in on the current position or the current time in the timeline. */
 	UPROPERTY( config, EditAnywhere, Category=Timeline )
 	TEnumAsByte<ESequencerZoomPosition> ZoomPosition;
 
+	/** Enable or disable auto scroll in the timeline. */
 	UPROPERTY( config, EditAnywhere, Category=Timeline )
 	bool bAutoScrollEnabled;
 
+	/** Enable or disable the curve editor. */
 	UPROPERTY( config, EditAnywhere, Category=CurveEditor )
 	bool bShowCurveEditor;
 
+	/** Enable or disable curve editor tooltips. */
 	UPROPERTY( config, EditAnywhere, Category=CurveEditor )
 	bool bShowCurveEditorCurveToolTips;
 
+	/** Enable or disable linking the curve editor time range to the sequencer timeline's time range. */
 	UPROPERTY( config, EditAnywhere, Category=CurveEditor )
 	bool bLinkCurveEditorTimeRange;
 
+	/** Enable or disable looping playback in the timeline. */
 	UPROPERTY( config )
 	bool bLooping;
 
+	/** Enable or disable keeping the cursor in the current playback range during playback. */
 	UPROPERTY( config, EditAnywhere, Category=Timeline )
 	bool bKeepCursorInPlayRange;
 
+	/** Enable or disable keeping the playback range constrained to the section bounds. */
 	UPROPERTY( config, EditAnywhere, Category=Timeline )
 	bool bKeepPlayRangeInSectionBounds;
 
+	/** The number of zeros to pad the frame numbers by. */
 	UPROPERTY( config, EditAnywhere, Category=Timeline )
 	uint8 ZeroPadFrames;
 
+	/** Enable or disable the combined keyframes at the top node level. Disabling can improve editor performance. */
 	UPROPERTY( config, EditAnywhere, Category=Timeline )
 	bool bShowCombinedKeyframes;
 
+	/** Enable or disable setting key area sections as infinite by default. */
 	UPROPERTY( config, EditAnywhere, Category=Timeline )
 	bool bInfiniteKeyAreas;
 
+	/** Enable or disable displaying channel bar colors for vector properties. */
 	UPROPERTY( config, EditAnywhere, Category=Timeline )
 	bool bShowChannelColors;
 
+	/** Enable or disable transport controls in the viewport. */
 	UPROPERTY( config )
 	bool bShowViewportTransportControls;
 
diff --git a/Engine/Source/Runtime/CinematicCamera/Private/CineCameraActor.cpp b/Engine/Source/Runtime/CinematicCamera/Private/CineCameraActor.cpp
--- a/Engine/Source/Runtime/CinematicCamera/Private/CineCameraActor.cpp
+++ b/Engine/Source/Runtime/CinematicCamera/Private/CineCameraActor.cpp
@@ -14,8 +14,10 @@ ACineCameraActor::ACineCameraActor(const FObjectInitializer& ObjectInitializer)
 			.SetDefaultSubobjectClass<UCineCameraComponent>(TEXT("CameraComponent"))
 	)
 {
+	CineCameraComponent = Cast<UCineCameraComponent>(GetCameraComponent());
+
 	PrimaryActorTick.bCanEverTick = true;
-	SetActorTickEnabled(LookatTrackingSettings.bEnableLookAtTracking);
+	SetActorTickEnabled(ShouldTickForTracking());
 }
 
 #if WITH_EDITOR
@@ -23,14 +25,14 @@ void ACineCameraActor::PostEditChangeProperty(FPropertyChangedEvent& PropertyCha
 {
 	Super::PostEditChangeProperty(PropertyChangedEvent);
 
-	SetActorTickEnabled(LookatTrackingSettings.bEnableLookAtTracking);
+	SetActorTickEnabled(ShouldTickForTracking());
 }
 #endif
 
 void ACineCameraActor::PostInitializeComponents()
 {
 	Super::PostInitializeComponents();
-	SetActorTickEnabled(LookatTrackingSettings.bEnableLookAtTracking);
+	SetActorTickEnabled(ShouldTickForTracking());
 
 	LookatTrackingSettings.LastLookatTrackingRotation = GetActorRotation();
 }
@@ -59,30 +61,55 @@ FVector ACineCameraActor::GetLookatLocation() const
 static const FColor DebugLookatTrackingPointSolidColor(200, 200, 32, 128);		// yellow
 static const FColor DebugLookatTrackingPointOutlineColor = FColor::Black;
 
+static const FColor DebugFocusPointSolidColor(102, 26, 204, 153);		// purple
+static const FColor DebugFocusPointOutlineColor = FColor::Black;
+
 void ACineCameraActor::Tick(float DeltaTime)
 {
 	Super::Tick(DeltaTime);
 
-	if (CameraComponent && LookatTrackingSettings.bEnableLookAtTracking)
+	if (CameraComponent && ShouldTickForTracking())
 	{
-		// do the lookat tracking
-		// #note this will turn the whole actor, which assumes the cameracomponent's transform is the same as the root component
-		// more complex component hierarchies will require different handling here
-		FVector const LookatLoc = GetLookatLocation();
-		FVector const ToLookat = LookatLoc - GetActorLocation();
-		FRotator const FinalRot = 
-			bResetInterplation
-			? ToLookat.Rotation()
-			: FMath::RInterpTo(LookatTrackingSettings.LastLookatTrackingRotation, ToLookat.Rotation(), DeltaTime, LookatTrackingSettings.LookAtTrackingInterpSpeed);
-		SetActorRotation(FinalRot);
-
-		// we store this ourselves in case other systems try to change our rotation, and end fighting the interpolation
-		LookatTrackingSettings.LastLookatTrackingRotation = FinalRot;
-
-		if (LookatTrackingSettings.bDrawDebugLookAtTrackingPosition)
+		if (LookatTrackingSettings.bEnableLookAtTracking)
+		{
+			// do the lookat tracking
+			// #note this will turn the whole actor, which assumes the cameracomponent's transform is the same as the root component
+			// more complex component hierarchies will require different handling here
+			FVector const LookatLoc = GetLookatLocation();
+			FVector const ToLookat = LookatLoc - GetActorLocation();
+			FRotator const FinalRot = 
+				bResetInterplation
+				? ToLookat.Rotation()
+				: FMath::RInterpTo(LookatTrackingSettings.LastLookatTrackingRotation, ToLookat.Rotation(), DeltaTime, LookatTrackingSettings.LookAtTrackingInterpSpeed);
+			SetActorRotation(FinalRot);
+
+			// we store this ourselves in case other systems try to change our rotation, and end fighting the interpolation
+			LookatTrackingSettings.LastLookatTrackingRotation = FinalRot;
+
+			if (LookatTrackingSettings.bDrawDebugLookAtTrackingPosition)
+			{
+				::DrawDebugSolidBox(GetWorld(), LookatLoc, FVector(12.f), DebugLookatTrackingPointSolidColor);
+				::DrawDebugBox(GetWorld(), LookatLoc, FVector(12.f), DebugLookatTrackingPointOutlineColor);
+			}
+		}
+
+		if (CineCameraComponent->FocusSettings.TrackingFocusSettings.bDrawDebugTrackingFocusPoint)
 		{
-			::DrawDebugSolidBox(GetWorld(), LookatLoc, FVector(12.f), DebugLookatTrackingPointSolidColor);
-			::DrawDebugBox(GetWorld(), LookatLoc, FVector(12.f), DebugLookatTrackingPointOutlineColor);
+			AActor const* const TrackedActor = CineCameraComponent->FocusSettings.TrackingFocusSettings.ActorToTrack;
+
+			FVector FocusPoint;
+			if (TrackedActor)
+			{
+				FTransform const BaseTransform = TrackedActor->GetActorTransform();
+				FocusPoint = BaseTransform.TransformPosition(CineCameraComponent->FocusSettings.TrackingFocusSettings.RelativeOffset);
+			}
+			else
+			{
+				FocusPoint = CineCameraComponent->FocusSettings.TrackingFocusSettings.RelativeOffset;
+			}
+			
+			::DrawDebugSolidBox(GetWorld(), FocusPoint, FVector(12.f), DebugFocusPointSolidColor);
+			::DrawDebugBox(GetWorld(), FocusPoint, FVector(12.f), DebugFocusPointOutlineColor);
 		}
 	}
 	else
@@ -102,5 +129,10 @@ void ACineCameraActor::NotifyCameraCut()
 	bResetInterplation = true;
 }
 
+bool ACineCameraActor::ShouldTickForTracking() const
+{
+	return LookatTrackingSettings.bEnableLookAtTracking || CineCameraComponent->FocusSettings.TrackingFocusSettings.bDrawDebugTrackingFocusPoint;
+}
+
 
 #undef LOCTEXT_NAMESPACE
diff --git a/Engine/Source/Runtime/CinematicCamera/Public/CineCameraActor.h b/Engine/Source/Runtime/CinematicCamera/Public/CineCameraActor.h
--- a/Engine/Source/Runtime/CinematicCamera/Public/CineCameraActor.h
+++ b/Engine/Source/Runtime/CinematicCamera/Public/CineCameraActor.h
@@ -2,6 +2,8 @@
 
 #pragma once
 
+#include "CineCameraComponent.h"
+
 #include "CineCameraActor.generated.h"
 
 /** Settings to control the camera's lookat feature */
@@ -64,4 +66,10 @@ protected:
 	FVector GetLookatLocation() const;
 
 	virtual void NotifyCameraCut() override;
+
+	bool ShouldTickForTracking() const;
+
+private:
+	/** Returns CineCameraComponent subobject **/
+	class UCineCameraComponent* CineCameraComponent;
 };
diff --git a/Engine/Source/Runtime/Core/Public/Math/CapsuleShape.h b/Engine/Source/Runtime/Core/Public/Math/CapsuleShape.h
new file mode 100644
--- /dev/null
+++ b/Engine/Source/Runtime/Core/Public/Math/CapsuleShape.h
@@ -0,0 +1,44 @@
+// Copyright 1998-2016 Epic Games, Inc. All Rights Reserved.
+
+#pragma once
+
+
+/**
+ * Structure for capsules.
+ *
+ * A capsule consists of two sphere connected by a cylinder.
+ */
+struct FCapsuleShape
+{
+	/** The capsule's center point. */
+	FVector Center;
+
+	/** The capsule's radius. */
+	float Radius;
+
+	/** The capsule's orientation in space. */
+	FVector Orientation;
+
+	/** The capsule's length. */
+	float Length;
+
+public:
+
+	/** Default constructor. */
+	FCapsuleShape() { }
+
+	/**
+	 * Create and inintialize a new instance.
+	 *
+	 * @param InCenter The capsule's center point.
+	 * @param InRadius The capsule's radius.
+	 * @param InOrientation The capsule's orientation in space.
+	 * @param InLength The capsule's length.
+	 */
+	FCapsuleShape(FVector InCenter, float InRadius, FVector InOrientation, float InLength)
+		: Center(InCenter)
+		, Radius(InRadius)
+		, Orientation(InOrientation)
+		, Length(InLength)
+	{ }
+};
diff --git a/Engine/Source/Runtime/Core/Public/Math/Plane.h b/Engine/Source/Runtime/Core/Public/Math/Plane.h
--- a/Engine/Source/Runtime/Core/Public/Math/Plane.h
+++ b/Engine/Source/Runtime/Core/Public/Math/Plane.h
@@ -6,14 +6,15 @@
 /**
  * Structure for three dimensional planes.
  *
- * Stores the coeffecients as Ax+By+Cz=D.
- * Note that this is different than many other Plane classes that use Ax+By+Cz+D=0.
+ * Stores the coeffecients as Xx+Yy+Zz=W.
+ * Note that this is different from many other Plane classes that use Xx+Yy+Zz+W=0.
  */
 MS_ALIGN(16) struct FPlane
 	: public FVector
 {
 public:
 
+	/** The w-component. */
 	float W;
 
 public:
@@ -26,32 +27,32 @@ public:
 	 *
 	 * @param P Plane to copy from.
 	 */
-	FORCEINLINE FPlane( const FPlane& P );
+	FORCEINLINE FPlane(const FPlane& P);
 
 	/**
 	 * Constructor.
 	 *
 	 * @param V 4D vector to set up plane.
 	 */
-	FORCEINLINE FPlane( const FVector4& V );
+	FORCEINLINE FPlane(const FVector4& V);
 
 	/**
 	 * Constructor.
 	 *
-	 * @param InX X Coordinate.
-	 * @param InY Y Coordinate.
-	 * @param InZ Z Coordinate.
-	 * @param InW W Coordinate.
+	 * @param InX X-coefficient.
+	 * @param InY Y-coefficient.
+	 * @param InZ Z-coefficient.
+	 * @param InW W-coefficient.
 	 */
-	FORCEINLINE FPlane( float InX, float InY, float InZ, float InW );
+	FORCEINLINE FPlane(float InX, float InY, float InZ, float InW);
 
 	/**
 	 * Constructor.
 	 *
 	 * @param InNormal Plane Normal Vector.
-	 * @param InW Plane W Coordinate.
+	 * @param InW Plane W-coefficient.
 	 */
-	FORCEINLINE FPlane( FVector InNormal, float InW );
+	FORCEINLINE FPlane(FVector InNormal, float InW);
 
 	/**
 	 * Constructor.
@@ -59,7 +60,7 @@ public:
 	 * @param InBase Base point in plane.
 	 * @param InNormal Plane Normal Vector.
 	 */
-	FORCEINLINE FPlane( FVector InBase, const FVector &InNormal );
+	FORCEINLINE FPlane(FVector InBase, const FVector &InNormal);
 
 	/**
 	 * Constructor.
@@ -68,7 +69,7 @@ public:
 	 * @param B Second point in the plane.
 	 * @param C Third point in the plane.
 	 */
-	FPlane( FVector A, FVector B, FVector C );
+	FPlane(FVector A, FVector B, FVector C);
 
 	/**
 	 * Constructor
@@ -85,7 +86,7 @@ public:
 	 * @param P The other point.
 	 * @return >0: point is in front of the plane, <0: behind, =0: on the plane.
 	 */
-	FORCEINLINE float PlaneDot( const FVector &P ) const;
+	FORCEINLINE float PlaneDot(const FVector &P) const;
 
 	/**
 	 * Get a flipped version of the plane.
@@ -100,7 +101,7 @@ public:
 	 * @param M The matrix to transform plane with.
 	 * @return The result of transform.
 	 */
-	FPlane TransformBy( const FMatrix& M ) const;
+	FPlane TransformBy(const FMatrix& M) const;
 
 	/**
 	 * You can optionally pass in the matrices transpose-adjoint, which save it recalculating it.
@@ -112,7 +113,7 @@ public:
 	 * @param TA Transpose-adjoint of Matrix.
 	 * @return The result of transform.
 	 */
-	FPlane TransformByUsingAdjointT( const FMatrix& M, float DetM, const FMatrix& TA ) const;
+	FPlane TransformByUsingAdjointT(const FMatrix& M, float DetM, const FMatrix& TA) const;
 
 	/**
 	 * Check if two planes are identical.
@@ -120,7 +121,7 @@ public:
 	 * @param V The other plane.
 	 * @return true if planes are identical, otherwise false.
 	 */
-	bool operator==( const FPlane& V ) const;
+	bool operator==(const FPlane& V) const;
 
 	/**
 	 * Check if two planes are different.
@@ -128,7 +129,7 @@ public:
 	 * @param V The other plane.
 	 * @return true if planes are different, otherwise false.
 	 */
-	bool operator!=( const FPlane& V ) const;
+	bool operator!=(const FPlane& V) const;
 
 	/**
 	 * Checks whether two planes are equal within specified tolerance.
@@ -145,7 +146,7 @@ public:
 	 * @param V The other plane.
 	 * @return The dot product.
 	 */
-	FORCEINLINE float operator|( const FPlane& V ) const;
+	FORCEINLINE float operator|(const FPlane& V) const;
 
 	/**
 	 * Gets result of adding a plane to this.
@@ -153,7 +154,7 @@ public:
 	 * @param V The other plane.
 	 * @return The result of adding a plane to this.
 	 */
-	FPlane operator+( const FPlane& V ) const;
+	FPlane operator+(const FPlane& V) const;
 
 	/**
 	 * Gets result of subtracting a plane from this.
@@ -161,7 +162,7 @@ public:
 	 * @param V The other plane.
 	 * @return The result of subtracting a plane from this.
 	 */
-	FPlane operator-( const FPlane& V ) const;
+	FPlane operator-(const FPlane& V) const;
 
 	/**
 	 * Gets result of dividing a plane.
@@ -169,7 +170,7 @@ public:
 	 * @param Scale What to divide by.
 	 * @return The result of division.
 	 */
-	FPlane operator/( float Scale ) const;
+	FPlane operator/(float Scale) const;
 
 	/**
 	 * Gets result of scaling a plane.
@@ -177,7 +178,7 @@ public:
 	 * @param Scale The scaling factor.
 	 * @return The result of scaling.
 	 */
-	FPlane operator*( float Scale ) const;
+	FPlane operator*(float Scale) const;
 
 	/**
 	 * Gets result of multiplying a plane with this.
@@ -185,7 +186,7 @@ public:
 	 * @param V The other plane.
 	 * @return The result of multiplying a plane with this.
 	 */
-	FPlane operator*( const FPlane& V );
+	FPlane operator*(const FPlane& V);
 
 	/**
 	 * Add another plane to this.
@@ -193,7 +194,7 @@ public:
 	 * @param V The other plane.
 	 * @return Copy of plane after addition.
 	 */
-	FPlane operator+=( const FPlane& V );
+	FPlane operator+=(const FPlane& V);
 
 	/**
 	 * Subtract another plane from this.
@@ -201,7 +202,7 @@ public:
 	 * @param V The other plane.
 	 * @return Copy of plane after subtraction.
 	 */
-	FPlane operator-=( const FPlane& V );
+	FPlane operator-=(const FPlane& V);
 
 	/**
 	 * Scale this plane.
@@ -209,7 +210,7 @@ public:
 	 * @param Scale The scaling factor.
 	 * @return Copy of plane after scaling.
 	 */
-	FPlane operator*=( float Scale );
+	FPlane operator*=(float Scale);
 
 	/**
 	 * Multiply another plane with this.
@@ -217,7 +218,7 @@ public:
 	 * @param V The other plane.
 	 * @return Copy of plane after multiplication.
 	 */
-	FPlane operator*=( const FPlane& V );
+	FPlane operator*=(const FPlane& V);
 
 	/**
 	 * Divide this plane.
@@ -225,7 +226,7 @@ public:
 	 * @param V What to divide by.
 	 * @return Copy of plane after division.
 	 */
-	FPlane operator/=( float V );
+	FPlane operator/=(float V);
 
 	/**
 	 * Serializer.
@@ -234,12 +235,12 @@ public:
 	 * @param P Plane to serialize.
 	 * @return Reference to Archive after serialization.
 	 */
-	friend FArchive& operator<<( FArchive& Ar, FPlane &P )
+	friend FArchive& operator<<(FArchive& Ar, FPlane &P)
 	{
 		return Ar << (FVector&)P << P.W;
 	}
 
-	bool Serialize( FArchive& Ar )
+	bool Serialize(FArchive& Ar)
 	{
 		if (Ar.UE4Ver() >= VER_UE4_ADDED_NATIVE_SERIALIZATION_FOR_IMMUTABLE_STRUCTURES)
 		{
@@ -256,12 +257,9 @@ public:
 	 */
 	bool NetSerialize(FArchive& Ar, class UPackageMap*, bool& bOutSuccess)
 	{
-		if( Ar.IsLoading() )
+		if(Ar.IsLoading())
 		{
-			int16 iX;
-			int16 iY;
-			int16 iZ;
-			int16 iW;
+			int16 iX, iY, iZ, iW;
 			Ar << iX << iY << iZ << iW;
 			*this = FPlane(iX,iY,iZ,iW);
 		}
@@ -286,37 +284,37 @@ FORCEINLINE FPlane::FPlane()
 {}
 
 
-FORCEINLINE FPlane::FPlane( const FPlane& P )
+FORCEINLINE FPlane::FPlane(const FPlane& P)
 	:	FVector(P)
 	,	W(P.W)
 {}
 
 
-FORCEINLINE FPlane::FPlane( const FVector4& V )
+FORCEINLINE FPlane::FPlane(const FVector4& V)
 	:	FVector(V)
 	,	W(V.W)
 {}
 
 
-FORCEINLINE FPlane::FPlane( float InX, float InY, float InZ, float InW )
+FORCEINLINE FPlane::FPlane(float InX, float InY, float InZ, float InW)
 	:	FVector(InX,InY,InZ)
 	,	W(InW)
 {}
 
 
-FORCEINLINE FPlane::FPlane( FVector InNormal, float InW )
+FORCEINLINE FPlane::FPlane(FVector InNormal, float InW)
 	:	FVector(InNormal), W(InW)
 {}
 
 
-FORCEINLINE FPlane::FPlane( FVector InBase, const FVector &InNormal )
+FORCEINLINE FPlane::FPlane(FVector InBase, const FVector &InNormal)
 	:	FVector(InNormal)
 	,	W(InBase | InNormal)
 {}
 
 
-FORCEINLINE FPlane::FPlane( FVector A, FVector B, FVector C )
-	:	FVector( ((B-A)^(C-A)).GetSafeNormal() )
+FORCEINLINE FPlane::FPlane(FVector A, FVector B, FVector C)
+	:	FVector(((B-A)^(C-A)).GetSafeNormal())
 {
 	W = A | (FVector)(*this);
 }
@@ -327,104 +325,104 @@ FORCEINLINE FPlane::FPlane(EForceInit)
 {}
 
 
-FORCEINLINE float FPlane::PlaneDot( const FVector &P ) const
+FORCEINLINE float FPlane::PlaneDot(const FVector &P) const
 {
-	return X*P.X + Y*P.Y + Z*P.Z - W;
+	return X * P.X + Y * P.Y + Z * P.Z - W;
 }
 
 
 FORCEINLINE FPlane FPlane::Flip() const
 {
-	return FPlane(-X,-Y,-Z,-W);
+	return FPlane(-X, -Y, -Z, -W);
 }
 
 
-FORCEINLINE bool FPlane::operator==( const FPlane& V ) const
+FORCEINLINE bool FPlane::operator==(const FPlane& V) const
 {
-	return X==V.X && Y==V.Y && Z==V.Z && W==V.W;
+	return (X == V.X) && (Y == V.Y) && (Z == V.Z) && (W == V.W);
 }
 
 
-FORCEINLINE bool FPlane::operator!=( const FPlane& V ) const
+FORCEINLINE bool FPlane::operator!=(const FPlane& V) const
 {
-	return X!=V.X || Y!=V.Y || Z!=V.Z || W!=V.W;
+	return (X != V.X) || (Y != V.Y) || (Z != V.Z) || (W != V.W);
 }
 
 
 FORCEINLINE bool FPlane::Equals(const FPlane& V, float Tolerance) const
 {
-	return FMath::Abs(X-V.X) < Tolerance && FMath::Abs(Y-V.Y) < Tolerance && FMath::Abs(Z-V.Z) < Tolerance && FMath::Abs(W-V.W) < Tolerance;
+	return (FMath::Abs(X - V.X) < Tolerance) && (FMath::Abs(Y - V.Y) < Tolerance) && (FMath::Abs(Z - V.Z) < Tolerance) && (FMath::Abs(W - V.W) < Tolerance);
 }
 
 
-FORCEINLINE float FPlane::operator|( const FPlane& V ) const
+FORCEINLINE float FPlane::operator|(const FPlane& V) const
 {
-	return X*V.X + Y*V.Y + Z*V.Z + W*V.W;
+	return X * V.X + Y * V.Y + Z * V.Z + W * V.W;
 }
 
 
-FORCEINLINE FPlane FPlane::operator+( const FPlane& V ) const
+FORCEINLINE FPlane FPlane::operator+(const FPlane& V) const
 {
-	return FPlane( X + V.X, Y + V.Y, Z + V.Z, W + V.W );
+	return FPlane(X + V.X, Y + V.Y, Z + V.Z, W + V.W);
 }
 
 
-FORCEINLINE FPlane FPlane::operator-( const FPlane& V ) const
+FORCEINLINE FPlane FPlane::operator-(const FPlane& V) const
 {
-	return FPlane( X - V.X, Y - V.Y, Z - V.Z, W - V.W );
+	return FPlane(X - V.X, Y - V.Y, Z - V.Z, W - V.W);
 }
 
 
-FORCEINLINE FPlane FPlane::operator/( float Scale ) const
+FORCEINLINE FPlane FPlane::operator/(float Scale) const
 {
-	const float RScale = 1.f/Scale;
-	return FPlane( X * RScale, Y * RScale, Z * RScale, W * RScale );
+	const float RScale = 1.f / Scale;
+	return FPlane(X * RScale, Y * RScale, Z * RScale, W * RScale);
 }
 
 
-FORCEINLINE FPlane FPlane::operator*( float Scale ) const
+FORCEINLINE FPlane FPlane::operator*(float Scale) const
 {
-	return FPlane( X * Scale, Y * Scale, Z * Scale, W * Scale );
+	return FPlane(X * Scale, Y * Scale, Z * Scale, W * Scale);
 }
 
 
-FORCEINLINE FPlane FPlane::operator*( const FPlane& V )
+FORCEINLINE FPlane FPlane::operator*(const FPlane& V)
 {
-	return FPlane ( X*V.X,Y*V.Y,Z*V.Z,W*V.W );
+	return FPlane (X * V.X, Y * V.Y, Z * V.Z, W * V.W);
 }
 
 
-FORCEINLINE FPlane FPlane::operator+=( const FPlane& V )
+FORCEINLINE FPlane FPlane::operator+=(const FPlane& V)
 {
 	X += V.X; Y += V.Y; Z += V.Z; W += V.W;
 	return *this;
 }
 
 
-FORCEINLINE FPlane FPlane::operator-=( const FPlane& V )
+FORCEINLINE FPlane FPlane::operator-=(const FPlane& V)
 {
 	X -= V.X; Y -= V.Y; Z -= V.Z; W -= V.W;
 	return *this;
 }
 
 
-FORCEINLINE FPlane FPlane::operator*=( float Scale )
+FORCEINLINE FPlane FPlane::operator*=(float Scale)
 {
 	X *= Scale; Y *= Scale; Z *= Scale; W *= Scale;
 	return *this;
 }
 
 
-FORCEINLINE FPlane FPlane::operator*=( const FPlane& V )
+FORCEINLINE FPlane FPlane::operator*=(const FPlane& V)
 {
 	X *= V.X; Y *= V.Y; Z *= V.Z; W *= V.W;
 	return *this;
 }
 
 
-FORCEINLINE FPlane FPlane::operator/=( float V )
+FORCEINLINE FPlane FPlane::operator/=(float V)
 {
-	const float RV = 1.f/V;
+	const float RV = 1.f / V;
 	X *= RV; Y *= RV; Z *= RV; W *= RV;
 	return *this;
 }
