Copying //UE4/Dev-Sequencer to //UE4/Dev-Main (Source: //UE4/Dev-Sequencer @ 3003857)

diff --git a/Engine/Source/Runtime/Core/Public/Math/Quat.h b/Engine/Source/Runtime/Core/Public/Math/Quat.h
--- a/Engine/Source/Runtime/Core/Public/Math/Quat.h
+++ b/Engine/Source/Runtime/Core/Public/Math/Quat.h
@@ -18,28 +18,26 @@ MS_ALIGN(16) struct FQuat
 {
 public:
 
-	/** Holds the quaternion's X-component. */
+	/** The quaternion's X-component. */
 	float X;
 
-	/** Holds the quaternion's Y-component. */
+	/** The quaternion's Y-component. */
 	float Y;
 
-	/** Holds the quaternion's Z-component. */
+	/** The quaternion's Z-component. */
 	float Z;
 
-	/** Holds the quaternion's W-component. */
+	/** The quaternion's W-component. */
 	float W;
 
 public:
 
-	//Identity quaternion
+	/** Identity quaternion. */
 	static CORE_API const FQuat Identity;
 
 public:
 
-	/**
-	 * Default constructor (no initialization).
-	 */
+	/** Default constructor (no initialization). */
 	FORCEINLINE FQuat() { }
 
 	/**
@@ -47,7 +45,7 @@ public:
 	 *
 	 * @param EForceInit Force init enum: if equal to ForceInitToZero then W is 0, otherwise W = 1 (creating an identity transform)
 	 */
-	explicit FORCEINLINE FQuat( EForceInit );
+	explicit FORCEINLINE FQuat(EForceInit);
 
 	/**
 	 * Constructor.
@@ -57,28 +55,28 @@ public:
 	 * @param InZ Z component of the quaternion
 	 * @param InW W component of the quaternion
 	 */
-	FORCEINLINE FQuat( float InX, float InY, float InZ, float InW );
+	FORCEINLINE FQuat(float InX, float InY, float InZ, float InW);
 
 	/**
 	 * Copy constructor.
 	 *
 	 * @param Q A FQuat object to use to create new quaternion from.
 	 */
-	FORCEINLINE FQuat( const FQuat& Q );
+	FORCEINLINE FQuat(const FQuat& Q);
 	
 	/**
 	 * Creates and initializes a new quaternion from the given matrix.
 	 *
 	 * @param M The rotation matrix to initialize from.
 	 */
-	explicit FQuat( const FMatrix& M );	
+	explicit FQuat(const FMatrix& M);	
 
 	/**
 	 * Creates and initializes a new quaternion from the given rotator.
 	 *
 	 * @param R The rotator to initialize from.
 	 */
-	explicit FQuat( const FRotator& R);
+	explicit FQuat(const FRotator& R);
 
 	/**
 	 * Creates and initializes a new quaternion from the a rotation around the given axis.
@@ -86,7 +84,7 @@ public:
 	 * @param Axis assumed to be a normalized vector
 	 * @param Angle angle to rotate above the given axis (in radians)
 	 */
-	FQuat( FVector Axis, float AngleRad );
+	FQuat(FVector Axis, float AngleRad);
 
 public:
 
@@ -106,7 +104,7 @@ public:
 	 * @param Q The Quaternion to add.
 	 * @return The result of addition.
 	 */
-	FORCEINLINE FQuat operator+( const FQuat& Q ) const;
+	FORCEINLINE FQuat operator+(const FQuat& Q) const;
 
 	/**
 	 * Adds to this quaternion.
@@ -124,7 +122,7 @@ public:
 	 * @param Q The Quaternion to subtract.
 	 * @return The result of subtraction.
 	 */
-	FORCEINLINE FQuat operator-( const FQuat& Q ) const;
+	FORCEINLINE FQuat operator-(const FQuat& Q) const;
 
 	/**
 	 * Checks whether another Quaternion is equal to this, within specified tolerance.
@@ -162,7 +160,7 @@ public:
 	 * @param Q The Quaternion to multiply this by.
 	 * @return The result of multiplication (this * Q).
 	 */
-	FORCEINLINE FQuat operator*( const FQuat& Q ) const;
+	FORCEINLINE FQuat operator*(const FQuat& Q) const;
 
 	/**
 	 * Multiply this by a quaternion (this = this * Q).
@@ -180,9 +178,9 @@ public:
 	 *
 	 * @param V the vector to be rotated
 	 * @return vector after rotation
-	 * @see RotateVector()
+	 * @see RotateVector
 	 */
-	FVector operator*( const FVector& V ) const;
+	FVector operator*(const FVector& V) const;
 
 	/** 
 	 * Multiply this by a matrix.
@@ -193,7 +191,7 @@ public:
 	 * @param M Matrix to multiply by.
 	 * @return Matrix result after multiplication.
 	 */	
-	FMatrix operator*( const FMatrix& M ) const;
+	FMatrix operator*(const FMatrix& M) const;
 	
 	/**
 	 * Multiply this quaternion by a scaling factor.
@@ -201,7 +199,7 @@ public:
 	 * @param Scale The scaling factor.
 	 * @return a reference to this after scaling.
 	 */
-	FORCEINLINE FQuat operator*=( const float Scale );
+	FORCEINLINE FQuat operator*=(const float Scale);
 
 	/**
 	 * Get the result of scaling this quaternion.
@@ -209,7 +207,7 @@ public:
 	 * @param Scale The scaling factor.
 	 * @return The result of scaling.
 	 */
-	FORCEINLINE FQuat operator*( const float Scale ) const;
+	FORCEINLINE FQuat operator*(const float Scale) const;
 	
 	/**
 	 * Divide this quaternion by scale.
@@ -217,7 +215,7 @@ public:
 	 * @param Scale What to divide by.
 	 * @return a reference to this after scaling.
 	 */
-	FORCEINLINE FQuat operator/=( const float Scale );
+	FORCEINLINE FQuat operator/=(const float Scale);
 
 	/**
 	 * Divide this quaternion by scale.
@@ -225,7 +223,7 @@ public:
 	 * @param Scale What to divide by.
 	 * @return new Quaternion of this after division by scale.
 	 */
-	FORCEINLINE FQuat operator/( const float Scale ) const;
+	FORCEINLINE FQuat operator/(const float Scale) const;
 
  	/**
 	 * Checks whether two quaternions are identical.
@@ -234,9 +232,9 @@ public:
 	 *
 	 * @param Q The other quaternion.
 	 * @return true if two quaternion are identical, otherwise false.
-	 * @see Equals()
+	 * @see Equals
 	 */
-	bool operator==( const FQuat& Q ) const;
+	bool operator==(const FQuat& Q) const;
 
  	/**
 	 * Checks whether two quaternions are not identical.
@@ -244,7 +242,7 @@ public:
 	 * @param Q The other quaternion.
 	 * @return true if two quaternion are not identical, otherwise false.
 	 */
-	bool operator!=( const FQuat& Q ) const;
+	bool operator!=(const FQuat& Q) const;
 
 	/**
 	 * Calculates dot product of two quaternions.
@@ -252,7 +250,7 @@ public:
 	 * @param Q The other quaternions.
 	 * @return The dot product.
 	 */
-	float operator|( const FQuat& Q ) const;
+	float operator|(const FQuat& Q) const;
 
 public:
 
@@ -262,7 +260,7 @@ public:
 	 * @param Euler the Euler angles
 	 * @return constructed FQuat
 	 */
-	static CORE_API FQuat MakeFromEuler( const FVector& Euler );
+	static CORE_API FQuat MakeFromEuler(const FVector& Euler);
 
 	/** Convert a Quaternion into floating-point Euler angles (in degrees). */
 	CORE_API FVector Euler() const;
@@ -273,7 +271,7 @@ public:
 	 *
 	 * @param Tolerance Minimum squared length of quaternion for normalization.
 	 */
-	FORCEINLINE void Normalize( float Tolerance=SMALL_NUMBER );
+	FORCEINLINE void Normalize(float Tolerance = SMALL_NUMBER);
 
 	/**
 	 * Get a normalized copy of this quaternion.
@@ -281,7 +279,7 @@ public:
 	 *
 	 * @param Tolerance Minimum squared length of quaternion for normalization.
 	 */
-	FORCEINLINE FQuat GetNormalized( float Tolerance=SMALL_NUMBER ) const;
+	FORCEINLINE FQuat GetNormalized(float Tolerance = SMALL_NUMBER) const;
 
 	// Return true if this quaternion is normalized
 	bool IsNormalized() const;
@@ -307,7 +305,7 @@ public:
 	 * @param Angle{out] angle of the quaternion
 	 * @warning : assumes normalized quaternions.
 	 */
-	void ToAxisAndAngle( FVector& Axis, float& Angle ) const;
+	void ToAxisAndAngle(FVector& Axis, float& Angle) const;
 
 	/**
 	 * Rotate a vector by this quaternion.
@@ -315,7 +313,7 @@ public:
 	 * @param V the vector to be rotated
 	 * @return vector after rotation
 	 */
-	FVector RotateVector( FVector V ) const;
+	FVector RotateVector(FVector V) const;
 	
 	/**
 	 * Rotate a vector by the inverse of this quaternion.
@@ -323,7 +321,7 @@ public:
 	 * @param V the vector to be rotated
 	 * @return vector after rotation by the inverse of this quaternion.
 	 */
-	FVector UnrotateVector( FVector V ) const;
+	FVector UnrotateVector(FVector V) const;
 
 	/**
 	 * @return quaternion with W=0 and V=theta*v.
@@ -346,7 +344,7 @@ public:
 	 * Enforce that the delta between this Quaternion and another one represents
 	 * the shortest possible rotation angle
 	 */
-	void EnforceShortestArcWith( const FQuat& OtherQuat );
+	void EnforceShortestArcWith(const FQuat& OtherQuat);
 	
 	/** Get the forward direction (X axis) after it has been rotated by this Quaternion. */
 	FORCEINLINE FVector GetAxisX() const;
@@ -384,7 +382,7 @@ public:
 	 * @param Ar Archive to serialize to/ from.
 	 * @return false to allow the ordinary struct code to run (this never happens).
 	 */
-	CORE_API bool NetSerialize( FArchive& Ar, class UPackageMap* Map, bool& bOutSuccess );
+	CORE_API bool NetSerialize(FArchive& Ar, class UPackageMap* Map, bool& bOutSuccess);
 
 	/**
 	 * Utility to check if there are any non-finite values (NaN or Inf) in this Quat.
@@ -430,7 +428,7 @@ public:
 	/**
 	 * Generates the 'smallest' (geodesic) rotation between two vectors of arbitrary length.
 	 */
-	static FORCEINLINE FQuat FindBetween( const FVector& Vector1, const FVector& Vector2 )
+	static FORCEINLINE FQuat FindBetween(const FVector& Vector1, const FVector& Vector2)
 	{
 		return FindBetweenVectors(Vector1, Vector2);
 	}
@@ -438,43 +436,43 @@ public:
 	/**
 	 * Generates the 'smallest' (geodesic) rotation between two normals (assumed to be unit length).
 	 */
-	static CORE_API FQuat FindBetweenNormals( const FVector& Normal1, const FVector& Normal2 );
+	static CORE_API FQuat FindBetweenNormals(const FVector& Normal1, const FVector& Normal2);
 
 	/**
 	 * Generates the 'smallest' (geodesic) rotation between two vectors of arbitrary length.
 	 */
-	static CORE_API FQuat FindBetweenVectors( const FVector& Vector1, const FVector& Vector2 );
+	static CORE_API FQuat FindBetweenVectors(const FVector& Vector1, const FVector& Vector2);
 
 	/**
 	 * Error measure (angle) between two quaternions, ranged [0..1].
 	 * Returns the hypersphere-angle between two quaternions; alignment shouldn't matter, though 
 	 * @note normalized input is expected.
 	 */
-	static FORCEINLINE float Error( const FQuat& Q1, const FQuat& Q2 );
+	static FORCEINLINE float Error(const FQuat& Q1, const FQuat& Q2);
 
 	/**
 	 * FQuat::Error with auto-normalization.
 	 */
-	static FORCEINLINE float ErrorAutoNormalize( const FQuat& A, const FQuat& B );
+	static FORCEINLINE float ErrorAutoNormalize(const FQuat& A, const FQuat& B);
 
 	/** 
 	 * Fast Linear Quaternion Interpolation.
 	 * Result is NOT normalized.
 	 */
-	static FORCEINLINE FQuat FastLerp( const FQuat& A, const FQuat& B, const float Alpha );
+	static FORCEINLINE FQuat FastLerp(const FQuat& A, const FQuat& B, const float Alpha);
 
 	/** 
 	 * Bi-Linear Quaternion interpolation.
 	 * Result is NOT normalized.
 	 */
-	static FORCEINLINE FQuat FastBilerp( const FQuat& P00, const FQuat& P10, const FQuat& P01, const FQuat& P11, float FracX, float FracY );
+	static FORCEINLINE FQuat FastBilerp(const FQuat& P00, const FQuat& P10, const FQuat& P01, const FQuat& P11, float FracX, float FracY);
 
 
 	/** Spherical interpolation. Will correct alignment. Result is NOT normalized. */
-	static CORE_API FQuat Slerp_NotNormalized( const FQuat &Quat1, const FQuat &Quat2, float Slerp );
+	static CORE_API FQuat Slerp_NotNormalized(const FQuat &Quat1, const FQuat &Quat2, float Slerp);
 
 	/** Spherical interpolation. Will correct alignment. Result is normalized. */
-	static FORCEINLINE FQuat Slerp( const FQuat &Quat1, const FQuat &Quat2, float Slerp )
+	static FORCEINLINE FQuat Slerp(const FQuat &Quat1, const FQuat &Quat2, float Slerp)
 	{
 		return Slerp_NotNormalized(Quat1, Quat2, Slerp).GetNormalized();
 	}
@@ -484,14 +482,14 @@ public:
 	 * We need this for the cubic interpolation stuff so that the multiple Slerps dont go in different directions.
 	 * Result is NOT normalized.
 	 */
-	static CORE_API FQuat SlerpFullPath_NotNormalized( const FQuat &quat1, const FQuat &quat2, float Alpha );
+	static CORE_API FQuat SlerpFullPath_NotNormalized(const FQuat &quat1, const FQuat &quat2, float Alpha);
 
 	/**
 	 * Simpler Slerp that doesn't do any checks for 'shortest distance' etc.
 	 * We need this for the cubic interpolation stuff so that the multiple Slerps dont go in different directions.
 	 * Result is normalized.
 	 */
-	static FORCEINLINE FQuat SlerpFullPath( const FQuat &quat1, const FQuat &quat2, float Alpha )
+	static FORCEINLINE FQuat SlerpFullPath(const FQuat &quat1, const FQuat &quat2, float Alpha)
 	{
 		return SlerpFullPath_NotNormalized(quat1, quat2, Alpha).GetNormalized();
 	}
@@ -500,7 +498,7 @@ public:
 	 * Given start and end quaternions of quat1 and quat2, and tangents at those points tang1 and tang2, calculate the point at Alpha (between 0 and 1) between them. Result is normalized.
 	 * This will correct alignment by ensuring that the shortest path is taken.
 	 */
-	static CORE_API FQuat Squad( const FQuat& quat1, const FQuat& tang1, const FQuat& quat2, const FQuat& tang2, float Alpha );
+	static CORE_API FQuat Squad(const FQuat& quat1, const FQuat& tang1, const FQuat& quat2, const FQuat& tang2, float Alpha);
 
 	/**
 	 * Simpler Squad that doesn't do any checks for 'shortest distance' etc.
@@ -517,7 +515,7 @@ public:
 	 * @param Tension @todo document
 	 * @param OutTan Out control point
 	 */
-	static CORE_API void CalcTangents( const FQuat& PrevP, const FQuat& P, const FQuat& NextP, float Tension, FQuat& OutTan );
+	static CORE_API void CalcTangents(const FQuat& PrevP, const FQuat& P, const FQuat& NextP, float Tension, FQuat& OutTan);
 
 public:
 
@@ -528,12 +526,12 @@ public:
 	 * @param F Reference to the quaternion being serialized.
 	 * @return Reference to the Archive after serialization.
 	 */
-	friend FArchive& operator<<( FArchive& Ar, FQuat& F )
+	friend FArchive& operator<<(FArchive& Ar, FQuat& F)
 	{
 		return Ar << F.X << F.Y << F.Z << F.W;
 	}
 
-	bool Serialize( FArchive& Ar )
+	bool Serialize(FArchive& Ar)
 	{
 		Ar << *this;
 		return true;
@@ -545,7 +543,7 @@ public:
 /* FQuat inline functions
  *****************************************************************************/
 
-inline FQuat::FQuat( const FMatrix& M )
+inline FQuat::FQuat(const FMatrix& M)
 {
 	// If Matrix is NULL, return Identity quaternion. If any of them is 0, you won't be able to construct rotation
 	// if you have two plane at least, we can reconstruct the frame using cross product, but that's a bit expensive op to do here
@@ -620,20 +618,20 @@ inline FQuat::FQuat( const FMatrix& M )
 }
 
 
-FORCEINLINE FQuat::FQuat( const FRotator& R )
+FORCEINLINE FQuat::FQuat(const FRotator& R)
 {
 	*this = R.Quaternion();
 	DiagnosticCheckNaN();
 }
 
 
-FORCEINLINE FVector FQuat::operator*( const FVector& V ) const
+FORCEINLINE FVector FQuat::operator*(const FVector& V) const
 {
 	return RotateVector(V);
 }
 
 
-inline FMatrix FQuat::operator*( const FMatrix& M ) const
+inline FMatrix FQuat::operator*(const FMatrix& M) const
 {
 	FMatrix Result;
 	FQuat VT, VR;
@@ -656,12 +654,12 @@ inline FMatrix FQuat::operator*( const FMatrix& M ) const
 /* FQuat inline functions
  *****************************************************************************/
 
-FORCEINLINE FQuat::FQuat( EForceInit ZeroOrNot )
+FORCEINLINE FQuat::FQuat(EForceInit ZeroOrNot)
 	:	X(0), Y(0), Z(0), W(ZeroOrNot == ForceInitToZero ? 0.0f : 1.0f)
 { }
 
 
-FORCEINLINE FQuat::FQuat( float InX, float InY, float InZ, float InW )
+FORCEINLINE FQuat::FQuat(float InX, float InY, float InZ, float InW)
 	: X(InX)
 	, Y(InY)
 	, Z(InZ)
@@ -671,7 +669,7 @@ FORCEINLINE FQuat::FQuat( float InX, float InY, float InZ, float InW )
 }
 
 
-FORCEINLINE FQuat::FQuat( const FQuat& Q )
+FORCEINLINE FQuat::FQuat(const FQuat& Q)
 	: X(Q.X)
 	, Y(Q.Y)
 	, Z(Q.Z)
@@ -686,7 +684,7 @@ FORCEINLINE FString FQuat::ToString() const
 
 
 #ifdef IMPLEMENT_ASSIGNMENT_OPERATOR_MANUALLY
-FORCEINLINE FQuat& FQuat::operator=( const FQuat& Other )
+FORCEINLINE FQuat& FQuat::operator=(const FQuat& Other)
 {
 	this->X = Other.X;
 	this->Y = Other.Y;
@@ -698,7 +696,7 @@ FORCEINLINE FQuat& FQuat::operator=( const FQuat& Other )
 #endif
 
 
-FORCEINLINE FQuat::FQuat( FVector Axis, float AngleRad )
+FORCEINLINE FQuat::FQuat(FVector Axis, float AngleRad)
 {
 	const float half_a = 0.5f * AngleRad;
 	float s, c;
@@ -713,9 +711,9 @@ FORCEINLINE FQuat::FQuat( FVector Axis, float AngleRad )
 }
 
 
-FORCEINLINE FQuat FQuat::operator+( const FQuat& Q ) const
+FORCEINLINE FQuat FQuat::operator+(const FQuat& Q) const
 {		
-	return FQuat( X + Q.X, Y + Q.Y, Z + Q.Z, W + Q.W );
+	return FQuat(X + Q.X, Y + Q.Y, Z + Q.Z, W + Q.W);
 }
 
 
@@ -732,9 +730,9 @@ FORCEINLINE FQuat FQuat::operator+=(const FQuat& Q)
 }
 
 
-FORCEINLINE FQuat FQuat::operator-( const FQuat& Q ) const
+FORCEINLINE FQuat FQuat::operator-(const FQuat& Q) const
 {
-	return FQuat( X - Q.X, Y - Q.Y, Z - Q.Z, W - Q.W );
+	return FQuat(X - Q.X, Y - Q.Y, Z - Q.Z, W - Q.W);
 }
 
 
@@ -772,7 +770,7 @@ FORCEINLINE FQuat FQuat::operator-=(const FQuat& Q)
 }
 
 
-FORCEINLINE FQuat FQuat::operator*( const FQuat& Q ) const
+FORCEINLINE FQuat FQuat::operator*(const FQuat& Q) const
 {
 	FQuat Result;
 	VectorQuaternionMultiply(&Result, this, &Q);
@@ -783,7 +781,7 @@ FORCEINLINE FQuat FQuat::operator*( const FQuat& Q ) const
 }
 
 
-FORCEINLINE FQuat FQuat::operator*=( const FQuat& Q )
+FORCEINLINE FQuat FQuat::operator*=(const FQuat& Q)
 {
 	/**
 	 * Now this uses VectorQuaternionMultiply that is optimized per platform.
@@ -800,7 +798,7 @@ FORCEINLINE FQuat FQuat::operator*=( const FQuat& Q )
 }
 
 
-FORCEINLINE FQuat FQuat::operator*=( const float Scale )
+FORCEINLINE FQuat FQuat::operator*=(const float Scale)
 {
 	X *= Scale;
 	Y *= Scale;
@@ -813,7 +811,7 @@ FORCEINLINE FQuat FQuat::operator*=( const float Scale )
 }
 
 
-FORCEINLINE FQuat FQuat::operator*( const float Scale ) const
+FORCEINLINE FQuat FQuat::operator*(const float Scale) const
 {
 	return FQuat(Scale * X, Scale * Y, Scale * Z, Scale * W);
 }
@@ -840,7 +838,7 @@ FORCEINLINE FQuat FQuat::operator/(const float Scale) const
 }
 
 
-FORCEINLINE bool FQuat::operator==( const FQuat& Q ) const
+FORCEINLINE bool FQuat::operator==(const FQuat& Q) const
 {
 #if PLATFORM_ENABLE_VECTORINTRINSICS
 	const VectorRegister A = VectorLoadAligned(this);
@@ -852,7 +850,7 @@ FORCEINLINE bool FQuat::operator==( const FQuat& Q ) const
 }
 
 
-FORCEINLINE bool FQuat::operator!=( const FQuat& Q ) const
+FORCEINLINE bool FQuat::operator!=(const FQuat& Q) const
 {
 #if PLATFORM_ENABLE_VECTORINTRINSICS
 	const VectorRegister A = VectorLoadAligned(this);
@@ -864,7 +862,7 @@ FORCEINLINE bool FQuat::operator!=( const FQuat& Q ) const
 }
 
 
-FORCEINLINE float FQuat::operator|( const FQuat& Q ) const
+FORCEINLINE float FQuat::operator|(const FQuat& Q) const
 {
 	return X * Q.X + Y * Q.Y + Z * Q.Z + W * Q.W;
 }
@@ -929,7 +927,7 @@ FORCEINLINE float FQuat::SizeSquared() const
 }
 
 
-FORCEINLINE void FQuat::ToAxisAndAngle( FVector& Axis, float& Angle ) const
+FORCEINLINE void FQuat::ToAxisAndAngle(FVector& Axis, float& Angle) const
 {
 	Angle = 2.f * FMath::Acos(W);
 	Axis = GetRotationAxis();
@@ -950,7 +948,7 @@ FORCEINLINE FVector FQuat::GetRotationAxis() const
 }
 
 
-FORCEINLINE FVector FQuat::RotateVector( FVector V ) const
+FORCEINLINE FVector FQuat::RotateVector(FVector V) const
 {	
 #if WITH_DIRECTXMATH
 	FVector Result;
@@ -973,7 +971,7 @@ FORCEINLINE FVector FQuat::RotateVector( FVector V ) const
 #endif
 }
 
-FORCEINLINE FVector FQuat::UnrotateVector( FVector V ) const
+FORCEINLINE FVector FQuat::UnrotateVector(FVector V) const
 {	
 #if WITH_DIRECTXMATH
 	FVector Result;
@@ -998,7 +996,7 @@ FORCEINLINE FQuat FQuat::Inverse() const
 }
 
 
-FORCEINLINE void FQuat::EnforceShortestArcWith( const FQuat& OtherQuat )
+FORCEINLINE void FQuat::EnforceShortestArcWith(const FQuat& OtherQuat)
 {
 	const float DotResult = (OtherQuat | *this);
 	const float Bias = FMath::FloatSelect(DotResult, 1.0f, -1.0f);
@@ -1049,14 +1047,14 @@ FORCEINLINE FVector FQuat::Vector() const
 }
 
 
-FORCEINLINE float FQuat::Error( const FQuat& Q1, const FQuat& Q2 )
+FORCEINLINE float FQuat::Error(const FQuat& Q1, const FQuat& Q2)
 {
 	const float cosom = FMath::Abs(Q1.X * Q2.X + Q1.Y * Q2.Y + Q1.Z * Q2.Z + Q1.W * Q2.W);
 	return (FMath::Abs(cosom) < 0.9999999f) ? FMath::Acos(cosom)*(1.f / PI) : 0.0f;
 }
 
 
-FORCEINLINE float FQuat::ErrorAutoNormalize( const FQuat& A, const FQuat& B )
+FORCEINLINE float FQuat::ErrorAutoNormalize(const FQuat& A, const FQuat& B)
 {
 	FQuat Q1 = A;
 	Q1.Normalize();
@@ -1068,7 +1066,7 @@ FORCEINLINE float FQuat::ErrorAutoNormalize( const FQuat& A, const FQuat& B )
 }
 
 
-FORCEINLINE FQuat FQuat::FastLerp( const FQuat& A, const FQuat& B, const float Alpha )
+FORCEINLINE FQuat FQuat::FastLerp(const FQuat& A, const FQuat& B, const float Alpha)
 {
 	// To ensure the 'shortest route', we make sure the dot product between the both rotations is positive.
 	const float DotResult = (A | B);
@@ -1077,7 +1075,7 @@ FORCEINLINE FQuat FQuat::FastLerp( const FQuat& A, const FQuat& B, const float A
 }
 
 
-FORCEINLINE FQuat FQuat::FastBilerp( const FQuat& P00, const FQuat& P10, const FQuat& P01, const FQuat& P11, float FracX, float FracY )
+FORCEINLINE FQuat FQuat::FastBilerp(const FQuat& P00, const FQuat& P10, const FQuat& P01, const FQuat& P11, float FracX, float FracY)
 {
 	return FQuat::FastLerp(
 		FQuat::FastLerp(P00,P10,FracX),
diff --git a/Engine/Source/Runtime/Core/Public/Math/Range.h b/Engine/Source/Runtime/Core/Public/Math/Range.h
--- a/Engine/Source/Runtime/Core/Public/Math/Range.h
+++ b/Engine/Source/Runtime/Core/Public/Math/Range.h
@@ -56,7 +56,7 @@ public:
 	 *
 	 * @param A The element in the range.
 	 */
-	explicit TRange( const ElementType& A )
+	explicit TRange(const ElementType& A)
 		: LowerBound(BoundsType::Inclusive(A))
 		, UpperBound(BoundsType::Inclusive(A))
 	{ }
@@ -69,7 +69,7 @@ public:
 	 * @param A The range's lower bound value (inclusive).
 	 * @param B The range's upper bound value (exclusive).
 	 */
-	explicit TRange( const ElementType& A, const ElementType& B )
+	explicit TRange(const ElementType& A, const ElementType& B)
 		: LowerBound(BoundsType::Inclusive(A))
 		, UpperBound(BoundsType::Exclusive(B))
 	{ }
@@ -80,7 +80,7 @@ public:
 	 * @param InLowerBound The range's lower bound.
 	 * @param InUpperBound The range's upper bound.
 	 */
-	explicit TRange( const BoundsType& InLowerBound, const BoundsType& InUpperBound )
+	explicit TRange(const BoundsType& InLowerBound, const BoundsType& InUpperBound)
 		: LowerBound(InLowerBound)
 		, UpperBound(InUpperBound)
 	{ }
@@ -93,7 +93,7 @@ public:
 	 * @param Other The range to compare with.
 	 * @return true if the ranges are equal, false otherwise.
 	 */
-	bool operator==( const TRange& Other ) const
+	bool operator==(const TRange& Other) const
 	{
 		if (IsEmpty() && Other.IsEmpty())
 		{
@@ -109,7 +109,7 @@ public:
 	 * @param Other The range to compare with.
 	 * @return true if the ranges are not equal, false otherwise.
 	 */
-	bool operator!=( const TRange& Other ) const
+	bool operator!=(const TRange& Other) const
 	{
 		if (IsEmpty() && Other.IsEmpty())
 		{
@@ -131,7 +131,7 @@ public:
 	 * @param Other The other range.
 	 * @return true if this range adjoins the other, false otherwise.
 	 */
-	bool Adjoins( const TRange& Other ) const
+	bool Adjoins(const TRange& Other) const
 	{
 		if (IsEmpty() || Other.IsEmpty())
 		{
@@ -163,7 +163,7 @@ public:
 	 * @param Y The second range.
 	 * @return true if this range conjoins the two ranges, false otherwise.
 	 */
-	bool Conjoins( const TRange& X, const TRange& Y ) const
+	bool Conjoins(const TRange& X, const TRange& Y) const
 	{
 		if (X.Overlaps(Y))
 		{
@@ -179,7 +179,7 @@ public:
 	 * @param Element The element to check.
 	 * @return true if the range contains the element, false otherwise.
 	 */
-	bool Contains( const ElementType& Element ) const
+	bool Contains(const ElementType& Element) const
 	{
 		return ((BoundsType::MinLower(LowerBound, Element) == LowerBound) &&
 				(BoundsType::MaxUpper(UpperBound, Element) == UpperBound));
@@ -191,7 +191,7 @@ public:
 	 * @param Other The range to check.
 	 * @return true if the range contains the other range, false otherwise.
 	 */
-	bool Contains( const TRange& Other ) const
+	bool Contains(const TRange& Other) const
 	{
 		return ((BoundsType::MinLower(LowerBound, Other.LowerBound) == LowerBound) &&
 				(BoundsType::MaxUpper(UpperBound, Other.UpperBound) == UpperBound));
@@ -205,7 +205,7 @@ public:
 	 * @param Other The other range.
 	 * @return true if the ranges are contiguous, false otherwise.
 	 */
-	bool Contiguous( const TRange& Other ) const
+	bool Contiguous(const TRange& Other) const
 	{
 		return (Overlaps(Other) || Adjoins(Other));
 	}
@@ -320,7 +320,7 @@ public:
 	 * @param Other The other range.
 	 * @return true if the ranges overlap, false otherwise.
 	 */
-	bool Overlaps( const TRange& Other ) const
+	bool Overlaps(const TRange& Other) const
 	{
 		if (IsEmpty() || Other.IsEmpty())
 		{
@@ -379,7 +379,7 @@ public:
 	 *
 	 * @param Element The element at which to split the range.
 	 */
-	TArray<TRange> Split( const ElementType& Element ) const
+	TArray<TRange> Split(const ElementType& Element) const
 	{
 		TArray<TRange> Result;
 		
@@ -406,7 +406,7 @@ public:
 	 * @return Between 0 and 2 remaining ranges.
 	 * @see Hull, Intersection, Union
 	 */
-	static FORCEINLINE TArray<TRange> Difference( const TRange& X, const TRange& Y )
+	static FORCEINLINE TArray<TRange> Difference(const TRange& X, const TRange& Y)
 	{
 		TArray<TRange> Result;
 
@@ -453,7 +453,7 @@ public:
 	 * @return The hull.
 	 * @see Difference, Intersection, Union
 	 */
-	static FORCEINLINE TRange Hull( const TRange& X, const TRange& Y )
+	static FORCEINLINE TRange Hull(const TRange& X, const TRange& Y)
 	{
 		if (X.IsEmpty())
 		{
@@ -475,7 +475,7 @@ public:
 	 * @return The hull.
 	 * @see Difference, Intersection, Union
 	 */
-	static FORCEINLINE TRange Hull( const TArray<TRange>& Ranges )
+	static FORCEINLINE TRange Hull(const TArray<TRange>& Ranges)
 	{
 		if (Ranges.Num() == 0)
 		{
@@ -502,7 +502,7 @@ public:
 	 * @return The intersection, or an empty range if the ranges do not overlap.
 	 * @see Difference, Hull, Union
 	 */
-	static FORCEINLINE TRange Intersection( const TRange& X, const TRange& Y )
+	static FORCEINLINE TRange Intersection(const TRange& X, const TRange& Y)
 	{
 		if (X.IsEmpty())
 		{
@@ -524,7 +524,7 @@ public:
 	 * @return The intersection.
 	 * @see Difference, Hull, Union
 	 */
-	static FORCEINLINE TRange Intersection( const TArray<TRange>& Ranges )
+	static FORCEINLINE TRange Intersection(const TArray<TRange>& Ranges)
 	{
 		if (Ranges.Num() == 0)
 		{
@@ -551,7 +551,7 @@ public:
 	 * @return The union, or both ranges if the two ranges are not contiguous, or no ranges if both ranges are empty.
 	 * @see Difference, Hull, Intersection
 	 */
-	static FORCEINLINE TArray<TRange> Union( const TRange& X, const TRange& Y )
+	static FORCEINLINE TArray<TRange> Union(const TRange& X, const TRange& Y)
 	{
 		TArray<TRange> OutRanges;
 
@@ -593,7 +593,7 @@ public:
 	 * @param Value The value.
 	 * @return A new range.
 	 */
-	static FORCEINLINE TRange AtLeast( const ElementType& Value )
+	static FORCEINLINE TRange AtLeast(const ElementType& Value)
 	{
 		return TRange(BoundsType::Inclusive(Value), BoundsType::Open());
 	}
@@ -604,7 +604,7 @@ public:
 	 * @param Value The value.
 	 * @return A new range.
 	 */
-	static FORCEINLINE TRange AtMost( const ElementType& Value )
+	static FORCEINLINE TRange AtMost(const ElementType& Value)
 	{
 		return TRange(BoundsType::Open(), BoundsType::Inclusive(Value));
 	}
@@ -615,7 +615,7 @@ public:
 	 * @param Value The value.
 	 * @return A new range.
 	 */
-	static FORCEINLINE TRange GreaterThan( const ElementType& Value )
+	static FORCEINLINE TRange GreaterThan(const ElementType& Value)
 	{
 		return TRange(BoundsType::Exclusive(Value), BoundsType::Open());
 	}
@@ -626,7 +626,7 @@ public:
 	 * @param Value The value.
 	 * @return A new range.
 	 */
-	static FORCEINLINE TRange LessThan( const ElementType& Value )
+	static FORCEINLINE TRange LessThan(const ElementType& Value)
 	{
 		return TRange(BoundsType::Open(), BoundsType::Exclusive(Value));
 	}
@@ -640,7 +640,7 @@ public:
 	 * @param Range The range to serialize.
 	 * @return The archive.
 	 */
-	friend class FArchive& operator<<( class FArchive& Ar, TRange& Range )
+	friend class FArchive& operator<<(class FArchive& Ar, TRange& Range)
 	{
 		return Ar << Range.LowerBound << Range.UpperBound;
 	}
@@ -651,17 +651,17 @@ public:
 	 * @param Range The range to get the hash for.
 	 * @return Hash value.
 	 */
-	friend uint32 GetTypeHash( const TRange& Range )
+	friend uint32 GetTypeHash(const TRange& Range)
 	{
 		return (GetTypeHash(Range.LowerBound) + 23 * GetTypeHash(Range.UpperBound));
 	}
 
 private:
 
-	// Holds the range's lower bound.
+	/** Holds the range's lower bound. */
 	BoundsType LowerBound;
 
-	// Holds the range's upper bound.
+	/** Holds the range's upper bound. */
 	BoundsType UpperBound;
 };
 
@@ -681,34 +681,34 @@ private:
 		{ \
 		} \
 		 \
-		Name( const Super& Rhs ) \
-			: Super( Rhs ) \
+		Name(const Super& Rhs) \
+			: Super(Rhs) \
 		{ \
 		} \
 		 \
-		explicit Name( const ElementType& A ) \
-			: Super( A ) \
+		explicit Name(const ElementType& A) \
+			: Super(A) \
 		{ \
 		} \
 		 \
-		explicit Name( const ElementType& A, const ElementType& B ) \
-			: Super( A, B ) \
+		explicit Name(const ElementType& A, const ElementType& B) \
+			: Super(A, B) \
 		{ \
 		} \
 		 \
-		explicit Name( const TRangeBound<ElementType>& InLowerBound, const TRangeBound<ElementType>& InUpperBound ) \
-			: Super( InLowerBound, InUpperBound ) \
+		explicit Name(const TRangeBound<ElementType>& InLowerBound, const TRangeBound<ElementType>& InUpperBound) \
+			: Super(InLowerBound, InUpperBound) \
 		{ \
 		} \
 		 \
-		TArray<Name> Split( const ElementType& Element ) const \
+		TArray<Name> Split(const ElementType& Element) const \
 		{ \
-			return TArray<Name>( Super::Split( Element ) ); \
+			return TArray<Name>(Super::Split(Element)); \
 		} \
 		 \
-		static FORCEINLINE TArray<Name> Difference( const Name& X, const Name& Y ) \
+		static FORCEINLINE TArray<Name> Difference(const Name& X, const Name& Y) \
 		{ \
-			return TArray<Name>( Super::Difference( X, Y ) ); \
+			return TArray<Name>(Super::Difference(X, Y)); \
 		} \
 		 \
 		static FORCEINLINE Name Empty() \
@@ -716,29 +716,29 @@ private:
 			return Super::Empty(); \
 		} \
 		 \
-		static FORCEINLINE Name Hull( const Name& X, const Name& Y ) \
+		static FORCEINLINE Name Hull(const Name& X, const Name& Y) \
 		{ \
-			return Super::Hull( X, Y ); \
+			return Super::Hull(X, Y); \
 		} \
 		 \
-		static FORCEINLINE Name Hull( const TArray<Name>& Ranges ) \
+		static FORCEINLINE Name Hull(const TArray<Name>& Ranges) \
 		{ \
-			return Super::Hull( reinterpret_cast<const TArray<Super>&>( Ranges ) ); \
+			return Super::Hull(reinterpret_cast<const TArray<Super>&>(Ranges)); \
 		} \
 		 \
-		static FORCEINLINE Name Intersection( const Name& X, const Name& Y ) \
+		static FORCEINLINE Name Intersection(const Name& X, const Name& Y) \
 		{ \
-			return Super::Intersection( X, Y ); \
+			return Super::Intersection(X, Y); \
 		} \
 		 \
-		static FORCEINLINE Name Intersection( const TArray<Name>& Ranges ) \
+		static FORCEINLINE Name Intersection(const TArray<Name>& Ranges) \
 		{ \
-			return Super::Intersection( reinterpret_cast<const TArray<Super>&>( Ranges ) ); \
+			return Super::Intersection(reinterpret_cast<const TArray<Super>&>(Ranges)); \
 		} \
 		 \
-		static FORCEINLINE TArray<Name> Union( const Name& X, const Name& Y ) \
+		static FORCEINLINE TArray<Name> Union(const Name& X, const Name& Y) \
 		{ \
-			return TArray<Name>( Super::Union( X, Y ) ); \
+			return TArray<Name>(Super::Union(X, Y)); \
 		} \
 		 \
 		static FORCEINLINE Name All() \
@@ -746,24 +746,24 @@ private:
 			return Super::All(); \
 		} \
 		 \
-		static FORCEINLINE Name AtLeast( const ElementType& Value ) \
+		static FORCEINLINE Name AtLeast(const ElementType& Value) \
 		{ \
-			return Super::AtLeast( Value ); \
+			return Super::AtLeast(Value); \
 		} \
 		 \
-		static FORCEINLINE Name AtMost( const ElementType& Value ) \
+		static FORCEINLINE Name AtMost(const ElementType& Value) \
 		{ \
-			return Super::AtMost( Value ); \
+			return Super::AtMost(Value); \
 		} \
 		 \
-		static FORCEINLINE TRange GreaterThan( const ElementType& Value ) \
+		static FORCEINLINE TRange GreaterThan(const ElementType& Value) \
 		{ \
-			return Super::GreaterThan( Value ); \
+			return Super::GreaterThan(Value); \
 		} \
 		 \
-		static FORCEINLINE TRange LessThan( const ElementType& Value ) \
+		static FORCEINLINE TRange LessThan(const ElementType& Value) \
 		{ \
-			return Super::LessThan( Value ); \
+			return Super::LessThan(Value); \
 		} \
 	}; \
 	 \
diff --git a/Engine/Source/Runtime/Core/Public/Math/RangeBound.h b/Engine/Source/Runtime/Core/Public/Math/RangeBound.h
--- a/Engine/Source/Runtime/Core/Public/Math/RangeBound.h
+++ b/Engine/Source/Runtime/Core/Public/Math/RangeBound.h
@@ -45,7 +45,7 @@ public:
 	 * @param InValue The bound's value.
 	 * @see Exclusive, Inclusive, Open
 	 */
-	TRangeBound( const ElementType& InValue )
+	TRangeBound(const ElementType& InValue)
 		: Type(ERangeBoundTypes::Inclusive)
 		, Value(InValue)
 	{ }
@@ -58,7 +58,7 @@ public:
 	 * @param Other The bound to compare with.
 	 * @return true if the bounds are equal, false otherwise.
 	 */
-	bool operator==( const TRangeBound& Other ) const
+	bool operator==(const TRangeBound& Other) const
 	{
 		return ((Type == Other.Type) && (IsOpen() || (Value == Other.Value)));
 	}
@@ -69,7 +69,7 @@ public:
 	 * @param Other The bound to compare with.
 	 * @return true if the bounds are not equal, false otherwise.
 	 */
-	bool operator!=( const TRangeBound& Other ) const
+	bool operator!=(const TRangeBound& Other) const
 	{
 		return ((Type != Other.Type) || (!IsOpen() && (Value != Other.Value)));
 	}
@@ -140,7 +140,7 @@ public:
 	 * @param Bound The bound to serialize.
 	 * @return The archive.
 	 */
-	friend class FArchive& operator<<( class FArchive& Ar, TRangeBound& Bound )
+	friend class FArchive& operator<<(class FArchive& Ar, TRangeBound& Bound)
 	{
 		return Ar << (uint8&)Bound.Type << Bound.Value;
 	}
@@ -151,7 +151,7 @@ public:
 	 * @param Bound The bound to get the hash for.
 	 * @return Hash value.
 	 */
-	friend uint32 GetTypeHash( const TRangeBound& Bound )
+	friend uint32 GetTypeHash(const TRangeBound& Bound)
 	{
 		return (GetTypeHash((uint8)Bound.Type) + 23 * GetTypeHash(Bound.Value));
 	}
@@ -164,7 +164,7 @@ public:
 	 * @param Value The bound value.
 	 * @return An exclusive closed bound.
 	 */
-	static FORCEINLINE TRangeBound Exclusive( const ElementType& Value )
+	static FORCEINLINE TRangeBound Exclusive(const ElementType& Value)
 	{
 		TRangeBound Result;
 
@@ -180,7 +180,7 @@ public:
 	 * @param Value The bound value.
 	 * @return An inclusive closed bound.
 	 */
-	static FORCEINLINE TRangeBound Inclusive( const ElementType& Value )
+	static FORCEINLINE TRangeBound Inclusive(const ElementType& Value)
 	{
 		TRangeBound Result;
 
@@ -213,7 +213,7 @@ public:
 	 *
 	 * @return A new bound.
 	 */
-	static FORCEINLINE TRangeBound FlipInclusion( const TRangeBound& Bound )
+	static FORCEINLINE TRangeBound FlipInclusion(const TRangeBound& Bound)
 	{
 		if (Bound.IsExclusive())
 		{
@@ -235,7 +235,7 @@ public:
 	 * @param B The second lower bound.
 	 * @return The greater lower bound.
 	 */
-	static FORCEINLINE const TRangeBound& MaxLower( const TRangeBound& A, const TRangeBound& B )
+	static FORCEINLINE const TRangeBound& MaxLower(const TRangeBound& A, const TRangeBound& B)
 	{
 		if (A.IsOpen()) { return B; }
 		if (B.IsOpen()) { return A; }
@@ -253,7 +253,7 @@ public:
 	 * @param B The second upper bound.
 	 * @return The greater upper bound.
 	 */
-	static FORCEINLINE const TRangeBound& MaxUpper( const TRangeBound& A, const TRangeBound& B )
+	static FORCEINLINE const TRangeBound& MaxUpper(const TRangeBound& A, const TRangeBound& B)
 	{
 		if (A.IsOpen()) { return A; }
 		if (B.IsOpen()) { return B; }
@@ -271,7 +271,7 @@ public:
 	 * @param B The second lower bound.
 	 * @return The lesser lower bound.
 	 */
-	static FORCEINLINE const TRangeBound& MinLower( const TRangeBound& A, const TRangeBound& B )
+	static FORCEINLINE const TRangeBound& MinLower(const TRangeBound& A, const TRangeBound& B)
 	{
 		if (A.IsOpen()) { return A; }
 		if (B.IsOpen()) { return B; }
@@ -289,7 +289,7 @@ public:
 	 * @param B The second upper bound.
 	 * @return The lesser upper bound.
 	 */
-	static FORCEINLINE const TRangeBound& MinUpper( const TRangeBound& A, const TRangeBound& B )
+	static FORCEINLINE const TRangeBound& MinUpper(const TRangeBound& A, const TRangeBound& B)
 	{
 		if (A.IsOpen()) { return B; }
 		if (B.IsOpen()) { return A; }
@@ -302,15 +302,15 @@ public:
 
 private:
 
-	// Holds the type of the bound.
+	/** Holds the type of the bound. */
 	TEnumAsByte<ERangeBoundTypes::Type> Type;
 
-	// Holds the bound's value.
+	/** Holds the bound's value. */
 	ElementType Value;
 };
 
 
-/* Default range bounds for built-in types
+/* Default range bounds for built-in types (for UProperty support)
  *****************************************************************************/
 
 #define DEFINE_RANGEBOUND_WRAPPER_STRUCT(Name, ElementType) \
@@ -322,57 +322,54 @@ private:
 	public: \
 		Name() \
 			: Super() \
-		{ \
-		} \
+		{ } \
 		 \
 		Name(const Super& Other) \
-			: Super( Other ) \
-		{ \
-		} \
+			: Super(Other) \
+		{ } \
 		 \
-		Name( const int64& InValue ) \
-			: Super( InValue ) \
-		{ \
-		} \
+		Name(const int64& InValue) \
+			: Super(InValue) \
+		{ } \
 		 \
-		static FORCEINLINE Name Exclusive( const ElementType& Value ) \
+		static FORCEINLINE Name Exclusive(const ElementType& Value) \
 		{ \
-			return static_cast<const Name&>( Super::Exclusive( Value ) ); \
+			return static_cast<const Name&>(Super::Exclusive(Value)); \
 		} \
 		 \
-		static FORCEINLINE Name Inclusive( const ElementType& Value ) \
+		static FORCEINLINE Name Inclusive(const ElementType& Value) \
 		{ \
-			return static_cast<const Name&>( Super::Inclusive( Value ) ); \
+			return static_cast<const Name&>(Super::Inclusive(Value)); \
 		} \
 		 \
 		static FORCEINLINE Name Open() \
 		{ \
-			return static_cast<const Name&>( Super::Open() ); \
+			return static_cast<const Name&>(Super::Open()); \
 		} \
 		 \
-		static FORCEINLINE Name FlipInclusion( const Name& Bound ) \
+		static FORCEINLINE Name FlipInclusion(const Name& Bound) \
 		{ \
-			return static_cast<const Name&>( Super::FlipInclusion( Bound ) ); \
+			return static_cast<const Name&>(Super::FlipInclusion(Bound)); \
 		} \
 		 \
-		static FORCEINLINE const Name& MaxLower( const Name& A, const Name& B ) \
+		static FORCEINLINE const Name& MaxLower(const Name& A, const Name& B) \
 		{ \
-			return static_cast<const Name&>( Super::MaxLower( A, B ) ); \
+			return static_cast<const Name&>(Super::MaxLower(A, B)); \
 		} \
 		 \
-		static FORCEINLINE const Name& MaxUpper( const Name& A, const Name& B ) \
+		static FORCEINLINE const Name& MaxUpper(const Name& A, const Name& B) \
 		{ \
-			return static_cast<const Name&>( Super::MaxUpper( A, B ) ); \
+			return static_cast<const Name&>(Super::MaxUpper(A, B)); \
 		} \
 		 \
-		static FORCEINLINE const Name& MinLower( const Name& A, const Name& B ) \
+		static FORCEINLINE const Name& MinLower(const Name& A, const Name& B) \
 		{ \
-			return static_cast<const Name&>( Super::MinLower( A, B ) ); \
+			return static_cast<const Name&>(Super::MinLower(A, B)); \
 		} \
 		 \
-		static FORCEINLINE const Name& MinUpper( const Name& A, const Name& B ) \
+		static FORCEINLINE const Name& MinUpper(const Name& A, const Name& B) \
 		{ \
-			return static_cast<const Name&>( Super::MinUpper( A, B ) ); \
+			return static_cast<const Name&>(Super::MinUpper(A, B)); \
 		} \
 	}; \
 	 \
@@ -388,6 +385,7 @@ private:
 		enum { Value = true }; \
 	};
 
+
 DEFINE_RANGEBOUND_WRAPPER_STRUCT(FDateRangeBound,   FDateTime)
 DEFINE_RANGEBOUND_WRAPPER_STRUCT(FDoubleRangeBound, double)
 DEFINE_RANGEBOUND_WRAPPER_STRUCT(FFloatRangeBound,  float)
diff --git a/Engine/Source/Runtime/Core/Public/Math/RangeSet.h b/Engine/Source/Runtime/Core/Public/Math/RangeSet.h
--- a/Engine/Source/Runtime/Core/Public/Math/RangeSet.h
+++ b/Engine/Source/Runtime/Core/Public/Math/RangeSet.h
@@ -30,7 +30,7 @@ public:
 	 *
 	 * @param Range The range to add.
 	 *//*
-	void Add( RangeType Range )
+	void Add(RangeType Range)
 	{
 		for (int32 Index = 0; Index < Ranges.Num(); ++Index)
 		{
@@ -54,7 +54,7 @@ public:
 	 *
 	 * @param Other The range set to merge.
 	 */
-/*	void Merge( const TRangeSet& Other )
+/*	void Merge(const TRangeSet& Other)
 	{
 		for (typename TArray<RangeType>::TConstIterator It(Other.Ranges); It; ++It)
 		{
@@ -69,7 +69,7 @@ public:
 	 *
 	 * @param Range The range to remove.
 	 */
-/*	void Remove( const RangeType& Range )
+/*	void Remove(const RangeType& Range)
 	{
 		for (int32 Index = 0; Index < Ranges.Num(); ++Index)
 		{
@@ -94,9 +94,7 @@ public:
 		}
 	}*/
 
-	/**
-	 * Removes all ranges from the set.
-	 */
+	/** Removes all ranges from the set. */
 	void Empty()
 	{
 		Ranges.Empty();
@@ -110,7 +108,7 @@ public:
 	 * @param Element The element to check.
 	 * @return true if the element is in the set, false otherwise.
 	 */
-	bool Contains( const ElementType& Element ) const
+	bool Contains(const ElementType& Element) const
 	{
 		for (typename TArray<RangeType>::TConstIterator It(Ranges); It; ++It)
 		{
@@ -129,7 +127,7 @@ public:
 	 * @param Range The range to check.
 	 * @return true if the set contains the range, false otherwise.
 	 */
-	bool Contains( const RangeType& Range ) const
+	bool Contains(const RangeType& Range) const
 	{
 		for (typename TArray<RangeType>::TConstIterator It(Ranges); It; ++It)
 		{
@@ -168,7 +166,7 @@ public:
 	 * @param Range The range to check.
 	 * @return true if this set overlaps with the range, false otherwise.
 	 */
-	bool Overlaps( const RangeType& Range ) const
+	bool Overlaps(const RangeType& Range) const
 	{
 		for (typename TArray<RangeType>::TConstIterator It(Ranges); It; ++It)
 		{
@@ -189,7 +187,7 @@ public:
 	 *
 	 * @todo gmp: This could be optimized to O(n*logn) using a line sweep on a pre-sorted array of bounds.
 	 */
-	bool Overlaps( const TRangeSet& Other ) const
+	bool Overlaps(const TRangeSet& Other) const
 	{
 		for (typename TArray<RangeType>::TConstIterator It(Other.Ranges); It; ++It)
 		{
@@ -211,13 +209,13 @@ public:
 	 * @param RangeSet The range set to serialize.
 	 * @return The archive.
 	 */
-	friend class FArchive& operator<<( class FArchive& Ar, TRangeSet& RangeSet )
+	friend class FArchive& operator<<(class FArchive& Ar, TRangeSet& RangeSet)
 	{
 		return Ar << RangeSet.Ranges;
 	}
 
 private:
 
-	// Holds the set of ranges.
+	/** Holds the set of ranges. */
 	TArray<RangeType> Ranges;
 };
diff --git a/Engine/Source/Runtime/Core/Public/Math/Sphere.h b/Engine/Source/Runtime/Core/Public/Math/Sphere.h
--- a/Engine/Source/Runtime/Core/Public/Math/Sphere.h
+++ b/Engine/Source/Runtime/Core/Public/Math/Sphere.h
@@ -10,10 +10,10 @@ class FSphere
 {
 public:
 
-	/** Holds the sphere's center. */
+	/** The sphere's center point. */
 	FVector Center;
 
-	/** Holds the sphere's radius. */
+	/** The sphere's radius. */
 	float W;
 
 public:
@@ -26,7 +26,7 @@ public:
 	 *
 	 * @param int32 Passing int32 sets up zeroed sphere.
 	 */
-	FSphere( int32 )
+	FSphere(int32)
 		: Center(0.0f, 0.0f, 0.0f)
 		, W(0)
 	{ }
@@ -37,7 +37,7 @@ public:
 	 * @param InV Center of sphere.
 	 * @param InW Radius of sphere.
 	 */
-	FSphere( FVector InV, float InW )
+	FSphere(FVector InV, float InW)
 		: Center(InV)
 		, W(InW)
 	{ }
@@ -47,7 +47,7 @@ public:
 	 *
 	 * @param EForceInit Force Init Enum.
 	 */
-	explicit FORCEINLINE FSphere( EForceInit )
+	explicit FORCEINLINE FSphere(EForceInit)
 		: Center(ForceInit)
 		, W(0.0f)
 	{ }
@@ -58,7 +58,7 @@ public:
 	 * @param Pts Pointer to list of points this sphere must contain.
 	 * @param Count How many points are in the list.
 	 */
-	CORE_API FSphere( const FVector* Pts, int32 Count );
+	CORE_API FSphere(const FVector* Pts, int32 Count);
 
 public:
 
@@ -69,7 +69,7 @@ public:
 	 * @param Tolerance Error Tolerance.
 	 * @return true if spheres are equal within specified tolerance, otherwise false.
 	 */
-	bool Equals( const FSphere& Sphere, float Tolerance = KINDA_SMALL_NUMBER ) const
+	bool Equals(const FSphere& Sphere, float Tolerance = KINDA_SMALL_NUMBER) const
 	{
 		return Center.Equals(Sphere.Center, Tolerance) && FMath::Abs(W - Sphere.W) <= Tolerance;
 	}
@@ -81,7 +81,7 @@ public:
 	 * @param Tolerance Error Tolerance.
 	 * @return true if sphere is inside another, otherwise false.
 	 */
-	bool IsInside( const FSphere& Other, float Tolerance = KINDA_SMALL_NUMBER ) const
+	bool IsInside(const FSphere& Other, float Tolerance = KINDA_SMALL_NUMBER) const
 	{
 		if (W > Other.W + Tolerance)
 		{
@@ -91,7 +91,6 @@ public:
 		return (Center - Other.Center).SizeSquared() <= FMath::Square(Other.W + Tolerance - W);
 	}
 
-
 	/**
 	* Checks whether the given location is inside this sphere.
 	*
@@ -103,7 +102,6 @@ public:
 		return (Center - In).SizeSquared() <= FMath::Square(W + Tolerance);
 	}
 
-	
 	/**
 	 * Test whether this sphere intersects another.
 	 * 
@@ -111,7 +109,7 @@ public:
 	 * @param  Tolerance Error tolerance.
 	 * @return true if spheres intersect, false otherwise.
 	 */
-	FORCEINLINE bool Intersects( const FSphere& Other, float Tolerance = KINDA_SMALL_NUMBER ) const
+	FORCEINLINE bool Intersects(const FSphere& Other, float Tolerance = KINDA_SMALL_NUMBER) const
 	{
 		return (Center - Other.Center).SizeSquared() <= FMath::Square(FMath::Max(0.f, Other.W + W + Tolerance));
 	}
@@ -122,7 +120,7 @@ public:
 	 * @param M Matrix to transform by.
 	 * @return Result of transformation.
 	 */
-	CORE_API FSphere TransformBy( const FMatrix& M ) const;
+	CORE_API FSphere TransformBy(const FMatrix& M) const;
 
 	/**
 	 * Get result of Transforming sphere with Transform.
@@ -130,11 +128,12 @@ public:
 	 * @param M Transform information.
 	 * @return Result of transformation.
 	 */
-	CORE_API FSphere TransformBy( const FTransform& M ) const;
+	CORE_API FSphere TransformBy(const FTransform& M) const;
 
 	/**
 	 * Get volume of the current sphere
-	 * 
+	 *
+	 * @return Volume (in Unreal units).
 	 */
 	CORE_API float GetVolume() const;
 
@@ -144,7 +143,7 @@ public:
 	 * @param Other the bounding volume to increase the bounding volume to.
 	 * @return Reference to this bounding volume after resizing to include the other bounding volume.
 	 */
-	CORE_API FSphere& operator+=( const FSphere& Other );
+	CORE_API FSphere& operator+=(const FSphere& Other);
 
 	/**
 	 * Gets the result of addition to this bounding volume.
@@ -152,7 +151,7 @@ public:
 	 * @param Other The other volume to add to this.
 	 * @return A new bounding volume.
 	 */
-	FSphere operator+( const FSphere& Other ) const
+	FSphere operator+(const FSphere& Other) const
 	{
 		return FSphere(*this) += Other;
 	}
@@ -166,28 +165,10 @@ public:
 	 * @param Sphere The sphere to serialize.
 	 * @return The archive.
 	 */
-	friend FArchive& operator<<( FArchive& Ar, FSphere& Sphere )
+	friend FArchive& operator<<(FArchive& Ar, FSphere& Sphere)
 	{
 		Ar << Sphere.Center << Sphere.W;
 
 		return Ar;
 	}
 };
-
-struct FCapsuleShape
-{
-	FVector Center;
-	float Radius;
-	FVector Orientation;
-	float Length;
-
-	FCapsuleShape()
-	{}
-
-	FCapsuleShape(FVector InCenter, float InRadius, FVector InOrientation, float InLength) :
-		Center(InCenter),
-		Radius(InRadius),
-		Orientation(InOrientation),
-		Length(InLength)
-	{}
-};
\ No newline at end of file
diff --git a/Engine/Source/Runtime/Core/Public/Math/Transform.h b/Engine/Source/Runtime/Core/Public/Math/Transform.h
--- a/Engine/Source/Runtime/Core/Public/Math/Transform.h
+++ b/Engine/Source/Runtime/Core/Public/Math/Transform.h
@@ -18,7 +18,6 @@
  * Example: LocalToWorld = (DeltaRotation * LocalToWorld) will change rotation in local space by DeltaRotation.
  * Example: LocalToWorld = (LocalToWorld * DeltaRotation) will change rotation in world space by DeltaRotation.
  */
-
 struct FTransform
 {
 #if !defined(COREUOBJECT_API)
@@ -93,9 +92,7 @@ public:
 	FORCEINLINE void DiagnosticCheck_IsValid() const {}
 #endif
 
-	/**
-	 * Constructor with initialization to the identity transform.
-	 */
+	/** Default constructor. */
 	FORCEINLINE FTransform()
 		: Rotation(0.f,0.f,0.f,1.f)
 		, Translation(0.f)
@@ -231,12 +228,12 @@ public:
 	CORE_API FString ToString() const;
 
 	/** Acceptable form: "%f,%f,%f|%f,%f,%f|%f,%f,%f" */
-	CORE_API bool InitFromString( const FString& InSourceString );
+	CORE_API bool InitFromString(const FString& InSourceString);
 
 #ifdef IMPLEMENT_ASSIGNMENT_OPERATOR_MANUALLY
 	/**
-	* Copy another Transform into this one
-	*/
+	 * Copy another Transform into this one
+	 */
 	FORCEINLINE FTransform& operator=(const FTransform& Other)
 	{
 		this->Rotation = Other.Rotation;
@@ -248,15 +245,15 @@ public:
 #endif
 
 	/**
-	* Convert this Transform to a transformation matrix with scaling.
-	*/
+	 * Convert this Transform to a transformation matrix with scaling.
+	 */
 	FORCEINLINE FMatrix ToMatrixWithScale() const
 	{
 		FMatrix OutMatrix;
 
 #if !(UE_BUILD_SHIPPING || UE_BUILD_TEST) && WITH_EDITORONLY_DATA
 		// Make sure Rotation is normalized when we turn it into a matrix.
-		check( IsRotationNormalized() );
+		check(IsRotationNormalized());
 #endif
 		OutMatrix.M[3][0] = Translation.X;
 		OutMatrix.M[3][1] = Translation.Y;
@@ -305,8 +302,8 @@ public:
 	}
 
 	/**
-	* Convert this Transform to matrix with scaling and compute the inverse of that.
-	*/
+	 * Convert this Transform to matrix with scaling and compute the inverse of that.
+	 */
 	FORCEINLINE FMatrix ToInverseMatrixWithScale() const
 	{
 		// todo: optimize
@@ -314,8 +311,8 @@ public:
 	}
 
 	/**
-	* Convert this Transform to inverse.
-	*/
+	 * Convert this Transform to inverse.
+	 */
 	FORCEINLINE FTransform Inverse() const
 	{
 		FQuat   InvRotation    = Rotation.Inverse();
@@ -327,15 +324,15 @@ public:
 	}
 
 	/**
-	* Convert this Transform to a transformation matrix, ignoring its scaling
-	*/
+	 * Convert this Transform to a transformation matrix, ignoring its scaling
+	 */
 	FORCEINLINE FMatrix ToMatrixNoScale() const
 	{
 		FMatrix OutMatrix;
 
 #if !(UE_BUILD_SHIPPING || UE_BUILD_TEST) && WITH_EDITORONLY_DATA
 		// Make sure Rotation is normalized when we turn it into a matrix.
-		check( IsRotationNormalized() );
+		check(IsRotationNormalized());
 #endif
 		OutMatrix.M[3][0] = Translation.X;
 		OutMatrix.M[3][1] = Translation.Y;
@@ -388,15 +385,15 @@ public:
 	{
 #if !(UE_BUILD_SHIPPING || UE_BUILD_TEST) && WITH_EDITORONLY_DATA
 		// Check that all bone atoms coming from animation are normalized
-		check( Atom1.IsRotationNormalized() );
-		check( Atom2.IsRotationNormalized() );
+		check(Atom1.IsRotationNormalized());
+		check(Atom2.IsRotationNormalized());
 #endif
-		if( Alpha <= ZERO_ANIMWEIGHT_THRESH )
+		if(Alpha <= ZERO_ANIMWEIGHT_THRESH)
 		{
 			// if blend is all the way for child1, then just copy its bone atoms
 			(*this) = Atom1;
 		}
-		else if( Alpha >= 1.f - ZERO_ANIMWEIGHT_THRESH )
+		else if(Alpha >= 1.f - ZERO_ANIMWEIGHT_THRESH)
 		{
 			// if blend is all the way for child2, then just copy its bone atoms
 			(*this) = Atom2;
@@ -418,12 +415,12 @@ public:
 	{
 #if !(UE_BUILD_SHIPPING || UE_BUILD_TEST) && WITH_EDITORONLY_DATA
 		// Check that all bone atoms coming from animation are normalized
-		check( IsRotationNormalized() );
-		check( OtherAtom.IsRotationNormalized() );
+		check(IsRotationNormalized());
+		check(OtherAtom.IsRotationNormalized());
 #endif
-		if( Alpha > ZERO_ANIMWEIGHT_THRESH )
+		if(Alpha > ZERO_ANIMWEIGHT_THRESH)
 		{
-			if( Alpha >= 1.f - ZERO_ANIMWEIGHT_THRESH )
+			if(Alpha >= 1.f - ZERO_ANIMWEIGHT_THRESH)
 			{
 				// if blend is all the way for child2, then just copy its bone atoms
 				(*this) = OtherAtom;
@@ -504,20 +501,20 @@ public:
 	FORCEINLINE void operator*=(const FTransform& Other);
 	
 	/**
-	* Return a transform that is the result of this multiplied by another transform (made only from a rotation).
-	* Order matters when composing transforms : C = A * B will yield a transform C that logically first applies A then B to any subsequent transformation.
-	*
-	* @param  Other other quaternion rotation by which to multiply.
-	* @return new transform: this * FTransform(Other)
-	*/
+	 * Return a transform that is the result of this multiplied by another transform (made only from a rotation).
+	 * Order matters when composing transforms : C = A * B will yield a transform C that logically first applies A then B to any subsequent transformation.
+	 *
+	 * @param  Other other quaternion rotation by which to multiply.
+	 * @return new transform: this * FTransform(Other)
+	 */
 	FORCEINLINE FTransform operator*(const FQuat& Other) const;
 	
 	/**
-	* Sets this transform to the result of this multiplied by another transform (made only from a rotation).
-	* Order matters when composing transforms : C = A * B will yield a transform C that logically first applies A then B to any subsequent transformation.
-	*
-	* @param  Other other quaternion rotation by which to multiply.
-	*/
+	 * Sets this transform to the result of this multiplied by another transform (made only from a rotation).
+	 * Order matters when composing transforms : C = A * B will yield a transform C that logically first applies A then B to any subsequent transformation.
+	 *
+	 * @param  Other other quaternion rotation by which to multiply.
+	 */
 	FORCEINLINE void operator*=(const FQuat& Other);
 
 	FORCEINLINE void ScaleTranslation(const FVector& InScale3D);
@@ -525,6 +522,7 @@ public:
 	FORCEINLINE void RemoveScaling(float Tolerance=SMALL_NUMBER);
 	FORCEINLINE float GetMaximumAxisScale() const;
 	FORCEINLINE float GetMinimumAxisScale() const;
+
 	// Inverse does not work well with VQS format(in particular non-uniform), so removing it, but made two below functions to be used instead. 
 
 	/*******************************************************************************************
@@ -540,37 +538,32 @@ public:
 	 * Set current transform and the relative to ParentTransform.
 	 * Equates to This = This->GetRelativeTransform(Parent), but saves the intermediate FTransform storage and copy.
 	 */
-	CORE_API void		SetToRelativeTransform(const FTransform& ParentTransform);
-
-	FORCEINLINE FVector4	TransformFVector4(const FVector4& V) const;
-	FORCEINLINE FVector4	TransformFVector4NoScale(const FVector4& V) const;
-	FORCEINLINE FVector		TransformPosition(const FVector& V) const;
-	FORCEINLINE FVector		TransformPositionNoScale(const FVector& V) const;
+	CORE_API void SetToRelativeTransform(const FTransform& ParentTransform);
 
+	FORCEINLINE FVector4 TransformFVector4(const FVector4& V) const;
+	FORCEINLINE FVector4 TransformFVector4NoScale(const FVector4& V) const;
+	FORCEINLINE FVector TransformPosition(const FVector& V) const;
+	FORCEINLINE FVector TransformPositionNoScale(const FVector& V) const;
 
 	/** Inverts the matrix and then transforms V - correctly handles scaling in this matrix. */
-	FORCEINLINE FVector		InverseTransformPosition(const FVector &V) const;
-
-	FORCEINLINE FVector		InverseTransformPositionNoScale(const FVector &V) const;
-
-	FORCEINLINE FVector		TransformVector(const FVector& V) const;
-
-	FORCEINLINE FVector		TransformVectorNoScale(const FVector& V) const;
+	FORCEINLINE FVector InverseTransformPosition(const FVector &V) const;
+	FORCEINLINE FVector InverseTransformPositionNoScale(const FVector &V) const;
+	FORCEINLINE FVector TransformVector(const FVector& V) const;
+	FORCEINLINE FVector TransformVectorNoScale(const FVector& V) const;
 
 	/** 
 	 *	Transform a direction vector by the inverse of this matrix - will not take into account translation part.
 	 *	If you want to transform a surface normal (or plane) and correctly account for non-uniform scaling you should use TransformByUsingAdjointT with adjoint of matrix inverse.
 	 */
 	FORCEINLINE FVector InverseTransformVector(const FVector &V) const;
-
 	FORCEINLINE FVector InverseTransformVectorNoScale(const FVector &V) const;
 
-	FORCEINLINE FTransform	GetScaled(float Scale) const;
-	FORCEINLINE FTransform	GetScaled(FVector Scale) const;
-	FORCEINLINE FVector		GetScaledAxis(EAxis::Type InAxis) const;
-	FORCEINLINE FVector		GetUnitAxis(EAxis::Type InAxis) const;
-	FORCEINLINE void		Mirror(EAxis::Type MirrorAxis, EAxis::Type FlipAxis);
-	FORCEINLINE FVector		GetSafeScaleReciprocal(const FVector& InScale, float Tolerance=SMALL_NUMBER) const;
+	FORCEINLINE FTransform GetScaled(float Scale) const;
+	FORCEINLINE FTransform GetScaled(FVector Scale) const;
+	FORCEINLINE FVector GetScaledAxis(EAxis::Type InAxis) const;
+	FORCEINLINE FVector GetUnitAxis(EAxis::Type InAxis) const;
+	FORCEINLINE void Mirror(EAxis::Type MirrorAxis, EAxis::Type FlipAxis);
+	FORCEINLINE FVector GetSafeScaleReciprocal(const FVector& InScale, float Tolerance=SMALL_NUMBER) const;
 
 	// temp function for easy conversion
 	FORCEINLINE FVector GetLocation() const
@@ -607,12 +600,12 @@ public:
 
 	inline bool IsValid() const
 	{
-		if ( ContainsNaN() )
+		if (ContainsNaN())
 		{
 			return false;
 		}
 
-		if ( !Rotation.IsNormalized() )
+		if (!Rotation.IsNormalized())
 		{
 			return false;
 		}
@@ -631,11 +624,11 @@ public:
 
 	// Binary comparison operators.
 	/*
-	bool operator==( const FTransform& Other ) const
+	bool operator==(const FTransform& Other) const
 	{
 		return Rotation==Other.Rotation && Translation==Other.Translation && Scale3D==Other.Scale3D;
 	}
-	bool operator!=( const FTransform& Other ) const
+	bool operator!=(const FTransform& Other) const
 	{
 		return Rotation!=Other.Rotation || Translation!=Other.Translation || Scale3D!=Other.Scale3D;
 	}
@@ -887,7 +880,7 @@ public:
 		FTransform Atom(DeltaAtom * Weight);
 
 		// To ensure the 'shortest route', we make sure the dot product between the accumulator and the incoming child atom is positive.
-		if( (Atom.Rotation | Rotation) < 0.f )
+		if((Atom.Rotation | Rotation) < 0.f)
 		{
 			Rotation.X -= Atom.Rotation.X;
 			Rotation.Y -= Atom.Rotation.Y;
@@ -920,7 +913,7 @@ public:
 	FORCEINLINE void Accumulate(const FTransform& SourceAtom)
 	{
 		// Add ref pose relative animation to base animation, only if rotation is significant.
-		if( FMath::Square(SourceAtom.Rotation.W) < 1.f - DELTA * DELTA )
+		if(FMath::Square(SourceAtom.Rotation.W) < 1.f - DELTA * DELTA)
 		{
 			Rotation = SourceAtom.Rotation * Rotation;
 		}
@@ -930,7 +923,7 @@ public:
 
 		DiagnosticCheckNaN_All();
 
-		checkSlow( IsRotationNormalized() );
+		checkSlow(IsRotationNormalized());
 	}
 
    /** Accumulates another transform with this one, with a blending weight
@@ -950,7 +943,7 @@ public:
 		FTransform SourceAtom(Atom * BlendWeight);
 
 		// Add ref pose relative animation to base animation, only if rotation is significant.
-		if( FMath::Square(SourceAtom.Rotation.W) < 1.f - DELTA * DELTA )
+		if(FMath::Square(SourceAtom.Rotation.W) < 1.f - DELTA * DELTA)
 		{
 			Rotation = SourceAtom.Rotation * Rotation;
 		}
@@ -991,7 +984,7 @@ public:
 	 */
 	FORCEINLINE void AccumulateWithAdditiveScale3D(const FTransform& SourceAtom)
 	{
-		if( FMath::Square(SourceAtom.Rotation.W) < 1.f - DELTA * DELTA )
+		if(FMath::Square(SourceAtom.Rotation.W) < 1.f - DELTA * DELTA)
 		{
 			Rotation = SourceAtom.Rotation * Rotation;
 		}
@@ -1038,13 +1031,13 @@ public:
 		const FTransform IdentityAtom = FTransform::Identity;
 
 		// Scale delta by weight
-		if( BlendWeight < (1.f - ZERO_ANIMWEIGHT_THRESH) )
+		if(BlendWeight < (1.f - ZERO_ANIMWEIGHT_THRESH))
 		{
 			SourceAtom.Blend(IdentityAtom, SourceAtom, BlendWeight);
 		}
 
 		// Add ref pose relative animation to base animation, only if rotation is significant.
-		if( FMath::Square(SourceAtom.Rotation.W) < 1.f - DELTA * DELTA )
+		if(FMath::Square(SourceAtom.Rotation.W) < 1.f - DELTA * DELTA)
 		{
 			FinalAtom.Rotation = SourceAtom.Rotation * FinalAtom.Rotation;
 		}
@@ -1054,7 +1047,7 @@ public:
 
 		FinalAtom.DiagnosticCheckNaN_All();
 
-		checkSlow( FinalAtom.IsRotationNormalized() );
+		checkSlow(FinalAtom.IsRotationNormalized());
 	}
 
 	/**
@@ -1170,7 +1163,7 @@ public:
 			// Assume it is along X and modify transform accordingly. 
 			// It doesn't actually matter which axis we choose, the 'appearance' will be the same
 			Scale3D.X *= -1.f;
-			M.SetAxis(0, -M.GetScaledAxis( EAxis::X ));
+			M.SetAxis(0, -M.GetScaledAxis(EAxis::X));
 		}
 
 		Rotation = FQuat(M);
@@ -1181,8 +1174,10 @@ public:
 	}
 };
 
+
 template <> struct TIsPODType<FTransform> { enum { Value = true }; };
 
+
 /** Scale the translation part of the Transform by the supplied vector. */
 FORCEINLINE void FTransform::ScaleTranslation(const FVector& InScale3D)
 {
@@ -1191,6 +1186,7 @@ FORCEINLINE void FTransform::ScaleTranslation(const FVector& InScale3D)
 	DiagnosticCheckNaN_Translate();
 }
 
+
 FORCEINLINE void FTransform::ScaleTranslation(const float& Scale)
 {
 	Translation *= Scale;
@@ -1198,6 +1194,7 @@ FORCEINLINE void FTransform::ScaleTranslation(const float& Scale)
 	DiagnosticCheckNaN_Translate();
 }
 
+
 // this function is from matrix, and all it does is to normalize rotation portion
 FORCEINLINE void FTransform::RemoveScaling(float Tolerance/*=SMALL_NUMBER*/)
 {
@@ -1295,6 +1292,7 @@ FORCEINLINE FTransform FTransform::GetScaled(float InScale) const
 	return A;
 }
 
+
 /** 
  * Apply Scale to this transform
  */
@@ -1308,6 +1306,7 @@ FORCEINLINE FTransform FTransform::GetScaled(FVector InScale) const
 	return A;
 }
 
+
 /** Transform homogenous FVector4, ignoring the scaling part of this transform **/
 FORCEINLINE FVector4 FTransform::TransformFVector4NoScale(const FVector4& V) const
 {
@@ -1327,6 +1326,7 @@ FORCEINLINE FVector4 FTransform::TransformFVector4NoScale(const FVector4& V) con
 	return Transform;
 }
 
+
 /** Transform FVector4 **/
 FORCEINLINE FVector4 FTransform::TransformFVector4(const FVector4& V) const
 {
@@ -1348,43 +1348,47 @@ FORCEINLINE FVector4 FTransform::TransformFVector4(const FVector4& V) const
 }
 
 
-
 FORCEINLINE FVector FTransform::TransformPosition(const FVector& V) const
 {
 	DiagnosticCheckNaN_All();
 	return Rotation.RotateVector(Scale3D*V) + Translation;
 }
 
+
 FORCEINLINE FVector FTransform::TransformPositionNoScale(const FVector& V) const
 {
 	DiagnosticCheckNaN_All();
 	return Rotation.RotateVector(V) + Translation;
 }
 
+
 FORCEINLINE FVector FTransform::TransformVector(const FVector& V) const
 {
 	DiagnosticCheckNaN_All();
 	return Rotation.RotateVector(Scale3D*V);
 }
 
+
 FORCEINLINE FVector FTransform::TransformVectorNoScale(const FVector& V) const
 {
 	DiagnosticCheckNaN_All();
 	return Rotation.RotateVector(V);
 }
 
+
 // do backward operation when inverse, translation -> rotation -> scale
 FORCEINLINE FVector FTransform::InverseTransformPosition(const FVector &V) const
 {
 	DiagnosticCheckNaN_All();
-	return ( Rotation.UnrotateVector(V-Translation) ) * GetSafeScaleReciprocal(Scale3D);
+	return (Rotation.UnrotateVector(V-Translation)) * GetSafeScaleReciprocal(Scale3D);
 }
 
+
 // do backward operation when inverse, translation -> rotation
 FORCEINLINE FVector FTransform::InverseTransformPositionNoScale(const FVector &V) const
 {
 	DiagnosticCheckNaN_All();
-	return ( Rotation.UnrotateVector(V-Translation) );
+	return (Rotation.UnrotateVector(V-Translation));
 }
 
 
@@ -1392,16 +1396,18 @@ FORCEINLINE FVector FTransform::InverseTransformPositionNoScale(const FVector &V
 FORCEINLINE FVector FTransform::InverseTransformVector(const FVector &V) const
 {
 	DiagnosticCheckNaN_All();
-	return ( Rotation.UnrotateVector(V) ) * GetSafeScaleReciprocal(Scale3D);
+	return (Rotation.UnrotateVector(V)) * GetSafeScaleReciprocal(Scale3D);
 }
 
+
 // do backward operation when inverse, translation -> rotation
 FORCEINLINE FVector FTransform::InverseTransformVectorNoScale(const FVector &V) const
 {
 	DiagnosticCheckNaN_All();
-	return ( Rotation.UnrotateVector(V) );
+	return (Rotation.UnrotateVector(V));
 }
 
+
 FORCEINLINE FTransform FTransform::operator*(const FTransform& Other) const
 {
 	FTransform Output;
@@ -1409,11 +1415,13 @@ FORCEINLINE FTransform FTransform::operator*(const FTransform& Other) const
 	return Output;
 }
 
+
 FORCEINLINE void FTransform::operator*=(const FTransform& Other)
 {
 	Multiply(this, this, &Other);
 }
 
+
 FORCEINLINE FTransform FTransform::operator*(const FQuat& Other) const
 {
 	FTransform Output, OtherTransform(Other, FVector::ZeroVector, FVector(1.f));
@@ -1421,20 +1429,22 @@ FORCEINLINE FTransform FTransform::operator*(const FQuat& Other) const
 	return Output;
 }
 
+
 FORCEINLINE void FTransform::operator*=(const FQuat& Other)
 {
 	FTransform OtherTransform(Other, FVector::ZeroVector, FVector(1.f));
 	Multiply(this, this, &OtherTransform);
 }
 
+
 // x = 0, y = 1, z = 2
-FORCEINLINE FVector FTransform::GetScaledAxis( EAxis::Type InAxis ) const
+FORCEINLINE FVector FTransform::GetScaledAxis(EAxis::Type InAxis) const
 {
-	if ( InAxis == EAxis::X )
+	if (InAxis == EAxis::X)
 	{
 		return TransformVector(FVector(1.f, 0.f, 0.f));
 	}
-	else if ( InAxis == EAxis::Y )
+	else if (InAxis == EAxis::Y)
 	{
 		return TransformVector(FVector(0.f, 1.f, 0.f));
 	}
@@ -1442,14 +1452,15 @@ FORCEINLINE FVector FTransform::GetScaledAxis( EAxis::Type InAxis ) const
 	return TransformVector(FVector(0.f, 0.f, 1.f));
 }
 
+
 // x = 0, y = 1, z = 2
-FORCEINLINE FVector FTransform::GetUnitAxis( EAxis::Type InAxis ) const
+FORCEINLINE FVector FTransform::GetUnitAxis(EAxis::Type InAxis) const
 {
-	if ( InAxis == EAxis::X )
+	if (InAxis == EAxis::X)
 	{
 		return TransformVectorNoScale(FVector(1.f, 0.f, 0.f));
 	}
-	else if ( InAxis == EAxis::Y )
+	else if (InAxis == EAxis::Y)
 	{
 		return TransformVectorNoScale(FVector(0.f, 1.f, 0.f));
 	}
@@ -1457,7 +1468,8 @@ FORCEINLINE FVector FTransform::GetUnitAxis( EAxis::Type InAxis ) const
 	return TransformVectorNoScale(FVector(0.f, 0.f, 1.f));
 }
 
-FORCEINLINE void FTransform::Mirror( EAxis::Type MirrorAxis, EAxis::Type FlipAxis)
+
+FORCEINLINE void FTransform::Mirror(EAxis::Type MirrorAxis, EAxis::Type FlipAxis)
 {
 	// We do convert to Matrix for mirroring. 
 	FMatrix M = ToMatrixWithScale();
@@ -1465,6 +1477,7 @@ FORCEINLINE void FTransform::Mirror( EAxis::Type MirrorAxis, EAxis::Type FlipAxi
 	SetFromMatrix(M);
 }
 
+
 /** same version of FMatrix::GetMaximumAxisScale function **/
 /** @return the maximum magnitude of all components of the 3D scale. */
 inline float FTransform::GetMaximumAxisScale() const
@@ -1473,6 +1486,7 @@ inline float FTransform::GetMaximumAxisScale() const
 	return Scale3D.GetAbsMax();
 }
 
+
 /** @return the minimum magnitude of all components of the 3D scale. */
 inline float FTransform::GetMinimumAxisScale() const
 {
@@ -1480,6 +1494,7 @@ inline float FTransform::GetMinimumAxisScale() const
 	return Scale3D.GetAbsMin();
 }
 
+
 // mathematically if you have 0 scale, it should be infinite, 
 // however, in practice if you have 0 scale, and relative transform doesn't make much sense 
 // anymore because you should be instead of showing gigantic infinite mesh
@@ -1516,4 +1531,4 @@ FORCEINLINE FVector FTransform::GetSafeScaleReciprocal(const FVector& InScale, f
 	}
 
 	return SafeReciprocalScale;
-}
\ No newline at end of file
+}
diff --git a/Engine/Source/Runtime/Core/Public/Math/UnrealMath.h b/Engine/Source/Runtime/Core/Public/Math/UnrealMath.h
--- a/Engine/Source/Runtime/Core/Public/Math/UnrealMath.h
+++ b/Engine/Source/Runtime/Core/Public/Math/UnrealMath.h
@@ -14,6 +14,7 @@
 #include "Edge.h"
 #include "Plane.h"
 #include "Sphere.h"
+#include "CapsuleShape.h"
 #include "Rotator.h"
 #include "RangeBound.h"
 #include "Range.h"
