Copying //UE4/Dev-Sequencer to //UE4/Dev-Main (Source: //UE4/Dev-Sequencer @ 3003857)

diff --git a/Engine/Source/Runtime/Core/Public/Math/Vector.h b/Engine/Source/Runtime/Core/Public/Math/Vector.h
--- a/Engine/Source/Runtime/Core/Public/Math/Vector.h
+++ b/Engine/Source/Runtime/Core/Public/Math/Vector.h
@@ -75,7 +75,7 @@ public:
 	 * @param InY Y Coordinate.
 	 * @param InZ Z Coordinate.
 	 */
-	FORCEINLINE FVector( float InX, float InY, float InZ );
+	FORCEINLINE FVector(float InX, float InY, float InZ);
 
 	/**
 	 * Constructs a vector from an FVector2D and Z value.
@@ -83,14 +83,14 @@ public:
 	 * @param V Vector to copy from.
 	 * @param InZ Z Coordinate.
 	 */
-	explicit FORCEINLINE FVector( const FVector2D V, float InZ );
+	explicit FORCEINLINE FVector(const FVector2D V, float InZ);
 
 	/**
 	 * Constructor using the XYZ components from a 4D vector.
 	 *
 	 * @param V 4D Vector to copy from.
 	 */
-	FORCEINLINE FVector( const FVector4& V );
+	FORCEINLINE FVector(const FVector4& V);
 
 	/**
 	 * Constructs a vector from an FLinearColor.
@@ -111,7 +111,7 @@ public:
 	 *
 	 * @param A Int Point used to set X and Y coordinates, Z is set to zero.
 	 */
-	explicit FVector( FIntPoint A );
+	explicit FVector(FIntPoint A);
 
 	/**
 	 * Constructor which initializes all components to zero.
@@ -136,7 +136,7 @@ public:
 	 * @param V The other vector.
 	 * @return The cross product.
 	 */
-	FORCEINLINE FVector operator^( const FVector& V ) const;
+	FORCEINLINE FVector operator^(const FVector& V) const;
 
 	/**
 	 * Calculate the cross product of two vectors.
@@ -145,7 +145,7 @@ public:
 	 * @param B The second vector.
 	 * @return The cross product.
 	 */
-	FORCEINLINE static FVector CrossProduct( const FVector& A, const FVector& B );
+	FORCEINLINE static FVector CrossProduct(const FVector& A, const FVector& B);
 
 	/**
 	 * Calculate the dot product between this and another vector.
@@ -153,7 +153,7 @@ public:
 	 * @param V The other vector.
 	 * @return The dot product.
 	 */
-	FORCEINLINE float operator|( const FVector& V ) const;
+	FORCEINLINE float operator|(const FVector& V) const;
 
 	/**
 	 * Calculate the dot product of two vectors.
@@ -162,7 +162,7 @@ public:
 	 * @param B The second vector.
 	 * @return The dot product.
 	 */
-	FORCEINLINE static float DotProduct( const FVector& A, const FVector& B );
+	FORCEINLINE static float DotProduct(const FVector& A, const FVector& B);
 
 	/**
 	 * Gets the result of component-wise addition of this and another vector.
@@ -170,7 +170,7 @@ public:
 	 * @param V The vector to add to this.
 	 * @return The result of vector addition.
 	 */
-	FORCEINLINE FVector operator+( const FVector& V ) const;
+	FORCEINLINE FVector operator+(const FVector& V) const;
 
 	/**
 	 * Gets the result of component-wise subtraction of this by another vector.
@@ -178,7 +178,7 @@ public:
 	 * @param V The vector to subtract from this.
 	 * @return The result of vector subtraction.
 	 */
-	FORCEINLINE FVector operator-( const FVector& V ) const;
+	FORCEINLINE FVector operator-(const FVector& V) const;
 
 	/**
 	 * Gets the result of subtracting from each component of the vector.
@@ -186,7 +186,7 @@ public:
 	 * @param Bias How much to subtract from each component.
 	 * @return The result of subtraction.
 	 */
-	FORCEINLINE FVector operator-( float Bias ) const;
+	FORCEINLINE FVector operator-(float Bias) const;
 
 	/**
 	 * Gets the result of adding to each component of the vector.
@@ -194,7 +194,7 @@ public:
 	 * @param Bias How much to add to each component.
 	 * @return The result of addition.
 	 */
-	FORCEINLINE FVector operator+( float Bias ) const;
+	FORCEINLINE FVector operator+(float Bias) const;
 
 	/**
 	 * Gets the result of scaling the vector (multiplying each component by a value).
@@ -202,7 +202,7 @@ public:
 	 * @param Scale What to multiply each component by.
 	 * @return The result of multiplication.
 	 */
-	FORCEINLINE FVector operator*( float Scale ) const;
+	FORCEINLINE FVector operator*(float Scale) const;
 
 	/**
 	 * Gets the result of dividing each component of the vector by a value.
@@ -210,7 +210,7 @@ public:
 	 * @param Scale What to divide each component by.
 	 * @return The result of division.
 	 */
-	FVector operator/( float Scale ) const;
+	FVector operator/(float Scale) const;
 
 	/**
 	 * Gets the result of component-wise multiplication of this vector by another.
@@ -218,7 +218,7 @@ public:
 	 * @param V The vector to multiply with.
 	 * @return The result of multiplication.
 	 */
-	FORCEINLINE FVector operator*( const FVector& V ) const;
+	FORCEINLINE FVector operator*(const FVector& V) const;
 
 	/**
 	 * Gets the result of component-wise division of this vector by another.
@@ -226,7 +226,7 @@ public:
 	 * @param V The vector to divide by.
 	 * @return The result of division.
 	 */
-	FORCEINLINE FVector operator/( const FVector& V ) const;
+	FORCEINLINE FVector operator/(const FVector& V) const;
 
 	// Binary comparison operators.
 
@@ -236,7 +236,7 @@ public:
 	 * @param V The vector to check against.
 	 * @return true if the vectors are equal, false otherwise.
 	 */
-	bool operator==( const FVector& V ) const;
+	bool operator==(const FVector& V) const;
 
 	/**
 	 * Check against another vector for inequality.
@@ -244,7 +244,7 @@ public:
 	 * @param V The vector to check against.
 	 * @return true if the vectors are not equal, false otherwise.
 	 */
-	bool operator!=( const FVector& V ) const;
+	bool operator!=(const FVector& V) const;
 
 	/**
 	 * Check against another vector for equality, within specified error limits.
@@ -277,7 +277,7 @@ public:
 	 * @param V Vector to add to this.
 	 * @return Copy of the vector after addition.
 	 */
-	FORCEINLINE FVector operator+=( const FVector& V );
+	FORCEINLINE FVector operator+=(const FVector& V);
 
 	/**
 	 * Subtracts another vector from this.
@@ -286,7 +286,7 @@ public:
 	 * @param V Vector to subtract from this.
 	 * @return Copy of the vector after subtraction.
 	 */
-	FORCEINLINE FVector operator-=( const FVector& V );
+	FORCEINLINE FVector operator-=(const FVector& V);
 
 	/**
 	 * Scales the vector.
@@ -294,7 +294,7 @@ public:
 	 * @param Scale Amount to scale this vector by.
 	 * @return Copy of the vector after scaling.
 	 */
-	FORCEINLINE FVector operator*=( float Scale );
+	FORCEINLINE FVector operator*=(float Scale);
 
 	/**
 	 * Divides the vector by a number.
@@ -302,7 +302,7 @@ public:
 	 * @param V What to divide this vector by.
 	 * @return Copy of the vector after division.
 	 */
-	FVector operator/=( float V );
+	FVector operator/=(float V);
 
 	/**
 	 * Multiplies the vector with another vector, using component-wise multiplication.
@@ -310,7 +310,7 @@ public:
 	 * @param V What to multiply this vector with.
 	 * @return Copy of the vector after multiplication.
 	 */
-	FVector operator*=( const FVector& V );
+	FVector operator*=(const FVector& V);
 
 	/**
 	 * Divides the vector by another vector, using component-wise division.
@@ -318,7 +318,7 @@ public:
 	 * @param V What to divide vector by.
 	 * @return Copy of the vector after division.
 	 */
-	FVector operator/=( const FVector& V );
+	FVector operator/=(const FVector& V);
 
 	/**
 	 * Gets specific component of the vector.
@@ -326,7 +326,7 @@ public:
 	 * @param Index the index of vector component
 	 * @return reference to component.
 	 */
-	float& operator[]( int32 Index );
+	float& operator[](int32 Index);
 
 	/**
 	 * Gets specific component of the vector.
@@ -334,7 +334,7 @@ public:
 	 * @param Index the index of vector component
 	 * @return Copy of the component.
 	 */
-	float operator[]( int32 Index )const;
+	float operator[](int32 Index)const;
 
 	/**
 	* Gets a specific component of the vector.
@@ -353,6 +353,8 @@ public:
 	*/
 	float Component(int32 Index) const;
 
+public:
+
 	// Simple functions.
 
 	/**
@@ -362,7 +364,7 @@ public:
 	 * @param InY New Y coordinate.
 	 * @param InZ New Z coordinate.
 	 */
-	void Set( float InX, float InY, float InZ );
+	void Set(float InX, float InY, float InZ);
 
 	/**
 	 * Get the maximum value of the vector's components.
@@ -504,7 +506,7 @@ public:
 	 * @return A copy of this vector snapped to a grid.
 	 * @see FMath::GridSnap()
 	 */
-	FVector GridSnap( const float& GridSz ) const;
+	FVector GridSnap(const float& GridSz) const;
 
 	/**
 	 * Get a copy of this vector, clamped inside of a cube.
@@ -512,7 +514,7 @@ public:
 	 * @param Radius Half size of the cube.
 	 * @return A copy of this vector, bound by cube.
 	 */
-	FVector BoundToCube( float Radius ) const;
+	FVector BoundToCube(float Radius) const;
 
 	/** Create a copy of this vector, with its magnitude clamped between Min and Max. */
 	FVector GetClampedToSize(float Min, float Max) const;
@@ -544,7 +546,7 @@ public:
 	 * @param V Vector to add.
 	 * @param Radius Half size of the cube.
 	 */
-	void AddBounded( const FVector& V, float Radius=MAX_int16 );
+	void AddBounded(const FVector& V, float Radius=MAX_int16);
 
 	/**
 	 * Gets the reciprocal of this vector, avoiding division by zero.
@@ -568,7 +570,7 @@ public:
 	 * @param MirrorNormal Normal vector to mirror about.
 	 * @return Mirrored vector.
 	 */
-	FVector MirrorByVector( const FVector& MirrorNormal ) const;
+	FVector MirrorByVector(const FVector& MirrorNormal) const;
 
 	/**
 	 * Mirrors a vector about a plane.
@@ -576,7 +578,7 @@ public:
 	 * @param Plane Plane to mirror about.
 	 * @return Mirrored vector.
 	 */
-	FVector MirrorByPlane( const FPlane& Plane ) const;
+	FVector MirrorByPlane(const FPlane& Plane) const;
 
 	/**
 	 * Rotates around Axis (assumes Axis.Size() == 1).
@@ -585,7 +587,7 @@ public:
 	 * @param Axis Axis to rotate around.
 	 * @return Rotated Vector.
 	 */
-	FVector RotateAngleAxis( const float AngleDeg, const FVector& Axis ) const;
+	FVector RotateAngleAxis(const float AngleDeg, const FVector& Axis) const;
 
 	/**
 	 * Gets a normalized copy of the vector, checking it is safe to do so based on the length.
@@ -625,7 +627,7 @@ public:
 	 * @param A	Vector to project onto, does not assume it is normalized.
 	 * @return Projected vector.
 	 */
-	FORCEINLINE FVector ProjectOnTo( const FVector& A ) const ;
+	FORCEINLINE FVector ProjectOnTo(const FVector& A) const ;
 
 	/**
 	 * Gets a copy of this vector projected onto the input vector, which is assumed to be unit length.
@@ -633,7 +635,7 @@ public:
 	 * @param  Normal Vector to project onto (assumed to be unit length).
 	 * @return Projected vector.
 	 */
-	FORCEINLINE FVector ProjectOnToNormal(const FVector& Normal ) const;
+	FORCEINLINE FVector ProjectOnToNormal(const FVector& Normal) const;
 
 	/**
 	 * Return the FRotator orientation corresponding to the direction in which the vector points.
@@ -672,7 +674,7 @@ public:
 	 * @param Axis1 Reference to first axis.
 	 * @param Axis2 Reference to second axis.
 	 */
-	CORE_API void FindBestAxisVectors( FVector& Axis1, FVector& Axis2 ) const;
+	CORE_API void FindBestAxisVectors(FVector& Axis1, FVector& Axis2) const;
 
 	/** When this vector contains Euler angles (degrees), ensure that angles are between +/-180 */
 	CORE_API void UnwindEuler();
@@ -690,7 +692,7 @@ public:
 	 * @param LengthSquaredTolerance Tolerance against squared length.
 	 * @return true if the vector is a unit vector within the specified tolerance.
 	 */
-	FORCEINLINE bool IsUnit(float LengthSquaredTolerance = KINDA_SMALL_NUMBER ) const;
+	FORCEINLINE bool IsUnit(float LengthSquaredTolerance = KINDA_SMALL_NUMBER) const;
 
 	/**
 	 * Get a textual representation of this vector.
@@ -719,7 +721,7 @@ public:
 	 * @param	InSourceString	FString containing the vector values.
 	 * @return true if the X,Y,Z values were read successfully; false otherwise.
 	 */
-	bool InitFromString( const FString& InSourceString );
+	bool InitFromString(const FString& InSourceString);
 
 	/** 
 	 * Converts a Cartesian unit vector into spherical coordinates on the unit sphere.
@@ -753,7 +755,7 @@ public:
 	 * @param Q Second vector.
 	 * @return Whether points are the same within a threshold. Uses fast distance approximation (linear per-component distance).
 	 */
-	static bool PointsAreSame( const FVector &P, const FVector &Q );
+	static bool PointsAreSame(const FVector &P, const FVector &Q);
 	
 	/**
 	 * Compare two points and see if they're within specified distance.
@@ -763,7 +765,7 @@ public:
 	 * @param Dist Specified distance.
 	 * @return Whether two points are within the specified distance. Uses fast distance approximation (linear per-component distance).
 	 */
-	static bool PointsAreNear( const FVector &Point1, const FVector &Point2, float Dist );
+	static bool PointsAreNear(const FVector &Point1, const FVector &Point2, float Dist);
 
 	/**
 	 * Calculate the signed distance (in the direction of the normal) between a point and a plane.
@@ -773,7 +775,7 @@ public:
 	 * @param PlaneNormal The Normal of the plane (assumed to be unit length).
 	 * @return Signed distance between point and plane.
 	 */
-	static float PointPlaneDist( const FVector &Point, const FVector &PlaneBase, const FVector &PlaneNormal );
+	static float PointPlaneDist(const FVector &Point, const FVector &PlaneBase, const FVector &PlaneNormal);
 
 	/**
 	 * Calculate the projection of a point on the given plane.
@@ -821,7 +823,7 @@ public:
 	 * @param V2 The second point.
 	 * @return The distance between two points.
 	 */
-	static FORCEINLINE float Dist( const FVector &V1, const FVector &V2 );
+	static FORCEINLINE float Dist(const FVector &V1, const FVector &V2);
 
 	/**
 	 * Squared distance between two points.
@@ -830,7 +832,7 @@ public:
 	 * @param V2 The second point.
 	 * @return The squared distance between two points.
 	 */
-	static FORCEINLINE float DistSquared( const FVector &V1, const FVector &V2 );
+	static FORCEINLINE float DistSquared(const FVector &V1, const FVector &V2);
 
 	/**
 	 * Squared distance between two points in the XY plane only.
@@ -839,7 +841,7 @@ public:
 	 * @param V2 The second point.
 	 * @return The squared distance between two points in the XY plane
 	 */
-	static FORCEINLINE float DistSquaredXY( const FVector &V1, const FVector &V2 );
+	static FORCEINLINE float DistSquaredXY(const FVector &V1, const FVector &V2);
 	
 	/**
 	 * Compute pushout of a box from a plane.
@@ -848,7 +850,7 @@ public:
 	 * @param Size The size of the box.
 	 * @return Pushout required.
 	 */
-	static FORCEINLINE float BoxPushOut( const FVector& Normal, const FVector& Size );
+	static FORCEINLINE float BoxPushOut(const FVector& Normal, const FVector& Size);
 
 	/**
 	 * See if two normal vectors are nearly parallel, meaning the angle between them is close to 0 degrees.
@@ -900,7 +902,7 @@ public:
 	 * @param Z The third vector.
 	 * @return The triple product: X dot (Y cross Z).
 	 */
-	static float Triple( const FVector& X, const FVector& Y, const FVector& Z );
+	static float Triple(const FVector& X, const FVector& Y, const FVector& Z);
 
 	/**
 	 * Generates a list of sample points on a Bezier curve defined by 2 points.
@@ -946,14 +948,14 @@ public:
 	 * @param V Vector to serialize.
 	 * @return Reference to Archive after serialization.
 	 */
-	friend FArchive& operator<<( FArchive& Ar, FVector& V )
+	friend FArchive& operator<<(FArchive& Ar, FVector& V)
 	{
 		// @warning BulkSerialize: FVector is serialized as memory dump
 		// See TArray::BulkSerialize for detailed description of implied limitations.
 		return Ar << V.X << V.Y << V.Z;
 	}
 	
-	bool Serialize( FArchive& Ar )
+	bool Serialize(FArchive& Ar)
 	{
 		Ar << *this;
 		return true;
@@ -979,11 +981,12 @@ public:
  * @param V Vector to scale.
  * @return Result of multiplication.
  */
-FORCEINLINE FVector operator*( float Scale, const FVector& V )
+FORCEINLINE FVector operator*(float Scale, const FVector& V)
 {
-	return V.operator*( Scale );
+	return V.operator*(Scale);
 }
 
+
 /**
  * Creates a hash value from a FVector. 
  *
@@ -996,6 +999,7 @@ FORCEINLINE uint32 GetTypeHash(const FVector& Vector)
 	return FCrc::MemCrc_DEPRECATED(&Vector,sizeof(Vector));
 }
 
+
 /**
  * Creates a hash value from a FVector2D. 
  *
@@ -1018,6 +1022,7 @@ FORCEINLINE uint32 GetTypeHash(const FVector2D& Vector)
 	}
 #endif
 
+
 /** 
  * Util to calculate distance from a point to a bounding box 
  *
@@ -1026,7 +1031,7 @@ FORCEINLINE uint32 GetTypeHash(const FVector2D& Vector)
  * @param Point 3D position of interest
  * @return the distance from the Point to the bounding box.
  */
-FORCEINLINE float ComputeSquaredDistanceFromBoxToPoint( const FVector& Mins, const FVector& Maxs, const FVector& Point )
+FORCEINLINE float ComputeSquaredDistanceFromBoxToPoint(const FVector& Mins, const FVector& Maxs, const FVector& Point)
 {
 	// Accumulates the distance as we iterate axis
 	float DistSquared = 0.f;
@@ -1064,14 +1069,14 @@ FORCEINLINE float ComputeSquaredDistanceFromBoxToPoint( const FVector& Mins, con
 }
 
 
-FORCEINLINE FVector::FVector( const FVector2D V, float InZ )
+FORCEINLINE FVector::FVector(const FVector2D V, float InZ)
 	: X(V.X), Y(V.Y), Z(InZ)
 {
 	DiagnosticCheckNaN();
 }
 
 
-inline FVector FVector::RotateAngleAxis( const float AngleDeg, const FVector& Axis ) const
+inline FVector FVector::RotateAngleAxis(const float AngleDeg, const FVector& Axis) const
 {
 	float S, C;
 	FMath::SinCos(&S, &C, FMath::DegreesToRadians(AngleDeg));
@@ -1091,23 +1096,23 @@ inline FVector FVector::RotateAngleAxis( const float AngleDeg, const FVector& Ax
 	const float OMC	= 1.f - C;
 
 	return FVector(
-		(OMC * XX + C ) * X + (OMC * XY - ZS) * Y + (OMC * ZX + YS) * Z,
-		(OMC * XY + ZS) * X + (OMC * YY + C ) * Y + (OMC * YZ - XS) * Z,
-		(OMC * ZX - YS) * X + (OMC * YZ + XS) * Y + (OMC * ZZ + C ) * Z
+		(OMC * XX + C) * X + (OMC * XY - ZS) * Y + (OMC * ZX + YS) * Z,
+		(OMC * XY + ZS) * X + (OMC * YY + C) * Y + (OMC * YZ - XS) * Z,
+		(OMC * ZX - YS) * X + (OMC * YZ + XS) * Y + (OMC * ZZ + C) * Z
 		);
 }
 
-inline bool FVector::PointsAreSame( const FVector &P, const FVector &Q )
+inline bool FVector::PointsAreSame(const FVector &P, const FVector &Q)
 {
 	float Temp;
 	Temp=P.X-Q.X;
-	if( (Temp > -THRESH_POINTS_ARE_SAME) && (Temp < THRESH_POINTS_ARE_SAME) )
+	if((Temp > -THRESH_POINTS_ARE_SAME) && (Temp < THRESH_POINTS_ARE_SAME))
 	{
 		Temp=P.Y-Q.Y;
-		if( (Temp > -THRESH_POINTS_ARE_SAME) && (Temp < THRESH_POINTS_ARE_SAME) )
+		if((Temp > -THRESH_POINTS_ARE_SAME) && (Temp < THRESH_POINTS_ARE_SAME))
 		{
 			Temp=P.Z-Q.Z;
-			if( (Temp > -THRESH_POINTS_ARE_SAME) && (Temp < THRESH_POINTS_ARE_SAME) )
+			if((Temp > -THRESH_POINTS_ARE_SAME) && (Temp < THRESH_POINTS_ARE_SAME))
 			{
 				return true;
 			}
@@ -1116,7 +1121,7 @@ inline bool FVector::PointsAreSame( const FVector &P, const FVector &Q )
 	return false;
 }
 
-inline bool FVector::PointsAreNear( const FVector &Point1, const FVector &Point2, float Dist )
+inline bool FVector::PointsAreNear(const FVector &Point1, const FVector &Point2, float Dist)
 {
 	float Temp;
 	Temp=(Point1.X - Point2.X); if (FMath::Abs(Temp)>=Dist) return false;
@@ -1173,12 +1178,12 @@ inline bool FVector::Coplanar(const FVector &Base1, const FVector &Normal1, cons
 	else return true;
 }
 
-inline float FVector::Triple( const FVector& X, const FVector& Y, const FVector& Z )
+inline float FVector::Triple(const FVector& X, const FVector& Y, const FVector& Z)
 {
 	return
 	(	(X.X * (Y.Y * Z.Z - Y.Z * Z.Y))
 	+	(X.Y * (Y.Z * Z.X - Y.X * Z.Z))
-	+	(X.Z * (Y.X * Z.Y - Y.Y * Z.X)) );
+	+	(X.Z * (Y.X * Z.Y - Y.Y * Z.X)));
 }
 
 inline FVector FVector::RadiansToDegrees(const FVector& RadVector)
@@ -1200,7 +1205,7 @@ FORCEINLINE FVector::FVector(float InF)
 	DiagnosticCheckNaN();
 }
 
-FORCEINLINE FVector::FVector( float InX, float InY, float InZ )
+FORCEINLINE FVector::FVector(float InX, float InY, float InZ)
 	: X(InX), Y(InY), Z(InZ)
 {
 	DiagnosticCheckNaN();
@@ -1218,7 +1223,7 @@ FORCEINLINE FVector::FVector(FIntVector InVector)
 	DiagnosticCheckNaN();
 }
 
-FORCEINLINE FVector::FVector( FIntPoint A )
+FORCEINLINE FVector::FVector(FIntPoint A)
 	: X(A.X), Y(A.Y), Z(0.f)
 {
 	DiagnosticCheckNaN();
@@ -1243,7 +1248,7 @@ FORCEINLINE FVector& FVector::operator=(const FVector& Other)
 }
 #endif
 
-FORCEINLINE FVector FVector::operator^( const FVector& V ) const
+FORCEINLINE FVector FVector::operator^(const FVector& V) const
 {
 	return FVector
 		(
@@ -1253,68 +1258,68 @@ FORCEINLINE FVector FVector::operator^( const FVector& V ) const
 		);
 }
 
-FORCEINLINE FVector FVector::CrossProduct( const FVector& A, const FVector& B )
+FORCEINLINE FVector FVector::CrossProduct(const FVector& A, const FVector& B)
 {
 	return A ^ B;
 }
 
-FORCEINLINE float FVector::operator|( const FVector& V ) const
+FORCEINLINE float FVector::operator|(const FVector& V) const
 {
 	return X*V.X + Y*V.Y + Z*V.Z;
 }
 
-FORCEINLINE float FVector::DotProduct( const FVector& A, const FVector& B )
+FORCEINLINE float FVector::DotProduct(const FVector& A, const FVector& B)
 {
 	return A | B;
 }
 
-FORCEINLINE FVector FVector::operator+( const FVector& V ) const
+FORCEINLINE FVector FVector::operator+(const FVector& V) const
 {
-	return FVector( X + V.X, Y + V.Y, Z + V.Z );
+	return FVector(X + V.X, Y + V.Y, Z + V.Z);
 }
 
-FORCEINLINE FVector FVector::operator-( const FVector& V ) const
+FORCEINLINE FVector FVector::operator-(const FVector& V) const
 {
-	return FVector( X - V.X, Y - V.Y, Z - V.Z );
+	return FVector(X - V.X, Y - V.Y, Z - V.Z);
 }
 
-FORCEINLINE FVector FVector::operator-( float Bias ) const
+FORCEINLINE FVector FVector::operator-(float Bias) const
 {
-	return FVector( X - Bias, Y - Bias, Z - Bias );
+	return FVector(X - Bias, Y - Bias, Z - Bias);
 }
 
-FORCEINLINE FVector FVector::operator+( float Bias ) const
+FORCEINLINE FVector FVector::operator+(float Bias) const
 {
-	return FVector( X + Bias, Y + Bias, Z + Bias );
+	return FVector(X + Bias, Y + Bias, Z + Bias);
 }
 
-FORCEINLINE FVector FVector::operator*( float Scale ) const
+FORCEINLINE FVector FVector::operator*(float Scale) const
 {
-	return FVector( X * Scale, Y * Scale, Z * Scale );
+	return FVector(X * Scale, Y * Scale, Z * Scale);
 }
 
-FORCEINLINE FVector FVector::operator/( float Scale ) const
+FORCEINLINE FVector FVector::operator/(float Scale) const
 {
 	const float RScale = 1.f/Scale;
-	return FVector( X * RScale, Y * RScale, Z * RScale );
+	return FVector(X * RScale, Y * RScale, Z * RScale);
 }
 
-FORCEINLINE FVector FVector::operator*( const FVector& V ) const
+FORCEINLINE FVector FVector::operator*(const FVector& V) const
 {
-	return FVector( X * V.X, Y * V.Y, Z * V.Z );
+	return FVector(X * V.X, Y * V.Y, Z * V.Z);
 }
 
-FORCEINLINE FVector FVector::operator/( const FVector& V ) const
+FORCEINLINE FVector FVector::operator/(const FVector& V) const
 {
-	return FVector( X / V.X, Y / V.Y, Z / V.Z );
+	return FVector(X / V.X, Y / V.Y, Z / V.Z);
 }
 
-FORCEINLINE bool FVector::operator==( const FVector& V ) const
+FORCEINLINE bool FVector::operator==(const FVector& V) const
 {
 	return X==V.X && Y==V.Y && Z==V.Z;
 }
 
-FORCEINLINE bool FVector::operator!=( const FVector& V ) const
+FORCEINLINE bool FVector::operator!=(const FVector& V) const
 {
 	return X!=V.X || Y!=V.Y || Z!=V.Z;
 }
@@ -1326,38 +1331,38 @@ FORCEINLINE bool FVector::Equals(const FVector& V, float Tolerance) const
 
 FORCEINLINE bool FVector::AllComponentsEqual(float Tolerance) const
 {
-	return FMath::Abs( X - Y ) <= Tolerance && FMath::Abs( X - Z ) <= Tolerance && FMath::Abs( Y - Z ) <= Tolerance;
+	return FMath::Abs(X - Y) <= Tolerance && FMath::Abs(X - Z) <= Tolerance && FMath::Abs(Y - Z) <= Tolerance;
 }
 
 
 FORCEINLINE FVector FVector::operator-() const
 {
-	return FVector( -X, -Y, -Z );
+	return FVector(-X, -Y, -Z);
 }
 
 
-FORCEINLINE FVector FVector::operator+=( const FVector& V )
+FORCEINLINE FVector FVector::operator+=(const FVector& V)
 {
 	X += V.X; Y += V.Y; Z += V.Z;
 	DiagnosticCheckNaN();
 	return *this;
 }
 
-FORCEINLINE FVector FVector::operator-=( const FVector& V )
+FORCEINLINE FVector FVector::operator-=(const FVector& V)
 {
 	X -= V.X; Y -= V.Y; Z -= V.Z;
 	DiagnosticCheckNaN();
 	return *this;
 }
 
-FORCEINLINE FVector FVector::operator*=( float Scale )
+FORCEINLINE FVector FVector::operator*=(float Scale)
 {
 	X *= Scale; Y *= Scale; Z *= Scale;
 	DiagnosticCheckNaN();
 	return *this;
 }
 
-FORCEINLINE FVector FVector::operator/=( float V )
+FORCEINLINE FVector FVector::operator/=(float V)
 {
 	const float RV = 1.f/V;
 	X *= RV; Y *= RV; Z *= RV;
@@ -1365,28 +1370,28 @@ FORCEINLINE FVector FVector::operator/=( float V )
 	return *this;
 }
 
-FORCEINLINE FVector FVector::operator*=( const FVector& V )
+FORCEINLINE FVector FVector::operator*=(const FVector& V)
 {
 	X *= V.X; Y *= V.Y; Z *= V.Z;
 	DiagnosticCheckNaN();
 	return *this;
 }
 
-FORCEINLINE FVector FVector::operator/=( const FVector& V )
+FORCEINLINE FVector FVector::operator/=(const FVector& V)
 {
 	X /= V.X; Y /= V.Y; Z /= V.Z;
 	DiagnosticCheckNaN();
 	return *this;
 }
 
-FORCEINLINE float& FVector::operator[]( int32 Index )
+FORCEINLINE float& FVector::operator[](int32 Index)
 {
 	check(Index >= 0 && Index < 3);
-	if( Index == 0 )
+	if(Index == 0)
 	{
 		return X;
 	}
-	else if( Index == 1)
+	else if(Index == 1)
 	{
 		return Y;
 	}
@@ -1396,14 +1401,14 @@ FORCEINLINE float& FVector::operator[]( int32 Index )
 	}
 }
 
-FORCEINLINE float FVector::operator[]( int32 Index )const
+FORCEINLINE float FVector::operator[](int32 Index)const
 {
 	check(Index >= 0 && Index < 3);
-	if( Index == 0 )
+	if(Index == 0)
 	{
 		return X;
 	}
-	else if( Index == 1)
+	else if(Index == 1)
 	{
 		return Y;
 	}
@@ -1413,7 +1418,7 @@ FORCEINLINE float FVector::operator[]( int32 Index )const
 	}
 }
 
-FORCEINLINE void FVector::Set( float InX, float InY, float InZ )
+FORCEINLINE void FVector::Set(float InX, float InY, float InZ)
 {
 	X = InX;
 	Y = InY;
@@ -1458,7 +1463,7 @@ FORCEINLINE FVector FVector::GetAbs() const
 
 FORCEINLINE float FVector::Size() const
 {
-	return FMath::Sqrt( X*X + Y*Y + Z*Z );
+	return FMath::Sqrt(X*X + Y*Y + Z*Z);
 }
 
 FORCEINLINE float FVector::SizeSquared() const
@@ -1468,7 +1473,7 @@ FORCEINLINE float FVector::SizeSquared() const
 
 FORCEINLINE float FVector::Size2D() const 
 {
-	return FMath::Sqrt( X*X + Y*Y );
+	return FMath::Sqrt(X*X + Y*Y);
 }
 
 FORCEINLINE float FVector::SizeSquared2D() const 
@@ -1492,7 +1497,7 @@ FORCEINLINE bool FVector::IsZero() const
 FORCEINLINE bool FVector::Normalize(float Tolerance)
 {
 	const float SquareSum = X*X + Y*Y + Z*Z;
-	if( SquareSum > Tolerance )
+	if(SquareSum > Tolerance)
 	{
 		const float Scale = FMath::InvSqrt(SquareSum);
 		X *= Scale; Y *= Scale; Z *= Scale;
@@ -1534,13 +1539,13 @@ FORCEINLINE FVector FVector::GetSignVector() const
 FORCEINLINE FVector FVector::Projection() const
 {
 	const float RZ = 1.f/Z;
-	return FVector( X*RZ, Y*RZ, 1 );
+	return FVector(X*RZ, Y*RZ, 1);
 }
 
 FORCEINLINE FVector FVector::GetUnsafeNormal() const
 {
 	const float Scale = FMath::InvSqrt(X*X+Y*Y+Z*Z);
-	return FVector( X*Scale, Y*Scale, Z*Scale );
+	return FVector(X*Scale, Y*Scale, Z*Scale);
 }
 
 FORCEINLINE FVector FVector::UnsafeNormal() const
@@ -1548,12 +1553,12 @@ FORCEINLINE FVector FVector::UnsafeNormal() const
 	return GetUnsafeNormal();
 }
 
-FORCEINLINE FVector FVector::GridSnap( const float& GridSz ) const
+FORCEINLINE FVector FVector::GridSnap(const float& GridSz) const
 {
-	return FVector( FMath::GridSnap(X, GridSz),FMath::GridSnap(Y, GridSz),FMath::GridSnap(Z, GridSz) );
+	return FVector(FMath::GridSnap(X, GridSz),FMath::GridSnap(Y, GridSz),FMath::GridSnap(Z, GridSz));
 }
 
-FORCEINLINE FVector FVector::BoundToCube( float Radius ) const
+FORCEINLINE FVector FVector::BoundToCube(float Radius) const
 {
 	return FVector
 		(
@@ -1643,17 +1648,17 @@ FORCEINLINE FVector FVector::ClampMaxSize2D(float MaxSize) const
 }
 
 
-FORCEINLINE void FVector::AddBounded( const FVector& V, float Radius )
+FORCEINLINE void FVector::AddBounded(const FVector& V, float Radius)
 {
 	*this = (*this + V).BoundToCube(Radius);
 }
 
-FORCEINLINE float& FVector::Component( int32 Index )
+FORCEINLINE float& FVector::Component(int32 Index)
 {
 	return (&X)[Index];
 }
 
-FORCEINLINE float FVector::Component( int32 Index ) const
+FORCEINLINE float FVector::Component(int32 Index) const
 {
 	return (&X)[Index];
 }
@@ -1694,7 +1699,7 @@ FORCEINLINE bool FVector::IsUniform(float Tolerance) const
 	return AllComponentsEqual(Tolerance);
 }
 
-FORCEINLINE FVector FVector::MirrorByVector( const FVector& MirrorNormal ) const
+FORCEINLINE FVector FVector::MirrorByVector(const FVector& MirrorNormal) const
 {
 	return *this - MirrorNormal * (2.f * (*this | MirrorNormal));
 }
@@ -1704,11 +1709,11 @@ FORCEINLINE FVector FVector::GetSafeNormal(float Tolerance) const
 	const float SquareSum = X*X + Y*Y + Z*Z;
 
 	// Not sure if it's safe to add tolerance in there. Might introduce too many errors
-	if( SquareSum == 1.f )
+	if(SquareSum == 1.f)
 	{
 		return *this;
 	}		
-	else if( SquareSum < Tolerance )
+	else if(SquareSum < Tolerance)
 	{
 		return FVector::ZeroVector;
 	}
@@ -1726,9 +1731,9 @@ FORCEINLINE FVector FVector::GetSafeNormal2D(float Tolerance) const
 	const float SquareSum = X*X + Y*Y;
 
 	// Not sure if it's safe to add tolerance in there. Might introduce too many errors
-	if( SquareSum == 1.f )
+	if(SquareSum == 1.f)
 	{
-		if( Z == 0.f )
+		if(Z == 0.f)
 		{
 			return *this;
 		}
@@ -1737,7 +1742,7 @@ FORCEINLINE FVector FVector::GetSafeNormal2D(float Tolerance) const
 			return FVector(X, Y, 0.f);
 		}
 	}
-	else if( SquareSum < Tolerance )
+	else if(SquareSum < Tolerance)
 	{
 		return FVector::ZeroVector;
 	}
@@ -1761,7 +1766,7 @@ FORCEINLINE float FVector::CosineAngle2D(FVector B) const
 	return A | B;
 }
 
-FORCEINLINE FVector FVector::ProjectOnTo( const FVector& A ) const 
+FORCEINLINE FVector FVector::ProjectOnTo(const FVector& A) const 
 { 
 	return (A * ((*this | A) / (A | A))); 
 }
@@ -1779,7 +1784,7 @@ FORCEINLINE bool FVector::ContainsNaN() const
 			!FMath::IsFinite(Z));
 }
 
-FORCEINLINE bool FVector::IsUnit(float LengthSquaredTolerance ) const
+FORCEINLINE bool FVector::IsUnit(float LengthSquaredTolerance) const
 {
 	return FMath::Abs(1.0f - SizeSquared()) < LengthSquaredTolerance;
 }
@@ -1853,30 +1858,30 @@ FORCEINLINE FText FVector::ToCompactText() const
 
 FORCEINLINE FString FVector::ToCompactString() const
 {
-	if( IsNearlyZero() )
+	if(IsNearlyZero())
 	{
 		return FString::Printf(TEXT("V(0)"));
 	}
 
 	FString ReturnString(TEXT("V("));
 	bool bIsEmptyString = true;
-	if( !FMath::IsNearlyZero(X) )
+	if(!FMath::IsNearlyZero(X))
 	{
 		ReturnString += FString::Printf(TEXT("X=%.2f"), X);
 		bIsEmptyString = false;
 	}
-	if( !FMath::IsNearlyZero(Y) )
+	if(!FMath::IsNearlyZero(Y))
 	{
-		if( !bIsEmptyString )
+		if(!bIsEmptyString)
 		{
 			ReturnString += FString(TEXT(", "));
 		}
 		ReturnString += FString::Printf(TEXT("Y=%.2f"), Y);
 		bIsEmptyString = false;
 	}
-	if( !FMath::IsNearlyZero(Z) )
+	if(!FMath::IsNearlyZero(Z))
 	{
-		if( !bIsEmptyString )
+		if(!bIsEmptyString)
 		{
 			ReturnString += FString(TEXT(", "));
 		}
@@ -1887,12 +1892,12 @@ FORCEINLINE FString FVector::ToCompactString() const
 	return ReturnString;
 }
 
-FORCEINLINE bool FVector::InitFromString( const FString& InSourceString )
+FORCEINLINE bool FVector::InitFromString(const FString& InSourceString)
 {
 	X = Y = Z = 0;
 
 	// The initialization is only successful if the X, Y, and Z values can all be parsed from the string
-	const bool bSuccessful = FParse::Value( *InSourceString, TEXT("X=") , X ) && FParse::Value( *InSourceString, TEXT("Y="), Y ) && FParse::Value( *InSourceString, TEXT("Z="), Z );
+	const bool bSuccessful = FParse::Value(*InSourceString, TEXT("X=") , X) && FParse::Value(*InSourceString, TEXT("Y="), Y) && FParse::Value(*InSourceString, TEXT("Z="), Z);
 
 	return bSuccessful;
 }
@@ -1924,28 +1929,28 @@ FORCEINLINE float FVector::HeadingAngle() const
 
 
 
-FORCEINLINE float FVector::Dist( const FVector &V1, const FVector &V2 )
+FORCEINLINE float FVector::Dist(const FVector &V1, const FVector &V2)
 {
-	return FMath::Sqrt( FMath::Square(V2.X-V1.X) + FMath::Square(V2.Y-V1.Y) + FMath::Square(V2.Z-V1.Z) );
+	return FMath::Sqrt(FMath::Square(V2.X-V1.X) + FMath::Square(V2.Y-V1.Y) + FMath::Square(V2.Z-V1.Z));
 }
 
-FORCEINLINE float FVector::DistSquared( const FVector &V1, const FVector &V2 )
+FORCEINLINE float FVector::DistSquared(const FVector &V1, const FVector &V2)
 {
 	return FMath::Square(V2.X-V1.X) + FMath::Square(V2.Y-V1.Y) + FMath::Square(V2.Z-V1.Z);
 }
 
-FORCEINLINE float FVector::DistSquaredXY( const FVector &V1, const FVector &V2 )
+FORCEINLINE float FVector::DistSquaredXY(const FVector &V1, const FVector &V2)
 {
 	return FMath::Square(V2.X-V1.X) + FMath::Square(V2.Y-V1.Y);
 }
 
-FORCEINLINE float FVector::BoxPushOut( const FVector& Normal, const FVector& Size )
+FORCEINLINE float FVector::BoxPushOut(const FVector& Normal, const FVector& Size)
 {
 	return FMath::Abs(Normal.X*Size.X) + FMath::Abs(Normal.Y*Size.Y) + FMath::Abs(Normal.Z*Size.Z);
 }
 
 /** Component-wise clamp for FVector */
-FORCEINLINE FVector ClampVector( const FVector& V, const FVector& Min, const FVector& Max )
+FORCEINLINE FVector ClampVector(const FVector& V, const FVector& Min, const FVector& Max)
 {
 	return FVector(
 		FMath::Clamp(V.X,Min.X,Max.X),
diff --git a/Engine/Source/Runtime/Core/Public/Math/Vector2D.h b/Engine/Source/Runtime/Core/Public/Math/Vector2D.h
--- a/Engine/Source/Runtime/Core/Public/Math/Vector2D.h
+++ b/Engine/Source/Runtime/Core/Public/Math/Vector2D.h
@@ -24,9 +24,7 @@ public:
 
 public:
 
-	/**
-	 * Default constructor (no initialization).
-	 */
+	/** Default constructor (no initialization). */
 	FORCEINLINE FVector2D() { }
 
 	/**
@@ -35,21 +33,21 @@ public:
 	 * @param InX X coordinate.
 	 * @param InY Y coordinate.
 	 */
-	FORCEINLINE FVector2D( float InX,float InY );
+	FORCEINLINE FVector2D(float InX,float InY);
 
 	/**
 	 * Constructs a vector from an FIntPoint.
 	 *
 	 * @param InPos Integer point used to set this vector.
 	 */
-	FORCEINLINE FVector2D( FIntPoint InPos );
+	FORCEINLINE FVector2D(FIntPoint InPos);
 
 	/**
 	 * Constructor which initializes all components to zero.
 	 *
 	 * @param EForceInit Force init enum
 	 */
-	explicit FORCEINLINE FVector2D( EForceInit );
+	explicit FORCEINLINE FVector2D(EForceInit);
 
 	/**
 	 * Constructs a vector from an FVector.
@@ -57,7 +55,7 @@ public:
 	 *
 	 * @param V Vector to copy from.
 	 */
-	explicit FORCEINLINE FVector2D( const FVector& V );
+	explicit FORCEINLINE FVector2D(const FVector& V);
 
 public:
 
@@ -67,7 +65,7 @@ public:
 	 * @param V The other vector to add to this.
 	 * @return The result of adding the vectors together.
 	 */
-	FORCEINLINE FVector2D operator+( const FVector2D& V ) const;
+	FORCEINLINE FVector2D operator+(const FVector2D& V) const;
 
 	/**
 	 * Gets the result of subtracting a vector from this one.
@@ -75,7 +73,7 @@ public:
 	 * @param V The other vector to subtract from this.
 	 * @return The result of the subtraction.
 	 */
-	FORCEINLINE FVector2D operator-( const FVector2D& V ) const;
+	FORCEINLINE FVector2D operator-(const FVector2D& V) const;
 
 	/**
 	 * Gets the result of scaling the vector (multiplying each component by a value).
@@ -83,7 +81,7 @@ public:
 	 * @param Scale How much to scale the vector by.
 	 * @return The result of scaling this vector.
 	 */
-	FORCEINLINE FVector2D operator*( float Scale ) const;
+	FORCEINLINE FVector2D operator*(float Scale) const;
 
 	/**
 	 * Gets the result of dividing each component of the vector by a value.
@@ -91,23 +89,23 @@ public:
 	 * @param Scale How much to divide the vector by.
 	 * @return The result of division on this vector.
 	 */
-	FVector2D operator/( float Scale ) const;
+	FVector2D operator/(float Scale) const;
 
 	/**
 	 * Gets the result of this vector + float A.
 	 *
-	 * @param A		Float to add to each component.
-	 * @return		The result of this vector + float A.
+	 * @param A Float to add to each component.
+	 * @return The result of this vector + float A.
 	 */
-	FORCEINLINE FVector2D operator+( float A ) const;
+	FORCEINLINE FVector2D operator+(float A) const;
 
 	/**
 	 * Gets the result of subtracting from each component of the vector.
 	 *
-	 * @param A		Float to subtract from each component
-	 * @return		The result of this vector - float A.
+	 * @param A Float to subtract from each component
+	 * @return The result of this vector - float A.
 	 */
-	FORCEINLINE FVector2D operator-( float A ) const;
+	FORCEINLINE FVector2D operator-(float A) const;
 
 	/**
 	 * Gets the result of component-wise multiplication of this vector by another.
@@ -115,7 +113,7 @@ public:
 	 * @param V The other vector to multiply this by.
 	 * @return The result of the multiplication.
 	 */
-	FORCEINLINE FVector2D operator*( const FVector2D& V ) const;
+	FORCEINLINE FVector2D operator*(const FVector2D& V) const;
 
 	/**
 	 * Gets the result of component-wise division of this vector by another.
@@ -123,7 +121,7 @@ public:
 	 * @param V The other vector to divide this by.
 	 * @return The result of the division.
 	 */
-	FVector2D operator/( const FVector2D& V ) const;
+	FVector2D operator/(const FVector2D& V) const;
 
 	/**
 	 * Calculates dot product of this vector and another.
@@ -131,7 +129,7 @@ public:
 	 * @param V The other vector.
 	 * @return The dot product.
 	 */
-	FORCEINLINE float operator|( const FVector2D& V) const;
+	FORCEINLINE float operator|(const FVector2D& V) const;
 
 	/**
 	 * Calculates cross product of this vector and another.
@@ -139,7 +137,7 @@ public:
 	 * @param V The other vector.
 	 * @return The cross product.
 	 */
-	FORCEINLINE float operator^( const FVector2D& V) const;
+	FORCEINLINE float operator^(const FVector2D& V) const;
 
 public:
 
@@ -149,7 +147,7 @@ public:
 	 * @param V The vector to compare against.
 	 * @return true if the two vectors are equal, otherwise false.
 	 */
-	bool operator==( const FVector2D& V ) const;
+	bool operator==(const FVector2D& V) const;
 
 	/**
 	 * Compares this vector against another for inequality.
@@ -157,7 +155,7 @@ public:
 	 * @param V The vector to compare against.
 	 * @return true if the two vectors are not equal, otherwise false.
 	 */
-	bool operator!=( const FVector2D& V ) const;
+	bool operator!=(const FVector2D& V) const;
 
 	/**
 	 * Checks whether both components of this vector are less than another.
@@ -165,7 +163,7 @@ public:
 	 * @param Other The vector to compare against.
 	 * @return true if this is the smaller vector, otherwise false.
 	 */
-	bool operator<( const FVector2D& Other ) const;
+	bool operator<(const FVector2D& Other) const;
 
 	/**
 	 * Checks whether both components of this vector are greater than another.
@@ -173,7 +171,7 @@ public:
 	 * @param Other The vector to compare against.
 	 * @return true if this is the larger vector, otherwise false.
 	 */
-	bool operator>( const FVector2D& Other ) const;
+	bool operator>(const FVector2D& Other) const;
 
 	/**
 	 * Checks whether both components of this vector are less than or equal to another.
@@ -181,7 +179,7 @@ public:
 	 * @param Other The vector to compare against.
 	 * @return true if this vector is less than or equal to the other vector, otherwise false.
 	 */
-	bool operator<=( const FVector2D& Other ) const;
+	bool operator<=(const FVector2D& Other) const;
 
 	/**
 	 * Checks whether both components of this vector are greater than or equal to another.
@@ -189,7 +187,7 @@ public:
 	 * @param Other The vector to compare against.
 	 * @return true if this vector is greater than or equal to the other vector, otherwise false.
 	 */
-	bool operator>=( const FVector2D& Other ) const;
+	bool operator>=(const FVector2D& Other) const;
 
 	/**
 	 * Gets a negated copy of the vector.
@@ -204,7 +202,7 @@ public:
 	 * @param V The other vector to add.
 	 * @return Copy of the vector after addition.
 	 */
-	FORCEINLINE FVector2D operator+=( const FVector2D& V );
+	FORCEINLINE FVector2D operator+=(const FVector2D& V);
 
 	/**
 	 * Subtracts another vector from this.
@@ -212,7 +210,7 @@ public:
 	 * @param V The other vector to subtract.
 	 * @return Copy of the vector after subtraction.
 	 */
-	FORCEINLINE FVector2D operator-=( const FVector2D& V );
+	FORCEINLINE FVector2D operator-=(const FVector2D& V);
 
 	/**
 	 * Scales this vector.
@@ -220,7 +218,7 @@ public:
 	 * @param Scale The scale to multiply vector by.
 	 * @return Copy of the vector after scaling.
 	 */
-	FORCEINLINE FVector2D operator*=( float Scale );
+	FORCEINLINE FVector2D operator*=(float Scale);
 
 	/**
 	 * Divides this vector.
@@ -228,7 +226,7 @@ public:
 	 * @param V What to divide vector by.
 	 * @return Copy of the vector after division.
 	 */
-	FVector2D operator/=( float V );
+	FVector2D operator/=(float V);
 
 	/**
 	 * Multiplies this vector with another vector, using component-wise multiplication.
@@ -236,7 +234,7 @@ public:
 	 * @param V The vector to multiply with.
 	 * @return Copy of the vector after multiplication.
 	 */
-	FVector2D operator*=( const FVector2D& V );
+	FVector2D operator*=(const FVector2D& V);
 
 	/**
 	 * Divides this vector by another vector, using component-wise division.
@@ -244,7 +242,7 @@ public:
 	 * @param V The vector to divide by.
 	 * @return Copy of the vector after division.
 	 */
-	FVector2D operator/=( const FVector2D& V );
+	FVector2D operator/=(const FVector2D& V);
 
 	/**
 	 * Gets specific component of the vector.
@@ -252,7 +250,7 @@ public:
 	 * @param Index the index of vector component
 	 * @return reference to component.
 	 */
-	float& operator[]( int32 Index );
+	float& operator[](int32 Index);
 
 	/**
 	 * Gets specific component of the vector.
@@ -260,7 +258,7 @@ public:
 	 * @param Index the index of vector component
 	 * @return copy of component value.
 	 */
-	float operator[]( int32 Index ) const;
+	float operator[](int32 Index) const;
 
 	/**
 	* Gets a specific component of the vector.
@@ -287,7 +285,7 @@ public:
 	 * @param B The second vector.
 	 * @return The dot product.
 	 */
-	FORCEINLINE static float DotProduct( const FVector2D& A, const FVector2D& B );
+	FORCEINLINE static float DotProduct(const FVector2D& A, const FVector2D& B);
 
 	/**
 	 * Squared distance between two 2D points.
@@ -296,7 +294,7 @@ public:
 	 * @param V2 The second point.
 	 * @return The squared distance between two 2D points.
 	 */
-	FORCEINLINE static float DistSquared( const FVector2D& V1, const FVector2D& V2 );
+	FORCEINLINE static float DistSquared(const FVector2D& V1, const FVector2D& V2);
 
 	/**
 	 * Distance between two 2D points.
@@ -305,7 +303,7 @@ public:
 	 * @param V2 The second point.
 	 * @return The squared distance between two 2D points.
 	 */
-	FORCEINLINE static float Distance( const FVector2D& V1, const FVector2D& V2 );
+	FORCEINLINE static float Distance(const FVector2D& V1, const FVector2D& V2);
 
 	/**
 	 * Calculate the cross product of two vectors.
@@ -323,7 +321,7 @@ public:
 	 * @param Tolerance Error tolerance.
 	 * @return true if the vectors are equal within specified tolerance, otherwise false.
 	 */
-	bool Equals( const FVector2D& V, float Tolerance=KINDA_SMALL_NUMBER ) const;
+	bool Equals(const FVector2D& V, float Tolerance=KINDA_SMALL_NUMBER) const;
 
 	/**
 	 * Set the values of the vector directly.
@@ -331,7 +329,7 @@ public:
 	 * @param InX New X coordinate.
 	 * @param InY New Y coordinate.
 	 */
-	void Set( float InX, float InY );
+	void Set(float InX, float InY);
 
 	/**
 	 * Get the maximum value of the vector's components.
@@ -383,7 +381,7 @@ public:
 	 * @param Tolerance Minimum squared length of vector for normalization.
 	 * @return A normalized copy of the vector if safe, (0,0) otherwise.
 	 */
-	FVector2D GetSafeNormal( float Tolerance=SMALL_NUMBER ) const;
+	FVector2D GetSafeNormal(float Tolerance=SMALL_NUMBER) const;
 
 	DEPRECATED(4.7, "Deprecated due to unclear name, use GetSafeNormal instead.")
 	FVector2D SafeNormal(float Tolerance = SMALL_NUMBER) const;
@@ -394,7 +392,7 @@ public:
 	 * @param Tolerance Minimum squared length of vector for normalization.
 	 * @see GetSafeNormal()
 	 */
-	void Normalize( float Tolerance=SMALL_NUMBER );
+	void Normalize(float Tolerance=SMALL_NUMBER);
 
 	/**
 	 * Checks whether vector is near to zero within a specified tolerance.
@@ -402,7 +400,7 @@ public:
 	 * @param Tolerance Error tolerance.
 	 * @return true if vector is in tolerance to zero, otherwise false.
 	 */
-	bool IsNearlyZero( float Tolerance=KINDA_SMALL_NUMBER ) const;
+	bool IsNearlyZero(float Tolerance=KINDA_SMALL_NUMBER) const;
 
 	/**
 	 * Util to convert this vector into a unit direction vector and its original length.
@@ -430,7 +428,7 @@ public:
 	 * Creates a copy of this vector with both axes clamped to the given range.
 	 * @return New vector with clamped axes.
 	 */
-	FVector2D ClampAxes( float MinAxisVal, float MaxAxisVal ) const;
+	FVector2D ClampAxes(float MinAxisVal, float MaxAxisVal) const;
 
 	/**
 	* Get a copy of the vector as sign only.
@@ -461,7 +459,7 @@ public:
 	 * @param	InSourceString	FString containing the vector values.
 	 * @return true if the X,Y values were read successfully; false otherwise.
 	 */
-	bool InitFromString( const FString& InSourceString );
+	bool InitFromString(const FString& InSourceString);
 
 	/**
 	 * Serialize a vector.
@@ -470,14 +468,14 @@ public:
 	 * @param V Vector being serialized.
 	 * @return Reference to Archive after serialization.
 	 */
-	friend FArchive& operator<<( FArchive& Ar, FVector2D& V )
+	friend FArchive& operator<<(FArchive& Ar, FVector2D& V)
 	{
 		// @warning BulkSerialize: FVector2D is serialized as memory dump
 		// See TArray::BulkSerialize for detailed description of implied limitations.
 		return Ar << V.X << V.Y;
 	}
 
-	bool Serialize( FArchive& Ar )
+	bool Serialize(FArchive& Ar)
 	{
 		Ar << *this;
 		return true;
@@ -511,7 +509,7 @@ public:
 	 * Network serialization function.
 	 * FVectors NetSerialize without quantization (ie exact values are serialized).
 	 */
-	CORE_API bool NetSerialize( FArchive& Ar, class UPackageMap* Map, bool& bOutSuccess );
+	CORE_API bool NetSerialize(FArchive& Ar, class UPackageMap* Map, bool& bOutSuccess);
 
 	/** Converts spherical coordinates on the unit sphere into a Cartesian unit length vector. */
 	inline FVector SphericalToUnitCartesian() const;
@@ -521,9 +519,9 @@ public:
 /* FVector2D inline functions
  *****************************************************************************/
 
-FORCEINLINE FVector2D operator*( float Scale, const FVector2D& V )
+FORCEINLINE FVector2D operator*(float Scale, const FVector2D& V)
 {
-	return V.operator*( Scale );
+	return V.operator*(Scale);
 }
 
 template <> struct TIsPODType<FVector2D> { enum { Value = true }; };
@@ -547,74 +545,74 @@ FORCEINLINE FVector2D::FVector2D(EForceInit)
 }
 
 
-FORCEINLINE FVector2D FVector2D::operator+( const FVector2D& V ) const
+FORCEINLINE FVector2D FVector2D::operator+(const FVector2D& V) const
 {
-	return FVector2D( X + V.X, Y + V.Y );
+	return FVector2D(X + V.X, Y + V.Y);
 }
 
 
-FORCEINLINE FVector2D FVector2D::operator-( const FVector2D& V ) const
+FORCEINLINE FVector2D FVector2D::operator-(const FVector2D& V) const
 {
-	return FVector2D( X - V.X, Y - V.Y );
+	return FVector2D(X - V.X, Y - V.Y);
 }
 
 
-FORCEINLINE FVector2D FVector2D::operator*( float Scale ) const
+FORCEINLINE FVector2D FVector2D::operator*(float Scale) const
 {
-	return FVector2D( X * Scale, Y * Scale );
+	return FVector2D(X * Scale, Y * Scale);
 }
 
 
-FORCEINLINE FVector2D FVector2D::operator/( float Scale ) const
+FORCEINLINE FVector2D FVector2D::operator/(float Scale) const
 {
 	const float RScale = 1.f/Scale;
-	return FVector2D( X * RScale, Y * RScale );
+	return FVector2D(X * RScale, Y * RScale);
 }
 
 
-FORCEINLINE FVector2D FVector2D::operator+( float A ) const
+FORCEINLINE FVector2D FVector2D::operator+(float A) const
 {
-	return FVector2D( X + A, Y + A );
+	return FVector2D(X + A, Y + A);
 }
 
 
-FORCEINLINE FVector2D FVector2D::operator-( float A ) const
+FORCEINLINE FVector2D FVector2D::operator-(float A) const
 {
-	return FVector2D( X - A, Y - A );
+	return FVector2D(X - A, Y - A);
 }
 
 
-FORCEINLINE FVector2D FVector2D::operator*( const FVector2D& V ) const
+FORCEINLINE FVector2D FVector2D::operator*(const FVector2D& V) const
 {
-	return FVector2D( X * V.X, Y * V.Y );
+	return FVector2D(X * V.X, Y * V.Y);
 }
 
 
-FORCEINLINE FVector2D FVector2D::operator/( const FVector2D& V ) const
+FORCEINLINE FVector2D FVector2D::operator/(const FVector2D& V) const
 {
-	return FVector2D( X / V.X, Y / V.Y );
+	return FVector2D(X / V.X, Y / V.Y);
 }
 
 
-FORCEINLINE float FVector2D::operator|( const FVector2D& V) const
+FORCEINLINE float FVector2D::operator|(const FVector2D& V) const
 {
 	return X*V.X + Y*V.Y;
 }
 
 
-FORCEINLINE float FVector2D::operator^( const FVector2D& V) const
+FORCEINLINE float FVector2D::operator^(const FVector2D& V) const
 {
 	return X*V.Y - Y*V.X;
 }
 
 
-FORCEINLINE float FVector2D::DotProduct( const FVector2D& A, const FVector2D& B )
+FORCEINLINE float FVector2D::DotProduct(const FVector2D& A, const FVector2D& B)
 {
 	return A | B;
 }
 
 
-FORCEINLINE float FVector2D::DistSquared( const FVector2D &V1, const FVector2D &V2 )
+FORCEINLINE float FVector2D::DistSquared(const FVector2D &V1, const FVector2D &V2)
 {
 	return FMath::Square(V2.X-V1.X) + FMath::Square(V2.Y-V1.Y);
 }
@@ -632,37 +630,37 @@ FORCEINLINE float FVector2D::CrossProduct(const FVector2D& A, const FVector2D& B
 }
 
 
-FORCEINLINE bool FVector2D::operator==( const FVector2D& V ) const
+FORCEINLINE bool FVector2D::operator==(const FVector2D& V) const
 {
 	return X==V.X && Y==V.Y;
 }
 
 
-FORCEINLINE bool FVector2D::operator!=( const FVector2D& V ) const
+FORCEINLINE bool FVector2D::operator!=(const FVector2D& V) const
 {
 	return X!=V.X || Y!=V.Y;
 }
 
 
-FORCEINLINE bool FVector2D::operator<( const FVector2D& Other ) const
+FORCEINLINE bool FVector2D::operator<(const FVector2D& Other) const
 {
 	return X < Other.X && Y < Other.Y;
 }
 
 
-FORCEINLINE bool FVector2D::operator>( const FVector2D& Other ) const
+FORCEINLINE bool FVector2D::operator>(const FVector2D& Other) const
 {
 	return X > Other.X && Y > Other.Y;
 }
 
 
-FORCEINLINE bool FVector2D::operator<=( const FVector2D& Other ) const
+FORCEINLINE bool FVector2D::operator<=(const FVector2D& Other) const
 {
 	return X <= Other.X && Y <= Other.Y;
 }
 
 
-FORCEINLINE bool FVector2D::operator>=( const FVector2D& Other ) const
+FORCEINLINE bool FVector2D::operator>=(const FVector2D& Other) const
 {
 	return X >= Other.X && Y >= Other.Y;
 }
@@ -676,32 +674,32 @@ FORCEINLINE bool FVector2D::Equals(const FVector2D& V, float Tolerance) const
 
 FORCEINLINE FVector2D FVector2D::operator-() const
 {
-	return FVector2D( -X, -Y );
+	return FVector2D(-X, -Y);
 }
 
 
-FORCEINLINE FVector2D FVector2D::operator+=( const FVector2D& V )
+FORCEINLINE FVector2D FVector2D::operator+=(const FVector2D& V)
 {
 	X += V.X; Y += V.Y;
 	return *this;
 }
 
 
-FORCEINLINE FVector2D FVector2D::operator-=( const FVector2D& V )
+FORCEINLINE FVector2D FVector2D::operator-=(const FVector2D& V)
 {
 	X -= V.X; Y -= V.Y;
 	return *this;
 }
 
 
-FORCEINLINE FVector2D FVector2D::operator*=( float Scale )
+FORCEINLINE FVector2D FVector2D::operator*=(float Scale)
 {
 	X *= Scale; Y *= Scale;
 	return *this;
 }
 
 
-FORCEINLINE FVector2D FVector2D::operator/=( float V )
+FORCEINLINE FVector2D FVector2D::operator/=(float V)
 {
 	const float RV = 1.f/V;
 	X *= RV; Y *= RV;
@@ -709,35 +707,35 @@ FORCEINLINE FVector2D FVector2D::operator/=( float V )
 }
 
 
-FORCEINLINE FVector2D FVector2D::operator*=( const FVector2D& V )
+FORCEINLINE FVector2D FVector2D::operator*=(const FVector2D& V)
 {
 	X *= V.X; Y *= V.Y;
 	return *this;
 }
 
 
-FORCEINLINE FVector2D FVector2D::operator/=( const FVector2D& V )
+FORCEINLINE FVector2D FVector2D::operator/=(const FVector2D& V)
 {
 	X /= V.X; Y /= V.Y;
 	return *this;
 }
 
 
-FORCEINLINE float& FVector2D::operator[]( int32 Index )
+FORCEINLINE float& FVector2D::operator[](int32 Index)
 {
 	check(Index>=0 && Index<2);
 	return ((Index == 0) ? X : Y);
 }
 
 
-FORCEINLINE float FVector2D::operator[]( int32 Index ) const
+FORCEINLINE float FVector2D::operator[](int32 Index) const
 {
 	check(Index>=0 && Index<2);
 	return ((Index == 0) ? X : Y);
 }
 
 
-FORCEINLINE void FVector2D::Set( float InX, float InY )
+FORCEINLINE void FVector2D::Set(float InX, float InY)
 {
 	X = InX;
 	Y = InY;
@@ -764,7 +762,7 @@ FORCEINLINE float FVector2D::GetMin() const
 
 FORCEINLINE float FVector2D::Size() const
 {
-	return FMath::Sqrt( X*X + Y*Y );
+	return FMath::Sqrt(X*X + Y*Y);
 }
 
 
@@ -773,6 +771,7 @@ FORCEINLINE float FVector2D::SizeSquared() const
 	return X*X + Y*Y;
 }
 
+
 FORCEINLINE FVector2D FVector2D::GetRotated(const float AngleDeg) const
 {
 	// Based on FVector::RotateAngleAxis with Axis(0,0,1)
@@ -787,10 +786,11 @@ FORCEINLINE FVector2D FVector2D::GetRotated(const float AngleDeg) const
 		S * X + C * Y);
 }
 
+
 FORCEINLINE FVector2D FVector2D::GetSafeNormal(float Tolerance) const
 {	
 	const float SquareSum = X*X + Y*Y;
-	if( SquareSum > Tolerance )
+	if(SquareSum > Tolerance)
 	{
 		const float Scale = FMath::InvSqrt(SquareSum);
 		return FVector2D(X*Scale, Y*Scale);
@@ -808,7 +808,7 @@ FORCEINLINE FVector2D FVector2D::SafeNormal(float Tolerance) const
 FORCEINLINE void FVector2D::Normalize(float Tolerance)
 {
 	const float SquareSum = X*X + Y*Y;
-	if( SquareSum > Tolerance )
+	if(SquareSum > Tolerance)
 	{
 		const float Scale = FMath::InvSqrt(SquareSum);
 		X *= Scale;
@@ -848,13 +848,13 @@ FORCEINLINE bool FVector2D::IsZero() const
 }
 
 
-FORCEINLINE float& FVector2D::Component( int32 Index )
+FORCEINLINE float& FVector2D::Component(int32 Index)
 {
 	return (&X)[Index];
 }
 
 
-FORCEINLINE float FVector2D::Component( int32 Index ) const
+FORCEINLINE float FVector2D::Component(int32 Index) const
 {
 	return (&X)[Index];
 }
@@ -862,13 +862,13 @@ FORCEINLINE float FVector2D::Component( int32 Index ) const
 
 FORCEINLINE FIntPoint FVector2D::IntPoint() const
 {
-	return FIntPoint( FMath::RoundToInt(X), FMath::RoundToInt(Y) );
+	return FIntPoint(FMath::RoundToInt(X), FMath::RoundToInt(Y));
 }
 
 
-FORCEINLINE FVector2D FVector2D::ClampAxes( float MinAxisVal, float MaxAxisVal ) const
+FORCEINLINE FVector2D FVector2D::ClampAxes(float MinAxisVal, float MaxAxisVal) const
 {
-	return FVector2D( FMath::Clamp(X, MinAxisVal, MaxAxisVal), FMath::Clamp(Y, MinAxisVal, MaxAxisVal) );
+	return FVector2D(FMath::Clamp(X, MinAxisVal, MaxAxisVal), FMath::Clamp(Y, MinAxisVal, MaxAxisVal));
 }
 
 
@@ -886,18 +886,19 @@ FORCEINLINE FVector2D FVector2D::GetAbs() const
 	return FVector2D(FMath::Abs(X), FMath::Abs(Y));
 }
 
+
 FORCEINLINE FString FVector2D::ToString() const
 {
 	return FString::Printf(TEXT("X=%3.3f Y=%3.3f"), X, Y);
 }
 
 
-FORCEINLINE bool FVector2D::InitFromString( const FString& InSourceString )
+FORCEINLINE bool FVector2D::InitFromString(const FString& InSourceString)
 {
 	X = Y = 0;
 
 	// The initialization is only successful if the X and Y values can all be parsed from the string
-	const bool bSuccessful = FParse::Value( *InSourceString, TEXT("X=") , X ) && FParse::Value( *InSourceString, TEXT("Y="), Y ) ;
+	const bool bSuccessful = FParse::Value(*InSourceString, TEXT("X=") , X) && FParse::Value(*InSourceString, TEXT("Y="), Y) ;
 
 	return bSuccessful;
 }
diff --git a/Engine/Source/Runtime/Core/Public/Math/Vector4.h b/Engine/Source/Runtime/Core/Public/Math/Vector4.h
--- a/Engine/Source/Runtime/Core/Public/Math/Vector4.h
+++ b/Engine/Source/Runtime/Core/Public/Math/Vector4.h
@@ -10,16 +10,16 @@ MS_ALIGN(16) struct FVector4
 {
 public:
 
-	/** Holds the vector's X-component. */
+	/** The vector's X-component. */
 	float X;
 
-	/** Holds the vector's Y-component. */
+	/** The vector's Y-component. */
 	float Y;
 
-	/** Holds the vector's Z-component. */
+	/** The vector's Z-component. */
 	float Z;
 
-	/** Holds the vector's W-component. */
+	/** The vector's W-component. */
 	float W;
 
 public:
@@ -30,14 +30,14 @@ public:
 	 * @param InVector 3D Vector to set first three components.
 	 * @param InW W Coordinate.
 	 */
-	FVector4( const FVector& InVector, float InW = 1.0f );
+	FVector4(const FVector& InVector, float InW = 1.0f);
 
 	/**
 	 * Creates and initializes a new vector from a color value.
 	 *
 	 * @param InColour Color used to set vector.
 	 */
-	FVector4( const FLinearColor& InColor );
+	FVector4(const FLinearColor& InColor);
 
 	/**
 	 * Creates and initializes a new vector from the specified components.
@@ -47,7 +47,7 @@ public:
 	 * @param InZ Z Coordinate.
 	 * @param InW W Coordinate.
 	 */
-	explicit FVector4( float InX = 0.0f, float InY = 0.0f, float InZ = 0.0f, float InW = 1.0f );
+	explicit FVector4(float InX = 0.0f, float InY = 0.0f, float InZ = 0.0f, float InW = 1.0f);
 
 	/**
 	 * Creates and initializes a new vector from the specified 2D vectors.
@@ -55,14 +55,14 @@ public:
 	 * @param InXY A 2D vector holding the X- and Y-components.
 	 * @param InZW A 2D vector holding the Z- and W-components.
 	 */
-	explicit FVector4( FVector2D InXY, FVector2D InZW );
+	explicit FVector4(FVector2D InXY, FVector2D InZW);
 
 	/**
 	 * Creates and initializes a new vector to zero.
 	 *
 	 * @param EForceInit Force Init Enum.
 	 */
-	explicit FVector4( EForceInit );
+	explicit FVector4(EForceInit);
 
 public:
 
@@ -72,7 +72,7 @@ public:
 	 * @param ComponentIndex The index of the component.
 	 * @return Reference to the desired component.
 	 */
-	FORCEINLINE float& operator[]( int32 ComponentIndex );
+	FORCEINLINE float& operator[](int32 ComponentIndex);
 
 	/**
 	 * Access a specific component of the vector.
@@ -80,7 +80,7 @@ public:
 	 * @param ComponentIndex The index of the component.
 	 * @return Copy of the desired component.
 	 */
-	FORCEINLINE float operator[]( int32 ComponentIndex ) const;
+	FORCEINLINE float operator[](int32 ComponentIndex) const;
 
 	// Unary operators.
 
@@ -97,7 +97,7 @@ public:
 	 * @param V The vector to add.
 	 * @return The result of vector addition.
 	 */
-	FORCEINLINE FVector4 operator+( const FVector4& V ) const;
+	FORCEINLINE FVector4 operator+(const FVector4& V) const;
 
 	/**
 	 * Adds another vector to this one.
@@ -105,7 +105,7 @@ public:
 	 * @param V The other vector to add.
 	 * @return Copy of the vector after addition.
 	 */
-	FORCEINLINE FVector4 operator+=( const FVector4& V );
+	FORCEINLINE FVector4 operator+=(const FVector4& V);
 
 	/**
 	 * Gets the result of subtracting a vector from this.
@@ -113,7 +113,7 @@ public:
 	 * @param V The vector to subtract.
 	 * @return The result of vector subtraction.
 	 */
-	FORCEINLINE FVector4 operator-( const FVector4& V ) const;
+	FORCEINLINE FVector4 operator-(const FVector4& V) const;
 
 	/**
 	 * Gets the result of scaling this vector.
@@ -121,7 +121,7 @@ public:
 	 * @param Scale The scaling factor.
 	 * @return The result of vector scaling.
 	 */
-	FORCEINLINE FVector4 operator*( float Scale ) const;
+	FORCEINLINE FVector4 operator*(float Scale) const;
 
 	/**
 	 * Gets the result of dividing this vector.
@@ -129,7 +129,7 @@ public:
 	 * @param Scale What to divide by.
 	 * @return The result of division.
 	 */
-	FVector4 operator/( float Scale ) const;
+	FVector4 operator/(float Scale) const;
 
 	/**
 	 * Gets the result of dividing this vector.
@@ -137,7 +137,7 @@ public:
 	 * @param V What to divide by.
 	 * @return The result of division.
 	 */
-	FVector4 operator/( const FVector4& V ) const;
+	FVector4 operator/(const FVector4& V) const;
 
 	/**
 	 * Gets the result of multiplying a vector with this.
@@ -145,7 +145,7 @@ public:
 	 * @param V The vector to multiply.
 	 * @return The result of vector multiplication.
 	 */
-	FVector4 operator*( const FVector4& V ) const;
+	FVector4 operator*(const FVector4& V) const;
 
 	/**
 	 * Gets the result of multiplying a vector with another Vector (component wise).
@@ -153,7 +153,7 @@ public:
 	 * @param V The vector to multiply.
 	 * @return The result of vector multiplication.
 	 */
-	FVector4 operator*=( const FVector4& V );
+	FVector4 operator*=(const FVector4& V);
 
 	/**
 	 * Gets the result of dividing a vector with another Vector (component wise).
@@ -161,7 +161,7 @@ public:
 	 * @param V The vector to divide with.
 	 * @return The result of vector multiplication.
 	 */
-	FVector4 operator/=( const FVector4& V );
+	FVector4 operator/=(const FVector4& V);
 
 	/**
 	 * Gets the result of scaling this vector.
@@ -169,7 +169,7 @@ public:
 	 * @param Scale The scaling factor.
 	 * @return The result of vector scaling.
 	 */
-	FVector4 operator*=( float S );
+	FVector4 operator*=(float S);
 
 	/**
 	 * Calculates 3D Dot product of two 4D vectors.
@@ -178,7 +178,7 @@ public:
 	 * @param V2 The second vector.
 	 * @return The 3D Dot product.
 	 */
-	friend FORCEINLINE float Dot3( const FVector4& V1, const FVector4& V2 )
+	friend FORCEINLINE float Dot3(const FVector4& V1, const FVector4& V2)
 	{
 		return V1.X*V2.X + V1.Y*V2.Y + V1.Z*V2.Z;
 	}
@@ -190,7 +190,7 @@ public:
 	 * @param V2 The second vector.
 	 * @return The 4D Dot Product.
 	 */
-	friend FORCEINLINE float Dot4( const FVector4& V1, const FVector4& V2 )
+	friend FORCEINLINE float Dot4(const FVector4& V1, const FVector4& V2)
 	{
 		return V1.X*V2.X + V1.Y*V2.Y + V1.Z*V2.Z + V1.W*V2.W;
 	}
@@ -202,9 +202,9 @@ public:
 	 * @param V The vector to scale.
 	 * @return The result of scaling.
 	 */
-	friend FORCEINLINE FVector4 operator*( float Scale, const FVector4& V )
+	friend FORCEINLINE FVector4 operator*(float Scale, const FVector4& V)
 	{
-		return V.operator*( Scale );
+		return V.operator*(Scale);
 	}
 
 	/**
@@ -213,7 +213,7 @@ public:
 	 * @param V The other vector.
 	 * @return true if the two vectors are the same, otherwise false.
 	 */
-	bool operator==( const FVector4& V ) const;
+	bool operator==(const FVector4& V) const;
 
 	/**
 	 * Checks for inequality against another vector.
@@ -221,7 +221,7 @@ public:
 	 * @param V The other vector.
 	 * @return true if the two vectors are different, otherwise false.
 	 */
-	bool operator!=( const FVector4& V ) const;
+	bool operator!=(const FVector4& V) const;
 
 	/**
 	 * Calculate Cross product between this and another vector.
@@ -229,7 +229,7 @@ public:
 	 * @param V The other vector.
 	 * @return The Cross product.
 	 */
-	FVector4 operator^( const FVector4& V ) const;
+	FVector4 operator^(const FVector4& V) const;
 
 public:
 
@@ -241,7 +241,7 @@ public:
 	 * @param Index The index of the component.
 	 * @return Reference to the component.
 	 */
-	float& Component( int32 Index );
+	float& Component(int32 Index);
 
 	/**
 	 * Error tolerant comparison.
@@ -250,7 +250,7 @@ public:
 	 * @param Tolerance Error Tolerance.
 	 * @return true if the two vectors are equal within specified tolerance, otherwise false.
 	 */
-	bool Equals( const FVector4& V, float Tolerance=KINDA_SMALL_NUMBER ) const;
+	bool Equals(const FVector4& V, float Tolerance=KINDA_SMALL_NUMBER) const;
 
 	/**
 	 * Check if the vector is of unit length, with specified tolerance.
@@ -258,7 +258,7 @@ public:
 	 * @param LengthSquaredTolerance Tolerance against squared length.
 	 * @return true if the vector is a unit vector within the specified tolerance.
 	 */
-	bool IsUnit3( float LengthSquaredTolerance = KINDA_SMALL_NUMBER ) const;
+	bool IsUnit3(float LengthSquaredTolerance = KINDA_SMALL_NUMBER) const;
 
 	/**
 	 * Get a textual representation of the vector.
@@ -274,7 +274,7 @@ public:
 	 * @param InSourceString	FString containing the vector values.
 	 * @return true if the X,Y,Z values were read successfully; false otherwise.
 	 */
-	bool InitFromString( const FString& InSourceString );
+	bool InitFromString(const FString& InSourceString);
 
 	/**
 	 * Returns a normalized copy of the vector if safe to normalize.
@@ -285,7 +285,7 @@ public:
 	FORCEINLINE FVector4 GetSafeNormal(float Tolerance = SMALL_NUMBER) const;
 
 	DEPRECATED(4.7, "Deprecated due to unclear name, use GetSafeNormal instead.")
-	FORCEINLINE FVector4 SafeNormal( float Tolerance=SMALL_NUMBER ) const;
+	FORCEINLINE FVector4 SafeNormal(float Tolerance=SMALL_NUMBER) const;
 
 	/**
 	 * Calculates normalized version of vector without checking if it is non-zero.
@@ -327,7 +327,7 @@ public:
 	 * @param InZ New Z Coordinate.
 	 * @param InW New W Coordinate.
 	 */
-	FORCEINLINE void Set( float InX, float InY, float InZ, float InW );
+	FORCEINLINE void Set(float InX, float InY, float InZ, float InW);
 
 	/**
 	 * Get the length of this vector not taking W component into account.
@@ -347,16 +347,16 @@ public:
 	bool ContainsNaN() const;
 
 	/** Utility to check if all of the components of this vector are nearly zero given the tolerance. */
-	bool IsNearlyZero3( float Tolerance = KINDA_SMALL_NUMBER ) const;
+	bool IsNearlyZero3(float Tolerance = KINDA_SMALL_NUMBER) const;
 
 	/** Reflect vector. */
-	FVector4 Reflect3( const FVector4& Normal ) const;
+	FVector4 Reflect3(const FVector4& Normal) const;
 
 	/**
 	 * Find good arbitrary axis vectors to represent U and V axes of a plane,
 	 * given just the normal.
 	 */
-	void FindBestAxisVectors3( FVector4& Axis1, FVector4& Axis2 ) const;
+	void FindBestAxisVectors3(FVector4& Axis1, FVector4& Axis2) const;
 
 #if ENABLE_NAN_DIAGNOSTIC
 	FORCEINLINE void DiagnosticCheckNaN()
@@ -381,12 +381,12 @@ public:
 	 * @param V The vector being serialized.
 	 * @return Reference to the Archive after serialization.
 	 */
-	friend FArchive& operator<<( FArchive& Ar, FVector4& V )
+	friend FArchive& operator<<(FArchive& Ar, FVector4& V)
 	{
 		return Ar << V.X << V.Y << V.Z << V.W;
 	}
 
-	bool Serialize( FArchive& Ar )
+	bool Serialize(FArchive& Ar)
 	{
 		Ar << *this;
 		return true;
@@ -412,7 +412,7 @@ FORCEINLINE uint32 GetTypeHash(const FVector4& Vector)
 /* FVector4 inline functions
  *****************************************************************************/
 
-FORCEINLINE FVector4::FVector4( const FVector& InVector,float InW )
+FORCEINLINE FVector4::FVector4(const FVector& InVector,float InW)
 	: X(InVector.X)
 	, Y(InVector.Y)
 	, Z(InVector.Z)
@@ -422,7 +422,7 @@ FORCEINLINE FVector4::FVector4( const FVector& InVector,float InW )
 }
 
 
-FORCEINLINE FVector4::FVector4( const FLinearColor& InColor )
+FORCEINLINE FVector4::FVector4(const FLinearColor& InColor)
 	: X(InColor.R)
 	, Y(InColor.G)
 	, Z(InColor.B)
@@ -432,7 +432,7 @@ FORCEINLINE FVector4::FVector4( const FLinearColor& InColor )
 }
 
 
-FORCEINLINE FVector4::FVector4( float InX,float InY,float InZ,float InW )
+FORCEINLINE FVector4::FVector4(float InX,float InY,float InZ,float InW)
 	: X(InX)
 	, Y(InY)
 	, Z(InZ)
@@ -442,7 +442,7 @@ FORCEINLINE FVector4::FVector4( float InX,float InY,float InZ,float InW )
 }
 
 
-FORCEINLINE FVector4::FVector4( EForceInit )
+FORCEINLINE FVector4::FVector4(EForceInit)
 	: X(0.f)
 	, Y(0.f)
 	, Z(0.f)
@@ -452,7 +452,7 @@ FORCEINLINE FVector4::FVector4( EForceInit )
 }
 
 
-FORCEINLINE FVector4::FVector4( FVector2D InXY, FVector2D InZW )
+FORCEINLINE FVector4::FVector4(FVector2D InXY, FVector2D InZW)
 	: X(InXY.X)
 	, Y(InXY.Y)
 	, Z(InZW.X)
@@ -462,19 +462,19 @@ FORCEINLINE FVector4::FVector4( FVector2D InXY, FVector2D InZW )
 }
 
 
-FORCEINLINE float& FVector4::operator[]( int32 ComponentIndex )
+FORCEINLINE float& FVector4::operator[](int32 ComponentIndex)
 {
 	return (&X)[ ComponentIndex ];
 }
 
 
-FORCEINLINE float FVector4::operator[]( int32 ComponentIndex ) const
+FORCEINLINE float FVector4::operator[](int32 ComponentIndex) const
 {
 	return (&X)[ ComponentIndex ];
 }
 
 
-FORCEINLINE void FVector4::Set( float InX, float InY, float InZ, float InW )
+FORCEINLINE void FVector4::Set(float InX, float InY, float InZ, float InW)
 {
 	X = InX;
 	Y = InY;
@@ -487,17 +487,17 @@ FORCEINLINE void FVector4::Set( float InX, float InY, float InZ, float InW )
 
 FORCEINLINE FVector4 FVector4::operator-() const
 {
-	return FVector4( -X, -Y, -Z, -W );
+	return FVector4(-X, -Y, -Z, -W);
 }
 
 
-FORCEINLINE FVector4 FVector4::operator+( const FVector4& V ) const
+FORCEINLINE FVector4 FVector4::operator+(const FVector4& V) const
 {
-	return FVector4( X + V.X, Y + V.Y, Z + V.Z, W + V.W );
+	return FVector4(X + V.X, Y + V.Y, Z + V.Z, W + V.W);
 }
 
 
-FORCEINLINE FVector4 FVector4::operator+=( const FVector4& V )
+FORCEINLINE FVector4 FVector4::operator+=(const FVector4& V)
 {
 	X += V.X; Y += V.Y; Z += V.Z; W += V.W;
 	DiagnosticCheckNaN();
@@ -505,32 +505,32 @@ FORCEINLINE FVector4 FVector4::operator+=( const FVector4& V )
 }
 
 
-FORCEINLINE FVector4 FVector4::operator-( const FVector4& V ) const
+FORCEINLINE FVector4 FVector4::operator-(const FVector4& V) const
 {
-	return FVector4( X - V.X, Y - V.Y, Z - V.Z, W - V.W );
+	return FVector4(X - V.X, Y - V.Y, Z - V.Z, W - V.W);
 }
 
 
-FORCEINLINE FVector4 FVector4::operator*( float Scale ) const
+FORCEINLINE FVector4 FVector4::operator*(float Scale) const
 {
-	return FVector4( X * Scale, Y * Scale, Z * Scale, W * Scale );
+	return FVector4(X * Scale, Y * Scale, Z * Scale, W * Scale);
 }
 
 
-FORCEINLINE FVector4 FVector4::operator/( float Scale ) const
+FORCEINLINE FVector4 FVector4::operator/(float Scale) const
 {
 	const float RScale = 1.f/Scale;
-	return FVector4( X * RScale, Y * RScale, Z * RScale, W * RScale );
+	return FVector4(X * RScale, Y * RScale, Z * RScale, W * RScale);
 }
 
 
-FORCEINLINE FVector4 FVector4::operator*( const FVector4& V ) const
+FORCEINLINE FVector4 FVector4::operator*(const FVector4& V) const
 {
-	return FVector4( X * V.X, Y * V.Y, Z * V.Z, W * V.W );
+	return FVector4(X * V.X, Y * V.Y, Z * V.Z, W * V.W);
 }
 
 
-FORCEINLINE FVector4 FVector4::operator^( const FVector4& V ) const
+FORCEINLINE FVector4 FVector4::operator^(const FVector4& V) const
 {
 	return FVector4(
 		Y * V.Z - Z * V.Y,
@@ -541,25 +541,25 @@ FORCEINLINE FVector4 FVector4::operator^( const FVector4& V ) const
 }
 
 
-FORCEINLINE float& FVector4::Component( int32 Index )
+FORCEINLINE float& FVector4::Component(int32 Index)
 {
 	return (&X)[Index];
 }
 
 
-FORCEINLINE bool FVector4::operator==( const FVector4& V ) const
+FORCEINLINE bool FVector4::operator==(const FVector4& V) const
 {
 	return ((X == V.X) && (Y == V.Y) && (Z == V.Z) && (W == V.W));
 }
 
 
-FORCEINLINE bool FVector4::operator!=( const FVector4& V ) const
+FORCEINLINE bool FVector4::operator!=(const FVector4& V) const
 {
 	return ((X != V.X) || (Y != V.Y) || (Z != V.Z) || (W != V.W));
 }
 
 
-FORCEINLINE bool FVector4::Equals( const FVector4& V, float Tolerance ) const
+FORCEINLINE bool FVector4::Equals(const FVector4& V, float Tolerance) const
 {
 	return FMath::Abs(X-V.X) <= Tolerance && FMath::Abs(Y-V.Y) <= Tolerance && FMath::Abs(Z-V.Z) <= Tolerance && FMath::Abs(W-V.W) <= Tolerance;
 }
@@ -571,25 +571,25 @@ FORCEINLINE FString FVector4::ToString() const
 }
 
 
-FORCEINLINE bool FVector4::InitFromString( const FString& InSourceString )
+FORCEINLINE bool FVector4::InitFromString(const FString& InSourceString)
 {
 	X = Y = Z = 0;
 	W = 1.0f;
 
 	// The initialization is only successful if the X, Y, and Z values can all be parsed from the string
-	const bool bSuccessful = FParse::Value( *InSourceString, TEXT("X=") , X ) && FParse::Value( *InSourceString, TEXT("Y="), Y ) && FParse::Value( *InSourceString, TEXT("Z="), Z );
+	const bool bSuccessful = FParse::Value(*InSourceString, TEXT("X=") , X) && FParse::Value(*InSourceString, TEXT("Y="), Y) && FParse::Value(*InSourceString, TEXT("Z="), Z);
 
 	// W is optional, so don't factor in its presence (or lack thereof) in determining initialization success
-	FParse::Value( *InSourceString, TEXT("W="), W );
+	FParse::Value(*InSourceString, TEXT("W="), W);
 
 	return bSuccessful;
 }
 
 
-FORCEINLINE FVector4 FVector4::GetSafeNormal( float Tolerance ) const
+FORCEINLINE FVector4 FVector4::GetSafeNormal(float Tolerance) const
 {
 	const float SquareSum = X*X + Y*Y + Z*Z;
-	if( SquareSum > Tolerance )
+	if(SquareSum > Tolerance)
 	{
 		const float Scale = FMath::InvSqrt(SquareSum);
 		return FVector4(X*Scale, Y*Scale, Z*Scale, 0.0f);
@@ -607,7 +607,7 @@ FORCEINLINE FVector4 FVector4::SafeNormal(float Tolerance) const
 FORCEINLINE FVector4 FVector4::GetUnsafeNormal3() const
 {
 	const float Scale = FMath::InvSqrt(X*X+Y*Y+Z*Z);
-	return FVector4( X*Scale, Y*Scale, Z*Scale, 0.0f );
+	return FVector4(X*Scale, Y*Scale, Z*Scale, 0.0f);
 }
 
 
@@ -619,7 +619,7 @@ FORCEINLINE FVector4 FVector4::UnsafeNormal3() const
 
 FORCEINLINE float FVector4::Size3() const
 {
-	return FMath::Sqrt( X*X + Y*Y + Z*Z );
+	return FMath::Sqrt(X*X + Y*Y + Z*Z);
 }
 
 
@@ -629,7 +629,7 @@ FORCEINLINE float FVector4::SizeSquared3() const
 }
 
 
-FORCEINLINE bool FVector4::IsUnit3( float LengthSquaredTolerance ) const
+FORCEINLINE bool FVector4::IsUnit3(float LengthSquaredTolerance) const
 {
 	return FMath::Abs(1.0f - SizeSquared3()) < LengthSquaredTolerance;
 }
@@ -644,7 +644,7 @@ FORCEINLINE bool FVector4::ContainsNaN() const
 }
 
 
-FORCEINLINE bool FVector4::IsNearlyZero3( float Tolerance ) const
+FORCEINLINE bool FVector4::IsNearlyZero3(float Tolerance) const
 {
 	return
 			FMath::Abs(X)<=Tolerance
@@ -653,20 +653,20 @@ FORCEINLINE bool FVector4::IsNearlyZero3( float Tolerance ) const
 }
 
 
-FORCEINLINE FVector4 FVector4::Reflect3( const FVector4& Normal ) const
+FORCEINLINE FVector4 FVector4::Reflect3(const FVector4& Normal) const
 {
 	return 2.0f * Dot3(*this, Normal) * Normal - *this;
 }
 
 
-FORCEINLINE void FVector4::FindBestAxisVectors3( FVector4& Axis1, FVector4& Axis2 ) const
+FORCEINLINE void FVector4::FindBestAxisVectors3(FVector4& Axis1, FVector4& Axis2) const
 {
 	const float NX = FMath::Abs(X);
 	const float NY = FMath::Abs(Y);
 	const float NZ = FMath::Abs(Z);
 
 	// Find best basis vectors.
-	if( NZ>NX && NZ>NY )	Axis1 = FVector4(1,0,0);
+	if(NZ>NX && NZ>NY)	Axis1 = FVector4(1,0,0);
 	else					Axis1 = FVector4(0,0,1);
 
 	Axis1 = (Axis1 - *this * Dot3(Axis1, *this)).GetSafeNormal();
@@ -674,7 +674,7 @@ FORCEINLINE void FVector4::FindBestAxisVectors3( FVector4& Axis1, FVector4& Axis
 }
 
 
-FORCEINLINE FVector4 FVector4::operator*=( const FVector4& V )
+FORCEINLINE FVector4 FVector4::operator*=(const FVector4& V)
 {
 	X *= V.X; Y *= V.Y; Z *= V.Z; W *= V.W;
 	DiagnosticCheckNaN();
@@ -682,7 +682,7 @@ FORCEINLINE FVector4 FVector4::operator*=( const FVector4& V )
 }
 
 
-FORCEINLINE FVector4 FVector4::operator/=( const FVector4& V )
+FORCEINLINE FVector4 FVector4::operator/=(const FVector4& V)
 {
 	X /= V.X; Y /= V.Y; Z /= V.Z; W /= V.W;
 	DiagnosticCheckNaN();
@@ -690,7 +690,7 @@ FORCEINLINE FVector4 FVector4::operator/=( const FVector4& V )
 }
 
 
-FORCEINLINE FVector4 FVector4::operator*=( float S )
+FORCEINLINE FVector4 FVector4::operator*=(float S)
 {
 	X *= S; Y *= S; Z *= S; W *= S;
 	DiagnosticCheckNaN();
@@ -698,9 +698,9 @@ FORCEINLINE FVector4 FVector4::operator*=( float S )
 }
 
 
-FORCEINLINE FVector4 FVector4::operator/( const FVector4& V ) const
+FORCEINLINE FVector4 FVector4::operator/(const FVector4& V) const
 {
-	return FVector4( X / V.X, Y / V.Y, Z / V.Z, W / V.W );
+	return FVector4(X / V.X, Y / V.Y, Z / V.Z, W / V.W);
 }
 
 template <> struct TIsPODType<FVector4> { enum { Value = true }; };
\ No newline at end of file
diff --git a/Engine/Source/Runtime/Engine/Classes/Engine/TextureLODSettings.h b/Engine/Source/Runtime/Engine/Classes/Engine/TextureLODSettings.h
--- a/Engine/Source/Runtime/Engine/Classes/Engine/TextureLODSettings.h
+++ b/Engine/Source/Runtime/Engine/Classes/Engine/TextureLODSettings.h
@@ -1,11 +1,14 @@
 // Copyright 1998-2016 Epic Games, Inc. All Rights Reserved.
 
-
 #pragma once
 
 #include "TextureDefines.h"
 #include "TextureLODSettings.generated.h"
 
+
+class UTexture;
+
+
 /** LOD settings for a single texture group. */
 USTRUCT()
 struct ENGINE_API FTextureLODGroup
diff --git a/Engine/Source/Runtime/Engine/Private/BufferVisualizationData.cpp b/Engine/Source/Runtime/Engine/Private/BufferVisualizationData.cpp
--- a/Engine/Source/Runtime/Engine/Private/BufferVisualizationData.cpp
+++ b/Engine/Source/Runtime/Engine/Private/BufferVisualizationData.cpp
@@ -95,7 +95,7 @@ void FBufferVisualizationData::ConfigureConsoleCommand()
 	ConsoleDocumentationOverviewTargets += AvailableVisualizationMaterials;
 
 	IConsoleManager::Get().RegisterConsoleVariable(TEXT("r.BufferVisualizationOverviewTargets"),
-		TEXT("BaseColor,Specular,SubsurfaceColor,WorldNormal,SeparateTranslucencyRGB,,,Opacity,SeparateTranslucencyA,,,,SceneDepth,Roughness,Metallic,ShadingModel,,SceneDepthWorldUnits,SceneColor,HDRColor"), 
+		TEXT("BaseColor,Specular,SubsurfaceColor,WorldNormal,SeparateTranslucencyRGB,,,Opacity,SeparateTranslucencyA,,,,SceneDepth,Roughness,Metallic,ShadingModel,,SceneDepthWorldUnits,SceneColor,PostTonemapHDRColor"), 
 		*ConsoleDocumentationOverviewTargets,
 		ECVF_Default);
 }
diff --git a/Engine/Source/Runtime/LevelSequence/LevelSequence.Build.cs b/Engine/Source/Runtime/LevelSequence/LevelSequence.Build.cs
--- a/Engine/Source/Runtime/LevelSequence/LevelSequence.Build.cs
+++ b/Engine/Source/Runtime/LevelSequence/LevelSequence.Build.cs
@@ -14,7 +14,8 @@ public class LevelSequence : ModuleRules
                 "CoreUObject",
                 "Engine",
 				"MovieScene",
-                "MovieSceneTracks",
+				"MovieSceneTracks",
+				"UMG",
             }
         );
 	}
diff --git a/Engine/Source/Runtime/LevelSequence/Private/LevelSequenceActor.cpp b/Engine/Source/Runtime/LevelSequence/Private/LevelSequenceActor.cpp
--- a/Engine/Source/Runtime/LevelSequence/Private/LevelSequenceActor.cpp
+++ b/Engine/Source/Runtime/LevelSequence/Private/LevelSequenceActor.cpp
@@ -2,6 +2,7 @@
 
 #include "LevelSequencePCH.h"
 #include "LevelSequenceActor.h"
+#include "LevelSequenceBurnIn.h"
 
 
 ALevelSequenceActor::ALevelSequenceActor(const FObjectInitializer& Init)
@@ -33,6 +34,7 @@ ALevelSequenceActor::ALevelSequenceActor(const FObjectInitializer& Init)
 	}
 #endif // WITH_EDITORONLY_DATA
 
+	BurnInOptions = Init.CreateDefaultSubobject<ULevelSequenceBurnInOptions>(this, "BurnInOptions");
 	PrimaryActorTick.bCanEverTick = true;
 	bAutoPlay = false;
 }
@@ -85,9 +87,92 @@ void ALevelSequenceActor::InitializePlayer()
 		SequencePlayer = NewObject<ULevelSequencePlayer>(this, "AnimationPlayer");
 		SequencePlayer->Initialize(LevelSequenceAsset, GetWorld(), PlaybackSettings);
 
+		RefreshBurnIn();
+
 		if (bAutoPlay)
 		{
 			SequencePlayer->Play();
 		}
 	}
-}
\ No newline at end of file
+}
+
+void ALevelSequenceActor::RefreshBurnIn()
+{
+ 	if (!SequencePlayer)
+	{
+		return;
+	}
+
+	if (BurnInInstance)
+	{
+		BurnInInstance->RemoveFromViewport();
+		BurnInInstance = nullptr;
+	}
+	
+	if (BurnInOptions && BurnInOptions->bUseBurnIn)
+	{
+		// Create the burn-in if necessary
+		UClass* Class = BurnInOptions->BurnInClass.TryLoadClass<ULevelSequenceBurnIn>();
+		if (Class)
+		{
+			BurnInInstance = CreateWidget<ULevelSequenceBurnIn>(GetWorld(), Class);
+			if (BurnInInstance)
+			{
+				// Ensure we have a valid settings object if possible
+				BurnInOptions->ResetSettings();
+
+				BurnInInstance->SetSettings(BurnInOptions->Settings);
+				BurnInInstance->TakeSnapshotsFrom(*this);
+				BurnInInstance->AddToViewport();
+			}
+		}
+	}
+}
+
+ULevelSequenceBurnInOptions::ULevelSequenceBurnInOptions(const FObjectInitializer& Init)
+	: Super(Init)
+	, bUseBurnIn(false)
+	, BurnInClass(TEXT("/Engine/Sequencer/DefaultBurnIn.DefaultBurnIn_C"))
+	, Settings(nullptr)
+{
+}
+
+void ULevelSequenceBurnInOptions::ResetSettings()
+{
+	UClass* Class = BurnInClass.TryLoadClass<ULevelSequenceBurnIn>();
+	if (Class)
+	{
+		TSubclassOf<ULevelSequenceBurnInInitSettings> SettingsClass = Cast<ULevelSequenceBurnIn>(Class->GetDefaultObject())->GetSettingsClass();
+		if (SettingsClass)
+		{
+			if (!Settings || !Settings->IsA(SettingsClass))
+			{
+				Settings = NewObject<ULevelSequenceBurnInInitSettings>(this, SettingsClass);
+			}
+		}
+		else
+		{
+			Settings = nullptr;
+		}
+	}
+	else
+	{
+		Settings = nullptr;
+	}
+}
+
+#if WITH_EDITOR
+
+void ULevelSequenceBurnInOptions::PostEditChangeProperty( FPropertyChangedEvent& PropertyChangedEvent)
+{
+	FName PropertyName = (PropertyChangedEvent.Property != nullptr) ? PropertyChangedEvent.Property->GetFName() : NAME_None;
+
+	if (PropertyName == GET_MEMBER_NAME_CHECKED(ULevelSequenceBurnInOptions, bUseBurnIn) || PropertyName == GET_MEMBER_NAME_CHECKED(ULevelSequenceBurnInOptions, BurnInClass))
+	{
+		ResetSettings();
+	}
+
+	Super::PostEditChangeProperty(PropertyChangedEvent);
+}
+
+#endif // WITH_EDITOR
\ No newline at end of file
diff --git a/Engine/Source/Runtime/LevelSequence/Private/LevelSequenceBurnIn.cpp b/Engine/Source/Runtime/LevelSequence/Private/LevelSequenceBurnIn.cpp
new file mode 100644
--- /dev/null
+++ b/Engine/Source/Runtime/LevelSequence/Private/LevelSequenceBurnIn.cpp
@@ -0,0 +1,24 @@
+// Copyright 1998-2015 Epic Games, Inc. All Rights Reserved.
+
+#include "LevelSequencePCH.h"
+#include "LevelSequenceBurnIn.h"
+
+ULevelSequenceBurnIn::ULevelSequenceBurnIn( const FObjectInitializer& ObjectInitializer )
+	: Super(ObjectInitializer)
+{
+}
+
+void ULevelSequenceBurnIn::TakeSnapshotsFrom(ALevelSequenceActor& InActor)
+{
+	LevelSequenceActor = &InActor;
+	if (ensure(InActor.SequencePlayer))
+	{
+		InActor.SequencePlayer->OnSequenceUpdated().AddUObject(this, &ULevelSequenceBurnIn::OnSequenceUpdated);
+		InActor.SequencePlayer->TakeFrameSnapshot(FrameInformation);
+	}
+}
+
+void ULevelSequenceBurnIn::OnSequenceUpdated(const ULevelSequencePlayer& Player, float CurrentTime, float PreviousTime)
+{
+	Player.TakeFrameSnapshot(FrameInformation);
+}
\ No newline at end of file
diff --git a/Engine/Source/Runtime/LevelSequence/Private/LevelSequencePlayer.cpp b/Engine/Source/Runtime/LevelSequence/Private/LevelSequencePlayer.cpp
--- a/Engine/Source/Runtime/LevelSequence/Private/LevelSequencePlayer.cpp
+++ b/Engine/Source/Runtime/LevelSequence/Private/LevelSequencePlayer.cpp
@@ -9,7 +9,8 @@
 #include "MovieSceneTrack.h"
 #include "LevelSequenceSpawnRegister.h"
 #include "Engine/LevelStreaming.h"
-
+#include "Tracks/MovieSceneCinematicShotTrack.h"
+#include "Sections/MovieSceneCinematicShotSection.h"
 
 struct FTickAnimationPlayers : public FTickableGameObject
 {
@@ -216,16 +217,15 @@ void ULevelSequencePlayer::SetTickPrerequisites(bool bAddTickPrerequisites)
 
 		for (AActor* ControlledActor : ControlledActors)
 		{
-			USceneComponent* RootComponent = ControlledActor->GetRootComponent();
-			if (RootComponent)
+			for( UActorComponent* Component : ControlledActor->GetComponents() )
 			{
 				if (bAddTickPrerequisites)
 				{
-					RootComponent->PrimaryComponentTick.AddPrerequisite(LevelSequenceActor, LevelSequenceActor->PrimaryActorTick);
+					Component->PrimaryComponentTick.AddPrerequisite(LevelSequenceActor, LevelSequenceActor->PrimaryActorTick);
 				}
 				else
 				{
-					RootComponent->PrimaryComponentTick.RemovePrerequisite(LevelSequenceActor, LevelSequenceActor->PrimaryActorTick);
+					Component->PrimaryComponentTick.RemovePrerequisite(LevelSequenceActor, LevelSequenceActor->PrimaryActorTick);
 				}
 			}
 
@@ -522,17 +522,54 @@ void ULevelSequencePlayer::UpdateMovieSceneInstance(float CurrentPosition, float
 			MovieSceneSequence->GetMovieScene()->GetFixedFrameInterval() > 0 )
 		{
 			float FixedFrameInterval = MovieSceneSequence->GetMovieScene()->GetFixedFrameInterval();
-			Position = FMath::RoundToInt( Position / FixedFrameInterval ) * FixedFrameInterval;
-			LastPosition = FMath::RoundToInt( LastPosition / FixedFrameInterval ) * FixedFrameInterval;
+			Position = UMovieScene::CalculateFixedFrameTime( Position, FixedFrameInterval );
+			LastPosition = UMovieScene::CalculateFixedFrameTime( LastPosition, FixedFrameInterval );
 		}
 		EMovieSceneUpdateData UpdateData(Position, LastPosition);
 		RootMovieSceneInstance->Update(UpdateData, *this);
+
 #if WITH_EDITOR
 		OnLevelSequencePlayerUpdate.Broadcast(*this, CurrentPosition, PreviousPosition);
 #endif
 	}
 }
 
+void ULevelSequencePlayer::TakeFrameSnapshot(FLevelSequencePlayerSnapshot& OutSnapshot) const
+{
+	const float CurrentTime = StartTime + TimeCursorPosition;
+
+	OutSnapshot.Settings = SnapshotSettings;
+
+	OutSnapshot.MasterTime = CurrentTime;
+	OutSnapshot.MasterName = FText::FromString(LevelSequence->GetName());
+
+	OutSnapshot.CurrentShotName = OutSnapshot.MasterName;
+	OutSnapshot.CurrentShotLocalTime = CurrentTime;
+
+	UMovieSceneCinematicShotTrack* ShotTrack = LevelSequence->GetMovieScene()->FindMasterTrack<UMovieSceneCinematicShotTrack>();
+	if (ShotTrack)
+	{
+		int32 HighestRow = TNumericLimits<int32>::Max();
+		UMovieSceneCinematicShotSection* ActiveShot = nullptr;
+		for (UMovieSceneSection* Section : ShotTrack->GetAllSections())
+		{
+			if (Section->GetRange().Contains(CurrentTime) && (!ActiveShot || Section->GetRowIndex() < ActiveShot->GetRowIndex()))
+			{
+				ActiveShot = Cast<UMovieSceneCinematicShotSection>(Section);
+			}
+		}
+
+		if (ActiveShot)
+		{
+			const float ShotOffset = ActiveShot->StartOffset + ActiveShot->GetSequence()->GetMovieScene()->GetPlaybackRange().GetLowerBoundValue() - ActiveShot->PrerollTime;
+			const float ShotPosition = ShotOffset + (CurrentTime - (ActiveShot->GetStartTime() - ActiveShot->PrerollTime)) / ActiveShot->TimeScale;
+
+			OutSnapshot.CurrentShotName = ActiveShot->GetShotDisplayName();
+			OutSnapshot.CurrentShotLocalTime = ShotPosition;
+		}
+	}
+}
+
 IMovieSceneSpawnRegister& ULevelSequencePlayer::GetSpawnRegister()
 {
 	return *SpawnRegister;
diff --git a/Engine/Source/Runtime/LevelSequence/Public/LevelSequenceActor.h b/Engine/Source/Runtime/LevelSequence/Public/LevelSequenceActor.h
--- a/Engine/Source/Runtime/LevelSequence/Public/LevelSequenceActor.h
+++ b/Engine/Source/Runtime/LevelSequence/Public/LevelSequenceActor.h
@@ -5,6 +5,42 @@
 #include "LevelSequencePlayer.h"
 #include "LevelSequenceActor.generated.h"
 
+class ULevelSequenceBurnIn;
+
+UCLASS(Blueprintable)
+class LEVELSEQUENCE_API ULevelSequenceBurnInInitSettings : public UObject
+{
+	GENERATED_BODY()
+};
+
+UCLASS()
+class LEVELSEQUENCE_API ULevelSequenceBurnInOptions : public UObject
+{
+public:
+
+	GENERATED_BODY()
+	ULevelSequenceBurnInOptions(const FObjectInitializer& Init);
+
+	/** Ensure the settings object is up-to-date */
+	void ResetSettings();
+
+public:
+
+	UPROPERTY(EditAnywhere, Category="General")
+	bool bUseBurnIn;
+
+	UPROPERTY(EditAnywhere, Category="General", meta=(EditCondition=bUseBurnIn, MetaClass="LevelSequenceBurnIn"))
+	FStringClassReference BurnInClass;
+
+	UPROPERTY(EditAnywhere, Category="General", meta=(EditCondition=bUseBurnIn, EditInline))
+	ULevelSequenceBurnInInitSettings* Settings;
+
+protected:
+
+#if WITH_EDITOR
+	virtual void PostEditChangeProperty(FPropertyChangedEvent& PropertyChangedEvent) override;
+#endif //WITH_EDITOR
+};
 
 /**
  * Actor responsible for controlling a specific level sequence in the world.
@@ -34,11 +70,17 @@ public:
 	UPROPERTY(EditAnywhere, BlueprintReadOnly, Category="General", meta=(AllowedClasses="LevelSequence"))
 	FStringAssetReference LevelSequence;
 
+	UPROPERTY(EditAnywhere, BlueprintReadWrite, AdvancedDisplay, Category="General", meta=(EditInline))
+	ULevelSequenceBurnInOptions* BurnInOptions;
+
 public:
 
 	UFUNCTION(BlueprintCallable, Category="Game|Cinematic")
 	void SetSequence(ULevelSequence* InSequence);
 
+	/** Refresh this actor's burn in */
+	void RefreshBurnIn();
+
 public:
 
 	virtual void BeginPlay() override;
@@ -49,4 +91,9 @@ public:
 #endif //WITH_EDITOR
 
 	void InitializePlayer();
+
+private:
+	/** Burn-in widget */
+	UPROPERTY()
+	ULevelSequenceBurnIn* BurnInInstance;
 };
diff --git a/Engine/Source/Runtime/LevelSequence/Public/LevelSequenceBurnIn.h b/Engine/Source/Runtime/LevelSequence/Public/LevelSequenceBurnIn.h
new file mode 100644
--- /dev/null
+++ b/Engine/Source/Runtime/LevelSequence/Public/LevelSequenceBurnIn.h
@@ -0,0 +1,50 @@
+// Copyright 1998-2015 Epic Games, Inc. All Rights Reserved.
+
+#pragma once
+
+#include "UserWidget.h"
+#include "LevelSequenceActor.h"
+#include "LevelSequencePlayer.h"
+#include "LevelSequenceBurnIn.generated.h"
+
+class ULevelSequenceBurnInInitSettings;
+class ALevelSequenceActor;
+
+/**
+ * Base class for level sequence burn ins
+ */
+UCLASS()
+class LEVELSEQUENCE_API ULevelSequenceBurnIn : public UUserWidget
+{
+public:
+	GENERATED_BODY()
+
+	ULevelSequenceBurnIn(const FObjectInitializer& ObjectInitializer);
+
+	/** Initialize this burn in */
+	void TakeSnapshotsFrom(ALevelSequenceActor& InActor);
+
+	/** Called when this burn in is receiving its settings */
+	UFUNCTION(BlueprintImplementableEvent, Category="Burn In")
+	void SetSettings(UObject* InSettings);
+
+	/** Get the settings class to use for this burn in */
+	UFUNCTION(BlueprintNativeEvent, Category="Burn In")
+	TSubclassOf<ULevelSequenceBurnInInitSettings> GetSettingsClass() const;
+	TSubclassOf<ULevelSequenceBurnInInitSettings> GetSettingsClass_Implementation() const { return nullptr; }
+
+protected:
+
+	/** Called as part of the game tick loop when the sequence has been updated */
+	void OnSequenceUpdated(const ULevelSequencePlayer& Player, float CurrentTime, float PreviousTime);
+
+protected:
+
+	/** Snapshot of frame information. */
+	UPROPERTY(BlueprintReadOnly, Category="Burn In")
+	FLevelSequencePlayerSnapshot FrameInformation;
+
+	/** The actor to get our burn in frames from */
+	UPROPERTY(BlueprintReadOnly, Category="Burn In")
+	ALevelSequenceActor* LevelSequenceActor;
+};
\ No newline at end of file
diff --git a/Engine/Source/Runtime/LevelSequence/Public/LevelSequencePlayer.h b/Engine/Source/Runtime/LevelSequence/Public/LevelSequencePlayer.h
--- a/Engine/Source/Runtime/LevelSequence/Public/LevelSequencePlayer.h
+++ b/Engine/Source/Runtime/LevelSequence/Public/LevelSequencePlayer.h
@@ -11,6 +11,51 @@ class FMovieSceneSequenceInstance;
 class ULevel;
 class UMovieSceneBindings;
 
+USTRUCT(BlueprintType)
+struct FLevelSequenceSnapshotSettings
+{
+	GENERATED_BODY()
+
+	FLevelSequenceSnapshotSettings()
+		: ZeroPadAmount(4), FrameRate(30)
+	{}
+
+	FLevelSequenceSnapshotSettings(int32 InZeroPadAmount, float InFrameRate)
+		: ZeroPadAmount(InZeroPadAmount), FrameRate(InFrameRate)
+	{}
+
+	/** Zero pad frames */
+	UPROPERTY(VisibleAnywhere, BlueprintReadOnly, Category="General")
+	uint8 ZeroPadAmount;
+
+	/** Playback framerate */
+	UPROPERTY(VisibleAnywhere, BlueprintReadOnly, Category="General")
+	float FrameRate;
+};
+
+/**
+ * Frame snapshot information for a level sequence
+ */
+USTRUCT(BlueprintType)
+struct FLevelSequencePlayerSnapshot
+{
+	GENERATED_BODY()
+
+	UPROPERTY(VisibleAnywhere, BlueprintReadOnly, Category="General")
+	FText MasterName;
+
+	UPROPERTY(VisibleAnywhere, BlueprintReadOnly, Category="General")
+	float MasterTime;
+
+	UPROPERTY(VisibleAnywhere, BlueprintReadOnly, Category="General")
+	FText CurrentShotName;
+
+	UPROPERTY(VisibleAnywhere, BlueprintReadOnly, Category="General")
+	float CurrentShotLocalTime;
+
+	UPROPERTY(VisibleAnywhere, BlueprintReadOnly, Category="General")
+	FLevelSequenceSnapshotSettings Settings;
+};
 
 /**
  * Settings for the level sequence player actor.
@@ -134,6 +179,25 @@ public:
 	UFUNCTION(BlueprintCallable, Category="Game|Cinematic")
 	void SetPlaybackRange( const float NewStartTime, const float NewEndTime );
 
+	/** Get the offset within the level sequence to start playing */
+	UFUNCTION(BlueprintCallable, Category="Game|Cinematic")
+	float GetPlaybackStart() const { return StartTime; }
+
+	/** Get the offset within the level sequence to finish playing */
+	UFUNCTION(BlueprintCallable, Category="Game|Cinematic")
+	float GetPlaybackEnd() const { return EndTime; }
+
+	/** Set the settings used to capture snapshots with */
+	void SetSnapshotSettings(const FLevelSequenceSnapshotSettings& InSettings) { SnapshotSettings = InSettings; }
+
+public:
+
+	/**
+	 * Access the level sequence this player is playing
+	 * @return the level sequence currently assigned to this player
+	 */
+	ULevelSequence* GetLevelSequence() const { return LevelSequence; }
+	
 protected:
 
 	// IMovieScenePlayer interface
@@ -154,6 +218,9 @@ public:
 
 	void Update(const float DeltaSeconds);
 
+	/** Take a snapshot of the current state of this player */
+	void TakeFrameSnapshot(FLevelSequencePlayerSnapshot& OutSnapshot) const;
+
 private:
 
 	/** Update the movie scene instance from the specified previous position, to the specified time position. */
@@ -215,7 +282,11 @@ private:
 	/** The last view target to reset to when updating camera cuts to null */
 	mutable TWeakObjectPtr<AActor> LastViewTarget;
 
-#if WITH_EDITOR
+protected:
+
+	/** How to take snapshots */
+	FLevelSequenceSnapshotSettings SnapshotSettings;
+
 public:
 
 	/** An event that is broadcast each time this level sequence player is updated */
@@ -226,6 +297,4 @@ private:
 
 	/** The event that will be broadcast every time the sequence is updated */
 	mutable FOnLevelSequencePlayerUpdated OnLevelSequencePlayerUpdate;
-
-#endif
 };
diff --git a/Engine/Source/Runtime/Messaging/Public/Helpers/MessageEndpoint.h b/Engine/Source/Runtime/Messaging/Public/Helpers/MessageEndpoint.h
--- a/Engine/Source/Runtime/Messaging/Public/Helpers/MessageEndpoint.h
+++ b/Engine/Source/Runtime/Messaging/Public/Helpers/MessageEndpoint.h
@@ -31,7 +31,7 @@ DECLARE_DELEGATE_RetVal_OneParam(bool, FOnMessageEndpointReceiveMessage, const I
 /**
  * Implements a message endpoint for sending and receiving messages on a message bus.
  *
- * This class provides a convenient implementation of the IRecieveMessages and ISendMessages interfaces,
+ * This class provides a convenient implementation of the IReceiveMessages and ISendMessages interfaces,
  * which allow consumers to send and receive messages on a message bus. The endpoint allows for receiving
  * messages asynchronously as they arrive, as well as synchronously through an inbox that can be polled.
  *
diff --git a/Engine/Source/Runtime/Messaging/Public/IMessageBus.h b/Engine/Source/Runtime/Messaging/Public/IMessageBus.h
--- a/Engine/Source/Runtime/Messaging/Public/IMessageBus.h
+++ b/Engine/Source/Runtime/Messaging/Public/IMessageBus.h
@@ -57,10 +57,10 @@ DECLARE_MULTICAST_DELEGATE(FOnMessageBusShutdown);
  * when message recipients already know about each other and wish to communicate directly, i.e. to exchange commands or events.
  *
  * In the Publish-Subscribe pattern a message is sent to all message recipients on the bus using the see IMessageBus.Publish method.
- * Only recipients that previously subscribed to the type of the sent message using the @IMessageBus.Subscribe method will actually receive
+ * Only recipients that previously subscribed to the type of the sent message using the IMessageBus.Subscribe method will actually receive
  * the message. All other recipients will not receive the message. After a published message is received, the recipients may respond with
  * another message, either directly to the message sender using IMessageBus.Send method, or by publishing another message using the
- * the IMessageBus.Publish method. This pattern is useful for discovering recipients on the message bus and to recipients that are unknown.
+ * the IMessageBus.Publish method. This pattern is useful for discovering recipients on the bus and to message recipients that are unknown.
  *
  * Most applications will often use a combination of both Request-Reply and Publish-Subscribe. A typical implementation of distributed
  * applications involves service providers and service consumers (a service can be any useful functionality provided by a system). The
@@ -75,8 +75,8 @@ DECLARE_MULTICAST_DELEGATE(FOnMessageBusShutdown);
  * of messages that can be inspected for debugging purposes[3].
  *
  * It is possible to intercept messages before they are being routed to recipients by registering so called Message Interceptors (objects
- * implementing the IInterceptMessage interface) with the message bus. This feature enables advanced use cases that require inspection or
- * manipulation of messages contents, such as message filtering and enriching, splitting and aggregating, resequencing or authentication[4].
+ * implementing the IMessageInterceptor interface) with the message bus. This feature enables advanced use cases that require inspection or
+ * manipulation of messages contents, such as message filtering and enriching, splitting and aggregating, re-sequencing or authentication[4].
  *
  * The messaging system also provides a facility for debugging the system itself through the so called Message Tracer (an internal object
  * implementing the see IMessageTracer interface) that can be accessed with the IMessageBus.GetTracer() method. The message tracer is
diff --git a/Engine/Source/Runtime/MovieScene/Private/MovieScene.cpp b/Engine/Source/Runtime/MovieScene/Private/MovieScene.cpp
--- a/Engine/Source/Runtime/MovieScene/Private/MovieScene.cpp
+++ b/Engine/Source/Runtime/MovieScene/Private/MovieScene.cpp
@@ -23,9 +23,62 @@ UMovieScene::UMovieScene(const FObjectInitializer& ObjectInitializer)
 #endif
 }
 
+void UMovieScene::Serialize( FArchive& Ar )
+{
+#if WITH_EDITOR
+
+	// Perform optimizations for cooking
+	if (Ar.IsCooking())
+	{
+		// @todo: Optimize master tracks?
+
+		// Optimize object bindings
+		OptimizeObjectArray(Spawnables);
+		OptimizeObjectArray(Possessables);
+	}
+
+#endif // WITH_EDITOR
+
+	Super::Serialize(Ar);
+}
 
 #if WITH_EDITOR
 
+template<typename T>
+void UMovieScene::OptimizeObjectArray(TArray<T>& ObjectArray)
+{
+	for (int32 ObjectIndex = ObjectArray.Num() - 1; ObjectIndex >= 0; --ObjectIndex)
+	{
+		FGuid ObjectGuid = ObjectArray[ObjectIndex].GetGuid();
+
+		// Find the binding relating to this object, and optimize its tracks
+		// @todo: ObjectBindings mapped by ID to avoid linear search
+		for (int32 BindingIndex = 0; BindingIndex < ObjectBindings.Num(); ++BindingIndex)
+		{
+			FMovieSceneBinding& Binding =ObjectBindings[BindingIndex];
+			if (Binding.GetObjectGuid() != ObjectGuid)
+			{
+				continue;
+			}
+			
+			bool bShouldRemoveObject = false;
+
+			// Optimize any tracks
+			Binding.PerformCookOptimization(bShouldRemoveObject);
+
+			// Remove the object if it's completely redundant
+			if (bShouldRemoveObject)
+			{
+				ObjectBindings.RemoveAtSwap(BindingIndex, 1, false);
+				ObjectArray.RemoveAtSwap(ObjectIndex, 1, false);
+			}
+
+			// Process next object
+			break;
+		}
+	}
+}
+
 // @todo sequencer: Some of these methods should only be used by tools, and should probably move out of MovieScene!
 FGuid UMovieScene::AddSpawnable( const FString& Name, UObject& ObjectTemplate )
 {
@@ -303,6 +356,14 @@ void UMovieScene::SetFixedFrameInterval( float InFixedFrameInterval )
 }
 
 
+const float UMovieScene::FixedFrameIntervalEpsilon = .0001f;
+
+float UMovieScene::CalculateFixedFrameTime( float Time, float FixedFrameInterval )
+{
+	return ( FMath::RoundToInt( Time / FixedFrameInterval ) ) * FixedFrameInterval + FixedFrameIntervalEpsilon;
+}
+
+
 TArray<UMovieSceneSection*> UMovieScene::GetAllSections() const
 {
 	TArray<UMovieSceneSection*> OutSections;
diff --git a/Engine/Source/Runtime/MovieScene/Private/MovieSceneBinding.cpp b/Engine/Source/Runtime/MovieScene/Private/MovieSceneBinding.cpp
--- a/Engine/Source/Runtime/MovieScene/Private/MovieSceneBinding.cpp
+++ b/Engine/Source/Runtime/MovieScene/Private/MovieSceneBinding.cpp
@@ -31,3 +31,23 @@ bool FMovieSceneBinding::RemoveTrack(UMovieSceneTrack& Track)
 {
 	return (Tracks.RemoveSingle(&Track) != 0);
 }
+
+#if WITH_EDITOR
+void FMovieSceneBinding::PerformCookOptimization(bool& bShouldRemoveObject)
+{
+	for (int32 Index = Tracks.Num() - 1; Index >= 0; --Index)
+	{
+		ECookOptimizationFlags Flags = Tracks[Index]->GetCookOptimizationFlags();
+
+		if ((Flags & ECookOptimizationFlags::RemoveObject) != ECookOptimizationFlags::None)
+		{
+			bShouldRemoveObject = true;
+			return;
+		}
+		else if ((Flags & ECookOptimizationFlags::RemoveTrack) != ECookOptimizationFlags::None)
+		{
+			Tracks.RemoveAtSwap(Index, 1, false);
+		}
+	}
+}
+#endif
\ No newline at end of file
