Copying //UE4/Dev-Framework to //UE4/Dev-Main (Source: //UE4/Dev-Framework @ 3038004)

diff --git a/Engine/Source/Runtime/Core/Public/Math/Transform.h b/Engine/Source/Runtime/Core/Public/Math/Transform.h
--- a/Engine/Source/Runtime/Core/Public/Math/Transform.h
+++ b/Engine/Source/Runtime/Core/Public/Math/Transform.h
@@ -865,6 +865,67 @@
 		DiagnosticCheckNaN_Scale3D();
 	}
 
+	/** @note: Added template type function for Accumulate
+	  * The template type isn't much useful yet, but it is with the plan to move forward
+	  * to unify blending features with just type of additive or full pose
+	  * Eventually it would be nice to just call blend and it all works depending on full pose
+	  * or additive, but right now that is a lot more refactoring
+	  * For now this types only defines the different functionality of accumulate
+	  */
+	  
 	/**
+	 * Accumulates another transform with this one
+	 *
+	 * Rotation is accumulated multiplicatively (Rotation = SourceAtom.Rotation * Rotation)
+	 * Translation is accumulated additively (Translation += SourceAtom.Translation)
+	 * Scale3D is accumulated multiplicatively (Scale3D *= SourceAtom.Scale3D)
+	 *
+	 * @param SourceAtom The other transform to accumulate into this one
+	 */
+	FORCEINLINE void Accumulate(const FTransform& SourceAtom)
+	{
+		// Add ref pose relative animation to base animation, only if rotation is significant.
+		if(FMath::Square(SourceAtom.Rotation.W) < 1.f - DELTA * DELTA)
+		{
+			Rotation = SourceAtom.Rotation * Rotation;
+		}
+
+		Translation += SourceAtom.Translation;
+		Scale3D *= SourceAtom.Scale3D;
+
+		DiagnosticCheckNaN_All();
+
+		checkSlow(IsRotationNormalized());
+	}
+	
+	/** Accumulates another transform with this one, with a blending weight
+	*
+	* Let SourceAtom = Atom * BlendWeight
+	* Rotation is accumulated multiplicatively (Rotation = SourceAtom.Rotation * Rotation).
+	* Translation is accumulated additively (Translation += SourceAtom.Translation)
+	* Scale3D is accumulated multiplicatively (Scale3D *= SourceAtom.Scale3D)
+	*
+	* Note: Rotation will not be normalized! Will have to be done manually.
+	*
+	* @param Atom The other transform to accumulate into this one
+	* @param BlendWeight The weight to multiply Atom by before it is accumulated.
+	*/
+	FORCEINLINE void Accumulate(const FTransform& Atom, float BlendWeight/* default param doesn't work since vectorized version takes ref param */)
+	{
+		FTransform SourceAtom(Atom * BlendWeight);
+
+		// Add ref pose relative animation to base animation, only if rotation is significant.
+		if (FMath::Square(SourceAtom.Rotation.W) < 1.f - DELTA * DELTA)
+		{
+			Rotation = SourceAtom.Rotation * Rotation;
+		}
+
+		Translation += SourceAtom.Translation;
+		Scale3D *= SourceAtom.Scale3D;
+
+		DiagnosticCheckNaN_All();
+	}
+
+	/**
 	 * Accumulates another transform with this one, with an optional blending weight
 	 *
@@ -875,9 +936,9 @@
 	 * @param DeltaAtom The other transform to accumulate into this one
 	 * @param Weight The weight to multiply DeltaAtom by before it is accumulated.
 	 */
-	FORCEINLINE void AccumulateWithShortestRotation(const FTransform& DeltaAtom, float Weight = 1.0f)
+	FORCEINLINE void AccumulateWithShortestRotation(const FTransform& DeltaAtom, float BlendWeight/* default param doesn't work since vectorized version takes ref param */)
 	{
-		FTransform Atom(DeltaAtom * Weight);
+		FTransform Atom(DeltaAtom * BlendWeight);
 
 		// To ensure the 'shortest route', we make sure the dot product between the accumulator and the incoming child atom is positive.
 		if((Atom.Rotation | Rotation) < 0.f)
@@ -901,59 +962,41 @@
 		DiagnosticCheckNaN_All();
 	}
 
-	/**
-	 * Accumulates another transform with this one
-	 *
-	 * Rotation is accumulated multiplicatively (Rotation = SourceAtom.Rotation * Rotation)
-	 * Translation is accumulated additively (Translation += SourceAtom.Translation)
-	 * Scale3D is accumulated multiplicatively (Scale3D *= SourceAtom.Scale3D)
-	 *
-	 * @param SourceAtom The other transform to accumulate into this one
-	 */
-	FORCEINLINE void Accumulate(const FTransform& SourceAtom)
-	{
-		// Add ref pose relative animation to base animation, only if rotation is significant.
-		if(FMath::Square(SourceAtom.Rotation.W) < 1.f - DELTA * DELTA)
-		{
-			Rotation = SourceAtom.Rotation * Rotation;
-		}
-
-		Translation += SourceAtom.Translation;
-		Scale3D *= SourceAtom.Scale3D;
-
-		DiagnosticCheckNaN_All();
-
-		checkSlow(IsRotationNormalized());
-	}
-
-   /** Accumulates another transform with this one, with a blending weight
+	/** Accumulates another transform with this one, with a blending weight
 	*
 	* Let SourceAtom = Atom * BlendWeight
 	* Rotation is accumulated multiplicatively (Rotation = SourceAtom.Rotation * Rotation).
 	* Translation is accumulated additively (Translation += SourceAtom.Translation)
-	* Scale3D is accumulated multiplicatively (Scale3D *= SourceAtom.Scale3D)
+	* Scale3D is accumulated assuming incoming scale is additive scale (Scale3D *= (1 + SourceAtom.Scale3D))
+	* 
+	* When we create additive, we create additive scale based on [TargetScale/SourceScale -1]
+	* because that way when you apply weight of 0.3, you don't shrink. We only saves the % of grow/shrink
+	* when we apply that back to it, we add back the 1, so that it goes back to it. 
+	* This solves issue where you blend two additives with 0.3, you don't come back to 0.6 scale, but 1 scale at the end 
+	* because [1 + [1-1]*0.3 + [1-1]*0.3] becomes 1, so you don't shrink by applying additive scale
 	*
 	* Note: Rotation will not be normalized! Will have to be done manually.
 	*
 	* @param Atom The other transform to accumulate into this one
 	* @param BlendWeight The weight to multiply Atom by before it is accumulated.
 	*/
-	FORCEINLINE void Accumulate(const FTransform& Atom, float BlendWeight)
+	FORCEINLINE void AccumulateWithAdditiveScale(const FTransform& Atom, float BlendWeight/* default param doesn't work since vectorized version takes ref param */)
 	{
+		const FVector DefaultScale(1.f);
+
 		FTransform SourceAtom(Atom * BlendWeight);
-
+		
 		// Add ref pose relative animation to base animation, only if rotation is significant.
-		if(FMath::Square(SourceAtom.Rotation.W) < 1.f - DELTA * DELTA)
+		if (FMath::Square(SourceAtom.Rotation.W) < 1.f - DELTA * DELTA)
 		{
 			Rotation = SourceAtom.Rotation * Rotation;
 		}
 
 		Translation += SourceAtom.Translation;
-		Scale3D *= SourceAtom.Scale3D;
+		Scale3D *= (DefaultScale + SourceAtom.Scale3D);
 
 		DiagnosticCheckNaN_All();
 	}
-
 	/**
 	 * Set the translation and Scale3D components of this transform to a linearly interpolated combination of two other transforms
 	 *
@@ -974,27 +1017,5 @@
 	}
 
 	/**
-	 * Accumulates another transform with this one
-	 *
-	 * Rotation is accumulated multiplicatively (Rotation = SourceAtom.Rotation * Rotation)
-	 * Translation is accumulated additively (Translation += SourceAtom.Translation)
-	 * Scale3D is accumulated additively (Scale3D += SourceAtom.Scale3D)
-	 *
-	 * @param SourceAtom The other transform to accumulate into this one
-	 */
-	FORCEINLINE void AccumulateWithAdditiveScale3D(const FTransform& SourceAtom)
-	{
-		if(FMath::Square(SourceAtom.Rotation.W) < 1.f - DELTA * DELTA)
-		{
-			Rotation = SourceAtom.Rotation * Rotation;
-		}
-
-		Translation += SourceAtom.Translation;
-		Scale3D += SourceAtom.Scale3D;
-
-		DiagnosticCheckNaN_All();
-	}
-
-	/**
 	 * Normalize the rotation component of this transformation
 	 */
@@ -1028,12 +1049,13 @@
 	 */
 	FORCEINLINE static void BlendFromIdentityAndAccumulate(FTransform& FinalAtom, FTransform& SourceAtom, float BlendWeight)
 	{
-		const FTransform IdentityAtom = FTransform::Identity;
+		const  FTransform AdditiveIdentity(FQuat::Identity, FVector::ZeroVector, FVector::ZeroVector);
+		const FVector DefaultScale(1.f);
 
 		// Scale delta by weight
 		if(BlendWeight < (1.f - ZERO_ANIMWEIGHT_THRESH))
 		{
-			SourceAtom.Blend(IdentityAtom, SourceAtom, BlendWeight);
+			SourceAtom.Blend(AdditiveIdentity, SourceAtom, BlendWeight);
 		}
 
 		// Add ref pose relative animation to base animation, only if rotation is significant.
@@ -1043,7 +1065,7 @@
 		}
 
 		FinalAtom.Translation += SourceAtom.Translation;
-		FinalAtom.Scale3D *= SourceAtom.Scale3D;
+		FinalAtom.Scale3D *= (DefaultScale + SourceAtom.Scale3D);
 
 		FinalAtom.DiagnosticCheckNaN_All();
 
diff --git a/Engine/Source/Runtime/Core/Public/Math/TransformVectorized.h b/Engine/Source/Runtime/Core/Public/Math/TransformVectorized.h
--- a/Engine/Source/Runtime/Core/Public/Math/TransformVectorized.h
+++ b/Engine/Source/Runtime/Core/Public/Math/TransformVectorized.h
@@ -900,44 +900,66 @@ public:
 		DiagnosticCheckNaN_Translate();
 		DiagnosticCheckNaN_Scale3D();
 	}
-
+	
+	/** @note : Added template type function for Accumulate
+	  * The template type isn't much useful yet, but it is with the plan to move forward
+	  * to unify blending features with just type of additive or full pose
+	  * Eventually it would be nice to just call blend and it all works depending on full pose
+	  * or additive, but right now that is a lot more refactoring
+	  * For now this types only defines the different functionality of accumulate
+	  */
 
 	/**
-	 * Accumulates another transform with this one, with an optional blending weight
-	 *
-	 * Rotation is accumulated additively, in the shortest direction (Rotation = Rotation +/- DeltaAtom.Rotation * Weight)
-	 * Translation is accumulated additively (Translation += DeltaAtom.Translation * Weight)
-	 * Scale3D is accumulated additively (Scale3D += DeltaAtom.Scale * Weight)
-	 *
-	 * @param DeltaAtom The other transform to accumulate into this one
-	 * @param Weight The weight to multiply DeltaAtom by before it is accumulated.
-	 */
-	FORCEINLINE void AccumulateWithShortestRotation(const FTransform& DeltaAtom, const ScalarRegister& BlendWeight)
+	* Accumulates another transform with this one
+	*
+	* Rotation is accumulated multiplicatively (Rotation = SourceAtom.Rotation * Rotation)
+	* Translation is accumulated additively (Translation += SourceAtom.Translation)
+	* Scale3D is accumulated multiplicatively (Scale3D *= SourceAtom.Scale3D)
+	*
+	* @param SourceAtom The other transform to accumulate into this one
+	*/
+	FORCEINLINE void Accumulate(const FTransform& SourceAtom)
 	{
-		const VectorRegister BlendedRotation = VectorMultiply(DeltaAtom.Rotation, BlendWeight.Value);
+		const VectorRegister BlendedRotation = SourceAtom.Rotation;
+		const VectorRegister RotationW = VectorReplicate(BlendedRotation, 3);
 
-		Rotation = VectorAccumulateQuaternionShortestPath(Rotation, BlendedRotation);
+		// if( Square(SourceAtom.Rotation.W) < 1.f - DELTA * DELTA )
+		if (VectorAnyGreaterThan(GlobalVectorConstants::RotationSignificantThreshold, VectorMultiply(RotationW, RotationW)))
+		{
+			// Rotation = SourceAtom.Rotation * Rotation;
+			Rotation = VectorQuaternionMultiply2(BlendedRotation, Rotation);
+		}
 
-		Translation = VectorMultiplyAdd(DeltaAtom.Translation, BlendWeight, Translation);
-		Scale3D = VectorMultiplyAdd(DeltaAtom.Scale3D, BlendWeight, Scale3D);
+		// Translation += SourceAtom.Translation;
+		// Scale *= SourceAtom.Scale;
+		Translation = VectorAdd(Translation, SourceAtom.Translation);
+		Scale3D = VectorMultiply(Scale3D, SourceAtom.Scale3D);
 
 		DiagnosticCheckNaN_All();
+
+		checkSlow(IsRotationNormalized());
 	}
 
 	/**
-	 * Accumulates another transform with this one
-	 *
-	 * Rotation is accumulated multiplicatively (Rotation = SourceAtom.Rotation * Rotation)
-	 * Translation is accumulated additively (Translation += SourceAtom.Translation)
-	 * Scale3D is accumulated multiplicatively (Scale3D *= SourceAtom.Scale3D)
-	 *
-	 * @param SourceAtom The other transform to accumulate into this one
-	 */
-	FORCEINLINE void Accumulate(const FTransform& SourceAtom)
+	* Accumulates another transform with this one, with a blending weight
+	*
+	* Let SourceAtom = Atom * BlendWeight
+	* Rotation is accumulated multiplicatively (Rotation = Atom.Rotation * Rotation).
+	* Translation is accumulated additively (Translation += Atom.Translation)
+	* Scale3D is accumulated multiplicatively (Scale3D *= Atom.Scale3D)
+	*
+	* Note: Rotation will not be normalized! Will have to be done manually.
+	*
+	* @param Atom The other transform to accumulate into this one
+	* @param BlendWeight The weight to multiply Atom by before it is accumulated.
+	*/
+	FORCEINLINE void Accumulate(const FTransform& Atom, const ScalarRegister& BlendWeight)
 	{
-		const VectorRegister BlendedRotation = SourceAtom.Rotation;
+		// SourceAtom = Atom * BlendWeight;
+		const VectorRegister BlendedRotation = VectorMultiply(Atom.Rotation, BlendWeight.Value);
 		const VectorRegister RotationW = VectorReplicate(BlendedRotation, 3);
 
+		// Add ref pose relative animation to base animation, only if rotation is significant.
 		// if( Square(SourceAtom.Rotation.W) < 1.f - DELTA * DELTA )
 		if (VectorAnyGreaterThan(GlobalVectorConstants::RotationSignificantThreshold, VectorMultiply(RotationW, RotationW)))
 		{
@@ -947,29 +969,54 @@ public:
 
 		// Translation += SourceAtom.Translation;
 		// Scale *= SourceAtom.Scale;
-		Translation = VectorAdd( Translation, SourceAtom.Translation );
-		Scale3D = VectorMultiply( Scale3D, SourceAtom.Scale3D );
+		Translation = VectorAdd(Translation, Atom.Translation);
+		Scale3D = VectorMultiply(Scale3D, Atom.Scale3D);
 
 		DiagnosticCheckNaN_All();
-
-		checkSlow( IsRotationNormalized() );
 	}
-
 	/**
-	 * Accumulates another transform with this one, with a blending weight
-	 *
-	 * Let SourceAtom = Atom * BlendWeight
-	 * Rotation is accumulated multiplicatively (Rotation = Atom.Rotation * Rotation).
-	 * Translation is accumulated additively (Translation += Atom.Translation)
-	 * Scale3D is accumulated multiplicatively (Scale3D *= Atom.Scale3D)
+	 * Accumulates another transform with this one, with an optional blending weight
 	 *
-	 * Note: Rotation will not be normalized! Will have to be done manually.
+	 * Rotation is accumulated additively, in the shortest direction (Rotation = Rotation +/- DeltaAtom.Rotation * Weight)
+	 * Translation is accumulated additively (Translation += DeltaAtom.Translation * Weight)
+	 * Scale3D is accumulated additively (Scale3D += DeltaAtom.Scale * Weight)
 	 *
-	 * @param Atom The other transform to accumulate into this one
-	 * @param BlendWeight The weight to multiply Atom by before it is accumulated.
+	 * @param DeltaAtom The other transform to accumulate into this one
+	 * @param Weight The weight to multiply DeltaAtom by before it is accumulated.
 	 */
-	FORCEINLINE void Accumulate(const FTransform& Atom, const ScalarRegister& BlendWeight)
+	FORCEINLINE void AccumulateWithShortestRotation(const FTransform& DeltaAtom, const ScalarRegister& BlendWeight)
 	{
+		const VectorRegister BlendedRotation = VectorMultiply(DeltaAtom.Rotation, BlendWeight.Value);
+
+		Rotation = VectorAccumulateQuaternionShortestPath(Rotation, BlendedRotation);
+
+		Translation = VectorMultiplyAdd(DeltaAtom.Translation, BlendWeight, Translation);
+		Scale3D = VectorMultiplyAdd(DeltaAtom.Scale3D, BlendWeight, Scale3D);
+
+		DiagnosticCheckNaN_All();
+	}
+
+	/** Accumulates another transform with this one, with a blending weight
+	*
+	* Let SourceAtom = Atom * BlendWeight
+	* Rotation is accumulated multiplicatively (Rotation = SourceAtom.Rotation * Rotation).
+	* Translation is accumulated additively (Translation += SourceAtom.Translation)
+	* Scale3D is accumulated assuming incoming scale is additive scale (Scale3D *= (1 + SourceAtom.Scale3D))
+	*
+	* When we create additive, we create additive scale based on [TargetScale/SourceScale -1]
+	* because that way when you apply weight of 0.3, you don't shrink. We only saves the % of grow/shrink
+	* when we apply that back to it, we add back the 1, so that it goes back to it.
+	* This solves issue where you blend two additives with 0.3, you don't come back to 0.6 scale, but 1 scale at the end
+	* because [1 + [1-1]*0.3 + [1-1]*0.3] becomes 1, so you don't shrink by applying additive scale
+	*
+	* Note: Rotation will not be normalized! Will have to be done manually.
+	*
+	* @param Atom The other transform to accumulate into this one
+	* @param BlendWeight The weight to multiply Atom by before it is accumulated.
+	*/
+	FORCEINLINE void AccumulateWithAdditiveScale(const FTransform& Atom, const ScalarRegister& BlendWeight)
+	{
+		const VectorRegister DefaultScale = MakeVectorRegister(1.f, 1.f, 1.f, 0.f);
 		// SourceAtom = Atom * BlendWeight;
 		const VectorRegister BlendedRotation = VectorMultiply(Atom.Rotation, BlendWeight.Value);
 		const VectorRegister RotationW = VectorReplicate(BlendedRotation, 3);
@@ -984,14 +1031,12 @@ public:
 
 		// Translation += SourceAtom.Translation;
 		// Scale *= SourceAtom.Scale;
-		Translation = VectorAdd( Translation, Atom.Translation );
-		Scale3D = VectorMultiply( Scale3D, Atom.Scale3D );
+		Translation = VectorAdd(Translation, Atom.Translation);
+		Scale3D = VectorMultiply(Scale3D, VectorAdd(DefaultScale, Atom.Scale3D));
 
 		DiagnosticCheckNaN_All();
 	}
 
-
-
 	/**
 	 * Set the translation and Scale3D components of this transform to a linearly interpolated combination of two other transforms
 	 *
@@ -1011,33 +1056,6 @@ public:
 		DiagnosticCheckNaN_Scale3D();
 	}
 
-/**
-	 * Accumulates another transform with this one
-	 *
-	 * Rotation is accumulated multiplicatively (Rotation = SourceAtom.Rotation * Rotation)
-	 * Translation is accumulated additively (Translation += SourceAtom.Translation)
-	 * Scale is accumulated additively (Scale3D += SourceAtom.Scale3D)
-	 *
-	 * @param SourceAtom The other transform to accumulate into this one
-	 */
-	FORCEINLINE void AccumulateWithAdditiveScale3D(const FTransform& SourceAtom)
-	{
-		// Add ref pose relative animation to base animation, only if rotation is significant.
-		const VectorRegister RotationW = VectorReplicate(SourceAtom.Rotation, 3);
-
-		// if( Square(SourceAtom.Rotation.W) < 1.f - DELTA * DELTA )
-		if (VectorAnyGreaterThan(GlobalVectorConstants::RotationSignificantThreshold, VectorMultiply(RotationW, RotationW)))
-		{
-			// Rotation = SourceAtom.Rotation * Rotation;
-			Rotation = VectorQuaternionMultiply2(SourceAtom.Rotation, Rotation);
-		}
-
-		Translation = VectorAdd( Translation, SourceAtom.Translation );
-		Scale3D = VectorAdd( Scale3D, SourceAtom.Scale3D );
-
-		DiagnosticCheckNaN_All();
-	}
-
 	/**
 	 * Normalize the rotation component of this transformation
 	 */
@@ -1075,6 +1093,7 @@ public:
 		const VectorRegister Const0001 = GlobalVectorConstants::Float0001;
 		const VectorRegister ConstNegative0001 = VectorSubtract(VectorZero(), Const0001);
 		const VectorRegister VOneMinusAlpha = VectorSubtract(VectorOne(), BlendWeight.Value);
+		const VectorRegister DefaultScale = MakeVectorRegister(1.f, 1.f, 1.f, 0.f);
 
 		// Blend rotation
 		//     To ensure the 'shortest route', we make sure the dot product between the both rotations is positive.
@@ -1097,16 +1116,15 @@ public:
 
 		// Blend translation and scale
 		//    BlendedAtom.Translation = Lerp(Zero, SourceAtom.Translation, Alpha);
-		//    BlendedAtom.Scale = Lerp(1, SourceAtom.Scale, Alpha);
+		//    BlendedAtom.Scale = Lerp(0, SourceAtom.Scale, Alpha);
 		const VectorRegister BlendedTranslation	= FMath::Lerp(VectorZero(), SourceAtom.Translation, BlendWeight.Value);
-		const VectorRegister BlendedScale3D	= FMath::Lerp(VectorOne(), SourceAtom.Scale3D, BlendWeight.Value);
+		const VectorRegister BlendedScale3D	= FMath::Lerp(VectorZero(), SourceAtom.Scale3D, BlendWeight.Value);
 
 		// Apply translation and scale to final atom
 		//     FinalAtom.Translation += BlendedAtom.Translation
 		//     FinalAtom.Scale *= BlendedAtom.Scale
 		FinalAtom.Translation = VectorAdd( FinalAtom.Translation, BlendedTranslation );
-		FinalAtom.Scale3D = VectorMultiply( FinalAtom.Scale3D, BlendedScale3D );
-		
+		FinalAtom.Scale3D = VectorMultiply( FinalAtom.Scale3D, VectorAdd(DefaultScale, BlendedScale3D));
 		checkSlow( FinalAtom.IsRotationNormalized() );
 	}
 
diff --git a/Engine/Source/Runtime/Core/Public/UObject/FrameworkObjectVersion.h b/Engine/Source/Runtime/Core/Public/UObject/FrameworkObjectVersion.h
--- a/Engine/Source/Runtime/Core/Public/UObject/FrameworkObjectVersion.h
+++ b/Engine/Source/Runtime/Core/Public/UObject/FrameworkObjectVersion.h
@@ -38,6 +38,12 @@ struct CORE_API FFrameworkObjectVersion
 		// Refactor ConstraintInstance so that we have an easy way to swap behavior paramters
 		ConstraintInstanceBehaviorParameters,
 
+		// Pose Asset support mask per bone
+		PoseAssetSupportPerBoneMask,
+
+		// Physics Assets now use SkeletalBodySetup instead of BodySetup
+		PhysAssetUseSkeletalBodySetup,
+
 		// -----<new versions can be added above this line>-------------------------------------------------
 		VersionPlusOne,
 		LatestVersion = VersionPlusOne - 1
diff --git a/Engine/Source/Runtime/Core/Public/UObject/UnrealNames.inl b/Engine/Source/Runtime/Core/Public/UObject/UnrealNames.inl
--- a/Engine/Source/Runtime/Core/Public/UObject/UnrealNames.inl
+++ b/Engine/Source/Runtime/Core/Public/UObject/UnrealNames.inl
@@ -122,6 +122,7 @@ REGISTER_NAME(248,Location)
 REGISTER_NAME(249,Rotation)
 REGISTER_NAME(250,BSP)
 REGISTER_NAME(251,EditorSettings)
+REGISTER_NAME(252,AudioThread)
 
 // Online
 REGISTER_NAME(280,DGram)
diff --git a/Engine/Source/Runtime/CoreUObject/Classes/Object.h b/Engine/Source/Runtime/CoreUObject/Classes/Object.h
--- a/Engine/Source/Runtime/CoreUObject/Classes/Object.h
+++ b/Engine/Source/Runtime/CoreUObject/Classes/Object.h
@@ -384,13 +384,16 @@ struct FPlane : public FVector
 USTRUCT(immutable, noexport, BlueprintType, meta=(HasNativeMake="Engine.KismetMathLibrary.MakeRotator", HasNativeBreak="Engine.KismetMathLibrary.BreakRotator"))
 struct FRotator
 {
-	UPROPERTY(EditAnywhere, BlueprintReadWrite, Category=Rotator, SaveGame)
+	/** Pitch (degrees) around Y axis */
+	UPROPERTY(EditAnywhere, BlueprintReadWrite, Category=Rotator, SaveGame, meta=(DisplayName="Y"))
 	float Pitch;
 
-	UPROPERTY(EditAnywhere, BlueprintReadWrite, Category=Rotator, SaveGame)
+	/** Yaw (degrees) around Z axis */
+	UPROPERTY(EditAnywhere, BlueprintReadWrite, Category=Rotator, SaveGame, meta=(DisplayName="Z"))
 	float Yaw;
 
-	UPROPERTY(EditAnywhere, BlueprintReadWrite, Category=Rotator, SaveGame)
+	/** Roll (degrees) around X axis */
+	UPROPERTY(EditAnywhere, BlueprintReadWrite, Category=Rotator, SaveGame, meta=(DisplayName="X"))
 	float Roll;
 
 };
diff --git a/Engine/Source/Runtime/Engine/Classes/AI/Navigation/RecastNavMesh.h b/Engine/Source/Runtime/Engine/Classes/AI/Navigation/RecastNavMesh.h
--- a/Engine/Source/Runtime/Engine/Classes/AI/Navigation/RecastNavMesh.h
+++ b/Engine/Source/Runtime/Engine/Classes/AI/Navigation/RecastNavMesh.h
@@ -241,7 +241,7 @@ struct FRecastDebugPathfindingNode
 
 	FORCEINLINE bool operator==(const NavNodeRef& OtherPolyRef) const { return PolyRef == OtherPolyRef; }
 	FORCEINLINE bool operator==(const FRecastDebugPathfindingNode& Other) const { return PolyRef == Other.PolyRef; }
-	FORCEINLINE friend uint32 GetTypeHash(const FRecastDebugPathfindingNode& Other) { return Other.PolyRef; }
+	FORCEINLINE friend uint32 GetTypeHash(const FRecastDebugPathfindingNode& Other) { return GetTypeHash(Other.PolyRef); }
 
 	FORCEINLINE float GetHeuristicCost() const { return TotalCost - Cost; }
 };
diff --git a/Engine/Source/Runtime/Engine/Classes/Animation/AnimInstance.h b/Engine/Source/Runtime/Engine/Classes/Animation/AnimInstance.h
--- a/Engine/Source/Runtime/Engine/Classes/Animation/AnimInstance.h
+++ b/Engine/Source/Runtime/Engine/Classes/Animation/AnimInstance.h
@@ -31,10 +31,9 @@ struct FAnimInstanceProxy;
 UENUM()
 enum class EAnimCurveType : uint8 
 {
-	EventCurve,
+	AttributeCurve,
 	MaterialCurve, 
 	MorphTargetCurve, 
-	PoseCurve, 
 	// make sure to update MaxCurve 
 	MaxAnimCurveType
 };
@@ -93,7 +92,7 @@ struct ENGINE_API FA2CSPose : public FA2Pose
 
 private:
 	/** Pointer to current BoneContainer. */
-	const struct FBoneContainer * BoneContainer;
+	const struct FBoneContainer* BoneContainer;
 
 	/** Once evaluated to be mesh space, this flag will be set. */
 	UPROPERTY()
@@ -502,11 +501,11 @@ public:
 	void StopSlotAnimation(float InBlendOutTime = 0.25f, FName SlotNodeName = NAME_None);
 
 	/** Return true if it's playing the slot animation */
-	UFUNCTION(BlueprintCallable, Category="Animation")
-	bool IsPlayingSlotAnimation(UAnimSequenceBase* Asset, FName SlotNodeName); 
+	UFUNCTION(BlueprintPure, Category="Animation")
+	bool IsPlayingSlotAnimation(const UAnimSequenceBase* Asset, FName SlotNodeName) const;
 
 	/** Return true if this instance playing the slot animation, also returning the montage it is playing on */
-	bool IsPlayingSlotAnimation(UAnimSequenceBase* Asset, FName SlotNodeName, UAnimMontage*& OutMontage);
+	bool IsPlayingSlotAnimation(const UAnimSequenceBase* Asset, FName SlotNodeName, UAnimMontage*& OutMontage) const;
 
 	/*********************************************************************************************
 	 * AnimMontage
@@ -514,27 +513,27 @@ public:
 public:
 	/** Plays an animation montage. Returns the length of the animation montage in seconds. Returns 0.f if failed to play. */
 	UFUNCTION(BlueprintCallable, Category = "Animation")
-	float Montage_Play(UAnimMontage * MontageToPlay, float InPlayRate = 1.f);
+	float Montage_Play(UAnimMontage* MontageToPlay, float InPlayRate = 1.f);
 
 	/** Stops the animation montage. If reference is NULL, it will stop ALL active montages. */
 	UFUNCTION(BlueprintCallable, Category = "Animation")
-	void Montage_Stop(float InBlendOutTime, UAnimMontage * Montage = NULL);
+	void Montage_Stop(float InBlendOutTime, const UAnimMontage* Montage = NULL);
 
 	/** Pauses the animation montage. If reference is NULL, it will pause ALL active montages. */
 	UFUNCTION(BlueprintCallable, Category = "Animation")
-	void Montage_Pause(UAnimMontage * Montage = NULL);
+	void Montage_Pause(const UAnimMontage* Montage = NULL);
 
 	/** Resumes a paused animation montage. If reference is NULL, it will resume ALL active montages. */
 	UFUNCTION(BlueprintCallable, Category = "Animation")
-	void Montage_Resume(UAnimMontage* Montage);
+	void Montage_Resume(const UAnimMontage* Montage);
 
 	/** Makes a montage jump to a named section. If Montage reference is NULL, it will do that to all active montages. */
 	UFUNCTION(BlueprintCallable, Category="Animation")
-	void Montage_JumpToSection(FName SectionName, UAnimMontage * Montage = NULL);
+	void Montage_JumpToSection(FName SectionName, const UAnimMontage* Montage = NULL);
 
 	/** Makes a montage jump to the end of a named section. If Montage reference is NULL, it will do that to all active montages. */
 	UFUNCTION(BlueprintCallable, Category="Animation")
-	void Montage_JumpToSectionsEnd(FName SectionName, UAnimMontage * Montage = NULL);
+	void Montage_JumpToSectionsEnd(FName SectionName, const UAnimMontage* Montage = NULL);
 
 	/** Relink new next section AFTER SectionNameToChange in run-time
 	 *	You can link section order the way you like in editor, but in run-time if you'd like to change it dynamically, 
@@ -546,24 +545,24 @@ public:
 	 * @param NextSection	: new next section 
 	 */
 	UFUNCTION(BlueprintCallable, Category="Animation")
-	void Montage_SetNextSection(FName SectionNameToChange, FName NextSection, UAnimMontage * Montage = NULL);
+	void Montage_SetNextSection(FName SectionNameToChange, FName NextSection, const UAnimMontage* Montage = NULL);
 
 	/** Change AnimMontage play rate. NewPlayRate = 1.0 is the default playback rate. */
 	UFUNCTION(BlueprintCallable, Category="Animation")
-	void Montage_SetPlayRate(UAnimMontage* Montage, float NewPlayRate = 1.f);
+	void Montage_SetPlayRate(const UAnimMontage* Montage, float NewPlayRate = 1.f);
 
 	/** Returns true if the animation montage is active. If the Montage reference is NULL, it will return true if any Montage is active. */
-	UFUNCTION(BlueprintCallable, Category="Animation")
-	bool Montage_IsActive(UAnimMontage* Montage);  
+	UFUNCTION(BlueprintPure, Category="Animation")
+	bool Montage_IsActive(const UAnimMontage* Montage) const;
 
 	/** Returns true if the animation montage is currently active and playing. 
 	If reference is NULL, it will return true is ANY montage is currently active and playing. */
-	UFUNCTION(BlueprintCallable, Category="Animation")
-	bool Montage_IsPlaying(UAnimMontage* Montage); 
+	UFUNCTION(BlueprintPure, Category="Animation")
+	bool Montage_IsPlaying(const UAnimMontage* Montage) const;
 
 	/** Returns the name of the current animation montage section. */
-	UFUNCTION(BlueprintCallable, Category="Animation")
-	FName Montage_GetCurrentSection(UAnimMontage* Montage = NULL);
+	UFUNCTION(BlueprintPure, Category="Animation")
+	FName Montage_GetCurrentSection(const UAnimMontage* Montage = NULL) const;
 
 	/** Called when a montage starts blending out, whether interrupted or finished */
 	UPROPERTY(BlueprintAssignable)
@@ -581,51 +580,55 @@ public:
 	* AnimMontage native C++ interface
 	********************************************************************************************* */
 public:	
-	void Montage_SetEndDelegate(FOnMontageEnded & InOnMontageEnded, UAnimMontage * Montage = NULL);
+	void Montage_SetEndDelegate(FOnMontageEnded & InOnMontageEnded, UAnimMontage* Montage = NULL);
 	
 	void Montage_SetBlendingOutDelegate(FOnMontageBlendingOutStarted & InOnMontageBlendingOut, UAnimMontage* Montage = NULL);
 	
 	/** Get pointer to BlendingOutStarted delegate for Montage.
 	If Montage reference is NULL, it will pick the first active montage found. */
-	FOnMontageBlendingOutStarted * Montage_GetBlendingOutDelegate(UAnimMontage * Montage = NULL);
+	FOnMontageBlendingOutStarted* Montage_GetBlendingOutDelegate(UAnimMontage* Montage = NULL);
 
 	/** Get Current Montage Position */
-	float Montage_GetPosition(UAnimMontage* Montage);
+	float Montage_GetPosition(const UAnimMontage* Montage);
 	
 	/** Set position. */
-	void Montage_SetPosition(UAnimMontage* Montage, float NewPosition);
+	void Montage_SetPosition(const UAnimMontage* Montage, float NewPosition);
 	
 	/** return true if Montage is not currently active. (not valid or blending out) */
-	bool Montage_GetIsStopped(UAnimMontage* Montage);
+	bool Montage_GetIsStopped(const UAnimMontage* Montage);
 
 	/** Get the current blend time of the Montage.
 	If Montage reference is NULL, it will return the current blend time on the first active Montage found. */
-	float Montage_GetBlendTime(UAnimMontage* Montage);
+	float Montage_GetBlendTime(const UAnimMontage* Montage);
 
 	/** Get PlayRate for Montage.
 	If Montage reference is NULL, PlayRate for any Active Montage will be returned.
 	If Montage is not playing, 0 is returned. */
-	float Montage_GetPlayRate(UAnimMontage* Montage);
+	float Montage_GetPlayRate(const UAnimMontage* Montage);
 
 	/** Get next sectionID for given section ID */
-	int32 Montage_GetNextSectionID(UAnimMontage const * const Montage, int32 const & CurrentSectionID) const;
+	int32 Montage_GetNextSectionID(const UAnimMontage* Montage, int32 const & CurrentSectionID) const;
 
 	/** Returns true if any montage is playing currently. Doesn't mean it's active though, it could be blending out. */
 	bool IsAnyMontagePlaying() const;
 
 	/** Get a current Active Montage in this AnimInstance. 
 		Note that there might be multiple Active at the same time. This will only return the first active one it finds. **/
-	UAnimMontage * GetCurrentActiveMontage();
+	UAnimMontage* GetCurrentActiveMontage();
 
 	/** Get Currently active montage instance.
 		Note that there might be multiple Active at the same time. This will only return the first active one it finds. **/
-	FAnimMontageInstance * GetActiveMontageInstance() const;
+	FAnimMontageInstance* GetActiveMontageInstance() const;
 
 	/** Get Active FAnimMontageInstance for given Montage asset. Will return NULL if Montage is not currently Active. */
-	FAnimMontageInstance * GetActiveInstanceForMontage(UAnimMontage const & Montage) const;
+	DEPRECATED(4.13, "Please use GetActiveInstanceForMontage(const UAnimMontage* Montage)")
+	FAnimMontageInstance* GetActiveInstanceForMontage(UAnimMontage const& Montage) const;
+
+	/** Get Active FAnimMontageInstance for given Montage asset. Will return NULL if Montage is not currently Active. */
+	FAnimMontageInstance* GetActiveInstanceForMontage(const UAnimMontage* Montage) const;
 
 	/** Get the FAnimMontageInstance currently running that matches this ID.  Will return NULL if no instance is found. */
-	FAnimMontageInstance * GetMontageInstanceForID(int32 MontageInstanceID);
+	FAnimMontageInstance* GetMontageInstanceForID(int32 MontageInstanceID);
 
 	/** AnimMontage instances that are running currently
 	* - only one is primarily active per group, and the other ones are blending out
@@ -641,7 +644,7 @@ public:
 
 protected:
 	/** Map between Active Montages and their FAnimMontageInstance */
-	TMap<class UAnimMontage *, struct FAnimMontageInstance*> ActiveMontagesMap;
+	TMap<class UAnimMontage*, struct FAnimMontageInstance*> ActiveMontagesMap;
 
 	/** Stop all montages that are active **/
 	void StopAllMontages(float BlendOut);
@@ -1162,7 +1165,7 @@ public:
 
 private:
 	/** Active Root Motion Montage Instance, if any. */
-	struct FAnimMontageInstance * RootMotionMontageInstance;
+	struct FAnimMontageInstance* RootMotionMontageInstance;
 
 	/** Temporarily queued root motion blend */
 	struct FQueuedRootMotionBlend
@@ -1290,6 +1293,8 @@ protected:
 	// TODO: Remove after deprecation (4.11)
 	friend struct FAnimationBaseContext;
 	
+	friend struct FAnimNode_SubInstance;
+
 protected:
 	/** Proxy object, nothing should access this from an externally-callable API as it is used as a scratch area on worker threads */
 	mutable FAnimInstanceProxy* AnimInstanceProxy;
diff --git a/Engine/Source/Runtime/Engine/Classes/Animation/AnimNodeBase.h b/Engine/Source/Runtime/Engine/Classes/Animation/AnimNodeBase.h
--- a/Engine/Source/Runtime/Engine/Classes/Animation/AnimNodeBase.h
+++ b/Engine/Source/Runtime/Engine/Classes/Animation/AnimNodeBase.h
@@ -164,9 +164,9 @@ public:
 		Pose.ResetToRefPose();	
 	}
 
-	void ResetToIdentity()
+	void ResetToAdditiveIdentity()
 	{
-		Pose.ResetToIdentity();
+		Pose.ResetToAdditiveIdentity();
 	}
 
 	bool ContainsNaN() const
@@ -444,6 +444,7 @@ struct FExposedValueCopyRecord
 		, DestProperty(nullptr)
 		, DestArrayIndex(0)
 		, Size(0)
+		, bInstanceIsTarget(false)
 		, PostCopyOperation(EPostCopyOperation::None)
 		, CachedBoolSourceProperty(nullptr)
 		, CachedBoolDestProperty(nullptr)
@@ -476,6 +477,10 @@ struct FExposedValueCopyRecord
 	UPROPERTY()
 	int32 Size;
 
+	// Whether or not the anim instance object is the target for the copy instead of a node.
+	UPROPERTY()
+	bool bInstanceIsTarget;
+
 	UPROPERTY()
 	EPostCopyOperation PostCopyOperation;
 
@@ -602,4 +607,9 @@ struct ENGINE_API FAnimNode_Base
 protected:
 	/** return true if enabled, otherwise, return false. This is utility function that can be used per node level */
 	bool IsLODEnabled(FAnimInstanceProxy* AnimInstanceProxy, int32 InLODThreshold);
+
+	/** Called once, from game thread as the parent anim instance is created */
+	virtual void RootInitialize(const FAnimInstanceProxy* InProxy) {}
+
+	friend struct FAnimInstanceProxy;
 };
diff --git a/Engine/Source/Runtime/Engine/Classes/Animation/AnimNode_SubInput.h b/Engine/Source/Runtime/Engine/Classes/Animation/AnimNode_SubInput.h
new file mode 100644
--- /dev/null
+++ b/Engine/Source/Runtime/Engine/Classes/Animation/AnimNode_SubInput.h
@@ -0,0 +1,26 @@
+// Copyright 1998-2016 Epic Games, Inc. All Rights Reserved.
+
+#pragma once
+
+#include "Animation/AnimNodeBase.h"
+#include "AnimNode_SubInput.generated.h"
+
+USTRUCT()
+struct ENGINE_API FAnimNode_SubInput : public FAnimNode_Base
+{
+	GENERATED_BODY()
+
+	// Begin poses for the sub instance, these will be populated by the calling subinstance node
+	// before this graph is processed
+	FCompactHeapPose InputPose;
+	FBlendedHeapCurve InputCurve;
+
+	// FAnimNode_Base interface
+	virtual void Initialize(const FAnimationInitializeContext& Context) override;
+	virtual void CacheBones(const FAnimationCacheBonesContext& Context) override;
+	virtual void Update(const FAnimationUpdateContext& Context) override;
+	virtual void Evaluate(FPoseContext& Output) override;
+	virtual void GatherDebugData(FNodeDebugData& DebugData) override;
+	// End of FAnimNode_Base interface
+
+};
\ No newline at end of file
diff --git a/Engine/Source/Runtime/Engine/Classes/Animation/AnimNode_SubInstance.h b/Engine/Source/Runtime/Engine/Classes/Animation/AnimNode_SubInstance.h
new file mode 100644
--- /dev/null
+++ b/Engine/Source/Runtime/Engine/Classes/Animation/AnimNode_SubInstance.h
@@ -0,0 +1,70 @@
+// Copyright 1998-2016 Epic Games, Inc. All Rights Reserved.
+
+#pragma once
+
+#include "Animation/AnimNodeBase.h"
+#include "AnimNode_SubInstance.generated.h"
+
+USTRUCT()
+struct ENGINE_API FAnimNode_SubInstance : public FAnimNode_Base
+{
+	GENERATED_BODY()
+
+public:
+
+	FAnimNode_SubInstance();
+
+	/** 
+	 *  Input pose for the node, intentionally not accessible because if there's no input
+	 *  Node in the target class we don't want to show this as a pin
+	 */
+	UPROPERTY()
+	FPoseLink InPose;
+
+	UPROPERTY(EditAnywhere, BlueprintReadWrite, Category = Settings)
+	TSubclassOf<UAnimInstance> InstanceClass;
+
+	/** This is the actual instance allocated at runtime that will run */
+	UPROPERTY(transient)
+	UAnimInstance* InstanceToRun;
+
+	/** List of properties on the calling instance to push from */
+	UPROPERTY(transient)
+	TArray<UProperty*> InstanceProperties;
+
+	/** List of properties on the sub instance to push to, built from name list when initialised */
+	UPROPERTY(transient)
+	TArray<UProperty*> SubInstanceProperties;
+
+	/** List of source properties to use, 1-1 with Dest names below, built by the compiler */
+	UPROPERTY()
+	TArray<FName> SourcePropertyNames;
+
+	/** List of destination properties to use, 1-1 with Source names above, built by the compiler */
+	UPROPERTY()
+	TArray<FName> DestPropertyNames;
+
+	// Temporary storage for the output of the subinstance, will be copied into output pose.
+	// Declared at member level to avoid reallocating the objects
+	TArray<FTransform> BoneTransforms;
+	FBlendedHeapCurve BlendedCurve;
+
+	// Validate the instance and create if invalid
+	void ValidateInstance();
+
+	// FAnimNode_Base interface
+	virtual void Initialize(const FAnimationInitializeContext& Context) override;
+	virtual void CacheBones(const FAnimationCacheBonesContext& Context) override;
+	virtual void Update(const FAnimationUpdateContext& Context) override;
+	virtual void Evaluate(FPoseContext& Output) override;
+	virtual void GatherDebugData(FNodeDebugData& DebugData) override;
+	virtual bool HasPreUpdate() const override;
+	virtual void PreUpdate(const UAnimInstance* InAnimInstance) override;
+
+protected:
+	virtual void RootInitialize(const FAnimInstanceProxy* InProxy) override;
+	// End of FAnimNode_Base interface
+
+	// Shutdown the currently running instance
+	void TeardownInstance();
+};
\ No newline at end of file
diff --git a/Engine/Source/Runtime/Engine/Classes/Animation/AnimSequenceBase.h b/Engine/Source/Runtime/Engine/Classes/Animation/AnimSequenceBase.h
--- a/Engine/Source/Runtime/Engine/Classes/Animation/AnimSequenceBase.h
+++ b/Engine/Source/Runtime/Engine/Classes/Animation/AnimSequenceBase.h
@@ -59,7 +59,7 @@ class UAnimSequenceBase : public UAnimationAsset
 
 	/** Returns the total play length of the montage, if played back with a speed of 1.0. */
 	UFUNCTION(BlueprintCallable, Category = "Animation")
-	virtual float GetPlayLength();
+	ENGINE_API virtual float GetPlayLength();
 
 	/** Sort the Notifies array by time, earliest first. */
 	ENGINE_API void SortNotifies();	
@@ -70,7 +70,7 @@ class UAnimSequenceBase : public UAnimationAsset
 	 * Supports playing backwards (DeltaTime<0).
 	 * Returns notifies between StartTime (exclusive) and StartTime+DeltaTime (inclusive)
 	 */
-	void GetAnimNotifies(const float& StartTime, const float& DeltaTime, const bool bAllowLooping, TArray<const FAnimNotifyEvent *>& OutActiveNotifies) const;
+	ENGINE_API void GetAnimNotifies(const float& StartTime, const float& DeltaTime, const bool bAllowLooping, TArray<const FAnimNotifyEvent *>& OutActiveNotifies) const;
 
 	/** 
 	 * Retrieves AnimNotifies between two time positions. ]PreviousPosition, CurrentPosition]
@@ -88,10 +88,10 @@ class UAnimSequenceBase : public UAnimationAsset
 	virtual int32 GetNumberOfFrames() const;
 
 	/** Get the frame number for the provided time */
-	virtual int32 GetFrameAtTime(const float Time) const;
+	ENGINE_API virtual int32 GetFrameAtTime(const float Time) const;
 
 	/** Get the time at the given frame */
-	virtual float GetTimeAtFrame(const int32 Frame) const;
+	ENGINE_API virtual float GetTimeAtFrame(const int32 Frame) const;
 	
 	// @todo document
 	ENGINE_API void InitializeNotifyTrack();
@@ -100,7 +100,7 @@ class UAnimSequenceBase : public UAnimationAsset
 	ENGINE_API void ClampNotifiesAtEndOfSequence();
 
 	/** Calculates what (if any) offset should be applied to the trigger time of a notify given its display time */ 
-	virtual EAnimEventTriggerOffsets::Type CalculateOffsetForNotify(float NotifyDisplayTime) const;
+	ENGINE_API virtual EAnimEventTriggerOffsets::Type CalculateOffsetForNotify(float NotifyDisplayTime) const;
 
 	virtual void GetAssetRegistryTags(TArray<FAssetRegistryTag>& OutTags) const override;
 	
@@ -115,7 +115,7 @@ class UAnimSequenceBase : public UAnimationAsset
 	ENGINE_API virtual void RefreshCacheData();
 
 	//~ Begin UAnimationAsset Interface
-	virtual void TickAssetPlayer(FAnimTickRecord& Instance, struct FAnimNotifyQueue& NotifyQueue, FAnimAssetTickContext& Context) const override;
+	ENGINE_API virtual void TickAssetPlayer(FAnimTickRecord& Instance, struct FAnimNotifyQueue& NotifyQueue, FAnimAssetTickContext& Context) const override;
 
 	void TickByMarkerAsFollower(FMarkerTickRecord &Instance, FMarkerTickContext &MarkerContext, float& CurrentTime, float& OutPreviousTime, const float MoveDelta, const bool bLooping) const;
 
@@ -138,7 +138,7 @@ class UAnimSequenceBase : public UAnimationAsset
 	ENGINE_API virtual void GetAnimationPose(struct FCompactPose& OutPose, FBlendedCurve& OutCurve, const FAnimExtractContext& ExtractionContext) const PURE_VIRTUAL(UAnimSequenceBase::GetAnimationPose, );
 	
 	DEPRECATED(4.11, "This function is deprecated, please use HandleAssetPlayerTickedInternal")
-	virtual void OnAssetPlayerTickedInternal(FAnimAssetTickContext &Context, const float PreviousTime, const float MoveDelta, const FAnimTickRecord &Instance, class UAnimInstance* InAnimInstance) const;
+	ENGINE_API virtual void OnAssetPlayerTickedInternal(FAnimAssetTickContext &Context, const float PreviousTime, const float MoveDelta, const FAnimTickRecord &Instance, class UAnimInstance* InAnimInstance) const;
 
 	virtual void HandleAssetPlayerTickedInternal(FAnimAssetTickContext &Context, const float PreviousTime, const float MoveDelta, const FAnimTickRecord &Instance, struct FAnimNotifyQueue& NotifyQueue) const;
 
@@ -182,7 +182,7 @@ public:
 
 #endif
 	// return true if anim notify is available 
-	virtual bool IsNotifyAvailable() const;
+	ENGINE_API virtual bool IsNotifyAvailable() const;
 
 protected:
 	template <typename DataType>
diff --git a/Engine/Source/Runtime/Engine/Classes/Animation/AnimSingleNodeInstance.h b/Engine/Source/Runtime/Engine/Classes/Animation/AnimSingleNodeInstance.h
--- a/Engine/Source/Runtime/Engine/Classes/Animation/AnimSingleNodeInstance.h
+++ b/Engine/Source/Runtime/Engine/Classes/Animation/AnimSingleNodeInstance.h
@@ -97,7 +97,7 @@ public:
 	virtual void SetAnimationAsset(UAnimationAsset* NewAsset, bool bIsLooping=true, float InPlayRate=1.f);
 	/** Set pose value */
  	UFUNCTION(BlueprintCallable, Category = "Animation")
- 	void SetPreviewPoseOverride(const FName& PoseName, float Value);
+ 	void SetPreviewCurveOverride(const FName& PoseName, float Value, bool bRemoveIfZero);
 public:
 	/** AnimSequence specific **/
 	void StepForward();
diff --git a/Engine/Source/Runtime/Engine/Classes/Animation/AnimationAsset.h b/Engine/Source/Runtime/Engine/Classes/Animation/AnimationAsset.h
--- a/Engine/Source/Runtime/Engine/Classes/Animation/AnimationAsset.h
+++ b/Engine/Source/Runtime/Engine/Classes/Animation/AnimationAsset.h
@@ -154,19 +154,21 @@ namespace ERootMotionMode
 }
 
 /** Animation Extraction Context */
-USTRUCT()
 struct FAnimExtractContext
 {
-	GENERATED_USTRUCT_BODY()
-
 	/** Is Root Motion being extracted? */
-	UPROPERTY()
 	bool bExtractRootMotion;
 
 	/** Position in animation to extract pose from */
-	UPROPERTY()
 	float CurrentTime;
 
+	/** 
+	 * Pose Curve Values to extract pose from pose assets. 
+	 * This is used by pose asset extraction 
+	 * This always has to match with pose # in the asset it's extracting from
+	 */
+	TArray<float> PoseCurves;
+
 	FAnimExtractContext()
 		: bExtractRootMotion(false)
 		, CurrentTime(0.f)
@@ -184,6 +186,14 @@ struct FAnimExtractContext
 		, CurrentTime(InCurrentTime)
 	{
 	}
+
+	FAnimExtractContext(TArray<float>& InPoseCurves)
+		: bExtractRootMotion(false)
+		, CurrentTime(0.f)
+		, PoseCurves(InPoseCurves)
+	{
+		// @todo: no support on root motion
+	}
 };
 
 //Represent a current play position in an animation
@@ -768,10 +778,10 @@ public:
 	ENGINE_API virtual TArray<FName>* GetUniqueMarkerNames() { return NULL; }
 
 	//~ Begin IInterface_AssetUserData Interface
-	virtual void AddAssetUserData(UAssetUserData* InUserData) override;
-	virtual void RemoveUserDataOfClass(TSubclassOf<UAssetUserData> InUserDataClass) override;
-	virtual UAssetUserData* GetAssetUserDataOfClass(TSubclassOf<UAssetUserData> InUserDataClass) override;
-	virtual const TArray<UAssetUserData*>* GetAssetUserDataArray() const override;
+	ENGINE_API virtual void AddAssetUserData(UAssetUserData* InUserData) override;
+	ENGINE_API virtual void RemoveUserDataOfClass(TSubclassOf<UAssetUserData> InUserDataClass) override;
+	ENGINE_API virtual UAssetUserData* GetAssetUserDataOfClass(TSubclassOf<UAssetUserData> InUserDataClass) override;
+	ENGINE_API virtual const TArray<UAssetUserData*>* GetAssetUserDataArray() const override;
 	//~ End IInterface_AssetUserData Interface
 
 	/**
diff --git a/Engine/Source/Runtime/Engine/Classes/Animation/PoseAsset.h b/Engine/Source/Runtime/Engine/Classes/Animation/PoseAsset.h
--- a/Engine/Source/Runtime/Engine/Classes/Animation/PoseAsset.h
+++ b/Engine/Source/Runtime/Engine/Classes/Animation/PoseAsset.h
@@ -28,6 +28,13 @@ struct ENGINE_API FPoseData
 	UPROPERTY()
 	TArray<FTransform>		LocalSpacePose;
 
+	// whether or not, the joint contains dirty transform
+	// it only blends if this is true 
+	// this allows per bone blend
+	// @todo: convert to bit field?
+	UPROPERTY()
+	TArray<bool>			LocalSpacePoseMask;
+
 	// # of array match with # of Curves in PoseDataContainer
  	UPROPERTY()
  	TArray<float>			CurveData;
@@ -70,7 +77,9 @@ private:
 	void AddOrUpdatePose(const FSmartName& InPoseName, const TArray<FTransform>& InlocalSpacePose, const TArray<float>& InCurveData);
 
 	// remove track if all poses has identity key
-	// we don't want to save transform per track (good for compression) because that will slow down blending
+	// Shrink currently works with only full pose
+	// we could support addiitve with comparing with identity but right now PoseContainer it self doesn't know if it's additive or not
+	// not sure if that's a good or not yet
 	void Shrink(USkeleton* InSkeleton, FName& InRetargetSourceName);
 
 	bool InsertTrack(const FName& InTrackName, USkeleton* InSkeleton, FName& InRetargetSourceName);
@@ -81,6 +90,7 @@ private:
 	void RenamePose(FSmartName OldPoseName, FSmartName NewPoseName);
 	bool DeletePose(FSmartName PoseName);
 	bool DeleteCurve(FSmartName CurveName);
+	void DeleteTrack(int32 TrackIndex);
 
 	FPoseData* FindPoseData(FSmartName PoseName);
 	FPoseData* FindOrAddPoseData(FSmartName PoseName);
@@ -96,7 +106,9 @@ private:
 	friend class UPoseAsset;
 };
 
-// @todo: add proper WITH_EDITOR
+/*
+ * Pose Asset that can be blended by weight of curves 
+ */
 UCLASS(MinimalAPI, BlueprintType)
 class UPoseAsset : public UAnimationAsset
 {
@@ -107,13 +119,6 @@ private:
 	UPROPERTY()
 	struct FPoseDataContainer PoseContainer;
 
-	/** Total number of poses - redundant data but noce to see it in asset registry*/
-	UPROPERTY(Category=Animation, AssetRegistrySearchable, VisibleAnywhere)
-	int32 NumPoses;
-
-	UPROPERTY(Category = Animation, EditAnywhere)
-	bool bNormalizeWeight;
-
 	/** Whether or not Additive Pose or not - these are property that needs post process, so */
 	UPROPERTY()
 	bool bAdditivePose;
@@ -127,6 +132,11 @@ public:
 	UPROPERTY(Category=Animation, EditAnywhere)
 	FName RetargetSource;
 
+#if WITH_EDITORONLY_DATA
+	UPROPERTY(Category=Source, EditAnywhere)
+	UAnimSequence* SourceAnimation;
+#endif // WITH_EDITORONLY_DATA
+
 	/**
 	* Get Animation Pose from one pose of PoseIndex and with PoseWeight
 	* This returns OutPose and OutCurve of one pose of PoseIndex with PoseWeight
@@ -136,30 +146,58 @@ public:
 	* @param	PoseIndex			Index of Pose
 	* @param	PoseWeight			Weight of pose
 	*/
-	ENGINE_API void GetAnimationPose(struct FCompactPose& OutPose, FBlendedCurve& InOutCurve, const FAnimExtractContext& ExtractionContext) const;
-	ENGINE_API void GetBaseAnimationPose(struct FCompactPose& OutPose, FBlendedCurve& InOutCurve, const FAnimExtractContext& ExtractionContext) const;
+	ENGINE_API bool GetAnimationPose(struct FCompactPose& OutPose, FBlendedCurve& OutCurve, const FAnimExtractContext& ExtractionContext) const;
+	ENGINE_API void GetBaseAnimationPose(struct FCompactPose& OutPose, FBlendedCurve& OutCurve, const FAnimExtractContext& ExtractionContext) const;
 	virtual bool HasRootMotion() const { return false; }
 	virtual bool IsValidAdditive() const { return bAdditivePose; }
 
 	//Begin UObject Interface
 	virtual void PostLoad() override;
 	virtual void Serialize(FArchive& Ar) override;
+	virtual void GetAssetRegistryTags(TArray<FAssetRegistryTag>& OutTags) const override;
 	//End UObject Interface
 
 public:
+	ENGINE_API int32 GetNumPoses() const;
+	ENGINE_API int32 GetNumCurves() const;
 	ENGINE_API const TArray<FSmartName> GetPoseNames() const;
+	ENGINE_API const TArray<FName>		GetTrackNames() const;
 	ENGINE_API const TArray<FSmartName> GetCurveNames() const;
+	ENGINE_API const TArray<FAnimCurveBase> GetCurveData() const;
+	ENGINE_API const TArray<float> GetCurveValues(const int32 PoseIndex) const;
+
+	/** Find index of a track with a given bone name. Returns INDEX_NONE if not found. */
+	ENGINE_API const int32 GetTrackIndexByName(const FName& InTrackName) const;
+
+	/** 
+	 *	Get local space pose for a particular track (by name) in a particular pose (by index) 
+	 *	@return	Returns true if OutTransform is valid, false if not
+	 */
+	ENGINE_API bool GetLocalPoseForTrack(const int32 PoseIndex, const int32 TrackIndex, FTransform& OutTransform) const;
+
+	/** 
+	 *	Return value of a curve for a particular pose 
+	 *	@return	Returns true if OutValue is valid, false if not
+	 */
+	ENGINE_API bool GetCurveValue(const int32 PoseIndex, const int32 CurveIndex, float& OutValue) const;
+	 
+	ENGINE_API bool ContainsPose(const FSmartName& InPoseName) const { return PoseContainer.Contains(InPoseName); }
+	ENGINE_API bool ContainsPose(const FName& InPoseName) const;
 
-	bool ContainsPose(const FSmartName& InPoseName) const { return PoseContainer.Contains(InPoseName); }
-	bool ContainsPose(const FName& InPoseName) const;
-
-	UFUNCTION(BlueprintCallable, Category = "Animation")
 	ENGINE_API void AddOrUpdatePose(const FSmartName& PoseName, USkeletalMeshComponent* MeshComponent);
+	ENGINE_API void AddOrUpdatePoseWithUniqueName(USkeletalMeshComponent* MeshComponent);
 	ENGINE_API void AddOrUpdatePose(const FSmartName& PoseName, TArray<FName> TrackNames, TArray<FTransform>& LocalTransform);
 
 #if WITH_EDITOR
-	ENGINE_API void CreatePoseFromAnimation(class UAnimSequence* AnimSequence);
+	ENGINE_API void CreatePoseFromAnimation(class UAnimSequence* AnimSequence, const TArray<FSmartName>* InPoseNames = nullptr);
+	ENGINE_API void UpdatePoseFromAnimation(class UAnimSequence* AnimSequence);
+
+	// Begin AnimationAsset interface
+	virtual bool GetAllAnimationSequencesReferred(TArray<UAnimationAsset*>& AnimationAssets) override;
+	virtual void ReplaceReferredAnimations(const TMap<UAnimationAsset*, UAnimationAsset*>& ReplacementMap) override;
+	// End AnimationAsset interface
 #endif
+
 	ENGINE_API bool ModifyPoseName(FName OldPoseName, FName NewPoseName, const FSmartNameMapping::UID* NewUID);
 
 	ENGINE_API int32 DeletePoses(TArray<FName> PoseNamesToDelete);
@@ -170,6 +208,8 @@ public:
 	ENGINE_API const int32 GetPoseIndexByName(const FName& InBasePoseName) const;
 	ENGINE_API const FName GetPoseNameByIndex(int32 InBasePoseIndex) const { return PoseContainer.PoseNames.IsValidIndex(InBasePoseIndex)? PoseContainer.PoseNames[InBasePoseIndex].DisplayName : NAME_None; }
 
+	ENGINE_API const int32 GetCurveIndexByName(const FName& InCurveName) const;
+
 private: 
 	DECLARE_MULTICAST_DELEGATE(FOnPoseListChangedMulticaster)
 	FOnPoseListChangedMulticaster OnPoseListChanged;
@@ -191,6 +231,17 @@ public:
 #if WITH_EDITOR
 protected:
 	virtual void RemapTracksToNewSkeleton(USkeleton* NewSkeleton, bool bConvertSpaces) override;
+
+private:
+	DECLARE_MULTICAST_DELEGATE(FOnPoseAssetModifiedMulticaster);
+	FOnPoseAssetModifiedMulticaster OnAssetModifiedNotifier;
+
+public:
+	typedef FOnPoseAssetModifiedMulticaster::FDelegate FAssetModifiedNotifier;
+
+	/** Registers a delegate to be called after asset has changed*/
+	ENGINE_API void RegisterOnAssetModifiedNotifier(const FAssetModifiedNotifier& Delegate);
+	ENGINE_API void UnregisterOnAssetModifiedNotifier(void* Unregister);
 #endif // WITH_EDITOR	
 
 private:
@@ -204,4 +255,6 @@ private:
 
 	bool GetBasePoseTransform(TArray<FTransform>& OutBasePose, TArray<float>& OutCurve) const;
 	void RecacheTrackmap();
+
+	void Reinitialize();
 };
diff --git a/Engine/Source/Runtime/Engine/Classes/Animation/Skeleton.h b/Engine/Source/Runtime/Engine/Classes/Animation/Skeleton.h
--- a/Engine/Source/Runtime/Engine/Classes/Animation/Skeleton.h
+++ b/Engine/Source/Runtime/Engine/Classes/Animation/Skeleton.h
@@ -227,7 +227,7 @@ class USkeleton : public UObject
 {
 	GENERATED_UCLASS_BODY()
 
-private:
+protected:
 	/** Skeleton bone tree - each contains name and parent index**/
 	UPROPERTY(VisibleAnywhere, Category=Skeleton)
 	TArray<struct FBoneNode> BoneTree;
@@ -353,9 +353,9 @@ public:
 	ENGINE_API void RemoveSmartnamesAndModify(FName ContainerName, const TArray<FSmartNameMapping::UID>& Uids);
 
 	// quick wrapper function for Find UID by name, if not found, it will return FSmartNameMapping::MaxUID
-	ENGINE_API FSmartNameMapping::UID GetUIDByName(FName ContainerName, FName Name);
-	ENGINE_API bool GetSmartNameByUID(FName ContainerName, FSmartNameMapping::UID UID, FSmartName& OutSmartName);
-	ENGINE_API bool GetSmartNameByName(FName ContainerName, FName InName, FSmartName& OutSmartName);
+	ENGINE_API FSmartNameMapping::UID GetUIDByName(const FName& ContainerName, const FName& Name);
+	ENGINE_API bool GetSmartNameByUID(const FName& ContainerName, FSmartNameMapping::UID UID, FSmartName& OutSmartName);
+	ENGINE_API bool GetSmartNameByName(const FName& ContainerName, const FName& InName, FSmartName& OutSmartName);
 
 	// Adds a new name to the smart name container and modifies the skeleton so it can be saved
 	// return bool - Whether a name was added (false if already present)
@@ -382,7 +382,7 @@ private:
 	FRigConfiguration RigConfig;
 
 	/** rig property will be saved separately */
-	virtual void GetAssetRegistryTags(TArray<FAssetRegistryTag>& OutTags) const override;
+	ENGINE_API virtual void GetAssetRegistryTags(TArray<FAssetRegistryTag>& OutTags) const override;
 
 public:
 
@@ -492,7 +492,7 @@ public:
 	ENGINE_API bool IsCompatibleMesh(const USkeletalMesh* InSkelMesh) const;
 
 	/** Clears all cache data **/
-	void ClearCacheData();
+	ENGINE_API void ClearCacheData();
 
 	/** 
 	 * Find a mesh linkup table (mapping of skeleton bone tree indices to refpose indices) for a particular SkeletalMesh
@@ -602,10 +602,10 @@ public:
 
 	ENGINE_API void SetBoneTranslationRetargetingMode(const int32& BoneIndex, EBoneTranslationRetargetingMode::Type NewRetargetingMode, bool bChildrenToo=false);
 
-	virtual void PostLoad() override;
-	virtual void PostDuplicate(bool bDuplicateForPIE) override;
-	virtual void PostInitProperties() override;
-	virtual void Serialize(FArchive& Ar) override;
+	ENGINE_API virtual void PostLoad() override;
+	ENGINE_API virtual void PostDuplicate(bool bDuplicateForPIE) override;
+	ENGINE_API virtual void PostInitProperties() override;
+	ENGINE_API virtual void Serialize(FArchive& Ar) override;
 
 	ENGINE_API static void AddReferencedObjects(UObject* InThis, FReferenceCollector& Collector);
 
diff --git a/Engine/Source/Runtime/Engine/Classes/Components/AudioComponent.h b/Engine/Source/Runtime/Engine/Classes/Components/AudioComponent.h
--- a/Engine/Source/Runtime/Engine/Classes/Components/AudioComponent.h
+++ b/Engine/Source/Runtime/Engine/Classes/Components/AudioComponent.h
@@ -255,9 +255,14 @@ class ENGINE_API UAudioComponent : public USceneComponent
 	UFUNCTION(BlueprintCallable, Category="Audio|Components|Audio")
 	void AdjustAttenuation(const FAttenuationSettings& InAttenuationSettings);
 
+	static void PlaybackCompleted(uint64 AudioComponentID, bool bFailedToStart);
+
+private:
 	/** Called by the ActiveSound to inform the component that playback is finished */
 	void PlaybackCompleted(bool bFailedToStart);
 
+public:
+
 	/** Sets the sound instance parameter. */
 	void SetSoundParameter(const FAudioComponentParam& Param);
 
@@ -299,6 +304,7 @@ public:
 #endif // WITH_EDITOR
 	virtual FString GetDetailedInfoInternal() const override;
 	virtual void PostLoad() override;
+	virtual void BeginDestroy() override;
 	//~ End UObject Interface.
 
 	//~ Begin USceneComponent Interface
@@ -327,6 +333,12 @@ public:
 	/** Returns the active audio device to use for this component based on whether or not the component is playing in a world. */
 	FAudioDevice* GetAudioDevice() const;
 
+	uint64 GetAudioComponentID() const { return AudioComponentID; }
+
+	static UAudioComponent* GetAudioComponentFromID(uint64 AudioComponentID);
+
+	void UpdateInteriorSettings(bool bFullUpdate);
+
 private:
 	
 #if WITH_EDITORONLY_DATA
@@ -339,6 +351,12 @@ private:
 
 	/** A count of how many times we've started playing */
 	int32 ActiveCount;
+
+	static uint64 AudioComponentIDCounter;
+	static TMap<uint64, UAudioComponent*> AudioIDToComponentMap;
+
+	uint64 AudioComponentID;
+
 };
 
 
diff --git a/Engine/Source/Runtime/Engine/Classes/Components/InstancedStaticMeshComponent.h b/Engine/Source/Runtime/Engine/Classes/Components/InstancedStaticMeshComponent.h
--- a/Engine/Source/Runtime/Engine/Classes/Components/InstancedStaticMeshComponent.h
+++ b/Engine/Source/Runtime/Engine/Classes/Components/InstancedStaticMeshComponent.h
@@ -106,6 +106,8 @@ class ENGINE_API UInstancedStaticMeshComponent : public UStaticMeshComponent
 	UFUNCTION(BlueprintCallable, Category = "Components|InstancedStaticMesh")
 	bool GetInstanceTransform(int32 InstanceIndex, FTransform& OutInstanceTransform, bool bWorldSpace = false) const;
 
+	virtual void OnUpdateTransform(EUpdateTransformFlags UpdateTransformFlags, ETeleportType Teleport) override;
+
 	/**
 	* Update the transform for the instance specified.
 	*
diff --git a/Engine/Source/Runtime/Engine/Classes/Components/SceneComponent.h b/Engine/Source/Runtime/Engine/Classes/Components/SceneComponent.h
--- a/Engine/Source/Runtime/Engine/Classes/Components/SceneComponent.h
+++ b/Engine/Source/Runtime/Engine/Classes/Components/SceneComponent.h
@@ -81,7 +81,7 @@ FORCEINLINE void operator&=(EMoveComponentFlags& Dest,EMoveComponentFlags Arg)
 FORCEINLINE void operator|=(EMoveComponentFlags& Dest,EMoveComponentFlags Arg)					{ Dest = EMoveComponentFlags(Dest | Arg); }
 
 DECLARE_DYNAMIC_MULTICAST_DELEGATE_OneParam(FPhysicsVolumeChanged, class APhysicsVolume*, NewVolume);
-
+DECLARE_EVENT_ThreeParams(USceneComponent, FTransformUpdated, USceneComponent* /*UpdatedComponent*/, EUpdateTransformFlags /*UpdateTransformFlags*/, ETeleportType /*Teleport*/);
 
 /**
  * A SceneComponent has a transform and supports attachment, but has no rendering or collision capabilities.
@@ -749,6 +749,8 @@ public:
 	UPROPERTY(BlueprintAssignable, Category=PhysicsVolume, meta=(DisplayName="Physics Volume Changed"))
 	FPhysicsVolumeChanged PhysicsVolumeChangedDelegate;
 
+	FTransformUpdated TransformUpdated;
+
 	//~ Begin ActorComponent Interface
 	virtual void OnRegister() override;
 	/** Return true if CreateRenderState() should be called */
diff --git a/Engine/Source/Runtime/Engine/Classes/Components/SkeletalMeshComponent.h b/Engine/Source/Runtime/Engine/Classes/Components/SkeletalMeshComponent.h
--- a/Engine/Source/Runtime/Engine/Classes/Components/SkeletalMeshComponent.h
+++ b/Engine/Source/Runtime/Engine/Classes/Components/SkeletalMeshComponent.h
@@ -364,6 +364,10 @@ public:
 	UPROPERTY(transient)
 	class UAnimInstance* AnimScriptInstance;
 
+	/** Any running sub anim instances that need to be updates on the game thread */
+	UPROPERTY(transient)
+	TArray<UAnimInstance*> SubInstances;
+
 	UPROPERTY(EditAnywhere, BlueprintReadWrite, Category=Animation, meta=(ShowOnlyInnerProperties))
 	struct FSingleAnimationPlayData AnimationData;
 
@@ -785,6 +789,8 @@ private:
 	/** Copies the data from the external cloth simulation context. We copy instead of flipping because the API has to return the full struct to make backwards compat easy*/
 	void UpdateClothSimulationContext();
 
+
+
    /** 
 	* clothing actors will be created from clothing assets for cloth simulation 
 	* 1 actor should correspond to 1 asset
@@ -818,6 +824,10 @@ public:
 #endif // WITH_APEX_CLOTHING
 
 private:
+
+	/** Wrapper that calls our constraint broken delegate */
+	void OnConstraintBrokenWrapper(int32 ConstraintIndex);
+
 	/** 
 	 * Morph Target Curves. This will override AnimInstance MorphTargetCurves
 	 * if same curve is found
@@ -1176,6 +1186,9 @@ public:
 	
 	/** Update bHasValidBodies flag */
 	void UpdateHasValidBodies();
+
+	/** Update the bone mapping on each body instance. This is useful when changing skeletal mesh without recreating bodies */
+	void UpdateBoneBodyMapping();
 	
 	/** 
 	 * Initialize SkelControls
diff --git a/Engine/Source/Runtime/Engine/Classes/Engine/Engine.h b/Engine/Source/Runtime/Engine/Classes/Engine/Engine.h
--- a/Engine/Source/Runtime/Engine/Classes/Engine/Engine.h
+++ b/Engine/Source/Runtime/Engine/Classes/Engine/Engine.h
@@ -2924,8 +2924,10 @@ private:
 	bool ToggleStatUnitGraph(UWorld* World, FCommonViewportClient* ViewportClient, const TCHAR* Stream = nullptr);
 	bool ToggleStatUnitTime(UWorld* World, FCommonViewportClient* ViewportClient, const TCHAR* Stream = nullptr);
 	bool ToggleStatRaw(UWorld* World, FCommonViewportClient* ViewportClient, const TCHAR* Stream = nullptr);
-#endif
+	bool ToggleStatSoundWaves(UWorld* World, FCommonViewportClient* ViewportClient, const TCHAR* Stream = nullptr);
+	bool ToggleStatSoundCues(UWorld* World, FCommonViewportClient* ViewportClient, const TCHAR* Stream = nullptr);
 	bool ToggleStatSounds(UWorld* World, FCommonViewportClient* ViewportClient, const TCHAR* Stream = nullptr);
+#endif
 
 	/**
 	 * Functions for rendering the various simple stats, should only be used when registering with EngineStats.
@@ -2954,8 +2956,8 @@ private:
 	int32 RenderStatSoundMixes(UWorld* World, FViewport* Viewport, FCanvas* Canvas, int32 X, int32 Y, const FVector* ViewLocation = nullptr, const FRotator* ViewRotation = nullptr);
 	int32 RenderStatSoundWaves(UWorld* World, FViewport* Viewport, FCanvas* Canvas, int32 X, int32 Y, const FVector* ViewLocation = nullptr, const FRotator* ViewRotation = nullptr);
 	int32 RenderStatSoundCues(UWorld* World, FViewport* Viewport, FCanvas* Canvas, int32 X, int32 Y, const FVector* ViewLocation = nullptr, const FRotator* ViewRotation = nullptr);
-#endif // !UE_BUILD_SHIPPING
 	int32 RenderStatSounds(UWorld* World, FViewport* Viewport, FCanvas* Canvas, int32 X, int32 Y, const FVector* ViewLocation = nullptr, const FRotator* ViewRotation = nullptr);
+#endif // !UE_BUILD_SHIPPING
 	int32 RenderStatAI(UWorld* World, FViewport* Viewport, FCanvas* Canvas, int32 X, int32 Y, const FVector* ViewLocation = nullptr, const FRotator* ViewRotation = nullptr);
 #if STATS
 	int32 RenderStatSlateBatches(UWorld* World, FViewport* Viewport, FCanvas* Canvas, int32 X, int32 Y, const FVector* ViewLocation = nullptr, const FRotator* ViewRotation = nullptr);
diff --git a/Engine/Source/Runtime/Engine/Classes/Engine/EngineBaseTypes.h b/Engine/Source/Runtime/Engine/Classes/Engine/EngineBaseTypes.h
--- a/Engine/Source/Runtime/Engine/Classes/Engine/EngineBaseTypes.h
+++ b/Engine/Source/Runtime/Engine/Classes/Engine/EngineBaseTypes.h
@@ -199,9 +199,12 @@ public:
 	uint8 bRunOnAnyThread:1;
 
 private:
-	/** If true, means that this tick function is in the master array of tick functions **/
+	/** If true, means that this tick function is in the master array of tick functions */
 	uint8 bRegistered:1;
 
+	/** Cache whether this function was rescheduled as an interval function during StartParallel */
+	uint8 bWasInterval:1;
+
 	enum class ETickState : uint8
 	{
 		Disabled,
@@ -337,9 +340,8 @@ private:
 	 * Queues a tick function for execution from the game thread
 	 * @param TickContext - context to tick in
 	 * @param StackForCycleDetection - Stack For Cycle Detection
-	 * @param bWasInterval - true if this was an interval tick
 	 */
-	void QueueTickFunctionParallel(const struct FTickContext& TickContext, TArray<FTickFunction*, TInlineAllocator<8> >& StackForCycleDetection, bool bWasInterval);
+	void QueueTickFunctionParallel(const struct FTickContext& TickContext, TArray<FTickFunction*, TInlineAllocator<8> >& StackForCycleDetection);
 
 	/** 
 	 * Logs the prerequisites
diff --git a/Engine/Source/Runtime/Engine/Classes/Engine/EngineTypes.h b/Engine/Source/Runtime/Engine/Classes/Engine/EngineTypes.h
--- a/Engine/Source/Runtime/Engine/Classes/Engine/EngineTypes.h
+++ b/Engine/Source/Runtime/Engine/Classes/Engine/EngineTypes.h
@@ -1419,7 +1419,7 @@ struct FLocalizedSubtitle
 	 * as the contents of the subtitle is commonly identical to what is spoken.
 	 */
 	UPROPERTY()
-	TArray<struct FSubtitleCue> Subtitles;
+	TArray<FSubtitleCue> Subtitles;
 
 	/** true if this sound is considered to contain mature content. */
 	UPROPERTY()
diff --git a/Engine/Source/Runtime/Engine/Classes/Engine/World.h b/Engine/Source/Runtime/Engine/Classes/Engine/World.h
--- a/Engine/Source/Runtime/Engine/Classes/Engine/World.h
+++ b/Engine/Source/Runtime/Engine/Classes/Engine/World.h
@@ -698,7 +698,7 @@ public:
 	/** A static map that is populated before loading a world from a package. This is so UWorld can look up its WorldType in ::PostLoad */
 	static TMap<FName, EWorldType::Type> WorldTypePreLoadMap;
 
-	/** Map of blueprints that are bieng debugged and the object instance they are debugging. */
+	/** Map of blueprints that are being debugged and the object instance they are debugging. */
 	typedef TMap<TWeakObjectPtr<class UBlueprint>, TWeakObjectPtr<UObject> > FBlueprintToDebuggedObjectMap;
 
 	/** Return the array of objects currently bieng debugged. */
@@ -941,8 +941,9 @@ public:
 	/** An array of post processing volumes, sorted in ascending order of priority.					*/
 	TArray< IInterface_PostProcessVolume * > PostProcessVolumes;
 
-	/** Pointer to the higest priority audio volumes, each volume has a reference to the next lower priority volume creating a linked list of prioritized audio volumes in descending order */
-	TAutoWeakObjectPtr<class AAudioVolume> HighestPriorityAudioVolume;
+	/** Set of AudioVolumes sorted by  */
+	// TODO: Make this be property UPROPERTY(Transient)
+	TSet<class AAudioVolume*> AudioVolumes;
 
 	/** Handle to the active audio device for this world. */
 	uint32 AudioDeviceHandle;
@@ -2688,6 +2689,9 @@ public:
 	 */
 	class AAudioVolume* GetAudioSettings( const FVector& ViewLocation, struct FReverbSettings* OutReverbSettings, struct FInteriorSettings* OutInteriorSettings );
 
+	/** Returns the audio device handle for this world.*/
+	uint32 GetAudioDeviceHandle() const { return AudioDeviceHandle; }
+
 	/** Sets the audio device handle to the active audio device for this world.*/
 	void SetAudioDeviceHandle(const uint32 InAudioDeviceHandle);
 
diff --git a/Engine/Source/Runtime/Engine/Classes/GameFramework/CharacterMovementComponent.h b/Engine/Source/Runtime/Engine/Classes/GameFramework/CharacterMovementComponent.h
--- a/Engine/Source/Runtime/Engine/Classes/GameFramework/CharacterMovementComponent.h
+++ b/Engine/Source/Runtime/Engine/Classes/GameFramework/CharacterMovementComponent.h
@@ -753,8 +753,11 @@ public:
 	UPROPERTY(Transient, Category="Character Movement", EditAnywhere, BlueprintReadWrite)
 	uint32 bIgnoreClientMovementErrorChecksAndCorrection:1;
 
-	/** if true, event NotifyJumpApex() to CharacterOwner's controller when at apex of jump.  Is cleared when event is triggered. */
-	UPROPERTY(Category="Character Movement: Jumping / Falling", EditAnywhere, BlueprintReadWrite)
+	/**
+	 * If true, event NotifyJumpApex() to CharacterOwner's controller when at apex of jump. Is cleared when event is triggered.
+	 * By default this is off, and if you want the event to fire you typically set it to true when movement mode changes to "Falling" from another mode (see OnMovementModeChanged).
+	 */
+	UPROPERTY(Category="Character Movement: Jumping / Falling", VisibleAnywhere, BlueprintReadWrite)
 	uint32 bNotifyApex:1;
 
 	/** Instantly stop when in flying mode and no acceleration is being applied. */
diff --git a/Engine/Source/Runtime/Engine/Classes/GameFramework/WorldSettings.h b/Engine/Source/Runtime/Engine/Classes/GameFramework/WorldSettings.h
--- a/Engine/Source/Runtime/Engine/Classes/GameFramework/WorldSettings.h
+++ b/Engine/Source/Runtime/Engine/Classes/GameFramework/WorldSettings.h
@@ -438,11 +438,11 @@ class ENGINE_API AWorldSettings : public AInfo, public IInterface_AssetUserData
 	/** AUDIO SETTINGS **/
 	/** Default reverb settings used by audio volumes.													*/
 	UPROPERTY(EditAnywhere, config, Category=Audio)
-	struct FReverbSettings DefaultReverbSettings;
+	FReverbSettings DefaultReverbSettings;
 
 	/** Default interior settings used by audio volumes.												*/
 	UPROPERTY(EditAnywhere, config, Category=Audio)
-	struct FInteriorSettings DefaultAmbientZoneSettings;
+	FInteriorSettings DefaultAmbientZoneSettings;
 
 	/** Default Base SoundMix.																			*/
 	UPROPERTY(EditAnywhere, Category=Audio)
@@ -543,6 +543,7 @@ public:
 #if WITH_EDITOR
 	virtual bool CanEditChange(const UProperty* InProperty) const override;
 	virtual void PostEditChangeProperty(FPropertyChangedEvent& PropertyChangedEvent) override;
+	virtual void PostEditChangeChainProperty(FPropertyChangedChainEvent& PropertyChangedEvent) override;
 #endif // WITH_EDITOR
 	//~ End UObject Interface.
 
@@ -553,6 +554,7 @@ public:
 #endif // WITH_EDITOR
 	virtual void PreInitializeComponents() override;
 	virtual void PostInitializeComponents() override;
+	virtual void PostRegisterAllComponents() override;
 	//~ End AActor Interface.
 
 	/**
diff --git a/Engine/Source/Runtime/Engine/Classes/Kismet/GameplayStatics.h b/Engine/Source/Runtime/Engine/Classes/Kismet/GameplayStatics.h
--- a/Engine/Source/Runtime/Engine/Classes/Kismet/GameplayStatics.h
+++ b/Engine/Source/Runtime/Engine/Classes/Kismet/GameplayStatics.h
@@ -2,6 +2,7 @@
 
 #pragma once
 #include "Engine/LatentActionManager.h"
+#include "Kismet/KismetSystemLibrary.h"
 #include "Kismet/BlueprintFunctionLibrary.h"
 #include "Sound/DialogueTypes.h"
 #include "GameplayStatics.generated.h"
@@ -738,9 +739,9 @@ class ENGINE_API UGameplayStatics : public UBlueprintFunctionLibrary
 	static bool DoesSaveGameExist(const FString& SlotName, const int32 UserIndex);
 
 	/** 
-	 *	Save the contents of the SaveGameObject to a slot.
-	 *	@param SlotName			Name of save game slot to save to.
-	 *  @param UserIndex		For some platforms, master user index to identify the user doing the saving.
+	 *	Load the contents from a given slot.
+	 *	@param SlotName			Name of the save game slot to load from.
+	 *  @param UserIndex		For some platforms, master user index to identify the user doing the loading.
 	 *	@return SaveGameObject	Object containing loaded game state (NULL if load fails)
 	 */
 	UFUNCTION(BlueprintCallable, Category="Game")
@@ -748,8 +749,8 @@ class ENGINE_API UGameplayStatics : public UBlueprintFunctionLibrary
 
 	/**
 	 * Delete a save game in a particular slot.
-	 *	@param SlotName			Name of save game slot to save to.
-	 *  @param UserIndex		For some platforms, master user index to identify the user doing the saving.
+	 *	@param SlotName			Name of save game slot to delete.
+	 *  @param UserIndex		For some platforms, master user index to identify the user doing the deletion.
 	 *  @return True if a file was actually able to be deleted. use DoesSaveGameExist to distinguish between delete failures and failure due to file not existing.
 	 */
 	UFUNCTION(BlueprintCallable, Category = "Game")
@@ -804,6 +805,42 @@ class ENGINE_API UGameplayStatics : public UBlueprintFunctionLibrary
 	/** Native version, has more options than the Blueprint version. */
 	static bool SuggestProjectileVelocity(UObject* WorldContextObject, FVector& TossVelocity, FVector StartLocation, FVector EndLocation, float TossSpeed, bool bHighArc = false, float CollisionRadius = 0.f, float OverrideGravityZ = 0, ESuggestProjVelocityTraceOption::Type TraceOption = ESuggestProjVelocityTraceOption::TraceFullPath, const FCollisionResponseParams& ResponseParam = FCollisionResponseParams::DefaultResponseParam, const TArray<AActor*>& ActorsToIgnore = TArray<AActor*>(), bool bDrawDebug = false);
 
+	/**
+	* Predict the arc of a virtual projectile affected by gravity with collision checks along the arc. Returns a list of positions of the simulated arc and the destination reached by the simulation.
+	* Returns true if it hit something.
+	*
+	* @param OutPathPositions			Predicted projectile path. Ordered series of positions from StartPos to the end. Includes location at point of impact if it hit something.
+	* @param OutHit						Predicted hit result, if the projectile will hit something
+	* @param OutLastTraceDestination	Goal position of the final trace it did. Will not be in the path if there is a hit.
+	* @param StartPos					First start trace location
+	* @param LaunchVelocity				Velocity the "virtual projectile" is launched at
+	* @param bTracePath					Trace along the entire path to look for blocking hits
+	* @param ProjectileRadius			Radius of the virtual projectile to sweep against the environment
+	* @param ObjectTypes				ObjecTypes to trace against
+	* @param bTraceComplex				Use TraceComplex (trace against triangles not primitives)
+	* @param ActorsToIgnore				Actors to exclude from the traces
+	* @param DrawDebugType				Debug type (one-frame, duration, persistent)
+	* @param DrawDebugTime				Duration of debug lines (only relevant for DrawDebugType::Duration)
+	* @param SimFrequency				Determines size of each sub-step in the simulation (chopping up MaxSimTime)
+	* @param MaxSimTime					Maximum simulation time for the virtual projectile.
+	* @param OverrideGravityZ			Optional override of Gravity (default uses WorldGravityZ
+	*/
+	UFUNCTION(BlueprintCallable, Category = "Game", meta = (WorldContext = "WorldContextObject", AutoCreateRefTerm = "ActorsToIgnore", AdvancedDisplay = "DrawDebugTime, DrawDebugType, SimFrequency, MaxSimTime, OverrideGravityZ", TraceChannel = ECC_WorldDynamic, bTracePath = true))
+	static bool PredictProjectilePath(UObject* WorldContextObject, FHitResult& OutHit, TArray<FVector>& OutPathPositions, FVector& OutLastTraceDestination, FVector StartPos, FVector LaunchVelocity, bool bTracePath, float ProjectileRadius, const TArray<TEnumAsByte<EObjectTypeQuery> >& ObjectTypes, bool bTraceComplex, const TArray<AActor*>& ActorsToIgnore, EDrawDebugTrace::Type DrawDebugType, float DrawDebugTime, float SimFrequency = 30.f, float MaxSimTime = 2.f, float OverrideGravityZ = 0);
+
+	/**
+	* Returns the launch velocity needed for a projectile at rest at StartPos to land on EndPos.
+	* Assumes a medium arc (e.g. 45 deg on level ground). Projectile velocity is variable and unconstrained.
+	* Does no tracing.
+	*
+	* @param OutLaunchVelocity			Returns the launch velocity required to reach the EndPos
+	* @param StartPos					Start position of the simulation
+	* @param EndPos						Desired end location for the simulation
+	* @param OverrideGravityZ			Optional override of WorldGravityZ
+	*/
+	UFUNCTION(BlueprintCallable, Category = "Game", DisplayName = "SuggestProjectileVelocity MediumArc", meta = (WorldContext = "WorldContextObject", AdvancedDisplay = "OverrideGravityZ"))
+	static bool SuggestProjectileVelocity_MediumArc(UObject* WorldContextObject, FVector& OutLaunchVelocity, FVector StartPos, FVector EndPos, float OverrideGravityZ = 0);
+
 	/** Returns world origin current location. */
 	UFUNCTION(BlueprintPure, Category="Game", meta=(WorldContext="WorldContextObject") )
 	static FIntVector GetWorldOriginLocation(UObject* WorldContextObject);
diff --git a/Engine/Source/Runtime/Engine/Classes/Kismet/KismetSystemLibrary.h b/Engine/Source/Runtime/Engine/Classes/Kismet/KismetSystemLibrary.h
--- a/Engine/Source/Runtime/Engine/Classes/Kismet/KismetSystemLibrary.h
+++ b/Engine/Source/Runtime/Engine/Classes/Kismet/KismetSystemLibrary.h
@@ -294,9 +294,20 @@ class ENGINE_API UKismetSystemLibrary : public UBlueprintFunctionLibrary
 	UFUNCTION(BlueprintCallable, meta=(Latent, LatentInfo="LatentInfo", WorldContext="WorldContextObject", Duration="0.2", Keywords="sleep"), Category="Utilities|FlowControl")
 	static void RetriggerableDelay(UObject* WorldContextObject, float Duration, FLatentActionInfo LatentInfo);
 
-	/** Interpolate a component to the specified relative location and rotation over the course of OverTime seconds. */
+	/*
+	 * Interpolate a component to the specified relative location and rotation over the course of OverTime seconds. 
+	 * @param Component						Component to interpolate
+	 * @param TargetRelativeLocation		Relative target location
+	 * @param TargetRelativeRotation		Relative target rotation
+	 * @param bEaseOut						if true we will ease out (ie end slowly) during interpolation
+	 * @param bEaseIn						if true we will ease in (ie start slowly) during interpolation
+	 * @param OverTime						duration of interpolation
+	 * @param bForceShortestRotationPath	if true we will always use the shortest path for rotation
+	 * @param MoveAction					required movement behavior @see EMoveComponentAction
+	 * @param LatentInfo					The latent action
+	 */
 	UFUNCTION(BlueprintCallable, meta=(Latent, LatentInfo="LatentInfo", WorldContext="WorldContextObject", ExpandEnumAsExecs="MoveAction", OverTime="0.2"), Category="Components")
-	static void MoveComponentTo(USceneComponent* Component, FVector TargetRelativeLocation, FRotator TargetRelativeRotation, bool bEaseOut, bool bEaseIn, float OverTime, TEnumAsByte<EMoveComponentAction::Type> MoveAction, FLatentActionInfo LatentInfo);
+	static void MoveComponentTo(USceneComponent* Component, FVector TargetRelativeLocation, FRotator TargetRelativeRotation, bool bEaseOut, bool bEaseIn, float OverTime, bool bForceShortestRotationPath, TEnumAsByte<EMoveComponentAction::Type> MoveAction, FLatentActionInfo LatentInfo);
 
 	// --- Timer functions with delegate input ----------
 
diff --git a/Engine/Source/Runtime/Engine/Classes/PhysicsEngine/BodyInstance.h b/Engine/Source/Runtime/Engine/Classes/PhysicsEngine/BodyInstance.h
--- a/Engine/Source/Runtime/Engine/Classes/PhysicsEngine/BodyInstance.h
+++ b/Engine/Source/Runtime/Engine/Classes/PhysicsEngine/BodyInstance.h
@@ -28,6 +28,12 @@ DECLARE_DELEGATE_TwoParams(FCalculateCustomPhysics, float, FBodyInstance*);
   * Since this could be called many times by GetWorldTransform any expensive computations should be cached if possible.*/
 DECLARE_DELEGATE_TwoParams(FCalculateCustomProjection, const FBodyInstance*, FTransform&);
 
+/** Delegate for body initialization. Called at the end of FBodyInstance::InitBody.*/
+DECLARE_MULTICAST_DELEGATE_OneParam(FBodyInstanceInit, const FBodyInstance*);
+
+/** Delegate for body termination. Called at the beginning of FBodyInstance::TermBody.*/
+DECLARE_MULTICAST_DELEGATE_OneParam(FBodyInstanceTerm, const FBodyInstance*);
+
 #if WITH_PHYSX
 struct FShapeData;
 
@@ -159,6 +165,10 @@ struct ENGINE_API FBodyInstance
 {
 	GENERATED_USTRUCT_BODY()
 
+	/** Body initialization and termination delegates.*/
+	static FBodyInstanceInit InitBodyDelegate;
+	static FBodyInstanceTerm TermBodyDelegate;
+
 	/** 
 	 *	Index of this BodyInstance within the SkeletalMeshComponent/PhysicsAsset. 
 	 *	Is INDEX_NONE if a single body component
diff --git a/Engine/Source/Runtime/Engine/Classes/PhysicsEngine/ConstraintInstance.h b/Engine/Source/Runtime/Engine/Classes/PhysicsEngine/ConstraintInstance.h
--- a/Engine/Source/Runtime/Engine/Classes/PhysicsEngine/ConstraintInstance.h
+++ b/Engine/Source/Runtime/Engine/Classes/PhysicsEngine/ConstraintInstance.h
@@ -17,6 +17,8 @@ class FMaterialRenderProxy;
 class FPrimitiveDrawInterface;
 class FMaterialRenderProxy;
 
+DECLARE_DELEGATE_OneParam(FOnConstraintBroken, int32 /*ConstraintIndex*/);
+
 /** Container for properties of a physics constraint that can be easily swapped at runtime. This is useful for switching different setups when going from ragdoll to standup for example */
 USTRUCT()
 struct ENGINE_API FConstraintProfileProperties
@@ -101,10 +103,6 @@ struct ENGINE_API FConstraintInstance
 	/** Indicates position of this constraint within the array in SkeletalMeshComponent. */
 	int32 ConstraintIndex;
 
-	/** The component that created this instance. */
-	UPROPERTY()
-	USceneComponent* OwnerComponent;
-
 #if WITH_PHYSX
 	/** Internal use. Physics-engine representation of this constraint. */
 	physx::PxD6Joint*		ConstraintData;
@@ -434,7 +432,12 @@ public:
 	void SetLinearLimitSize(float NewLimitSize);
 
 	/** Create physics engine constraint. */
-	void InitConstraint(USceneComponent* Owner, FBodyInstance* Body1, FBodyInstance* Body2, float Scale);
+	void InitConstraint(FBodyInstance* Body1, FBodyInstance* Body2, float Scale, UObject* DebugOwner, FOnConstraintBroken InConstraintBrokenDelegate = FOnConstraintBroken());
+
+#if WITH_PHYSX
+	/** Create physics engine constraint using physx actors. */
+	void InitConstraintPhysX_AssumesLocked(physx::PxRigidActor* PActor1, physx::PxRigidActor* PActor2, physx::PxScene* PScene, float InScale, FOnConstraintBroken InConstraintBrokenDelegate = FOnConstraintBroken());
+#endif
 
 	/** Terminate physics engine constraint */
 	void TermConstraint();
@@ -550,6 +553,8 @@ private:
 	void UpdateBreakable();
 	void UpdateDriveTarget();
 
+	FOnConstraintBroken OnConstraintBrokenDelegate;
+
 public:
 
 	///////////////////////////// DEPRECATED
diff --git a/Engine/Source/Runtime/Engine/Classes/PhysicsEngine/ConstraintTypes.h b/Engine/Source/Runtime/Engine/Classes/PhysicsEngine/ConstraintTypes.h
--- a/Engine/Source/Runtime/Engine/Classes/PhysicsEngine/ConstraintTypes.h
+++ b/Engine/Source/Runtime/Engine/Classes/PhysicsEngine/ConstraintTypes.h
@@ -133,4 +133,7 @@ struct ENGINE_API FTwistConstraint : public FConstraintBaseParams
 	/** Updates physx twist limit from unreal data. */
 	void UpdatePhysXTwistLimit_AssumesLocked(physx::PxD6Joint* Joint, float AverageMass) const;
 #endif
-};
\ No newline at end of file
+};
+
+#define		RB_MinSizeToLockDOF				(0.1)
+#define		RB_MinAngleToLockDOF			(5.0)
\ No newline at end of file
diff --git a/Engine/Source/Runtime/Engine/Classes/PhysicsEngine/ConvexElem.h b/Engine/Source/Runtime/Engine/Classes/PhysicsEngine/ConvexElem.h
--- a/Engine/Source/Runtime/Engine/Classes/PhysicsEngine/ConvexElem.h
+++ b/Engine/Source/Runtime/Engine/Classes/PhysicsEngine/ConvexElem.h
@@ -49,7 +49,7 @@ struct FKConvexElem : public FKShapeElem
 	
 	ENGINE_API void	DrawElemWire(class FPrimitiveDrawInterface* PDI, const FTransform& ElemTM, const float Scale, const FColor Color) const;
 
-	void AddCachedSolidConvexGeom(TArray<FDynamicMeshVertex>& VertexBuffer, TArray<int32>& IndexBuffer, const float Scale, const FColor VertexColor, const bool bIsMirrored) const;
+	ENGINE_API void AddCachedSolidConvexGeom(TArray<FDynamicMeshVertex>& VertexBuffer, TArray<int32>& IndexBuffer, const FColor VertexColor) const;
 
 	/** Reset the hull to empty all arrays */
 	ENGINE_API void	Reset();
diff --git a/Engine/Source/Runtime/Engine/Classes/PhysicsEngine/PhysicsAsset.h b/Engine/Source/Runtime/Engine/Classes/PhysicsEngine/PhysicsAsset.h
--- a/Engine/Source/Runtime/Engine/Classes/PhysicsEngine/PhysicsAsset.h
+++ b/Engine/Source/Runtime/Engine/Classes/PhysicsEngine/PhysicsAsset.h
@@ -3,8 +3,12 @@
 #pragma once
 
 #include "RigidBodyIndexPair.h"
+#include "PhysicalAnimationComponent.h"
+#include "PhysicsEngine/BodySetup.h"
 #include "PhysicsAsset.generated.h"
 
+class USkeletalBodySetup;
+
 /**
  * PhysicsAsset contains a set of rigid bodies and constraints that make up a single ragdoll.
  * The asset is not limited to human ragdolls, and can be used for any physical simulation using bodies and constraints.
@@ -31,18 +35,25 @@ class UPhysicsAsset : public UObject
 	UPROPERTY()
 	TAssetPtr<class USkeletalMesh> PreviewSkeletalMesh;
 
+	UPROPERTY(EditAnywhere, Category = PhysicalAnimation)
+	TArray<FName> Profiles;
+
+	UPROPERTY(transient)
+	FName CurrentPhysicalAnimationProfileName;
+
 #endif // WITH_EDITORONLY_DATA
-	/** 
-	 *	Array of BodySetup objects. Stores information about collision shape etc. for each body.
-	 *	Does not include body position - those are taken from mesh.
-	 */
-	UPROPERTY(instanced)
-	TArray<class UBodySetup*> BodySetup;
 
 	/** Index of bodies that are marked bConsiderForBounds */
 	UPROPERTY()
 	TArray<int32> BoundsBodies;
 
+	/**
+	*	Array of SkeletalBodySetup objects. Stores information about collision shape etc. for each body.
+	*	Does not include body position - those are taken from mesh.
+	*/
+	UPROPERTY(instanced)
+	TArray<USkeletalBodySetup*> SkeletalBodySetups;
+
 	/** 
 	 *	Array of RB_ConstraintSetup objects. 
 	 *	Stores information about a joint between two bodies, such as position relative to each body, joint limits etc.
@@ -68,6 +79,13 @@ public:
 	virtual SIZE_T GetResourceSize(EResourceSizeMode::Type Mode) override;
 #if WITH_EDITOR
 	virtual void PostEditUndo() override;
+
+	ENGINE_API const TArray<FName>& GetPhysicalAnimationProfileNames(TArray<FName>& OutProfileNames) const
+	{
+		return Profiles;
+	}
+
+	virtual void PostEditChangeProperty(struct FPropertyChangedEvent& PropertyChangedEvent) override;
 #endif
 	//~ End UObject Interface
 
@@ -122,7 +140,66 @@ public:
 	 * @param	Constraints		Returns the found constraints
 	 **/
 	ENGINE_API void BodyFindConstraints(int32 BodyIndex, TArray<int32>& Constraints);
+
+private:
+	UPROPERTY(instanced)
+	TArray<class UBodySetup*> BodySetup_DEPRECATED;
 };
 
-#define		RB_MinSizeToLockDOF				(0.1)
-#define		RB_MinAngleToLockDOF			(5.0)
\ No newline at end of file
+USTRUCT()
+struct FPhysicalAnimationProfile
+{
+	GENERATED_BODY()
+	
+	/** Profile name used to identify set of physical animation parameters */
+	UPROPERTY()
+	FName ProfileName;
+
+	/** Physical animation parameters used to drive animation */
+	UPROPERTY(EditAnywhere, Category = PhysicalAnimation)
+	FPhysicalAnimationData PhysicalAnimationData;
+};
+
+UCLASS(MinimalAPI)
+class USkeletalBodySetup : public UBodySetup
+{
+	GENERATED_BODY()
+public:
+	const FPhysicalAnimationProfile* FindPhysicalAnimationProfile(const FName ProfileName) const
+	{
+		return PhysicalAnimationData.FindByPredicate([ProfileName](const FPhysicalAnimationProfile& Profile){ return ProfileName == Profile.ProfileName; });
+	}
+
+	FPhysicalAnimationProfile* FindPhysicalAnimationProfile(const FName ProfileName)
+	{
+		return PhysicalAnimationData.FindByPredicate([ProfileName](const FPhysicalAnimationProfile& Profile) { return ProfileName == Profile.ProfileName; });
+	}
+
+	const TArray<FPhysicalAnimationProfile>& GetPhysicalAnimationProfiles() const
+	{
+		return PhysicalAnimationData;
+	}
+
+#if WITH_EDITOR
+	virtual void PostEditChangeProperty(struct FPropertyChangedEvent& PropertyChangedEvent);
+	ENGINE_API FName GetCurrentPhysicalAnimationProfileName() const;
+	
+	/** Creates a new physical animation profile entry */
+	ENGINE_API void AddPhysicalAnimationProfile(FName ProfileName);
+
+	/** Removes physical animation profile */
+	ENGINE_API void RemovePhysicalAnimationProfile(FName ProfileName);
+
+	ENGINE_API void UpdatePhysicalAnimationProfiles(const TArray<FName>& Profiles);
+#endif
+
+#if WITH_EDITORONLY_DATA
+	//dummy place for customization inside phat. Profiles are ordered dynamically and we need a static place for detail customization
+	UPROPERTY(EditAnywhere, Category = PhysicalAnimation)
+	FPhysicalAnimationProfile CurrentPhysicalAnimationProfile;
+#endif
+
+private:
+	UPROPERTY()
+	TArray<FPhysicalAnimationProfile> PhysicalAnimationData;
+};
\ No newline at end of file
diff --git a/Engine/Source/Runtime/Engine/Classes/PhysicsEngine/PhysicsConstraintComponent.h b/Engine/Source/Runtime/Engine/Classes/PhysicsEngine/PhysicsConstraintComponent.h
--- a/Engine/Source/Runtime/Engine/Classes/PhysicsEngine/PhysicsConstraintComponent.h
+++ b/Engine/Source/Runtime/Engine/Classes/PhysicsEngine/PhysicsConstraintComponent.h
@@ -278,5 +278,9 @@ protected:
 
 	/** Returns the scale of the constraint as it will be passed into the ConstraintInstance*/
 	float GetConstraintScale() const;
+
+private:
+	/** Wrapper that calls our constraint broken delegate */
+	void OnConstraintBrokenWrapper(int32 ConstraintIndex);
 };
 
diff --git a/Engine/Source/Runtime/Engine/Classes/Sound/AudioVolume.h b/Engine/Source/Runtime/Engine/Classes/Sound/AudioVolume.h
--- a/Engine/Source/Runtime/Engine/Classes/Sound/AudioVolume.h
+++ b/Engine/Source/Runtime/Engine/Classes/Sound/AudioVolume.h
@@ -10,6 +10,9 @@
 #include "Audio.h"
 #include "AudioVolume.generated.h"
 
+class AAudioVolume;
+struct FBodyInstance;
+
 /**
  * DEPRECATED: Exists for backwards compatibility
  * Indicates a reverb preset to use.
@@ -70,19 +73,19 @@ struct FReverbSettings
 	UPROPERTY(EditAnywhere, BlueprintReadWrite, Category=ReverbSettings)
 	float FadeTime;
 
-
-
-		FReverbSettings()
+	FReverbSettings()
 		: bApplyReverb(true)
 		, ReverbType_DEPRECATED(REVERB_Default)
-		, ReverbEffect(NULL)
+		, ReverbEffect(nullptr)
 		, Volume(0.5f)
 		, FadeTime(2.0f)
-		{
-		}
+	{
+	}
 
-		void PostSerialize(const FArchive& Ar);
-	
+	bool operator==(const FReverbSettings& Other) const;
+	bool operator!=(const FReverbSettings& Other) const { return !(*this == Other); }
+
+	void PostSerialize(const FArchive& Ar);
 };
 
 template<>
@@ -180,35 +183,85 @@ struct TStructOpsTypeTraits<FInteriorSettings> : public TStructOpsTypeTraitsBase
 	};
 };
 
+struct FAudioVolumeProxy
+{
+	FAudioVolumeProxy()
+		: AudioVolumeID(0)
+		, WorldID(0)
+		, Priority(0.f)
+		, BodyInstance(nullptr)
+	{
+	}
+
+	FAudioVolumeProxy(const AAudioVolume* AudioVolume);
+
+	uint32 AudioVolumeID;
+	uint32 WorldID;
+	float Priority;
+	FReverbSettings ReverbSettings;
+	FInteriorSettings InteriorSettings;
+	FBodyInstance* BodyInstance; // This is scary
+};
+
 UCLASS(hidecategories=(Advanced, Attachment, Collision, Volume))
 class ENGINE_API AAudioVolume : public AVolume
 {
 	GENERATED_UCLASS_BODY()
 
+private:
 	/**
 	 * Priority of this volume. In the case of overlapping volumes the one with the highest priority
 	 * is chosen. The order is undefined if two or more overlapping volumes have the same priority.
 	 */
-	UPROPERTY(EditAnywhere, BlueprintReadWrite, Category=AudioVolume)
+	UPROPERTY(EditAnywhere, BlueprintReadOnly, Category=AudioVolume, meta=(AllowPrivateAccess="true"))
 	float Priority;
 
 	/** whether this volume is currently enabled and able to affect sounds */
-	UPROPERTY(EditAnywhere, BlueprintReadWrite, replicated, Category=AudioVolume)
+	UPROPERTY(EditAnywhere, BlueprintReadOnly, ReplicatedUsing=OnRep_bEnabled, Category=AudioVolume, meta=(AllowPrivateAccess="true"))
 	uint32 bEnabled:1;
 
 	/** Reverb settings to use for this volume. */
-	UPROPERTY(EditAnywhere, BlueprintReadWrite, Category=Reverb)
-	struct FReverbSettings Settings;
+	UPROPERTY(EditAnywhere, BlueprintReadOnly, Category=Reverb, meta=(AllowPrivateAccess="true"))
+	FReverbSettings Settings;
 
 	/** Interior settings used for this volume */
-	UPROPERTY(EditAnywhere, BlueprintReadWrite, Category=AmbientZone)
-	struct FInteriorSettings AmbientZoneSettings;
+	UPROPERTY(EditAnywhere, BlueprintReadOnly, Category=AmbientZone, meta=(AllowPrivateAccess="true"))
+	FInteriorSettings AmbientZoneSettings;
+
+public:
+
+	float GetPriority() const { return Priority; }
+	
+	UFUNCTION(BlueprintCallable, Category=AudioVolume)
+	void SetPriority(float NewPriority);
+
+	bool GetEnabled() const { return bEnabled; }
+	
+	UFUNCTION(BlueprintCallable, Category=AudioVolume)
+	void SetEnabled(bool bNewEnabled);
 
-	/** Next volume in linked listed, sorted by priority in descending order. */
-	UPROPERTY(transient)
-	class AAudioVolume* NextLowerPriorityVolume;
+	const FReverbSettings& GetReverbSettings() const { return Settings; }
+	
+	UFUNCTION(BlueprintCallable, Category=AudioVolume)
+	void SetReverbSettings(const FReverbSettings& NewReverbSettings);
 
+	const FInteriorSettings& GetInteriorSettings() const { return AmbientZoneSettings; }
 
+	UFUNCTION(BlueprintCallable, Category=AudioVolume)
+	void SetInteriorSettings(const FInteriorSettings& NewInteriorSettings);
+
+private:
+
+	UFUNCTION()
+	virtual void OnRep_bEnabled();
+
+	void TransformUpdated(USceneComponent* RootComponent, EUpdateTransformFlags UpdateTransformFlags, ETeleportType Teleport);
+
+	void AddProxy() const;
+	void RemoveProxy() const;
+	void UpdateProxy() const;
+
+public:
 
 	//~ Begin UObject Interface
 #if WITH_EDITOR
@@ -218,9 +271,7 @@ class ENGINE_API AAudioVolume : public AVolume
 
 	//~ Begin AActor Interface
 	virtual void PostUnregisterAllComponents() override;
-protected:
 	virtual void PostRegisterAllComponents() override;
-public:
 	//~ End AActor Interface
 };
 
diff --git a/Engine/Source/Runtime/Engine/Classes/Sound/SoundBase.h b/Engine/Source/Runtime/Engine/Classes/Sound/SoundBase.h
--- a/Engine/Source/Runtime/Engine/Classes/Sound/SoundBase.h
+++ b/Engine/Source/Runtime/Engine/Classes/Sound/SoundBase.h
@@ -105,6 +105,9 @@ public:
 	  * Returns the subtitle priority
 	  */
 	virtual float GetSubtitlePriority() const { return DEFAULT_SUBTITLE_PRIORITY; };
+	
+	/** Returns whether or not any part of this sound wants interior volumes applied to it */
+	virtual bool ShouldApplyInteriorVolumes() const;
 
 	/** Returns whether or not this sound is looping. */
 	bool IsLooping();
diff --git a/Engine/Source/Runtime/Engine/Classes/Sound/SoundConcurrency.h b/Engine/Source/Runtime/Engine/Classes/Sound/SoundConcurrency.h
--- a/Engine/Source/Runtime/Engine/Classes/Sound/SoundConcurrency.h
+++ b/Engine/Source/Runtime/Engine/Classes/Sound/SoundConcurrency.h
@@ -22,11 +22,14 @@ namespace EMaxConcurrentResolutionRule
 		/** When Max Concurrent sounds are active stop the furthest sound.  If all sounds are the same distance then stop the oldest. */
 		StopFarthestThenOldest,
 
-		/** Stop the lowest volume-weighted priority sound in the group. If all sounds are the same priority, then it will stop the oldest sound in the group. */
+		/** Stop the lowest priority sound in the group. If all sounds are the same priority, then it will stop the oldest sound in the group. */
 		StopLowestPriority,
 
 		/** Stop the sound that is quietest in the group. */
 		StopQuietest,
+
+		/** Stop the lowest priority sound in the group. If all sounds are the same priority, then it won't play a new sound. */
+		StopLowestPriorityThenPreventNew,
 	};
 }
 
diff --git a/Engine/Source/Runtime/Engine/Classes/Sound/SoundCue.h b/Engine/Source/Runtime/Engine/Classes/Sound/SoundCue.h
--- a/Engine/Source/Runtime/Engine/Classes/Sound/SoundCue.h
+++ b/Engine/Source/Runtime/Engine/Classes/Sound/SoundCue.h
@@ -86,6 +86,7 @@ public:
 
 	//~ Begin USoundBase Interface.
 	virtual bool IsPlayable() const override;
+	virtual bool ShouldApplyInteriorVolumes() const override;
 	virtual void Parse( class FAudioDevice* AudioDevice, const UPTRINT NodeWaveInstanceHash, FActiveSound& ActiveSound, const FSoundParseParameters& ParseParams, TArray<FWaveInstance*>& WaveInstances ) override;
 	virtual float GetVolumeMultiplier() override;
 	virtual float GetPitchMultiplier() override;
diff --git a/Engine/Source/Runtime/Engine/Classes/Sound/SoundWave.h b/Engine/Source/Runtime/Engine/Classes/Sound/SoundWave.h
--- a/Engine/Source/Runtime/Engine/Classes/Sound/SoundWave.h
+++ b/Engine/Source/Runtime/Engine/Classes/Sound/SoundWave.h
@@ -334,7 +334,7 @@ public:
 	/** 
 	 * Handle any special requirements when the sound starts (e.g. subtitles)
 	 */
-	FWaveInstance* HandleStart( FActiveSound& ActiveSound, const UPTRINT WaveInstanceHash );
+	FWaveInstance* HandleStart( FActiveSound& ActiveSound, const UPTRINT WaveInstanceHash ) const;
 
 	/** 
 	 * This is only for DTYPE_Procedural audio. Override this function.
@@ -438,6 +438,18 @@ public:
 	 * @param OutChunkData	Address of pointer that will store data.
 	 */
 	void GetChunkData(int32 ChunkIndex, uint8** OutChunkData);
+
+private:
+
+	enum class ESoundWaveResourceState
+	{
+		NeedsFree,
+		Freeing,
+		Freed
+	};
+
+	ESoundWaveResourceState ResourceState;
+
 };
 
 
