Copying //UE4/Dev-Platform to //UE4/Dev-Main (Source: //UE4/Dev-Platform @ 3046626)

diff --git a/Engine/Source/Runtime/Apple/MetalRHI/Private/MetalStructuredBuffer.cpp b/Engine/Source/Runtime/Apple/MetalRHI/Private/MetalStructuredBuffer.cpp
--- a/Engine/Source/Runtime/Apple/MetalRHI/Private/MetalStructuredBuffer.cpp
+++ b/Engine/Source/Runtime/Apple/MetalRHI/Private/MetalStructuredBuffer.cpp
@@ -3,6 +3,7 @@
 
 #include "MetalRHIPrivate.h"
 #include "MetalProfiler.h"
+#include "MetalCommandBuffer.h"
 
 FMetalStructuredBuffer::FMetalStructuredBuffer(uint32 Stride, uint32 Size, FResourceArrayInterface* ResourceArray, uint32 Usage)
 	: FRHIStructuredBuffer(Stride, Size, Usage)
@@ -61,6 +62,7 @@ void* FMetalStructuredBuffer::Lock(EResourceLockMode LockMode, uint32 Offset, ui
 		
 		// Synchronise the buffer with the CPU
 		id<MTLBlitCommandEncoder> Blitter = GetMetalDeviceContext().GetBlitContext();
+		METAL_DEBUG_COMMAND_BUFFER_BLIT_LOG((&GetMetalDeviceContext()), @"SynchronizeResource(StructuredBuffer %p)", this);
 		[Blitter synchronizeResource:Buffer];
 		
 		//kick the current command buffer.
diff --git a/Engine/Source/Runtime/Apple/MetalRHI/Private/MetalTexture.cpp b/Engine/Source/Runtime/Apple/MetalRHI/Private/MetalTexture.cpp
--- a/Engine/Source/Runtime/Apple/MetalRHI/Private/MetalTexture.cpp
+++ b/Engine/Source/Runtime/Apple/MetalRHI/Private/MetalTexture.cpp
@@ -6,6 +6,7 @@
 
 #include "MetalRHIPrivate.h"
 #include "MetalProfiler.h" // for STAT_MetalTexturePageOffTime
+#include "MetalCommandBuffer.h"
 
 uint8 FMetalSurface::NextKey = 1; // 0 is reserved for MTLPixelFormatInvalid
 int32 FMetalSurface::ActiveUploads = 0;
@@ -250,7 +251,7 @@ FMetalSurface::FMetalSurface(FMetalSurface& Source, NSRange const MipRange, EPix
 	
 #if PLATFORM_MAC // Recreate the texture to enable MTLTextureUsagePixelFormatView which must be off unless we definitely use this feature or we are throwing ~4% performance vs. Windows on the floor.
 	MTLTextureUsage Usage = Source.Texture.usage;
-	if(!(Usage & MTLTextureUsagePixelFormatView) && (Source.PixelFormat != PF_DepthStencil) && !bUseSourceTex)
+	if(!(Usage & MTLTextureUsagePixelFormatView) && (Source.PixelFormat == PF_DepthStencil && Format == PF_X24_G8 && GetMetalDeviceContext().SupportsFeature(EMetalFeaturesStencilView)) && !bUseSourceTex)
 	{
 		Source.PrepareTextureView();
 	}
@@ -265,6 +266,26 @@ FMetalSurface::FMetalSurface(FMetalSurface& Source, NSRange const MipRange, EPix
         Texture = [Source.Texture newTextureViewWithPixelFormat:MetalFormat textureType:Source.Texture.textureType levels:MipRange slices:Slices];
 		TRACK_OBJECT(STAT_MetalTextureCount, Texture);
     }
+    else if (Source.PixelFormat == PF_DepthStencil && Format == PF_X24_G8 && GetMetalDeviceContext().SupportsFeature(EMetalFeaturesStencilView))
+    {
+        switch (Source.Texture.pixelFormat)
+        {
+#if PLATFORM_MAC
+            case MTLPixelFormatDepth24Unorm_Stencil8:
+                MetalFormat = MTLPixelFormatX24_Stencil8;
+                break;
+#endif
+            case MTLPixelFormatDepth32Float_Stencil8:
+                MetalFormat = MTLPixelFormatX32_Stencil8;
+                break;
+            default:
+                check(false);
+                break;
+        }
+        
+        Texture = [Source.Texture newTextureViewWithPixelFormat:MetalFormat textureType:Source.Texture.textureType levels:MipRange slices:Slices];
+        TRACK_OBJECT(STAT_MetalTextureCount, Texture);
+    }
     else
 #endif
 	if(Source.PixelFormat == PF_DepthStencil && Format == PF_X24_G8)
@@ -319,6 +340,7 @@ FMetalSurface::FMetalSurface(FMetalSurface& Source, NSRange const MipRange, EPix
 					[Buf.Buffer retain];
 					INC_MEMORY_STAT_BY(STAT_MetalWastedPooledBufferMem, Buf.Buffer.length - SizePerImage);
 					
+					METAL_DEBUG_COMMAND_BUFFER_BLIT_LOG((&GetMetalDeviceContext()), @"FMetalSurface %p(copyFromBuffer: %p toTexture: %p)", this, Buf.Buffer, Source.StencilTexture);
 					[BlitEncoder copyFromBuffer:Buf.Buffer sourceOffset:0 sourceBytesPerRow:Source.Texture.width sourceBytesPerImage:SizePerImage sourceSize:MTLSizeMake(Source.Texture.width, Source.Texture.height, 1) toTexture:Source.StencilTexture destinationSlice:0 destinationLevel:0 destinationOrigin:MTLOriginMake(0, 0, 0)];
 					
 					bWritten = true;
@@ -396,6 +418,10 @@ FMetalSurface::FMetalSurface(ERHIResourceType ResourceType, EPixelFormat Format,
 		PixelFormatKeyMap.Add(MTLPixelFormatDepth32Float_Stencil8, NextKey++);
 #if PLATFORM_MAC
 		PixelFormatKeyMap.Add(MTLPixelFormatDepth24Unorm_Stencil8, NextKey++);
+        if (GetMetalDeviceContext().SupportsFeature(EMetalFeaturesDepth16))
+        {
+            PixelFormatKeyMap.Add(MTLPixelFormatDepth16Unorm, NextKey++);
+        }
 #endif
 #endif
 	}
@@ -861,6 +887,7 @@ void* FMetalSurface::Lock(uint32 MipIndex, uint32 ArrayIndex, EResourceLockMode
 			
 			if (bSupportsResourceOptions && Texture.storageMode == MTLStorageModePrivate)
 			{
+				METAL_DEBUG_COMMAND_BUFFER_BLIT_LOG((&GetMetalDeviceContext()), @"FMetalSurface::Lock %p(copyFromTexture:toBuffer:)", this);
 				[Blitter copyFromTexture:Texture sourceSlice:ArrayIndex sourceLevel:MipIndex sourceOrigin:Region.origin sourceSize:Region.size toBuffer:LockedMemory[MipIndex] destinationOffset:0 destinationBytesPerRow:DestStride destinationBytesPerImage:MipBytes];
 				
 				//kick the current command buffer.
@@ -986,6 +1013,7 @@ void FMetalSurface::Unlock(uint32 MipIndex, uint32 ArrayIndex)
 			// create a blitter object
 			id<MTLBlitCommandEncoder> Blitter = [CommandBuffer blitCommandEncoder];
 			
+			METAL_DEBUG_COMMAND_BUFFER_BLIT_ASYNC_LOG((&GetMetalDeviceContext()), CommandBuffer, @"FMetalSurface::Unlock %p(copyFromBuffer:toTexture)", this);
 			[Blitter copyFromBuffer:LockedMemory[MipIndex] sourceOffset:0 sourceBytesPerRow:Stride sourceBytesPerImage:BytesPerImage sourceSize:Region.size toTexture:Texture destinationSlice:ArrayIndex destinationLevel:MipIndex destinationOrigin:Region.origin];
 			
 			[Blitter endEncoding];
@@ -1111,7 +1139,7 @@ id<MTLTexture> FMetalSurface::GetDrawableTexture()
 void FMetalSurface::UpdateSRV(FTextureRHIRef SourceTex)
 {
 #if METAL_API_1_1
-	if (PixelFormat == PF_X24_G8 && Texture != StencilTexture && GetMetalDeviceContext().SupportsFeature(EMetalFeaturesDepthStencilBlitOptions))
+	if (PixelFormat == PF_X24_G8 && Texture != StencilTexture && GetMetalDeviceContext().SupportsFeature(EMetalFeaturesDepthStencilBlitOptions) && !GetMetalDeviceContext().SupportsFeature(EMetalFeaturesStencilView))
 	{
 		// In this case StencilTexture is the source depth/stencil texture and Texture is our target Stencil-only copy
 		
@@ -1122,6 +1150,7 @@ void FMetalSurface::UpdateSRV(FTextureRHIRef SourceTex)
 		[Buf.Buffer retain];
 		INC_MEMORY_STAT_BY(STAT_MetalWastedPooledBufferMem, Buf.Buffer.length - SizePerImage);
 		
+		METAL_DEBUG_COMMAND_BUFFER_BLIT_LOG((&GetMetalDeviceContext()), @"FMetalSurface::UpdateSRV %p(Source %p)", this, SourceTex.GetReference());
 		[BlitEncoder copyFromTexture:StencilTexture sourceSlice:0 sourceLevel:0 sourceOrigin:MTLOriginMake(0,0,0) sourceSize:MTLSizeMake(Texture.width, Texture.height, 1) toBuffer:Buf.Buffer destinationOffset:0 destinationBytesPerRow:Texture.width destinationBytesPerImage:SizePerImage options:MTLBlitOptionStencilFromDepthStencil];
 		
 		[BlitEncoder copyFromBuffer:Buf.Buffer sourceOffset:0 sourceBytesPerRow:Texture.width sourceBytesPerImage:SizePerImage sourceSize:MTLSizeMake(Texture.width, Texture.height, 1) toTexture:Texture destinationSlice:0 destinationLevel:0 destinationOrigin:MTLOriginMake(0,0,0)];
@@ -1247,6 +1276,7 @@ void FMetalDynamicRHI::RHIGenerateMips(FTextureRHIParamRef SourceSurfaceRHI)
 		// create a blitter object
 		id<MTLBlitCommandEncoder> Blitter = [CommandBuffer blitCommandEncoder];
 		
+		METAL_DEBUG_COMMAND_BUFFER_BLIT_ASYNC_LOG((&GetMetalDeviceContext()), CommandBuffer, @"RHIGenerateMips %p(Source %p)", this, SourceSurfaceRHI);
 		[Blitter generateMipmapsForTexture:Surf->Texture];
 		
 		// kick it off!
@@ -1309,6 +1339,8 @@ FTexture2DRHIRef FMetalDynamicRHI::RHIAsyncReallocateTexture2D(FTexture2DRHIPara
 	uint32 SliceIndex = 0;
 	MTLOrigin Origin = MTLOriginMake(0,0,0);
 	
+	METAL_DEBUG_COMMAND_BUFFER_BLIT_ASYNC_LOG((&GetMetalDeviceContext()), CommandBuffer, @"RHIAsyncReallocateTexture2D(OldTextureRHI %p, NewMipCount %d, NewSizeX %d, NewSizeY %d) -> %p", OldTextureRHI, NewMipCount, NewSizeX, NewSizeY, NewTexture);
+	
 	id<MTLTexture> Tex = OldTexture->Surface.Texture;
 	[Tex retain];
 
@@ -1453,6 +1485,8 @@ void FMetalDynamicRHI::RHIUpdateTexture2D(FTexture2DRHIParamRef TextureRHI, uint
 		
 		FMemory::Memcpy([LockedMemory contents], SourceData, BufferSize);
 		
+		METAL_DEBUG_COMMAND_BUFFER_BLIT_ASYNC_LOG((&GetMetalDeviceContext()), CommandBuffer, @"RHIUpdateTexture2D(TextureRHI %p, MipIndex %d, SourcePitch %d)", TextureRHI, MipIndex, SourcePitch);
+		
 		[Blitter copyFromBuffer:LockedMemory sourceOffset:0 sourceBytesPerRow:SourcePitch sourceBytesPerImage:BytesPerImage sourceSize:Region.size toTexture:Tex destinationSlice:0 destinationLevel:MipIndex destinationOrigin:Region.origin];
 		
 		[Blitter endEncoding];
@@ -1500,6 +1534,8 @@ void FMetalDynamicRHI::RHIUpdateTexture3D(FTexture3DRHIParamRef TextureRHI,uint3
 
 		FMemory::Memcpy([LockedMemory contents], SourceData, BufferSize);
 		
+		METAL_DEBUG_COMMAND_BUFFER_BLIT_ASYNC_LOG((&GetMetalDeviceContext()), CommandBuffer, @"RHIUpdateTexture3D(TextureRHI %p, MipIndex %d, SourceRowPitch %d, SourceDepthPitch %d)", TextureRHI, MipIndex, SourceRowPitch, SourceDepthPitch);
+		
 		[Blitter copyFromBuffer:LockedMemory sourceOffset:0 sourceBytesPerRow:SourceRowPitch sourceBytesPerImage:BytesPerImage sourceSize:Region.size toTexture:Tex destinationSlice:0 destinationLevel:MipIndex destinationOrigin:Region.origin];
 		
 		[Blitter endEncoding];
diff --git a/Engine/Source/Runtime/Apple/MetalRHI/Private/MetalUAV.cpp b/Engine/Source/Runtime/Apple/MetalRHI/Private/MetalUAV.cpp
--- a/Engine/Source/Runtime/Apple/MetalRHI/Private/MetalUAV.cpp
+++ b/Engine/Source/Runtime/Apple/MetalRHI/Private/MetalUAV.cpp
@@ -3,6 +3,7 @@
 
 #include "MetalRHIPrivate.h"
 #include "ShaderCache.h"
+#include "MetalCommandBuffer.h"
 
 FMetalShaderResourceView::~FMetalShaderResourceView()
 {
@@ -323,6 +324,7 @@ void FMetalRHICommandContext::RHIClearUAV(FUnorderedAccessViewRHIParamRef Unorde
 		
 		// Fill the buffer via a blit encoder - I hope that is sufficient.
 		id<MTLBlitCommandEncoder> Blitter = Context->GetBlitContext();
+		METAL_DEBUG_COMMAND_BUFFER_BLIT_LOG(Context, @"RHIClearUAV(UAV %p, %d)", UnorderedAccessViewRHI, Values[0]);
 		[Blitter fillBuffer:UnorderedAccessView->SourceVertexBuffer->Buffer range:NSMakeRange(0, UnorderedAccessView->SourceVertexBuffer->GetSize()) value:Values[0]];
 		
 		// If there are problems you may need to add calls to restore the render command encoder at this point
diff --git a/Engine/Source/Runtime/Apple/MetalRHI/Private/MetalVertexBuffer.cpp b/Engine/Source/Runtime/Apple/MetalRHI/Private/MetalVertexBuffer.cpp
--- a/Engine/Source/Runtime/Apple/MetalRHI/Private/MetalVertexBuffer.cpp
+++ b/Engine/Source/Runtime/Apple/MetalRHI/Private/MetalVertexBuffer.cpp
@@ -6,6 +6,7 @@
 
 #include "MetalRHIPrivate.h"
 #include "MetalProfiler.h"
+#include "MetalCommandBuffer.h"
 
 
 FMetalVertexBuffer::FMetalVertexBuffer(uint32 InSize, uint32 InUsage)
@@ -93,6 +94,7 @@ void* FMetalVertexBuffer::Lock(EResourceLockMode LockMode, uint32 Offset, uint32
 		
 		// Synchronise the buffer with the CPU
 		id<MTLBlitCommandEncoder> Blitter = GetMetalDeviceContext().GetBlitContext();
+		METAL_DEBUG_COMMAND_BUFFER_BLIT_LOG((&GetMetalDeviceContext()), @"SynchronizeResource(VertexBuffer %p)", this);
 		[Blitter synchronizeResource:Buffer];
 		
 		//kick the current command buffer.
diff --git a/Engine/Source/Runtime/Apple/MetalRHI/Public/MetalResources.h b/Engine/Source/Runtime/Apple/MetalRHI/Public/MetalResources.h
--- a/Engine/Source/Runtime/Apple/MetalRHI/Public/MetalResources.h
+++ b/Engine/Source/Runtime/Apple/MetalRHI/Public/MetalResources.h
@@ -80,9 +80,8 @@ public:
 	/** The binding for the buffer side-table if present */
 	int32 SideTableBinding;
 	
-	TArray<ANSICHAR> GlslCode;
+	/** The debuggable text source */
 	NSString* GlslCodeNSString;
-	const ANSICHAR*  GlslCodeString; // make it easier in VS to see shader code in debug mode; points to begin of GlslCode
 };
 
 typedef TMetalBaseShader<FRHIVertexShader, SF_Vertex> FMetalVertexShader;
@@ -153,11 +152,28 @@ public:
 	/**
 	 * Prepare a pipeline state object for the current state right before drawing
 	 */
-	void PrepareToDraw(FMetalContext* Context, const struct FMetalRenderPipelineDesc& RenderPipelineDesc);
+	void PrepareToDraw(FMetalContext* Context, MTLVertexDescriptor* VertexDesc, const struct FMetalRenderPipelineDesc& RenderPipelineDesc);
 
 protected:
 	FCriticalSection PipelineMutex;
-	TMap<FMetalRenderPipelineHash, id<MTLRenderPipelineState> > PipelineStates;
+	
+	struct FMetalPipelineHash
+	{
+		FMetalRenderPipelineHash RenderPipelineHash;
+		NSUInteger VertexDescHash;
+		
+		bool operator==(FMetalPipelineHash const& Other) const
+		{
+			return (RenderPipelineHash == Other.RenderPipelineHash && VertexDescHash == Other.VertexDescHash);
+		}
+		
+		friend uint32 GetTypeHash(FMetalPipelineHash const& Hash)
+		{
+			return GetTypeHash(Hash.RenderPipelineHash) ^ GetTypeHash((uint64)Hash.VertexDescHash);
+		}
+	};
+	
+	TMap<FMetalPipelineHash, id<MTLRenderPipelineState> > PipelineStates;
 };
 
 
diff --git a/Engine/Source/Runtime/Apple/MetalRHI/Public/MetalShaderResources.h b/Engine/Source/Runtime/Apple/MetalRHI/Public/MetalShaderResources.h
--- a/Engine/Source/Runtime/Apple/MetalRHI/Public/MetalShaderResources.h
+++ b/Engine/Source/Runtime/Apple/MetalRHI/Public/MetalShaderResources.h
@@ -139,6 +139,8 @@ struct FMetalCodeHeader
 	FMetalShaderBindings Bindings;
 	TArray<CrossCompiler::FUniformBufferCopyInfo> UniformBuffersCopyInfo;
 	FString ShaderName;
+	FString ShaderCode;
+	FString ShaderPath;
 	int32 SideTable;
 
 	uint8 NumThreadsX;
@@ -174,6 +176,8 @@ inline FArchive& operator<<(FArchive& Ar, FMetalCodeHeader& Header)
 	}
 	
 	Ar << Header.ShaderName;
+	Ar << Header.ShaderCode;
+	Ar << Header.ShaderPath;
 	
 	Ar << Header.SideTable;
 
diff --git a/Engine/Source/Runtime/Core/Private/Android/AndroidMemory.cpp b/Engine/Source/Runtime/Core/Private/Android/AndroidMemory.cpp
--- a/Engine/Source/Runtime/Core/Private/Android/AndroidMemory.cpp
+++ b/Engine/Source/Runtime/Core/Private/Android/AndroidMemory.cpp
@@ -5,6 +5,7 @@
 #include "MallocAnsi.h"
 #include "unistd.h"
 #include <jni.h>
+#include <sys/mman.h>
 
 #define JNI_CURRENT_VERSION JNI_VERSION_1_6
 extern JavaVM* GJavaVM;
@@ -95,14 +96,17 @@ FMalloc* FAndroidPlatformMemory::BaseAllocator()
 	return new FMallocBinned(MemoryConstants.PageSize, MemoryLimit);
 }
 
-void* FAndroidPlatformMemory::BinnedAllocFromOS( SIZE_T Size )
+void* FAndroidPlatformMemory::BinnedAllocFromOS(SIZE_T Size)
 {
-	// valloc was deprecated, this is a functional equivalent, for SDK 21
-	const FPlatformMemoryConstants& MemoryConstants = FPlatformMemory::GetConstants();
-	return memalign(MemoryConstants.PageSize, Size);
+	return mmap(nullptr, Size, PROT_READ | PROT_WRITE, MAP_PRIVATE | MAP_ANON, -1, 0);
 }
 
-void FAndroidPlatformMemory::BinnedFreeToOS( void* Ptr )
+void FAndroidPlatformMemory::BinnedFreeToOS(void* Ptr, SIZE_T Size)
 {
-	free(Ptr);
+	if (munmap(Ptr, Size) != 0)
+	{
+		const int ErrNo = errno;
+		UE_LOG(LogHAL, Fatal, TEXT("munmap(addr=%p, len=%llu) failed with errno = %d (%s)"), Ptr, Size,
+			ErrNo, StringCast< TCHAR >(strerror(ErrNo)).Get());
+	}
 }
diff --git a/Engine/Source/Runtime/Core/Private/Android/AndroidProcess.cpp b/Engine/Source/Runtime/Core/Private/Android/AndroidProcess.cpp
--- a/Engine/Source/Runtime/Core/Private/Android/AndroidProcess.cpp
+++ b/Engine/Source/Runtime/Core/Private/Android/AndroidProcess.cpp
@@ -10,6 +10,7 @@
 #include <sys/syscall.h>
 #include <pthread.h>
 
+#include "Android/AndroidApplication.h"
 
 const TCHAR* FAndroidPlatformProcess::ComputerName()
 {
@@ -57,3 +58,25 @@ void FAndroidPlatformProcess::LaunchURL(const TCHAR* URL, const TCHAR* Parms, FS
 		*Error = TEXT("");
 	}
 }
+
+FString FAndroidPlatformProcess::GetGameBundleId()
+{
+	JNIEnv* JEnv = FAndroidApplication::GetJavaEnv();
+	if (nullptr != JEnv)
+	{
+		jclass Class = FAndroidApplication::FindJavaClass("com/epicgames/ue4/GameActivity");
+		if (nullptr != Class)
+		{
+			jmethodID getAppPackageNameMethodId = JEnv->GetStaticMethodID(Class, "getAppPackageName", "()Ljava/lang/String;");
+			jstring JPackageName = (jstring)JEnv->CallStaticObjectMethod(Class, getAppPackageNameMethodId, nullptr);
+			const char * NativePackageNameString = JEnv->GetStringUTFChars(JPackageName, 0);
+			FString PackageName = FString(NativePackageNameString);
+			JEnv->ReleaseStringUTFChars(JPackageName, NativePackageNameString);
+			JEnv->DeleteLocalRef(JPackageName);
+			JEnv->DeleteLocalRef(Class);
+			return PackageName;
+		}
+	}
+	
+	return TEXT("");
+}
diff --git a/Engine/Source/Runtime/Core/Private/GenericPlatform/GenericPlatformMemory.cpp b/Engine/Source/Runtime/Core/Private/GenericPlatform/GenericPlatformMemory.cpp
--- a/Engine/Source/Runtime/Core/Private/GenericPlatform/GenericPlatformMemory.cpp
+++ b/Engine/Source/Runtime/Core/Private/GenericPlatform/GenericPlatformMemory.cpp
@@ -124,7 +124,7 @@ void FGenericPlatformMemory::OnOutOfMemory(uint64 Size, uint32 Alignment)
 	FPlatformMemoryStats PlatformMemoryStats = FPlatformMemory::GetStats();
 	if (BackupOOMMemoryPool)
 	{
-		FPlatformMemory::BinnedFreeToOS(BackupOOMMemoryPool);
+		FPlatformMemory::BinnedFreeToOS(BackupOOMMemoryPool, BackupOOMMemoryPoolSize);
 		UE_LOG(LogMemory, Warning, TEXT("Freeing %d bytes from backup pool to handle out of memory."), BackupOOMMemoryPoolSize);
 	}
 	UE_LOG(LogMemory, Warning, TEXT("MemoryStats:")\
@@ -201,7 +201,7 @@ void* FGenericPlatformMemory::BinnedAllocFromOS( SIZE_T Size )
 	return nullptr;
 }
 
-void FGenericPlatformMemory::BinnedFreeToOS( void* Ptr )
+void FGenericPlatformMemory::BinnedFreeToOS( void* Ptr, SIZE_T Size )
 {
 	UE_LOG(LogMemory, Error, TEXT("FGenericPlatformMemory::BinnedFreeToOS not implemented on this platform"));
 }
@@ -320,4 +320,4 @@ bool FGenericPlatformMemory::UnmapNamedSharedMemoryRegion(FSharedMemoryRegion *
 void FGenericPlatformMemory::InternalUpdateStats( const FPlatformMemoryStats& MemoryStats )
 {
 	// Generic method is empty. Implement at platform level.
-}
\ No newline at end of file
+}
diff --git a/Engine/Source/Runtime/Core/Private/GenericPlatform/GenericPlatformProcess.cpp b/Engine/Source/Runtime/Core/Private/GenericPlatform/GenericPlatformProcess.cpp
--- a/Engine/Source/Runtime/Core/Private/GenericPlatform/GenericPlatformProcess.cpp
+++ b/Engine/Source/Runtime/Core/Private/GenericPlatform/GenericPlatformProcess.cpp
@@ -183,6 +183,12 @@ bool FGenericPlatformProcess::CanLaunchURL(const TCHAR* URL)
 	return false;
 }
 
+FString FGenericPlatformProcess::GetGameBundleId()
+{
+	UE_LOG(LogHAL, Warning, TEXT("FGenericPlatformProcess::GetGameBundleId not implemented on this platform"));
+	return TEXT("");
+}
+
 FProcHandle FGenericPlatformProcess::CreateProc( const TCHAR* URL, const TCHAR* Parms, bool bLaunchDetached, bool bLaunchHidden, bool bLaunchReallyHidden, uint32* OutProcessID, int32 PriorityModifier, const TCHAR* OptionalWorkingDirectory, void* PipeWriteChild, void * PipeReadChild)
 {
 	UE_LOG(LogHAL, Fatal, TEXT("FGenericPlatformProcess::CreateProc not implemented on this platform"));
diff --git a/Engine/Source/Runtime/Core/Private/HAL/Allocators/CachedOSPageAllocator.cpp b/Engine/Source/Runtime/Core/Private/HAL/Allocators/CachedOSPageAllocator.cpp
--- a/Engine/Source/Runtime/Core/Private/HAL/Allocators/CachedOSPageAllocator.cpp
+++ b/Engine/Source/Runtime/Core/Private/HAL/Allocators/CachedOSPageAllocator.cpp
@@ -56,7 +56,7 @@ void* FCachedOSPageAllocator::AllocateImpl(SIZE_T Size, FFreePageBlock* First, F
 		// Are we holding on to much mem? Release it all.
 		for (FFreePageBlock* Block = First; Block != Last; ++Block)
 		{
-			FPlatformMemory::BinnedFreeToOS(Block->Ptr);
+			FPlatformMemory::BinnedFreeToOS(Block->Ptr, Block->ByteSize);
 			Block->Ptr      = nullptr;
 			Block->ByteSize = 0;
 		}
@@ -71,21 +71,22 @@ void FCachedOSPageAllocator::FreeImpl(void* Ptr, SIZE_T Size, uint32 NumCacheBlo
 {
 	if (Size > CachedByteLimit / 4)
 	{
-		FPlatformMemory::BinnedFreeToOS(Ptr);
+		FPlatformMemory::BinnedFreeToOS(Ptr, Size);
 		return;
 	}
 
-	while (FreedPageBlocksNum && (FreedPageBlocksNum >= NumCacheBlocks || CachedTotal + Size > CachedByteLimit)) 
+	while (FreedPageBlocksNum && (FreedPageBlocksNum >= NumCacheBlocks || CachedTotal + Size > CachedByteLimit))
 	{
 		//Remove the oldest one
 		void* FreePtr = First->Ptr;
-		CachedTotal -= First->ByteSize;
+		SIZE_T FreeSize = First->ByteSize;
+		CachedTotal -= FreeSize;
 		FreedPageBlocksNum--;
 		if (FreedPageBlocksNum)
 		{
 			FMemory::Memmove(First, First + 1, sizeof(FFreePageBlock) * FreedPageBlocksNum);
 		}
-		FPlatformMemory::BinnedFreeToOS(FreePtr);
+		FPlatformMemory::BinnedFreeToOS(FreePtr, FreeSize);
 	}
 
 	First[FreedPageBlocksNum].Ptr      = Ptr;
@@ -101,12 +102,13 @@ void FCachedOSPageAllocator::FreeAllImpl(FFreePageBlock* First, uint32& FreedPag
 	{
 		//Remove the oldest one
 		void* FreePtr = First->Ptr;
-		CachedTotal -= First->ByteSize;
+		SIZE_T FreeSize = First->ByteSize;
+		CachedTotal -= FreeSize;
 		FreedPageBlocksNum--;
 		if (FreedPageBlocksNum)
 		{
 			FMemory::Memmove(First, First + 1, sizeof(FFreePageBlock)* FreedPageBlocksNum);
 		}
-		FPlatformMemory::BinnedFreeToOS(FreePtr);
+		FPlatformMemory::BinnedFreeToOS(FreePtr, FreeSize);
 	}
 }
diff --git a/Engine/Source/Runtime/Core/Private/HAL/MallocBinned.cpp b/Engine/Source/Runtime/Core/Private/HAL/MallocBinned.cpp
--- a/Engine/Source/Runtime/Core/Private/HAL/MallocBinned.cpp
+++ b/Engine/Source/Runtime/Core/Private/HAL/MallocBinned.cpp
@@ -550,7 +550,7 @@ struct FMallocBinned::Private
 #endif
 		if (Size > MAX_CACHED_OS_FREES_BYTE_LIMIT / 4)
 		{
-			FPlatformMemory::BinnedFreeToOS(Ptr);
+			FPlatformMemory::BinnedFreeToOS(Ptr, Size);
 			return;
 		}
 
@@ -558,21 +558,21 @@ struct FMallocBinned::Private
 		{
 			//Remove the oldest one
 			void* FreePtr = Allocator.FreedPageBlocks[0].Ptr;
-			Allocator.CachedTotal -= Allocator.FreedPageBlocks[0].ByteSize;
+			SIZE_T FreeSize = Allocator.FreedPageBlocks[0].ByteSize;
+			Allocator.CachedTotal -= FreeSize;
 			Allocator.FreedPageBlocksNum--;
 			if (Allocator.FreedPageBlocksNum)
 			{
 				FMemory::Memmove(&Allocator.FreedPageBlocks[0], &Allocator.FreedPageBlocks[1], sizeof(FFreePageBlock) * Allocator.FreedPageBlocksNum);
 			}
-			FPlatformMemory::BinnedFreeToOS(FreePtr);
+			FPlatformMemory::BinnedFreeToOS(FreePtr, FreeSize);
 		}
 		Allocator.FreedPageBlocks[Allocator.FreedPageBlocksNum].Ptr      = Ptr;
 		Allocator.FreedPageBlocks[Allocator.FreedPageBlocksNum].ByteSize = Size;
 		Allocator.CachedTotal += Size;
 		++Allocator.FreedPageBlocksNum;
 #else
-		(void)Size;
-		FPlatformMemory::BinnedFreeToOS(Ptr);
+		FPlatformMemory::BinnedFreeToOS(Ptr, Size);
 #endif
 	}
 
@@ -644,7 +644,7 @@ struct FMallocBinned::Private
 		for (int i = 0, n = Allocator.FreedPageBlocksNum; i<n; ++i) 
 		{
 			//Remove allocs
-			FPlatformMemory::BinnedFreeToOS(Allocator.FreedPageBlocks[i].Ptr);
+			FPlatformMemory::BinnedFreeToOS(Allocator.FreedPageBlocks[i].Ptr, Allocator.FreedPageBlocks[i].ByteSize);
 			Allocator.FreedPageBlocks[i].Ptr = nullptr;
 			Allocator.FreedPageBlocks[i].ByteSize = 0;
 		}
diff --git a/Engine/Source/Runtime/Core/Private/HAL/MallocStomp.cpp b/Engine/Source/Runtime/Core/Private/HAL/MallocStomp.cpp
--- a/Engine/Source/Runtime/Core/Private/HAL/MallocStomp.cpp
+++ b/Engine/Source/Runtime/Core/Private/HAL/MallocStomp.cpp
@@ -118,7 +118,7 @@ void FMallocStomp::Free(void* InPtr)
 	// why we can do it.
 	munmap(AllocDataPtr->FullAllocationPointer, AllocDataPtr->FullSize);
 #else
-	FPlatformMemory::BinnedFreeToOS(AllocDataPtr->FullAllocationPointer);
+	FPlatformMemory::BinnedFreeToOS(AllocDataPtr->FullAllocationPointer, AllocDataPtr->FullSize);
 #endif // PLATFORM_LINUX || PLATFORM_MAC
 }
 
@@ -138,4 +138,4 @@ bool FMallocStomp::GetAllocationSize(void *Original, SIZE_T &SizeOut)
 	return true;
 }
 
-#endif // USE_MALLOC_STOMP
\ No newline at end of file
+#endif // USE_MALLOC_STOMP
diff --git a/Engine/Source/Runtime/Core/Private/HTML5/HTML5PlatformMemory.cpp b/Engine/Source/Runtime/Core/Private/HTML5/HTML5PlatformMemory.cpp
--- a/Engine/Source/Runtime/Core/Private/HTML5/HTML5PlatformMemory.cpp
+++ b/Engine/Source/Runtime/Core/Private/HTML5/HTML5PlatformMemory.cpp
@@ -64,7 +64,7 @@ void* FHTML5PlatformMemory::BinnedAllocFromOS( SIZE_T Size )
 #endif 
 }
 
-void FHTML5PlatformMemory::BinnedFreeToOS( void* Ptr )
+void FHTML5PlatformMemory::BinnedFreeToOS( void* Ptr, SIZE_T Size )
 {
 #if PLATFORM_HTML5_WIN32 
 	_aligned_free ( Ptr );
diff --git a/Engine/Source/Runtime/Core/Private/Linux/LinuxPlatformMemory.cpp b/Engine/Source/Runtime/Core/Private/Linux/LinuxPlatformMemory.cpp
--- a/Engine/Source/Runtime/Core/Private/Linux/LinuxPlatformMemory.cpp
+++ b/Engine/Source/Runtime/Core/Private/Linux/LinuxPlatformMemory.cpp
@@ -138,12 +138,17 @@ bool FLinuxPlatformMemory::PageProtect(void* const Ptr, const SIZE_T Size, const
 
 void* FLinuxPlatformMemory::BinnedAllocFromOS( SIZE_T Size )
 {	
-	return valloc(Size);	// equivalent to memalign(sysconf(_SC_PAGESIZE),size).
+	return mmap(nullptr, Size, PROT_READ | PROT_WRITE, MAP_PRIVATE | MAP_ANON, -1, 0);
 }
 
-void FLinuxPlatformMemory::BinnedFreeToOS( void* Ptr )
+void FLinuxPlatformMemory::BinnedFreeToOS( void* Ptr, SIZE_T Size )
 {
-	return free(Ptr);
+	if (munmap(Ptr, Size) != 0)
+	{
+		const int ErrNo = errno;
+		UE_LOG(LogHAL, Fatal, TEXT("munmap(addr=%p, len=%llu) failed with errno = %d (%s)"), Ptr, Size,
+			ErrNo, StringCast< TCHAR >(strerror(ErrNo)).Get());
+	}
 }
 
 namespace LinuxPlatformMemory
diff --git a/Engine/Source/Runtime/Core/Private/Linux/LinuxPlatformStackWalk.cpp b/Engine/Source/Runtime/Core/Private/Linux/LinuxPlatformStackWalk.cpp
--- a/Engine/Source/Runtime/Core/Private/Linux/LinuxPlatformStackWalk.cpp
+++ b/Engine/Source/Runtime/Core/Private/Linux/LinuxPlatformStackWalk.cpp
@@ -18,6 +18,10 @@
 #include <libdwarf.h>
 #include <dwarf.h>
 
+#ifndef DW_AT_MIPS_linkage_name
+	#define DW_AT_MIPS_linkage_name		0x2007			// common extension, used before DW_AT_linkage_name became standard
+#endif // DW_AT_MIPS_linkage_name
+
 namespace LinuxStackWalkHelpers
 {
 	struct LinuxBacktraceSymbols
@@ -53,16 +57,23 @@ namespace LinuxStackWalkHelpers
 		 * Gets information for the crash.
 		 *
 		 * @param Address the address to look up info for
-		 * @param OutFunctionNamePtr pointer to function name (may be NULL). Caller doesn't have to free it, but need to consider it temporary (i.e. next GetInfoForAddress() call on any thread may change it).
-		 * @param OutSourceFilePtr pointer to source filename (may be NULL). Caller doesn't have to free it, but need to consider it temporary (i.e. next GetInfoForAddress() call on any thread may change it).
-		 * @param OutLineNumberPtr pointer to line in a source file (may be NULL). Caller doesn't have to free it, but need to consider it temporary (i.e. next GetInfoForAddress() call on any thread may change it).
+		 * @param OutModuleNamePtr pointer to module name (may be null). Caller doesn't have to free it, but need to consider it temporary (i.e. next GetInfoForAddress() call on any thread may change it).
+		 * @param OutFunctionNamePtr pointer to function name (may be null). Caller doesn't have to free it, but need to consider it temporary (i.e. next GetInfoForAddress() call on any thread may change it).
+		 * @param OutSourceFilePtr pointer to source filename (may be null). Caller doesn't have to free it, but need to consider it temporary (i.e. next GetInfoForAddress() call on any thread may change it).
+		 * @param OutLineNumberPtr pointer to line in a source file (may be null). Caller doesn't have to free it, but need to consider it temporary (i.e. next GetInfoForAddress() call on any thread may change it).
 		 *
 		 * @return true if succeeded in getting the info. If false is returned, none of above parameters should be trusted to contain valid data!
 		 */
-		bool GetInfoForAddress(void * Address, const char **OutFunctionNamePtr, const char **OutSourceFilePtr, int *OutLineNumberPtr);
+		bool GetInfoForAddress(void* Address, const char** OutModuleNamePtr, const char** OutFunctionNamePtr, const char** OutSourceFilePtr, int* OutLineNumberPtr);
 
 		/**
-		 * Check check if address is inside this entry.
+		 * Checks check if address is inside this entry.
+		 *
+		 * @param DebugInfo DWARF debug info point
+		 * @param Die pointer to Debugging Information Entry
+		 * @param Addr address to check
+		 *
+		 * @return true if the address is in range
 		 */
 		static bool CheckAddressInRange(Dwarf_Debug DebugInfo, Dwarf_Die Die, Dwarf_Unsigned Addr);
 
@@ -82,6 +93,15 @@ namespace LinuxStackWalkHelpers
 		 */
 		static bool FindFunctionNameInDIE(Dwarf_Debug DebugInfo, Dwarf_Die Die, Dwarf_Addr Addr, const char **OutFuncName);
 
+		/**
+		 * Tries all usable attributes in DIE to determine function name (i.e. DW_AT_MIPS_linkage_name, DW_AT_linkage_name, DW_AT_name)
+		 *
+		 * @param Die Debugging Infromation Entry
+		 * @param OutFuncName Pointer to function name (volatile, copy it after call). Will not be touched if function returns false
+		 *
+		 * @return true if found
+		 */
+		static bool FindNameAttributeInDIE(Dwarf_Die Die, const char **OutFuncName);
 	};
 
 	enum
@@ -150,7 +170,7 @@ namespace LinuxStackWalkHelpers
 		}
 	}
 
-	bool LinuxBacktraceSymbols::GetInfoForAddress(void * Address, const char **OutFunctionNamePtr, const char **OutSourceFilePtr, int *OutLineNumberPtr)
+	bool LinuxBacktraceSymbols::GetInfoForAddress(void* Address, const char** OutModuleNamePtr, const char** OutFunctionNamePtr, const char** OutSourceFilePtr, int* OutLineNumberPtr)
 	{
 		if (DebugInfo == NULL)
 		{
@@ -271,24 +291,60 @@ namespace LinuxStackWalkHelpers
 			}
 		}
 
-		const char * FunctionName = NULL;
-		if (ReturnCode == DW_DLV_OK)
-		{
-			FindFunctionNameInDIEAndChildren(DebugInfo, Die, Addr, &FunctionName);
-		}
+		bool bSuccess = (ReturnCode == DW_DLV_OK);
 
-		if (OutFunctionNamePtr != NULL && FunctionName != NULL)
+		if (LIKELY(bSuccess))
 		{
-			*OutFunctionNamePtr = FunctionName;
-		}
+			if (LIKELY(OutFunctionNamePtr != nullptr))
+			{
+				const char * FunctionName = nullptr;
+				FindFunctionNameInDIEAndChildren(DebugInfo, Die, Addr, &FunctionName);
+				if (LIKELY(FunctionName != nullptr))
+				{
+					*OutFunctionNamePtr = FunctionName;
+				}
+				else
+				{
+					// make sure it's not null
+					*OutFunctionNamePtr = "Unknown";
+				}
+			}
 
-		if (OutSourceFilePtr != NULL && SrcFile != NULL)
-		{
-			*OutSourceFilePtr = SrcFile;
-			
-			if (OutLineNumberPtr != NULL)
+			if (LIKELY(OutSourceFilePtr != nullptr && OutLineNumberPtr != nullptr))
 			{
-				*OutLineNumberPtr = LineNumber;
+				if (SrcFile != nullptr)
+				{
+					*OutSourceFilePtr = SrcFile;
+					*OutLineNumberPtr = LineNumber;
+				}
+				else
+				{
+					*OutSourceFilePtr = "Unknown";
+					*OutLineNumberPtr = -1;
+				}
+			}
+
+			if (LIKELY(OutModuleNamePtr != nullptr))
+			{
+				const char* ModuleName = nullptr;
+
+				Dl_info DlInfo;
+				if (dladdr(Address, &DlInfo) != 0)
+				{
+					if (DlInfo.dli_fname != nullptr)
+					{
+						ModuleName = DlInfo.dli_fname;	// this is a pointer we don't own, but assuming it's good until at least the next dladdr call
+					}
+				}
+
+				if (LIKELY(ModuleName != nullptr))
+				{
+					*OutModuleNamePtr = ModuleName;
+				}
+				else
+				{
+					*OutModuleNamePtr = "Unknown";
+				}
 			}
 		}
 
@@ -296,69 +352,76 @@ namespace LinuxStackWalkHelpers
 		while (ReturnCode != DW_DLV_NO_ENTRY) 
 		{
 			if (ReturnCode == DW_DLV_ERROR)
+			{
 				break;
+			}
 			ReturnCode = dwarf_next_cu_header(DebugInfo, NULL, NULL, NULL, NULL, NULL, &ErrorInfo);
 		}
 
-		// if we weren't able to find a function name, don't trust the source file either
-		return FunctionName != NULL;
+		return bSuccess;
 	}
 
-	/**
-	 * Check check if address is inside this entry.
-	 */
 	bool LinuxBacktraceSymbols::CheckAddressInRange(Dwarf_Debug DebugInfo, Dwarf_Die Die, Dwarf_Unsigned Addr)
 	{
 		Dwarf_Attribute *AttrList;
 		Dwarf_Signed AttrCount;
 
-		if (dwarf_attrlist(Die, &AttrList, &AttrCount, NULL) != DW_DLV_OK)
-			return true;
+		if (UNLIKELY(dwarf_attrlist(Die, &AttrList, &AttrCount, NULL) != DW_DLV_OK))
+		{
+			// assume not in range if we couldn't get the information
+			return false;
+		}
+
+		Dwarf_Addr LowAddr = 0, HighAddr = 0, HighOffset = 0;
 
-		Dwarf_Addr LowAddr = 0;
 		for (int i = 0; i < AttrCount; i++)
 		{
 			Dwarf_Half Attr;
-			if (dwarf_whatattr(AttrList[i], &Attr, NULL) != DW_DLV_OK)
+			if (dwarf_whatattr(AttrList[i], &Attr, nullptr) != DW_DLV_OK)
+			{
 				continue;
+			}
+
 			switch (Attr)
 			{
 				case DW_AT_low_pc:
 					{
-						if (dwarf_formaddr(AttrList[i], &LowAddr, nullptr) == DW_DLV_OK)
+						Dwarf_Addr TempLowAddr;
+						if (dwarf_formaddr(AttrList[i], &TempLowAddr, nullptr) == DW_DLV_OK)
 						{
-							if (LowAddr > Addr)
+							if (LIKELY(TempLowAddr > Addr))	// shortcut
 							{
 								return false;
 							}
+
+							LowAddr = TempLowAddr;
 						}
 					}
 					break;
+
 				case DW_AT_high_pc:
 					{
-						Dwarf_Addr HighAddr;
-						if (dwarf_formaddr(AttrList[i], &HighAddr, nullptr) == DW_DLV_OK)
+						Dwarf_Addr TempHighAddr;
+						if (dwarf_formaddr(AttrList[i], &TempHighAddr, nullptr) == DW_DLV_OK)
 						{
-							if (HighAddr <= Addr)
+							if (LIKELY(TempHighAddr <= Addr))	// shortcut
 							{
 								return false;
 							}
+
+							HighAddr = TempHighAddr;
 						}
 
-						Dwarf_Unsigned HighOffset;
-						if (LowAddr != 0)
+						// Offset is used since DWARF-4. Store it, but don't compare right now in case
+						// we haven't yet initialized LowAddr
+						Dwarf_Unsigned TempHighOffset;
+						if (dwarf_formudata(AttrList[i], &TempHighOffset, nullptr) == DW_DLV_OK)
 						{
-							// Offset
-							if (dwarf_formudata(AttrList[i], &HighOffset, NULL) == DW_DLV_OK)
-							{
-								if (HighOffset + LowAddr <= Addr)
-								{
-									return false;
-								}
-							}
+							HighOffset = TempHighOffset;
 						}
 					}
 					break;
+
 				case DW_AT_ranges:
 					{
 						Dwarf_Unsigned Offset;
@@ -369,8 +432,11 @@ namespace LinuxStackWalkHelpers
 
 						Dwarf_Ranges *Ranges;
 						Dwarf_Signed Count;
-						if (dwarf_get_ranges(DebugInfo, (Dwarf_Off) Offset, &Ranges, &Count, NULL, NULL) != DW_DLV_OK)
+						if (dwarf_get_ranges(DebugInfo, (Dwarf_Off) Offset, &Ranges, &Count, nullptr, nullptr) != DW_DLV_OK)
+						{
 							continue;
+						}
+
 						for (int j = 0; j < Count; j++)
 						{
 							if (Ranges[j].dwr_type == DW_RANGES_END)
@@ -389,11 +455,73 @@ namespace LinuxStackWalkHelpers
 						return false;
 					}
 					break;
+
 				default:
 					break;
 			}
 		}
-		return true;
+
+		if (UNLIKELY(HighAddr == 0 && HighOffset != 0))
+		{
+			HighAddr = LowAddr + HighOffset;
+		}
+
+		return LowAddr <= Addr && Addr < HighAddr;
+	}
+
+	bool LinuxBacktraceSymbols::FindNameAttributeInDIE(Dwarf_Die Die, const char **OutFuncName)
+	{
+		Dwarf_Error ErrorInfo;
+		int ReturnCode;
+
+		// look first for DW_AT_linkage_name or DW_AT_MIPS_linkage_name, since they hold fully qualified (albeit mangled) name
+		Dwarf_Attribute LinkageNameAt;
+		// DW_AT_MIPS_linkage_name is preferred because we're using DWARF2 by default
+		ReturnCode = dwarf_attr(Die, DW_AT_MIPS_linkage_name, &LinkageNameAt, &ErrorInfo);
+		if (UNLIKELY(ReturnCode == DW_DLV_NO_ENTRY))
+		{
+			// retry with newer DW_AT_linkage_name
+			ReturnCode = dwarf_attr(Die, DW_AT_linkage_name, &LinkageNameAt, &ErrorInfo);
+		}
+
+		if (LIKELY(ReturnCode == DW_DLV_OK))
+		{
+			char *TempFuncName;
+			if (LIKELY(dwarf_formstring(LinkageNameAt, &TempFuncName, &ErrorInfo) == DW_DLV_OK))
+			{
+				// try to demangle
+				int DemangleStatus = 0xBAD;
+				char *Demangled = abi::__cxa_demangle(TempFuncName, nullptr, nullptr, &DemangleStatus);
+				if (DemangleStatus == 0 && Demangled != nullptr)
+				{
+					// cache the demangled name
+					static char CachedDemangledName[1024];
+					FCStringAnsi::Strcpy(CachedDemangledName, sizeof(CachedDemangledName), Demangled);
+
+					*OutFuncName = CachedDemangledName;
+				}
+				else
+				{
+					*OutFuncName = TempFuncName;
+				}
+
+				if (Demangled)
+				{
+					free(Demangled);
+				}
+				return true;
+			}
+		}
+
+		// if everything else fails, just take DW_AT_name, but in case of class methods, it is only a method name, so the information will be incomplete and almost useless
+		const char *TempMethodName;
+		if (LIKELY(dwarf_attrval_string(Die, DW_AT_name, &TempMethodName, &ErrorInfo) == DW_DLV_OK))
+		{
+			*OutFuncName = TempMethodName;
+			return true;
+		}
+
+		return false;
 	}
 
 	/**
@@ -406,9 +534,6 @@ namespace LinuxStackWalkHelpers
 	{
 		Dwarf_Error ErrorInfo;
 		Dwarf_Half Tag;
-		//Dwarf_Unsigned LowerPC, HigherPC;
-		char *TempFuncName;
-		int ReturnCode;
 
 		if (dwarf_tag(Die, &Tag, &ErrorInfo) != DW_DLV_OK || Tag != DW_TAG_subprogram)
 		{
@@ -420,54 +545,34 @@ namespace LinuxStackWalkHelpers
 		{
 			return false;
 		}
-		
-		// found it
-		*OutFuncName = NULL;
-		Dwarf_Attribute SubAt;
-		ReturnCode = dwarf_attr(Die, DW_AT_name, &SubAt, &ErrorInfo);
-		if (ReturnCode == DW_DLV_ERROR)
-		{
-			return true;	// error, but stop the search
-		}
-		else if (ReturnCode == DW_DLV_OK) 
+
+		// attempt to find the name in DW_TAG_subprogram DIE
+		if (FindNameAttributeInDIE(Die, OutFuncName))
 		{
-			if (dwarf_formstring(SubAt, &TempFuncName, &ErrorInfo))
-			{
-				*OutFuncName = NULL;
-			}
-			else
-			{
-				*OutFuncName = TempFuncName;
-			}
 			return true;
 		}
 
-		// DW_AT_Name is not present, look in DW_AT_specification
+		// If not found, navigate to specification DIE and look there
 		Dwarf_Attribute SpecAt;
-		if (dwarf_attr(Die, DW_AT_specification, &SpecAt, &ErrorInfo))
+		if (UNLIKELY(dwarf_attr(Die, DW_AT_specification, &SpecAt, &ErrorInfo) != DW_DLV_OK))
 		{
-			// not found, tough luck
+			// no specificaation die
 			return false;
 		}
 
 		Dwarf_Off Offset;
-		if (dwarf_global_formref(SpecAt, &Offset, &ErrorInfo))
+		if (UNLIKELY(dwarf_global_formref(SpecAt, &Offset, &ErrorInfo) != DW_DLV_OK))
 		{
 			return false;
 		}
 
 		Dwarf_Die SpecDie;
-		if (dwarf_offdie(DebugInfo, Offset, &SpecDie, &ErrorInfo))
+		if (UNLIKELY(dwarf_offdie(DebugInfo, Offset, &SpecDie, &ErrorInfo) != DW_DLV_OK))
 		{
 			return false;
 		}
 
-		if (dwarf_attrval_string(SpecDie, DW_AT_name, OutFuncName, &ErrorInfo))
-		{
-			*OutFuncName = NULL;
-		}
-
-		return true;
+		return FindNameAttributeInDIE(SpecDie, OutFuncName);
 	}
 
 	/**
@@ -667,7 +772,34 @@ void FLinuxPlatformStackWalk::ProgramCounterToSymbolInfo( uint64 ProgramCounter,
 	out_SymbolInfo.ProgramCounter = ProgramCounter;
 
 	// Get function, filename and line number.
-	// @TODO
+	const char* ModuleName = nullptr;
+	const char* FunctionName = nullptr;
+	const char* SourceFilename = nullptr;
+	int LineNumber = 0;
+
+	if (LinuxStackWalkHelpers::GetBacktraceSymbols()->GetInfoForAddress(reinterpret_cast<void*>(ProgramCounter), &ModuleName, &FunctionName, &SourceFilename, &LineNumber))
+	{
+		out_SymbolInfo.LineNumber = LineNumber;
+
+		if (LIKELY(ModuleName != nullptr))
+		{
+			FCStringAnsi::Strcpy(out_SymbolInfo.ModuleName, sizeof(out_SymbolInfo.ModuleName), ModuleName);
+		}
+
+		if (LIKELY(SourceFilename != nullptr))
+		{
+			FCStringAnsi::Strcpy(out_SymbolInfo.Filename, sizeof(out_SymbolInfo.Filename), SourceFilename);
+		}
+
+		if (FunctionName != nullptr)
+		{
+			FCStringAnsi::Strcpy(out_SymbolInfo.FunctionName, sizeof(out_SymbolInfo.Filename), FunctionName);
+		}
+		else
+		{
+			sprintf(out_SymbolInfo.FunctionName, "0x%016llx", ProgramCounter);
+		}
+	}
 }
 
 bool FLinuxPlatformStackWalk::ProgramCounterToHumanReadableString( int32 CurrentCallDepth, uint64 ProgramCounter, ANSICHAR* HumanReadableString, SIZE_T HumanReadableStringSize, FGenericCrashContext* Context )
@@ -701,42 +833,49 @@ bool FLinuxPlatformStackWalk::ProgramCounterToHumanReadableString( int32 Current
 			}
 			LinuxStackWalkHelpers::AppendToString(HumanReadableString, HumanReadableStringSize, Context, TempArray);
 
-			// Get filename.
+			// Get filename, source file and line number
+			FLinuxCrashContext* LinuxContext = static_cast< FLinuxCrashContext* >( Context );
+			if (LinuxContext)
 			{
-				const char * FunctionName = LinuxStackWalkHelpers::GetFunctionName(Context, CurrentCallDepth);
-				if (FunctionName)
+				const char * ModuleName = nullptr;
+				const char * FunctionName = nullptr;
+				const char * SourceFilename = nullptr;
+				int LineNumber;
+
+				// for ensure, use the fast path - do not even attempt to get detailed info as it will result in long hitch
+				bool bAddDetailedInfo = !LinuxContext->GetIsEnsure();
+
+				// attempt to get the said detailed info
+				bAddDetailedInfo = bAddDetailedInfo && LinuxStackWalkHelpers::GetBacktraceSymbols()->GetInfoForAddress(reinterpret_cast<void*>(ProgramCounter), &ModuleName, &FunctionName, &SourceFilename, &LineNumber);
+
+				if (bAddDetailedInfo)
 				{
+					// append FunctionName() [Source.cpp, line X] to HumanReadableString
 					LinuxStackWalkHelpers::AppendToString(HumanReadableString, HumanReadableStringSize, Context, FunctionName);
+					FCStringAnsi::Sprintf(TempArray, " [%s, line %d]", SourceFilename, LineNumber);
+					LinuxStackWalkHelpers::AppendToString(HumanReadableString, HumanReadableStringSize, Context, TempArray);
+
+					// append Module!FunctioName [Source.cpp:X] to MinidumpCallstackInfo
+					FCStringAnsi::Strcat(LinuxContext->MinidumpCallstackInfo, ARRAY_COUNT( LinuxContext->MinidumpCallstackInfo ) - 1, ModuleName);
+					FCStringAnsi::Strcat(LinuxContext->MinidumpCallstackInfo, ARRAY_COUNT( LinuxContext->MinidumpCallstackInfo ) - 1, "!");
+					LinuxStackWalkHelpers::AppendFunctionNameIfAny(*LinuxContext, FunctionName, ProgramCounter);
+					FCStringAnsi::Sprintf(TempArray, " [%s:%d]", SourceFilename, LineNumber);
+					FCStringAnsi::Strcat(LinuxContext->MinidumpCallstackInfo, ARRAY_COUNT( LinuxContext->MinidumpCallstackInfo ) - 1, TempArray);
 				}
-
-				// try to add source file and line number, too
-				FLinuxCrashContext* LinuxContext = static_cast< FLinuxCrashContext* >( Context );
-				if (LinuxContext)
+				else
 				{
-					const char * SourceFilename = NULL;
-					int LineNumber;
-					// do not symbolicate if we're handling ensure()
-					if (!LinuxContext->GetIsEnsure() && LinuxStackWalkHelpers::GetBacktraceSymbols()->GetInfoForAddress(reinterpret_cast< void* >( ProgramCounter ), NULL, &SourceFilename, &LineNumber) && FunctionName != nullptr)
+					// get the function name for backtrace, may be incorrect
+					FunctionName = LinuxStackWalkHelpers::GetFunctionName(Context, CurrentCallDepth);
+					if (FunctionName)
 					{
-						FCStringAnsi::Sprintf(TempArray, " [%s, line %d]", SourceFilename, LineNumber);
-						LinuxStackWalkHelpers::AppendToString(HumanReadableString, HumanReadableStringSize, Context, TempArray);
-						FCStringAnsi::Sprintf(TempArray, " [%s:%d]", SourceFilename, LineNumber);
-
-						// if we were able to find info for this line, it means it's in our code
-						// append program name (strong assumption of monolithic build here!)						
-						FCStringAnsi::Strcat(LinuxContext->MinidumpCallstackInfo, ARRAY_COUNT( LinuxContext->MinidumpCallstackInfo ) - 1, TCHAR_TO_ANSI(*FApp::GetName()));
-						FCStringAnsi::Strcat(LinuxContext->MinidumpCallstackInfo, ARRAY_COUNT( LinuxContext->MinidumpCallstackInfo ) - 1, "!");
-						LinuxStackWalkHelpers::AppendFunctionNameIfAny(*LinuxContext, FunctionName, ProgramCounter);
-						FCStringAnsi::Strcat(LinuxContext->MinidumpCallstackInfo, ARRAY_COUNT( LinuxContext->MinidumpCallstackInfo ) - 1, TempArray);
+						LinuxStackWalkHelpers::AppendToString(HumanReadableString, HumanReadableStringSize, Context, FunctionName);
 					}
-					else
-					{
-						// if we were NOT able to find info for this line, it means it's something else
-						FCStringAnsi::Strcat(LinuxContext->MinidumpCallstackInfo, ARRAY_COUNT( LinuxContext->MinidumpCallstackInfo ) - 1, "Unknown!");
-						LinuxStackWalkHelpers::AppendFunctionNameIfAny(*LinuxContext, FunctionName, ProgramCounter);
-					}
-					FCStringAnsi::Strcat(LinuxContext->MinidumpCallstackInfo, ARRAY_COUNT( LinuxContext->MinidumpCallstackInfo ) - 1, "\r\n");	// this one always uses Windows line terminators
+
+					FCStringAnsi::Strcat(LinuxContext->MinidumpCallstackInfo, ARRAY_COUNT( LinuxContext->MinidumpCallstackInfo ) - 1, "Unknown!");
+					LinuxStackWalkHelpers::AppendFunctionNameIfAny(*LinuxContext, FunctionName, ProgramCounter);
 				}
+
+				FCStringAnsi::Strcat(LinuxContext->MinidumpCallstackInfo, ARRAY_COUNT( LinuxContext->MinidumpCallstackInfo ) - 1, "\r\n");	// this one always uses Windows line terminators
 			}
 		}
 		return true;
diff --git a/Engine/Source/Runtime/Core/Private/Mac/MacPlatformMemory.cpp b/Engine/Source/Runtime/Core/Private/Mac/MacPlatformMemory.cpp
--- a/Engine/Source/Runtime/Core/Private/Mac/MacPlatformMemory.cpp
+++ b/Engine/Source/Runtime/Core/Private/Mac/MacPlatformMemory.cpp
@@ -153,10 +153,15 @@ bool FMacPlatformMemory::PageProtect(void* const Ptr, const SIZE_T Size, const b
 
 void* FMacPlatformMemory::BinnedAllocFromOS( SIZE_T Size )
 {
-	return valloc(Size);
+	return mmap(nullptr, Size, PROT_READ | PROT_WRITE, MAP_PRIVATE | MAP_ANON, -1, 0);
 }
 
-void FMacPlatformMemory::BinnedFreeToOS( void* Ptr )
+void FMacPlatformMemory::BinnedFreeToOS( void* Ptr, SIZE_T Size )
 {
-	free(Ptr);
+	if (munmap(Ptr, Size) != 0)
+	{
+		const int ErrNo = errno;
+		UE_LOG(LogHAL, Fatal, TEXT("munmap(addr=%p, len=%llu) failed with errno = %d (%s)"), Ptr, Size,
+			   ErrNo, StringCast< TCHAR >(strerror(ErrNo)).Get());
+	}
 }
diff --git a/Engine/Source/Runtime/Core/Private/Mac/MacPlatformMisc.cpp b/Engine/Source/Runtime/Core/Private/Mac/MacPlatformMisc.cpp
--- a/Engine/Source/Runtime/Core/Private/Mac/MacPlatformMisc.cpp
+++ b/Engine/Source/Runtime/Core/Private/Mac/MacPlatformMisc.cpp
@@ -237,6 +237,43 @@ struct FMacApplicationInfo
 		{
 			SystemLogSize = IFileManager::Get().FileSize(TEXT("/var/log/system.log"));
 		}
+		
+		if (!FPlatformMisc::IsDebuggerPresent() && FParse::Param(FCommandLine::Get(), TEXT("RedirectNSLog")))
+		{
+			fflush(stderr);
+			StdErrPipe = [NSPipe new];
+			int StdErr = dup2([StdErrPipe fileHandleForWriting].fileDescriptor, STDERR_FILENO);
+			if(StdErr > 0)
+			{
+				@try
+				{
+					NSFileHandle* StdErrFile = [StdErrPipe fileHandleForReading];
+					if(StdErrFile)
+					{
+						StdErrFile.readabilityHandler = ^(NSFileHandle* Handle){
+							NSData* FileData = Handle.availableData;
+							if (FileData.length > 0)
+							{
+								NSString* NewString = (NSString*)[[[NSString alloc] initWithData:FileData encoding:NSUTF8StringEncoding] autorelease];
+								UE_LOG(LogMac, Error, TEXT("NSLog: %s"), *FString(NewString));
+							}
+						};
+					}
+				}
+				@catch (NSException* Exc)
+				{
+					UE_LOG(LogMac, Warning, TEXT("Exception redirecting stderr to capture NSLog messages: %s"), *FString([Exc description]));
+					[StdErrPipe release];
+					StdErrPipe = nil;
+				}
+			}
+			else
+			{
+				UE_LOG(LogMac, Warning, TEXT("Failed to redirect stderr in order to capture NSLog messages."));
+				[StdErrPipe release];
+				StdErrPipe = nil;
+			}
+		}
 	}
 	
 	~FMacApplicationInfo()
@@ -331,6 +368,7 @@ struct FMacApplicationInfo
 	NSOperatingSystemVersion OSXVersion;
 	FGuid RunUUID;
 	FString XcodePath;
+	NSPipe* StdErrPipe;
 	static PLCrashReporter* CrashReporter;
 	static FMacMallocCrashHandler* CrashMalloc;
 };
@@ -452,6 +490,19 @@ void FMacPlatformMisc::PlatformPostInit(bool ShowSplashScreen)
 		[NSApp setMainMenu:MenuBar];
 		[AppMenuItem setSubmenu:AppMenu];
 
+		if (FApp::IsGame())
+		{
+			NSMenu* ViewMenu = [[FCocoaMenu new] autorelease];
+			[ViewMenu setTitle:@"View"];
+			NSMenuItem* ViewMenuItem = [[NSMenuItem new] autorelease];
+			[ViewMenuItem setSubmenu:ViewMenu];
+			[[NSApp mainMenu] addItem:ViewMenuItem];
+
+			NSMenuItem* ToggleFullscreenItem = [[[NSMenuItem alloc] initWithTitle:@"Enter Full Screen" action:@selector(toggleFullScreen:) keyEquivalent:@"f"] autorelease];
+			[ToggleFullscreenItem setKeyEquivalentModifierMask:NSCommandKeyMask | NSControlKeyMask];
+			[ViewMenu addItem:ToggleFullscreenItem];
+		}
+		
 		UpdateWindowMenu();
 	}
 
@@ -474,6 +525,16 @@ void FMacPlatformMisc::PlatformTearDown()
 		CommandletActivity = nil;
 	}
 	FApplePlatformSymbolication::EnableCoreSymbolication(false);
+	
+	if (GMacAppInfo.StdErrPipe)
+	{
+		NSFileHandle* StdErrFile = [GMacAppInfo.StdErrPipe fileHandleForReading];
+		if (StdErrFile)
+		{
+			StdErrFile.readabilityHandler = nil;
+		}
+		[GMacAppInfo.StdErrPipe release];
+	}
 }
 
 void FMacPlatformMisc::UpdateWindowMenu()
@@ -2342,6 +2403,8 @@ bool FMacPlatformMisc::HasPlatformFeature(const TCHAR* FeatureName)
 {
 	if (FCString::Stricmp(FeatureName, TEXT("Metal")) == 0)
 	{
+		bool bHasMetal = false;
+		
 		// Metal is only permitted on 10.11.4 and above now because of all the bug-fixes Apple made for us in 10.11.4.
 		if (FPlatformMisc::MacOSXVersionCompare(10, 11, 4) >= 0 && !FParse::Param(FCommandLine::Get(),TEXT("opengl")) && FModuleManager::Get().ModuleExists(TEXT("MetalRHI")))
 		{
@@ -2349,18 +2412,21 @@ bool FMacPlatformMisc::HasPlatformFeature(const TCHAR* FeatureName)
 			void* DLLHandle = FPlatformProcess::GetDllHandle(TEXT("/System/Library/Frameworks/Metal.framework/Metal"));
 			if (DLLHandle)
 			{
-				// Use the copy all function because we don't want to invoke a GPU switch at this point on dual-GPU Macbooks
-				MTLCopyAllDevices CopyDevicesPtr = (MTLCopyAllDevices)FPlatformProcess::GetDllExport(DLLHandle, TEXT("MTLCopyAllDevices"));
-				if (CopyDevicesPtr)
+				TArray<FMacPlatformMisc::FGPUDescriptor> const& GPUs = FPlatformMisc::GetGPUDescriptors();
+				for (FMacPlatformMisc::FGPUDescriptor const& GPU : GPUs)
 				{
-					SCOPED_AUTORELEASE_POOL;
-					NSArray* MetalDevices = CopyDevicesPtr();
-					[MetalDevices autorelease];
-					FPlatformProcess::FreeDllHandle(DLLHandle);
-					return MetalDevices && [MetalDevices count] > 0;
+					if (GPU.GPUMetalBundle && GPU.GPUMetalBundle.length > 0)
+					{
+						bHasMetal = true;
+						break;
+					}
 				}
+				
+				FPlatformProcess::FreeDllHandle(DLLHandle);
 			}
 		}
+		
+		return bHasMetal;
 	}
 
 	return FGenericPlatformMisc::HasPlatformFeature(FeatureName);
diff --git a/Engine/Source/Runtime/Core/Private/Mac/MacPlatformProcess.cpp b/Engine/Source/Runtime/Core/Private/Mac/MacPlatformProcess.cpp
--- a/Engine/Source/Runtime/Core/Private/Mac/MacPlatformProcess.cpp
+++ b/Engine/Source/Runtime/Core/Private/Mac/MacPlatformProcess.cpp
@@ -205,6 +205,11 @@ void FMacPlatformProcess::LaunchURL( const TCHAR* URL, const TCHAR* Parms, FStri
 	}
 }
 
+FString FMacPlatformProcess::GetGameBundleId()
+{
+	return FString([[NSBundle mainBundle] bundleIdentifier]);
+}
+
 @interface NSAutoReadPipe : NSObject
 
 /** The pipe itself */
diff --git a/Engine/Source/Runtime/Core/Private/Mac/MacWindow.cpp b/Engine/Source/Runtime/Core/Private/Mac/MacWindow.cpp
--- a/Engine/Source/Runtime/Core/Private/Mac/MacWindow.cpp
+++ b/Engine/Source/Runtime/Core/Private/Mac/MacWindow.cpp
@@ -147,6 +147,11 @@ void FMacWindow::Initialize( FMacApplication* const Application, const TSharedRe
 
 			ReshapeWindow( X, Y, SizeX, SizeY );
 
+			if (Definition->ShouldPreserveAspectRatio)
+			{
+				[WindowHandle setContentAspectRatio:NSMakeSize((float)SizeX / (float)SizeY, 1.0f)];
+			}
+
 			if (Definition->IsRegularWindow)
 			{
 				[NSApp addWindowsItem:WindowHandle title:Definition->Title.GetNSString() filename:NO];
@@ -243,7 +248,11 @@ void FMacWindow::ReshapeWindow( int32 X, int32 Y, int32 Width, int32 Height )
 					{
 						[WindowHandle setFrame: [WindowHandle screen].frame display:YES];
 					}
-					
+					else if (Definition->ShouldPreserveAspectRatio)
+					{
+						[WindowHandle setContentAspectRatio:NSMakeSize((float)Width / (float)Height, 1.0f)];
+					}
+
 					WindowHandle->bZoomed = [WindowHandle isZoomed];
 				}, UE4ResizeEventMode, true);
 			}
diff --git a/Engine/Source/Runtime/Core/Private/Windows/WindowsApplication.cpp b/Engine/Source/Runtime/Core/Private/Windows/WindowsApplication.cpp
--- a/Engine/Source/Runtime/Core/Private/Windows/WindowsApplication.cpp
+++ b/Engine/Source/Runtime/Core/Private/Windows/WindowsApplication.cpp
@@ -882,7 +882,79 @@ int32 FWindowsApplication::ProcessMessage( HWND hwnd, uint32 msg, WPARAM wParam,
 		case WM_SIZING:
 			{
 				DeferMessage( CurrentNativeEventWindowPtr, hwnd, msg, wParam, lParam, 0, 0 );
-			}
+		
+				if (CurrentNativeEventWindowPtr->GetDefinition().ShouldPreserveAspectRatio)
+				{
+					// The rect we get in lParam is window rect, but we need to preserve client's aspect ratio,
+					// so we need to find what the border and title bar sizes are, if window has them and adjust the rect.
+					WINDOWINFO WindowInfo;
+					FMemory::Memzero(WindowInfo);
+					WindowInfo.cbSize = sizeof(WindowInfo);
+					::GetWindowInfo(hwnd, &WindowInfo);
+
+					RECT TestRect;
+					TestRect.left = TestRect.right = TestRect.top = TestRect.bottom = 0;
+					AdjustWindowRectEx(&TestRect, WindowInfo.dwStyle, false, WindowInfo.dwExStyle);
+
+					RECT* Rect = (RECT*)lParam;
+					Rect->left -= TestRect.left;
+					Rect->right -= TestRect.right;
+					Rect->top -= TestRect.top;
+					Rect->bottom -= TestRect.bottom;
+
+					const float AspectRatio = CurrentNativeEventWindowPtr->GetAspectRatio();
+					int32 NewWidth = Rect->right - Rect->left;
+					int32 NewHeight = Rect->bottom - Rect->top;
+
+					switch (wParam)
+					{
+					case WMSZ_LEFT:
+					case WMSZ_RIGHT:
+						{
+							int32 AdjustedHeight = NewWidth / AspectRatio;
+							Rect->top -= (AdjustedHeight - NewHeight) / 2;
+							Rect->bottom += (AdjustedHeight - NewHeight) / 2;
+							break;
+						}
+					case WMSZ_TOP:
+					case WMSZ_BOTTOM:
+						{
+							int32 AdjustedWidth = NewHeight * AspectRatio;
+							Rect->left -= (AdjustedWidth - NewWidth) / 2;
+							Rect->right += (AdjustedWidth - NewWidth) / 2;
+							break;
+						}
+					case WMSZ_TOPLEFT:
+						{
+							int32 AdjustedHeight = NewWidth / AspectRatio;
+							Rect->top -= AdjustedHeight - NewHeight;
+							break;
+						}
+					case WMSZ_TOPRIGHT:
+						{
+							int32 AdjustedHeight = NewWidth / AspectRatio;
+							Rect->top -= AdjustedHeight - NewHeight;
+							break;
+						}
+					case WMSZ_BOTTOMLEFT:
+						{
+							int32 AdjustedHeight = NewWidth / AspectRatio;
+							Rect->bottom += AdjustedHeight - NewHeight;
+							break;
+						}
+					case WMSZ_BOTTOMRIGHT:
+						{
+							int32 AdjustedHeight = NewWidth / AspectRatio;
+							Rect->bottom += AdjustedHeight - NewHeight;
+							break;
+						}
+					}
+
+					AdjustWindowRectEx(Rect, WindowInfo.dwStyle, false, WindowInfo.dwExStyle);
+
+					return TRUE;
+				}
+		}
 			break;
 		case WM_ENTERSIZEMOVE:
 			{
@@ -1700,7 +1772,13 @@ int32 FWindowsApplication::ProcessDeferredMessage( const FDeferredWindowsMessage
 
 					const bool bWasMinimized = (wParam == SIZE_MINIMIZED);
 
-					const bool Result = MessageHandler->OnSizeChanged( CurrentNativeEventWindowPtr.ToSharedRef(), NewWidth, NewHeight, bWasMinimized );
+					const bool bIsFullscreen = (CurrentNativeEventWindowPtr->GetWindowMode() == EWindowMode::Type::Fullscreen);
+
+					// When in fullscreen Windows rendering size should be determined by the application. Do not adjust based on WM_SIZE messages.
+ 					if ( !bIsFullscreen )
+					{
+						const bool Result = MessageHandler->OnSizeChanged(CurrentNativeEventWindowPtr.ToSharedRef(), NewWidth, NewHeight, bWasMinimized);
+					}
 				}
 			}
 			break;
diff --git a/Engine/Source/Runtime/Core/Private/Windows/WindowsPlatformMemory.cpp b/Engine/Source/Runtime/Core/Private/Windows/WindowsPlatformMemory.cpp
--- a/Engine/Source/Runtime/Core/Private/Windows/WindowsPlatformMemory.cpp
+++ b/Engine/Source/Runtime/Core/Private/Windows/WindowsPlatformMemory.cpp
@@ -195,9 +195,10 @@ void* FWindowsPlatformMemory::BinnedAllocFromOS( SIZE_T Size )
 	return VirtualAlloc( NULL, Size, MEM_COMMIT, PAGE_READWRITE );
 }
 
-void FWindowsPlatformMemory::BinnedFreeToOS( void* Ptr )
+void FWindowsPlatformMemory::BinnedFreeToOS( void* Ptr, SIZE_T Size )
 {
 	CA_SUPPRESS(6001)
+	// Windows maintains the size of allocation internally, so Size is unused
 	verify(VirtualFree( Ptr, 0, MEM_RELEASE ) != 0);
 }
 
diff --git a/Engine/Source/Runtime/Core/Private/Windows/WindowsWindow.cpp b/Engine/Source/Runtime/Core/Private/Windows/WindowsWindow.cpp
--- a/Engine/Source/Runtime/Core/Private/Windows/WindowsWindow.cpp
+++ b/Engine/Source/Runtime/Core/Private/Windows/WindowsWindow.cpp
@@ -247,6 +247,7 @@ FWindowsWindow::FWindowsWindow()
 	: HWnd(NULL)
 	, WindowMode( EWindowMode::Windowed )
 	, OLEReferenceCount(0)
+	, AspectRatio(1.0f)
 	, bIsVisible( false )
 {
 	FMemory::Memzero(PreFullscreenWindowPlacement);
@@ -351,6 +352,8 @@ void FWindowsWindow::ReshapeWindow( int32 NewX, int32 NewY, int32 NewWidth, int3
 	WindowInfo.cbSize = sizeof( WindowInfo );
 	::GetWindowInfo( HWnd, &WindowInfo );
 
+	AspectRatio = (float)NewWidth / (float)NewHeight;
+
 	// X,Y, Width, Height defines the top-left pixel of the client area on the screen
 	if( Definition->HasOSWindowBorder )
 	{
diff --git a/Engine/Source/Runtime/Core/Public/Android/AndroidMemory.h b/Engine/Source/Runtime/Core/Public/Android/AndroidMemory.h
--- a/Engine/Source/Runtime/Core/Public/Android/AndroidMemory.h
+++ b/Engine/Source/Runtime/Core/Public/Android/AndroidMemory.h
@@ -25,7 +25,7 @@ struct CORE_API FAndroidPlatformMemory : public FGenericPlatformMemory
 	static const FPlatformMemoryConstants& GetConstants();
 	static FMalloc* BaseAllocator();
 	static void* BinnedAllocFromOS( SIZE_T Size );
-	static void BinnedFreeToOS( void* Ptr );
+	static void BinnedFreeToOS( void* Ptr, SIZE_T Size );
 	//~ End FGenericPlatformMemory Interface
 };
 
diff --git a/Engine/Source/Runtime/Core/Public/Android/AndroidProcess.h b/Engine/Source/Runtime/Core/Public/Android/AndroidProcess.h
--- a/Engine/Source/Runtime/Core/Public/Android/AndroidProcess.h
+++ b/Engine/Source/Runtime/Core/Public/Android/AndroidProcess.h
@@ -34,6 +34,7 @@ struct CORE_API FAndroidPlatformProcess : public FGenericPlatformProcess
 	static const TCHAR* ExecutableName(bool bRemoveExtension = true);
 	static class FRunnableThread* CreateRunnableThread();
 	static void LaunchURL(const TCHAR* URL, const TCHAR* Parms, FString* Error);
+	static FString GetGameBundleId();
 };
 
 typedef FAndroidPlatformProcess FPlatformProcess;
diff --git a/Engine/Source/Runtime/Core/Public/Android/AndroidString.h b/Engine/Source/Runtime/Core/Public/Android/AndroidString.h
--- a/Engine/Source/Runtime/Core/Public/Android/AndroidString.h
+++ b/Engine/Source/Runtime/Core/Public/Android/AndroidString.h
@@ -200,21 +200,37 @@ struct FAndroidPlatformString : public FGenericPlatformString
 
 	static FORCEINLINE const WIDECHAR* Strstr( const WIDECHAR* String, const WIDECHAR* Find)
 	{
-		int StringLen = Strlen(String);
-		int FindLen = Strlen(Find);
+		WIDECHAR FindChar = *Find;
 
-		ANSICHAR* AnsiString = (ANSICHAR*)FMemory_Alloca(StringLen+1);
-		ANSICHAR* AnsiFind = (ANSICHAR*)FMemory_Alloca(FindLen+1);
+		// Always find an empty string
+		if (!FindChar)
+			return String;
 
-		CopyWideToAnsi(AnsiString, String);
-		CopyWideToAnsi(AnsiFind, Find);
+		++Find;
+		size_t MemCmpLen = wcslen(Find);
+		WIDECHAR* FoundChar;
+		for (;;)
+		{
+			FoundChar = wcschr(String, FindChar);
+			if (!FoundChar)
+			{
+				//	No more instances of FindChar in String, Find does not exist in String
+				break;
+			}
 
-		ANSICHAR* Loc = strstr(AnsiString, AnsiFind);
+			// Found FindChar, now set String to one character after FoundChar...
+			// We can now compare characters after FoundChar to Find, which is pointing at original Find[1], to see if the rest of the characters match
+			String = FoundChar + 1;
+			if (!wmemcmp(String, Find, MemCmpLen))
+			{
+				// Strings match, return pointer to beginning of Find instance in String
+				return FoundChar;
+			}
 
-		if (!Loc)
-			return NULL;
+			// No match, String is already ready to loop again to find the next instance of FindChar
+		}
 
-		return String + (Loc - AnsiString);
+		return nullptr;
 	}
 
 	static FORCEINLINE const WIDECHAR* Strchr( const WIDECHAR* String, WIDECHAR C)
diff --git a/Engine/Source/Runtime/Core/Public/GenericPlatform/GenericPlatformMemory.h b/Engine/Source/Runtime/Core/Public/GenericPlatform/GenericPlatformMemory.h
--- a/Engine/Source/Runtime/Core/Public/GenericPlatform/GenericPlatformMemory.h
+++ b/Engine/Source/Runtime/Core/Public/GenericPlatform/GenericPlatformMemory.h
@@ -265,8 +265,9 @@ struct CORE_API FGenericPlatformMemory
 	 * Returns pages allocated by BinnedAllocFromOS to the OS.
 	 *
 	 * @param A pointer previously returned from BinnedAllocFromOS
+	 * @param Size size of the allocation previously passed to BinnedAllocFromOS
 	 */
-	static void BinnedFreeToOS( void* Ptr );
+	static void BinnedFreeToOS( void* Ptr, SIZE_T Size );
 
 	// These alloc/free memory that is mapped to the GPU
 	// Only for platforms with UMA (XB1/PS4/etc)
diff --git a/Engine/Source/Runtime/Core/Public/GenericPlatform/GenericPlatformProcess.h b/Engine/Source/Runtime/Core/Public/GenericPlatform/GenericPlatformProcess.h
--- a/Engine/Source/Runtime/Core/Public/GenericPlatform/GenericPlatformProcess.h
+++ b/Engine/Source/Runtime/Core/Public/GenericPlatform/GenericPlatformProcess.h
@@ -319,7 +319,14 @@ struct CORE_API FGenericPlatformProcess
 	 * Checks if the platform can launch a uniform resource locator (i.e. http://www.epicgames.com/unreal).
 	 **/
 	static bool CanLaunchURL(const TCHAR* URL);
-
+	
+	/**
+	 * Retrieves the platform-specific bundle identifier or package name of the game
+	 *
+	 * @return The game's bundle identifier or package name.
+	 */
+	static FString GetGameBundleId();
+	
 	/**
 	 * Creates a new process and its primary thread. The new process runs the
 	 * specified executable file in the security context of the calling process.
diff --git a/Engine/Source/Runtime/Core/Public/GenericPlatform/GenericWindowDefinition.h b/Engine/Source/Runtime/Core/Public/GenericPlatform/GenericWindowDefinition.h
--- a/Engine/Source/Runtime/Core/Public/GenericPlatform/GenericWindowDefinition.h
+++ b/Engine/Source/Runtime/Core/Public/GenericPlatform/GenericWindowDefinition.h
@@ -90,6 +90,8 @@ struct CORE_API FGenericWindowDefinition
 	bool HasSizingFrame;
 	/** true if we expect the size of this window to change often, such as if its animated, or if it recycled for tool-tips. */
 	bool SizeWillChangeOften;
+	/** true if the window should preserve its aspect ratio when resized by user */
+	bool ShouldPreserveAspectRatio;
 	/** The expected maximum width of the window.  May be used for performance optimization when SizeWillChangeOften is set. */
 	int32 ExpectedMaxWidth;
 	/** The expected maximum height of the window.  May be used for performance optimization when SizeWillChangeOften is set. */
diff --git a/Engine/Source/Runtime/Core/Public/HTML5/HTML5PlatformMemory.h b/Engine/Source/Runtime/Core/Public/HTML5/HTML5PlatformMemory.h
--- a/Engine/Source/Runtime/Core/Public/HTML5/HTML5PlatformMemory.h
+++ b/Engine/Source/Runtime/Core/Public/HTML5/HTML5PlatformMemory.h
@@ -24,7 +24,7 @@ struct CORE_API FHTML5PlatformMemory : public FGenericPlatformMemory
 	static FPlatformMemoryStats GetStats();
 	static FMalloc* BaseAllocator();
 	static void* BinnedAllocFromOS( SIZE_T Size );
-	static void BinnedFreeToOS( void* Ptr );
+	static void BinnedFreeToOS( void* Ptr, SIZE_T Size );
 	//~ End FGenericPlatformMemory Interface
 };
 
diff --git a/Engine/Source/Runtime/Core/Public/IOS/IOSPlatformCompilerPreSetup.h b/Engine/Source/Runtime/Core/Public/IOS/IOSPlatformCompilerPreSetup.h
--- a/Engine/Source/Runtime/Core/Public/IOS/IOSPlatformCompilerPreSetup.h
+++ b/Engine/Source/Runtime/Core/Public/IOS/IOSPlatformCompilerPreSetup.h
@@ -102,4 +102,5 @@
 // Rationale: we don't want to suppress those as there are plans to address them (e.g. UE-12341), but breaking builds due to these warnings is very expensive
 // since they cannot be caught by all compilers that we support. They are deemed to be relatively safe to be ignored, at least until all SDKs/toolchains start supporting them.
 #pragma clang diagnostic warning "-Wreorder"
-#pragma clang diagnostic warning "-Wparentheses-equality"
\ No newline at end of file
+#pragma clang diagnostic warning "-Wparentheses-equality"
+#pragma clang diagnostic warning "-Wdelete-non-virtual-dtor"
diff --git a/Engine/Source/Runtime/Core/Public/IOS/IOSPlatformMemory.h b/Engine/Source/Runtime/Core/Public/IOS/IOSPlatformMemory.h
--- a/Engine/Source/Runtime/Core/Public/IOS/IOSPlatformMemory.h
+++ b/Engine/Source/Runtime/Core/Public/IOS/IOSPlatformMemory.h
@@ -24,8 +24,9 @@ struct CORE_API FIOSPlatformMemory : public FGenericPlatformMemory
 	static FPlatformMemoryStats GetStats();
 	static const FPlatformMemoryConstants& GetConstants();
 	static FMalloc* BaseAllocator();
+	static bool PageProtect(void* const Ptr, const SIZE_T Size, const bool bCanRead, const bool bCanWrite);
 	static void* BinnedAllocFromOS( SIZE_T Size );
-	static void BinnedFreeToOS( void* Ptr );
+	static void BinnedFreeToOS( void* Ptr, SIZE_T Size );
 	//~ End FGenericPlatformMemory Interface
 };
 
diff --git a/Engine/Source/Runtime/Core/Public/IOS/IOSPlatformProcess.h b/Engine/Source/Runtime/Core/Public/IOS/IOSPlatformProcess.h
--- a/Engine/Source/Runtime/Core/Public/IOS/IOSPlatformProcess.h
+++ b/Engine/Source/Runtime/Core/Public/IOS/IOSPlatformProcess.h
@@ -32,6 +32,7 @@ struct CORE_API FIOSPlatformProcess : public FGenericPlatformProcess
 	static FRunnableThread* CreateRunnableThread();
 	static void LaunchURL( const TCHAR* URL, const TCHAR* Parms, FString* Error );
 	static bool CanLaunchURL(const TCHAR* URL);
+	static FString GetGameBundleId();
 	static void SetRealTimeMode();
 	static void SetupGameThread();
 	static void SetupRenderThread();
diff --git a/Engine/Source/Runtime/Core/Public/Mac/MacPlatformCompilerPreSetup.h b/Engine/Source/Runtime/Core/Public/Mac/MacPlatformCompilerPreSetup.h
--- a/Engine/Source/Runtime/Core/Public/Mac/MacPlatformCompilerPreSetup.h
+++ b/Engine/Source/Runtime/Core/Public/Mac/MacPlatformCompilerPreSetup.h
@@ -105,6 +105,7 @@
 #pragma clang diagnostic ignored "-Winconsistent-missing-override"
 #pragma clang diagnostic ignored "-Wundefined-bool-conversion"
 #pragma clang diagnostic ignored "-Wunused-local-typedef"
+#pragma clang diagnostic ignored "-Wdelete-non-virtual-dtor"
 
 // We can pragma optimisation's on and off as of Apple LLVM 7.3.0 but not before.
 #if __clang_major__ >= 7 && __clang_minor__ >= 3
diff --git a/Engine/Source/Runtime/Core/Public/Windows/WindowsWindow.h b/Engine/Source/Runtime/Core/Public/Windows/WindowsWindow.h
--- a/Engine/Source/Runtime/Core/Public/Windows/WindowsWindow.h
+++ b/Engine/Source/Runtime/Core/Public/Windows/WindowsWindow.h
@@ -76,6 +76,8 @@ public:
 	/** Called by the owning application when the level of transparency support has changed */
 	void OnTransparencySupportChanged(EWindowTransparency NewTransparency);
 
+	float GetAspectRatio() const { return AspectRatio; }
+
 public:
 
 	// FGenericWindow interface
@@ -166,5 +168,8 @@ private:
 	int32 VirtualWidth;
 	int32 VirtualHeight;
 
+	/** Current aspect ratio of window's client area */
+	float AspectRatio;
+
 	bool bIsVisible : 1;
 };
diff --git a/Engine/Source/Runtime/CoreUObject/Private/Serialization/AsyncLoading.cpp b/Engine/Source/Runtime/CoreUObject/Private/Serialization/AsyncLoading.cpp
--- a/Engine/Source/Runtime/CoreUObject/Private/Serialization/AsyncLoading.cpp
+++ b/Engine/Source/Runtime/CoreUObject/Private/Serialization/AsyncLoading.cpp
@@ -19,7 +19,7 @@
 #include "HAL/ExceptionHandling.h"
 #include "AsyncLoadingPrivate.h"
 
-#define FIND_MEMORY_STOMPS (1 && PLATFORM_WINDOWS && !WITH_EDITORONLY_DATA)
+#define FIND_MEMORY_STOMPS (1 && (PLATFORM_WINDOWS || PLATFORM_LINUX) && !WITH_EDITORONLY_DATA)
 
 //#pragma clang optimize off
 
@@ -2689,11 +2689,7 @@ static uint8* MallocAsyncBuffer(const SIZE_T Size, SIZE_T& OutAllocatedSize)
 	const SIZE_T AllocFullPageSize = AlignedSize + (PageSize - 1) & ~(PageSize - 1U);
 	check(AllocFullPageSize >= Size);
 	OutAllocatedSize = AllocFullPageSize;
-#if PLATFORM_LINUX || PLATFORM_MAC
-	Result = (uint8*)mmap(nullptr, AllocFullPageSize, PROT_READ | PROT_WRITE, MAP_PRIVATE | MAP_ANON, -1, 0);
-#else
 	Result = (uint8*)FPlatformMemory::BinnedAllocFromOS(AllocFullPageSize);
-#endif // PLATFORM_LINUX || PLATFORM_MAC
 #else
 	OutAllocatedSize = Size;
 	Result = (uint8*)FMemory::Malloc(Size);
@@ -2706,11 +2702,7 @@ static void FreeAsyncBuffer(uint8* Buffer, const SIZE_T AllocatedSize)
 	if (Buffer)
 	{
 #if FIND_MEMORY_STOMPS
-#if PLATFORM_LINUX || PLATFORM_MAC
-		munmap(Buffer, AllocatedSize);
-#else
-		FPlatformMemory::BinnedFreeToOS(Buffer);
-#endif // PLATFORM_LINUX || PLATFORM_MAC
+		FPlatformMemory::BinnedFreeToOS(Buffer, AllocatedSize);
 #else
 		FMemory::Free(Buffer);
 #endif // FIND_MEMORY_STOMPS
@@ -4297,4 +4289,4 @@ static FAutoConsoleCommand DumpSerializeCmd(
 
 #endif // USE_NEW_ASYNC_IO
 
-//#pragma clang optimize on
\ No newline at end of file
+//#pragma clang optimize on
diff --git a/Engine/Source/Runtime/Engine/Classes/Kismet/KismetSystemLibrary.h b/Engine/Source/Runtime/Engine/Classes/Kismet/KismetSystemLibrary.h
--- a/Engine/Source/Runtime/Engine/Classes/Kismet/KismetSystemLibrary.h
+++ b/Engine/Source/Runtime/Engine/Classes/Kismet/KismetSystemLibrary.h
@@ -106,6 +106,14 @@ class ENGINE_API UKismetSystemLibrary : public UBlueprintFunctionLibrary
 	UFUNCTION(BlueprintPure, Category="Game")
 	static FString GetGameName();
 
+	/**
+	 * Retrieves the game's platform-specific bundle identifier or package name of the game
+	 *
+	 * @return The game's bundle identifier or package name.
+	 */
+	UFUNCTION(BlueprintPure, Category="Game", meta=(Keywords = "bundle id package name"))
+	static FString GetGameBundleId();
+
 	/** Get the current user name from the OS */
 	UFUNCTION(BlueprintPure, Category="Utilities|Platform")
 	static FString GetPlatformUserName();
diff --git a/Engine/Source/Runtime/Engine/Private/Components/SplineMeshComponent.cpp b/Engine/Source/Runtime/Engine/Private/Components/SplineMeshComponent.cpp
--- a/Engine/Source/Runtime/Engine/Private/Components/SplineMeshComponent.cpp
+++ b/Engine/Source/Runtime/Engine/Private/Components/SplineMeshComponent.cpp
@@ -49,39 +49,42 @@ void FSplineMeshVertexFactoryShaderParameters::SetMesh(FRHICommandList& RHICmdLi
 		check(OverrideColorVertexBuffer);
 		static_cast<const FLocalVertexFactory*>(VertexFactory)->SetColorOverrideStream(RHICmdList, OverrideColorVertexBuffer);
 	}
-	if (Shader->GetVertexShader())
+
+	FVertexShaderRHIRef VertexShader = Shader->GetVertexShader();
+
+	if (VertexShader)
 	{
 		checkSlow(BatchElement.bIsSplineProxy);
 		FSplineMeshSceneProxy* SplineProxy = BatchElement.SplineMeshSceneProxy;
 		FSplineMeshParams& SplineParams = SplineProxy->SplineParams;
 
-		SetShaderValue(RHICmdList, Shader->GetVertexShader(), SplineStartPosParam, SplineParams.StartPos);
-		SetShaderValue(RHICmdList, Shader->GetVertexShader(), SplineStartTangentParam, SplineParams.StartTangent);
-		SetShaderValue(RHICmdList, Shader->GetVertexShader(), SplineStartRollParam, SplineParams.StartRoll);
-		SetShaderValue(RHICmdList, Shader->GetVertexShader(), SplineStartScaleParam, SplineParams.StartScale);
-		SetShaderValue(RHICmdList, Shader->GetVertexShader(), SplineStartOffsetParam, SplineParams.StartOffset);
+		SetShaderValue(RHICmdList, VertexShader, SplineStartPosParam, SplineParams.StartPos);
+		SetShaderValue(RHICmdList, VertexShader, SplineStartTangentParam, SplineParams.StartTangent);
+		SetShaderValue(RHICmdList, VertexShader, SplineStartRollParam, SplineParams.StartRoll);
+		SetShaderValue(RHICmdList, VertexShader, SplineStartScaleParam, SplineParams.StartScale);
+		SetShaderValue(RHICmdList, VertexShader, SplineStartOffsetParam, SplineParams.StartOffset);
 
-		SetShaderValue(RHICmdList, Shader->GetVertexShader(), SplineEndPosParam, SplineParams.EndPos);
-		SetShaderValue(RHICmdList, Shader->GetVertexShader(), SplineEndTangentParam, SplineParams.EndTangent);
-		SetShaderValue(RHICmdList, Shader->GetVertexShader(), SplineEndRollParam, SplineParams.EndRoll);
-		SetShaderValue(RHICmdList, Shader->GetVertexShader(), SplineEndScaleParam, SplineParams.EndScale);
-		SetShaderValue(RHICmdList, Shader->GetVertexShader(), SplineEndOffsetParam, SplineParams.EndOffset);
+		SetShaderValue(RHICmdList, VertexShader, SplineEndPosParam, SplineParams.EndPos);
+		SetShaderValue(RHICmdList, VertexShader, SplineEndTangentParam, SplineParams.EndTangent);
+		SetShaderValue(RHICmdList, VertexShader, SplineEndRollParam, SplineParams.EndRoll);
+		SetShaderValue(RHICmdList, VertexShader, SplineEndScaleParam, SplineParams.EndScale);
+		SetShaderValue(RHICmdList, VertexShader, SplineEndOffsetParam, SplineParams.EndOffset);
 
-		SetShaderValue(RHICmdList, Shader->GetVertexShader(), SplineUpDirParam, SplineProxy->SplineUpDir);
-		SetShaderValue(RHICmdList, Shader->GetVertexShader(), SmoothInterpRollScaleParam, SplineProxy->bSmoothInterpRollScale);
+		SetShaderValue(RHICmdList, VertexShader, SplineUpDirParam, SplineProxy->SplineUpDir);
+		SetShaderValue(RHICmdList, VertexShader, SmoothInterpRollScaleParam, SplineProxy->bSmoothInterpRollScale);
 
-		SetShaderValue(RHICmdList, Shader->GetVertexShader(), SplineMeshMinZParam, SplineProxy->SplineMeshMinZ);
-		SetShaderValue(RHICmdList, Shader->GetVertexShader(), SplineMeshScaleZParam, SplineProxy->SplineMeshScaleZ);
+		SetShaderValue(RHICmdList, VertexShader, SplineMeshMinZParam, SplineProxy->SplineMeshMinZ);
+		SetShaderValue(RHICmdList, VertexShader, SplineMeshScaleZParam, SplineProxy->SplineMeshScaleZ);
 
 		FVector DirMask(0, 0, 0);
 		DirMask[SplineProxy->ForwardAxis] = 1;
-		SetShaderValue(RHICmdList, Shader->GetVertexShader(), SplineMeshDirParam, DirMask);
+		SetShaderValue(RHICmdList, VertexShader, SplineMeshDirParam, DirMask);
 		DirMask = FVector::ZeroVector;
 		DirMask[(SplineProxy->ForwardAxis + 1) % 3] = 1;
-		SetShaderValue(RHICmdList, Shader->GetVertexShader(), SplineMeshXParam, DirMask);
+		SetShaderValue(RHICmdList, VertexShader, SplineMeshXParam, DirMask);
 		DirMask = FVector::ZeroVector;
 		DirMask[(SplineProxy->ForwardAxis + 2) % 3] = 1;
-		SetShaderValue(RHICmdList, Shader->GetVertexShader(), SplineMeshYParam, DirMask);
+		SetShaderValue(RHICmdList, VertexShader, SplineMeshYParam, DirMask);
 	}
 }
 
diff --git a/Engine/Source/Runtime/Engine/Private/GameEngine.cpp b/Engine/Source/Runtime/Engine/Private/GameEngine.cpp
--- a/Engine/Source/Runtime/Engine/Private/GameEngine.cpp
+++ b/Engine/Source/Runtime/Engine/Private/GameEngine.cpp
@@ -303,6 +303,15 @@ TSharedRef<SWindow> UGameEngine::CreateGameWindow()
 		IConsoleManager::Get().CallAllConsoleVariableSinks();
 	}
 
+	const FText WindowTitleOverride = GetDefault<UGeneralProjectSettings>()->ProjectDisplayedTitle;
+	const FText WindowTitleComponent = WindowTitleOverride.IsEmpty() ? NSLOCTEXT("UnrealEd", "GameWindowTitle", "{GameName}") : WindowTitleOverride;
+
+	FText WindowDebugInfoComponent;
+#if !UE_BUILD_SHIPPING
+	const FText WindowDebugInfoOverride = GetDefault<UGeneralProjectSettings>()->ProjectDebugTitleInfo;
+	WindowDebugInfoComponent = WindowDebugInfoOverride.IsEmpty() ? NSLOCTEXT("UnrealEd", "GameWindowTitleDebugInfo", "({PlatformArchitecture}-bit, {RHIName})") : WindowDebugInfoOverride;
+#endif
+
 #if PLATFORM_64BITS
 	//These are invariant strings so they don't need to be localized
 	const FText PlatformBits = FText::FromString( TEXT( "64" ) );
@@ -310,15 +319,15 @@ TSharedRef<SWindow> UGameEngine::CreateGameWindow()
 	const FText PlatformBits = FText::FromString( TEXT( "32" ) );
 #endif	//PLATFORM_64BITS
 
-	// Note: If these parameters are updated or renamed, please update the tooltip on the ProjectDisplayedTitle property
+	// Note: If these parameters are updated or renamed, please update the tooltip on the ProjectDisplayedTitle and ProjectDebugTitleInfo properties
 	FFormatNamedArguments Args;
 	Args.Add( TEXT("GameName"), FText::FromString( FApp::GetGameName() ) );
 	Args.Add( TEXT("PlatformArchitecture"), PlatformBits );
 	Args.Add( TEXT("RHIName"), FText::FromName( LegacyShaderPlatformToShaderFormat( GMaxRHIShaderPlatform ) ) );
 
-	const FText DefaultWindowTitle = NSLOCTEXT("UnrealEd", "GameWindowTitle", "{GameName} ({PlatformArchitecture}-bit, {RHIName})");
-	const FText WindowTitleOverride = GetDefault<UGeneralProjectSettings>()->ProjectDisplayedTitle;
-	const FText WindowTitle = FText::Format(WindowTitleOverride.IsEmpty() ? DefaultWindowTitle : WindowTitleOverride, Args);
+	const FText WindowTitleVar = FText::Format(FText::FromString(TEXT("{0} {1}")), WindowTitleComponent, WindowDebugInfoComponent);
+	const FText WindowTitle = FText::Format(WindowTitleVar, Args);
+	const bool bShouldPreserveAspectRatio = GetDefault<UGeneralProjectSettings>()->bShouldWindowPreserveAspectRatio;
 
 	// Allow optional winX/winY parameters to set initial window position
 	EAutoCenter::Type AutoCenterType = EAutoCenter::PrimaryWorkArea;
@@ -360,7 +369,8 @@ TSharedRef<SWindow> UGameEngine::CreateGameWindow()
 	.MaxHeight(MaxWindowHeight)
 	.FocusWhenFirstShown(true)
 	.SaneWindowPlacement(AutoCenterType == EAutoCenter::None)
-	.UseOSWindowBorder(true);
+	.UseOSWindowBorder(true)
+	.ShouldPreserveAspectRatio(bShouldPreserveAspectRatio);
 
 	const bool bShowImmediately = false;
 
diff --git a/Engine/Source/Runtime/Engine/Private/KismetSystemLibrary.cpp b/Engine/Source/Runtime/Engine/Private/KismetSystemLibrary.cpp
--- a/Engine/Source/Runtime/Engine/Private/KismetSystemLibrary.cpp
+++ b/Engine/Source/Runtime/Engine/Private/KismetSystemLibrary.cpp
@@ -85,6 +85,11 @@ FString UKismetSystemLibrary::GetGameName()
 	return FString(FApp::GetGameName());
 }
 
+FString UKismetSystemLibrary::GetGameBundleId()
+{
+	return FString(FPlatformProcess::GetGameBundleId());
+}
+
 FString UKismetSystemLibrary::GetPlatformUserName()
 {
 	return FString(FPlatformProcess::UserName());
diff --git a/Engine/Source/Runtime/Engine/Private/ScreenRendering.cpp b/Engine/Source/Runtime/Engine/Private/ScreenRendering.cpp
--- a/Engine/Source/Runtime/Engine/Private/ScreenRendering.cpp
+++ b/Engine/Source/Runtime/Engine/Private/ScreenRendering.cpp
@@ -14,4 +14,5 @@ TGlobalResource<FScreenVertexDeclaration> GScreenVertexDeclaration;
 // Shader implementations.
 IMPLEMENT_SHADER_TYPE(,FScreenPS,TEXT("ScreenPixelShader"),TEXT("Main"),SF_Pixel);
 IMPLEMENT_SHADER_TYPE(,FScreenVS,TEXT("ScreenVertexShader"),TEXT("Main"),SF_Vertex);
-IMPLEMENT_SHADER_TYPE(,FScreenVSForGS,TEXT("ScreenVertexShader"),TEXT("MainForGS"),SF_Vertex);
+IMPLEMENT_SHADER_TYPE(template<> ,TScreenVSForGS<false>,TEXT("ScreenVertexShader"),TEXT("MainForGS"),SF_Vertex);
+IMPLEMENT_SHADER_TYPE(template<> ,TScreenVSForGS<true>,TEXT("ScreenVertexShader"),TEXT("MainForGS"),SF_Vertex);
diff --git a/Engine/Source/Runtime/Engine/Public/ScreenRendering.h b/Engine/Source/Runtime/Engine/Public/ScreenRendering.h
--- a/Engine/Source/Runtime/Engine/Public/ScreenRendering.h
+++ b/Engine/Source/Runtime/Engine/Public/ScreenRendering.h
@@ -113,17 +113,31 @@ public:
 	}
 };
 
-class FScreenVSForGS : public FScreenVS
+template<bool bUsingVertexLayers=false>
+class TScreenVSForGS : public FScreenVS
 {
-	DECLARE_EXPORTED_SHADER_TYPE(FScreenVSForGS,Global,ENGINE_API);
+	DECLARE_EXPORTED_SHADER_TYPE(TScreenVSForGS,Global,ENGINE_API);
 public:
 
-	static bool ShouldCache(EShaderPlatform Platform) { return IsFeatureLevelSupported(Platform, ERHIFeatureLevel::SM4); }
+	static bool ShouldCache(EShaderPlatform Platform)
+	{
+		return IsFeatureLevelSupported(Platform, ERHIFeatureLevel::SM4) && (!bUsingVertexLayers || RHISupportsVertexShaderLayer(Platform));
+	}
 
-	FScreenVSForGS(const ShaderMetaType::CompiledShaderInitializerType& Initializer):
+	TScreenVSForGS(const ShaderMetaType::CompiledShaderInitializerType& Initializer):
 	  FScreenVS(Initializer)
 	  {
 	  }
-	FScreenVSForGS() {}
+	TScreenVSForGS() {}
+	
+	static void ModifyCompilationEnvironment(EShaderPlatform Platform, FShaderCompilerEnvironment& OutEnvironment)
+	{
+		FScreenVS::ModifyCompilationEnvironment(Platform, OutEnvironment);
+		OutEnvironment.SetDefine(TEXT("USING_LAYERS"), (uint32)(bUsingVertexLayers ? 1 : 0));
+		if (!bUsingVertexLayers)
+		{
+			OutEnvironment.CompilerFlags.Add( CFLAG_VertexToGeometryShader );
+		}
+	}
 };
 
diff --git a/Engine/Source/Runtime/EngineSettings/Classes/GeneralProjectSettings.h b/Engine/Source/Runtime/EngineSettings/Classes/GeneralProjectSettings.h
--- a/Engine/Source/Runtime/EngineSettings/Classes/GeneralProjectSettings.h
+++ b/Engine/Source/Runtime/EngineSettings/Classes/GeneralProjectSettings.h
@@ -58,4 +58,12 @@ class ENGINESETTINGS_API UGeneralProjectSettings
 	/** The project's title as displayed on the window title bar (can include the tokens {GameName}, {PlatformArchitecture}, or {RHIName}, which will be replaced with the specified text) */
 	UPROPERTY(config, EditAnywhere, Category=Displayed)
 	FText ProjectDisplayedTitle;
+
+	/** Additional data to be displayed on the window title bar in non-shipping configurations (can include the tokens {GameName}, {PlatformArchitecture}, or {RHIName}, which will be replaced with the specified text) */
+	UPROPERTY(config, EditAnywhere, Category=Displayed)
+	FText ProjectDebugTitleInfo;
+
+	/** Should the game's window preserve its aspect ratio when resized by user. */
+	UPROPERTY(config, EditAnywhere, Category = Settings)
+	bool bShouldWindowPreserveAspectRatio;
 };
diff --git a/Engine/Source/Runtime/EngineSettings/Private/EngineSettingsModule.cpp b/Engine/Source/Runtime/EngineSettings/Private/EngineSettingsModule.cpp
--- a/Engine/Source/Runtime/EngineSettings/Private/EngineSettingsModule.cpp
+++ b/Engine/Source/Runtime/EngineSettings/Private/EngineSettingsModule.cpp
@@ -57,6 +57,7 @@ UGeneralEngineSettings::UGeneralEngineSettings( const FObjectInitializer& Object
 
 UGeneralProjectSettings::UGeneralProjectSettings( const FObjectInitializer& ObjectInitializer )
 	: Super(ObjectInitializer)
+	, bShouldWindowPreserveAspectRatio(true)
 { }
 
 
diff --git a/Engine/Source/Runtime/IOS/IOSRuntimeSettings/Classes/IOSRuntimeSettings.h b/Engine/Source/Runtime/IOS/IOSRuntimeSettings/Classes/IOSRuntimeSettings.h
--- a/Engine/Source/Runtime/IOS/IOSRuntimeSettings/Classes/IOSRuntimeSettings.h
+++ b/Engine/Source/Runtime/IOS/IOSRuntimeSettings/Classes/IOSRuntimeSettings.h
@@ -177,6 +177,14 @@ public:
 	UPROPERTY(GlobalConfig, EditAnywhere, Category = Build, meta = (DisplayName = "Generate dSYM file for code debugging and profiling"))
 	bool bGeneratedSYMFile;
 
+	// Enable generation of dSYM bundle
+	UPROPERTY(GlobalConfig, EditAnywhere, Category = Build, meta = (DisplayName = "Generate dSYM bundle for third party crash tools"))
+	bool bGeneratedSYMBundle;
+
+	// Enable generation of xcode archive package
+	UPROPERTY(GlobalConfig, EditAnywhere, Category = Build, meta = (DisplayName = "Generate xcode archive package"))
+	bool bGenerateXCArchive;	
+	
 	// Enable ArmV7 support? (this will be used if all type are unchecked)
 	UPROPERTY(GlobalConfig, EditAnywhere, Category = Build, meta = (DisplayName = "Support armv7 in Development"))
 	bool bDevForArmV7;
@@ -202,7 +210,7 @@ public:
 	bool bShipForArmV7S;
 
 	// Enable bitcode compiling?
-	UPROPERTY(GlobalConfig, EditAnywhere, Category = Build, meta = (DisplayName = "Support bitcode compiling"))
+	UPROPERTY(GlobalConfig, EditAnywhere, Category = Build, meta = (DisplayName = "Support bitcode in Shipping"))
 	bool bShipForBitcode;
 	
 	// Any additional linker flags to pass to the linker in non-shipping builds
diff --git a/Engine/Source/Runtime/IOS/IOSRuntimeSettings/Private/IOSRuntimeSettings.cpp b/Engine/Source/Runtime/IOS/IOSRuntimeSettings/Private/IOSRuntimeSettings.cpp
--- a/Engine/Source/Runtime/IOS/IOSRuntimeSettings/Private/IOSRuntimeSettings.cpp
+++ b/Engine/Source/Runtime/IOS/IOSRuntimeSettings/Private/IOSRuntimeSettings.cpp
@@ -19,6 +19,8 @@ UIOSRuntimeSettings::UIOSRuntimeSettings(const FObjectInitializer& ObjectInitial
 	bSupportsIPhone = true;
 	MinimumiOSVersion = EIOSVersion::IOS_7;
 	bGeneratedSYMFile = false;
+	bGeneratedSYMBundle = false;
+	bGenerateXCArchive = false;
 	bDevForArmV7 = true;
 	bDevForArm64 = false;
 	bDevForArmV7S = false;
diff --git a/Engine/Source/Runtime/Launch/Private/LaunchEngineLoop.cpp b/Engine/Source/Runtime/Launch/Private/LaunchEngineLoop.cpp
--- a/Engine/Source/Runtime/Launch/Private/LaunchEngineLoop.cpp
+++ b/Engine/Source/Runtime/Launch/Private/LaunchEngineLoop.cpp
@@ -18,6 +18,7 @@
 #include "VersionManifest.h"
 #include "UObject/DevObjectVersion.h"
 #include "HAL/ThreadHeartBeat.h"
+#include "MallocProfiler.h"
 
 #include "NetworkVersion.h"
 
@@ -2428,8 +2429,6 @@ void FEngineLoop::Exit()
 	}
 #endif // WITH_ENGINE
 
-	MALLOC_PROFILER( GEngine->Exec( nullptr, TEXT( "MPROF STOP" ) ); )
-
 	if ( GEngine != nullptr )
 	{
 		GEngine->ShutdownAudioDeviceManager();
@@ -2455,6 +2454,9 @@ void FEngineLoop::Exit()
 
 	TermGamePhys();
 	ParticleVertexFactoryPool_FreePool();
+#else
+	// AppPreExit() stops malloc profiler, do it here instead
+	MALLOC_PROFILER( GMalloc->Exec(nullptr, TEXT("MPROF STOP"), *GLog);	);
 #endif // !ANDROID
 
 	// Stop the rendering thread.
@@ -3369,6 +3371,8 @@ void FEngineLoop::AppPreExit( )
 
 	FCoreDelegates::OnPreExit.Broadcast();
 
+	MALLOC_PROFILER( GMalloc->Exec(nullptr, TEXT("MPROF STOP"), *GLog);	);
+
 #if WITH_ENGINE
 	if (FString(FCommandLine::Get()).Contains(TEXT("CreatePak")) && GetDerivedDataCache())
 	{
diff --git a/Engine/Source/Runtime/NetworkFileSystem/Private/NetworkFileServerHttp.cpp b/Engine/Source/Runtime/NetworkFileSystem/Private/NetworkFileServerHttp.cpp
--- a/Engine/Source/Runtime/NetworkFileSystem/Private/NetworkFileServerHttp.cpp
+++ b/Engine/Source/Runtime/NetworkFileSystem/Private/NetworkFileServerHttp.cpp
@@ -118,6 +118,8 @@ void libwebsocket_debugLog(int level, const char *line)
 }
 #endif
 
+FNetworkFileServerHttp* user_space_patch = NULL;
+
 bool FNetworkFileServerHttp::Init()
 {
 	// setup log level.
@@ -151,6 +153,20 @@ bool FNetworkFileServerHttp::Init()
 		return false;
 	}
 
+	// ========================================
+	// March 11 2016 - nick.shin
+	// libwebsocket_create_context() above is not setting the user_space pointer properly
+	// this can be corrected with the following two code changes:
+//#include "private-libwebsockets.h" // put this in NetworkFileServerHttp.h
+//	Context->user_space = this;
+	// but this fails in UE4 frontend with multiple defined symbols...
+
+	// i will continue to analyze this problem after today's ZBR day
+
+	// since, there's only one of 'this' object -- for now...
+	user_space_patch = this;
+	// ========================================
+
 	Ready.Set(true);
 	return true;
 }
@@ -287,7 +303,8 @@ int FNetworkFileServerHttp::CallBack_HTTP(
 			size_t Len)
 {
 	PerSessionData* BufferInfo = (PerSessionData*)User;
-	FNetworkFileServerHttp* Server = (FNetworkFileServerHttp*)libwebsocket_context_user(Context);
+//	FNetworkFileServerHttp* Server = (FNetworkFileServerHttp*)libwebsocket_context_user(Context);
+	FNetworkFileServerHttp* Server = user_space_patch; // see Init() for details...
 
 	switch (Reason)
 	{
diff --git a/Engine/Source/Runtime/Online/HTTP/Private/Apple/AppleHTTP.cpp b/Engine/Source/Runtime/Online/HTTP/Private/Apple/AppleHTTP.cpp
--- a/Engine/Source/Runtime/Online/HTTP/Private/Apple/AppleHTTP.cpp
+++ b/Engine/Source/Runtime/Online/HTTP/Private/Apple/AppleHTTP.cpp
@@ -455,7 +455,7 @@ float FAppleHttpRequest::GetElapsedTime()
 		SecTrustRef PeerTrustInfo = reinterpret_cast<SecTrustRef>([[error userInfo] objectForKey:NSURLErrorFailingURLPeerTrustErrorKey]);
 		if (PeerTrustInfo != nullptr)
 		{
-			SecTrustResultType TrustResult = 0;
+			SecTrustResultType TrustResult = kSecTrustResultInvalid;
 			SecTrustGetTrustResult(PeerTrustInfo, &TrustResult);
 
 			FString TrustResultString;
diff --git a/Engine/Source/Runtime/RHI/Private/Apple/AppleDynamicRHI.cpp b/Engine/Source/Runtime/RHI/Private/Apple/AppleDynamicRHI.cpp
--- a/Engine/Source/Runtime/RHI/Private/Apple/AppleDynamicRHI.cpp
+++ b/Engine/Source/Runtime/RHI/Private/Apple/AppleDynamicRHI.cpp
@@ -7,7 +7,7 @@ int32 GMacMetalEnabled = 1;
 static FAutoConsoleVariableRef CVarMacMetalEnabled(
 	TEXT("r.Mac.UseMetal"),
 	GMacMetalEnabled,
-	TEXT("If set to true uses Metal when available rather than OpenGL as the graphics API. (Default: False)"));
+	TEXT("If set to true uses Metal when available rather than OpenGL as the graphics API. (Default: True)"));
 
 FDynamicRHI* PlatformCreateDynamicRHI()
 {
diff --git a/Engine/Source/Runtime/RHI/Public/RHIDefinitions.h b/Engine/Source/Runtime/RHI/Public/RHIDefinitions.h
--- a/Engine/Source/Runtime/RHI/Public/RHIDefinitions.h
+++ b/Engine/Source/Runtime/RHI/Public/RHIDefinitions.h
@@ -542,8 +542,6 @@ enum ETextureCreateFlags
 	TexCreate_NoFastClear = 1 << 25,
 	// Texture is a depth stencil resolve target
 	TexCreate_DepthStencilResolveTarget = 1 << 26,
-	// RenderTarget will create with delta color compression
-	TexCreate_DeltaColorCompression = 1 << 27,
 };
 
 enum EAsyncComputePriority
diff --git a/Engine/Source/Runtime/Renderer/Private/LightRendering.cpp b/Engine/Source/Runtime/Renderer/Private/LightRendering.cpp
--- a/Engine/Source/Runtime/Renderer/Private/LightRendering.cpp
+++ b/Engine/Source/Runtime/Renderer/Private/LightRendering.cpp
@@ -335,7 +335,7 @@ void FDeferredShadingSceneRenderer::RenderLights(FRHICommandListImmediate& RHICm
 	TArray<FSortedLightSceneInfo, SceneRenderingAllocator> SortedLights;
 	SortedLights.Empty(Scene->Lights.Num());
 
-	bool bDynamicShadows = ViewFamily.EngineShowFlags.DynamicShadows;
+	bool bDynamicShadows = ViewFamily.EngineShowFlags.DynamicShadows && GetShadowQuality() > 0;
 	
 	// Build a list of visible lights.
 	for (TSparseArray<FLightSceneInfoCompact>::TConstIterator LightIt(Scene->Lights); LightIt; ++LightIt)
diff --git a/Engine/Source/Runtime/Renderer/Private/PostProcess/SceneFilterRendering.cpp b/Engine/Source/Runtime/Renderer/Private/PostProcess/SceneFilterRendering.cpp
--- a/Engine/Source/Runtime/Renderer/Private/PostProcess/SceneFilterRendering.cpp
+++ b/Engine/Source/Runtime/Renderer/Private/PostProcess/SceneFilterRendering.cpp
@@ -174,7 +174,8 @@ void DrawRectangle(
 	FIntPoint TargetSize,
 	FIntPoint TextureSize,
 	FShader* VertexShader,
-	EDrawRectangleFlags Flags
+	EDrawRectangleFlags Flags,
+	uint32 InstanceCount
 	)
 {
 	float ClipSpaceQuadZ = 0.0f;
@@ -214,7 +215,7 @@ void DrawRectangle(
 			/*NumVertices=*/ GTesselatedScreenRectangleIndexBuffer.NumVertices(),
 			/*StartIndex=*/ 0,
 			/*NumPrimitives=*/ GTesselatedScreenRectangleIndexBuffer.NumPrimitives(),
-			/*NumInstances=*/ 1
+			/*NumInstances=*/ InstanceCount
 			);
 	}
 	else
@@ -235,7 +236,7 @@ void DrawRectangle(
 				/*NumVertices=*/ 3,
 				/*StartIndex=*/ 6,
 				/*NumPrimitives=*/ 1,
-				/*NumInstances=*/ 1
+				/*NumInstances=*/ InstanceCount
 				);
 		}
 		else
@@ -250,7 +251,7 @@ void DrawRectangle(
 				/*NumVertices=*/ 4,
 				/*StartIndex=*/ 0,
 				/*NumPrimitives=*/ 2,
-				/*NumInstances=*/ 1
+				/*NumInstances=*/ InstanceCount
 				);
 		}
 	}
diff --git a/Engine/Source/Runtime/Renderer/Private/PostProcess/SceneFilterRendering.h b/Engine/Source/Runtime/Renderer/Private/PostProcess/SceneFilterRendering.h
--- a/Engine/Source/Runtime/Renderer/Private/PostProcess/SceneFilterRendering.h
+++ b/Engine/Source/Runtime/Renderer/Private/PostProcess/SceneFilterRendering.h
@@ -23,6 +23,7 @@
  * TextureSize                  Size in texels of the source texture
  * VertexShader					The vertex shader used for rendering
  * Flags						see EDrawRectangleFlags
+ * InstanceCount				Number of instances of rectangle
  */
 extern void DrawRectangle(
 	FRHICommandList& RHICmdList,
@@ -37,7 +38,8 @@ extern void DrawRectangle(
 	FIntPoint TargetSize,
 	FIntPoint TextureSize,
 	class FShader* VertexShader,
-	EDrawRectangleFlags Flags = EDRF_Default
+	EDrawRectangleFlags Flags = EDRF_Default,
+	uint32 InstanceCount = 1
 	);
 
 extern void DrawTransformedRectangle(
diff --git a/Engine/Source/Runtime/Renderer/Private/PostProcess/SceneRenderTargets.cpp b/Engine/Source/Runtime/Renderer/Private/PostProcess/SceneRenderTargets.cpp
--- a/Engine/Source/Runtime/Renderer/Private/PostProcess/SceneRenderTargets.cpp
+++ b/Engine/Source/Runtime/Renderer/Private/PostProcess/SceneRenderTargets.cpp
@@ -661,7 +661,7 @@ void FSceneRenderTargets::AllocSceneColor(FRHICommandList& RHICmdList)
 
 	// Create the scene color.
 	{
-		FPooledRenderTargetDesc Desc(FPooledRenderTargetDesc::Create2DDesc(BufferSize, SceneColorBufferFormat, FClearValueBinding::Black, TexCreate_None, TexCreate_RenderTargetable | TexCreate_DeltaColorCompression, false));
+		FPooledRenderTargetDesc Desc(FPooledRenderTargetDesc::Create2DDesc(BufferSize, SceneColorBufferFormat, FClearValueBinding::Black, TexCreate_None, TexCreate_RenderTargetable, false));
 
 		Desc.Flags |= TexCreate_FastVRAM;
 
@@ -772,7 +772,7 @@ void FSceneRenderTargets::AllocGBufferTargets(FRHICommandList& RHICmdList)
 			NormalGBufferFormat = PF_FloatRGBA;
 		}
 
-		FPooledRenderTargetDesc Desc(FPooledRenderTargetDesc::Create2DDesc(BufferSize, NormalGBufferFormat, FClearValueBinding::Transparent, TexCreate_None, TexCreate_RenderTargetable | TexCreate_DeltaColorCompression, false));
+		FPooledRenderTargetDesc Desc(FPooledRenderTargetDesc::Create2DDesc(BufferSize, NormalGBufferFormat, FClearValueBinding::Transparent, TexCreate_None, TexCreate_RenderTargetable, false));
 		GRenderTargetPool.FindFreeElement(RHICmdList, Desc, GBufferA, TEXT("GBufferA"));
 	}
 
@@ -780,26 +780,26 @@ void FSceneRenderTargets::AllocGBufferTargets(FRHICommandList& RHICmdList)
 	{
 		const EPixelFormat SpecularGBufferFormat = bHighPrecisionGBuffers ? PF_FloatRGBA : PF_B8G8R8A8;
 
-		FPooledRenderTargetDesc Desc(FPooledRenderTargetDesc::Create2DDesc(BufferSize, SpecularGBufferFormat, FClearValueBinding::Transparent, TexCreate_None, TexCreate_RenderTargetable | TexCreate_DeltaColorCompression, false));
+		FPooledRenderTargetDesc Desc(FPooledRenderTargetDesc::Create2DDesc(BufferSize, SpecularGBufferFormat, FClearValueBinding::Transparent, TexCreate_None, TexCreate_RenderTargetable, false));
 		GRenderTargetPool.FindFreeElement(RHICmdList, Desc, GBufferB, TEXT("GBufferB"));
 	}
 
 	// Create the diffuse color g-buffer.
 	{
 		const EPixelFormat DiffuseGBufferFormat = bHighPrecisionGBuffers ? PF_FloatRGBA : PF_B8G8R8A8;
-		FPooledRenderTargetDesc Desc(FPooledRenderTargetDesc::Create2DDesc(BufferSize, DiffuseGBufferFormat, FClearValueBinding::Transparent, TexCreate_SRGB, TexCreate_RenderTargetable | TexCreate_DeltaColorCompression, false));
+		FPooledRenderTargetDesc Desc(FPooledRenderTargetDesc::Create2DDesc(BufferSize, DiffuseGBufferFormat, FClearValueBinding::Transparent, TexCreate_SRGB, TexCreate_RenderTargetable, false));
 		GRenderTargetPool.FindFreeElement(RHICmdList, Desc, GBufferC, TEXT("GBufferC"));
 	}
 
 	// Create the mask g-buffer (e.g. SSAO, subsurface scattering, wet surface mask, skylight mask, ...).
 	{
-		FPooledRenderTargetDesc Desc(FPooledRenderTargetDesc::Create2DDesc(BufferSize, PF_B8G8R8A8, FClearValueBinding(FLinearColor(0, 1, 1, 1)), TexCreate_None, TexCreate_RenderTargetable | TexCreate_DeltaColorCompression, false));
+		FPooledRenderTargetDesc Desc(FPooledRenderTargetDesc::Create2DDesc(BufferSize, PF_B8G8R8A8, FClearValueBinding(FLinearColor(0, 1, 1, 1)), TexCreate_None, TexCreate_RenderTargetable, false));
 		GRenderTargetPool.FindFreeElement(RHICmdList, Desc, GBufferD, TEXT("GBufferD"));
 	}
 
 	if (bAllowStaticLighting)
 	{
-		FPooledRenderTargetDesc Desc(FPooledRenderTargetDesc::Create2DDesc(BufferSize, PF_B8G8R8A8, FClearValueBinding(FLinearColor(1, 1, 1, 1)), TexCreate_None, TexCreate_RenderTargetable | TexCreate_DeltaColorCompression, false));
+		FPooledRenderTargetDesc Desc(FPooledRenderTargetDesc::Create2DDesc(BufferSize, PF_B8G8R8A8, FClearValueBinding(FLinearColor(1, 1, 1, 1)), TexCreate_None, TexCreate_RenderTargetable, false));
 		GRenderTargetPool.FindFreeElement(RHICmdList, Desc, GBufferE, TEXT("GBufferE"));
 	}
 
