Copying //UE4/Dev-Editor to //UE4/Dev-Main (Source: //UE4/Dev-Editor @ 3082391)

diff --git a/Engine/Source/Programs/AutomationTool/Scripts/Localisation.Automation.cs b/Engine/Source/Programs/AutomationTool/Scripts/Localisation.Automation.cs
--- a/Engine/Source/Programs/AutomationTool/Scripts/Localisation.Automation.cs
+++ b/Engine/Source/Programs/AutomationTool/Scripts/Localisation.Automation.cs
@@ -10,20 +10,40 @@ using UnrealBuildTool;
 using EpicGames.Localization;
 
 [Help("Updates the external localization data using the arguments provided.")]
-[Help("UEProjectDirectory", "Sub-path to the project we're gathering for (relative to CmdEnv.LocalRoot).")]
+[Help("UEProjectRoot", "Optional root-path to the project we're gathering for (defaults to CmdEnv.LocalRoot if unset).")]
+[Help("UEProjectDirectory", "Sub-path to the project we're gathering for (relative to UEProjectRoot).")]
 [Help("UEProjectName", "Optional name of the project we're gathering for (should match its .uproject file, eg QAGame).")]
 [Help("LocalizationProjectNames", "Comma separated list of the projects to gather text from.")]
 [Help("LocalizationBranch", "Optional suffix to use when uploading the new data to the localization provider.")]
 [Help("LocalizationProvider", "Optional localization provide override (default is OneSky).")]
 [Help("LocalizationSteps", "Optional comma separated list of localization steps to perform [Download, Gather, Import, Export, Compile, GenerateReports, Upload] (default is all). Only valid for projects using a modular config.")]
+[Help("IncludePlugins", "Optional flag to include plugins from within the given UEProjectDirectory as part of the gather. This may optionally specify a comma separated list of the specific plugins to gather (otherwise all plugins will be gathered).")]
+[Help("ExcludePlugins", "Optional comma separated list of plugins to exclude from the gather.")]
 [Help("AdditionalCommandletArguments", "Optional arguments to pass to the gather process.")]
 class Localise : BuildCommand
 {
+	private struct LocalizationBatch
+	{
+		public LocalizationBatch(string InUEProjectDirectory, string InRemoteFilenamePrefix, List<string> InLocalizationProjectNames)
+		{
+			UEProjectDirectory = InUEProjectDirectory;
+			RemoteFilenamePrefix = InRemoteFilenamePrefix;
+			LocalizationProjectNames = InLocalizationProjectNames;
+		}
+
+		public string UEProjectDirectory;
+		public string RemoteFilenamePrefix;
+		public List<string> LocalizationProjectNames;
+	};
+
 	public override void ExecuteBuild()
 	{
-		var EditorExe = CombinePaths(CmdEnv.LocalRoot, @"Engine/Binaries/Win64/UE4Editor-Cmd.exe");
+		var UEProjectRoot = ParseParamValue("UEProjectRoot");
+		if (UEProjectRoot == null)
+		{
+			UEProjectRoot = CmdEnv.LocalRoot;
+		}
 
-		// Parse out the required command line arguments
 		var UEProjectDirectory = ParseParamValue("UEProjectDirectory");
 		if (UEProjectDirectory == null)
 		{
@@ -39,13 +59,12 @@ class Localise : BuildCommand
 		var LocalizationProjectNames = new List<string>();
 		{
 			var LocalizationProjectNamesStr = ParseParamValue("LocalizationProjectNames");
-			if (LocalizationProjectNamesStr == null)
-			{
-				throw new AutomationException("Missing required command line argument: 'LocalizationProjectNames'");
-			}
-			foreach (var ProjectName in LocalizationProjectNamesStr.Split(','))
+			if (LocalizationProjectNamesStr != null)
 			{
-				LocalizationProjectNames.Add(ProjectName.Trim());
+				foreach (var ProjectName in LocalizationProjectNamesStr.Split(','))
+				{
+					LocalizationProjectNames.Add(ProjectName.Trim());
+				}
 			}
 		}
 
@@ -72,19 +91,88 @@ class Localise : BuildCommand
 			LocalizationSteps.Add("Monolithic"); // Always allow the monolithic scripts to run as we don't know which steps they do
 		}
 
+		var ShouldGatherPlugins = ParseParam("IncludePlugins");
+		var IncludePlugins = new List<string>();
+		var ExcludePlugins = new List<string>();
+		if (ShouldGatherPlugins)
+		{
+			var IncludePluginsStr = ParseParamValue("IncludePlugins");
+			if (IncludePluginsStr != null)
+			{
+				foreach (var PluginName in IncludePluginsStr.Split(','))
+				{
+					IncludePlugins.Add(PluginName.Trim());
+				}
+			}
+
+			var ExcludePluginsStr = ParseParamValue("ExcludePlugins");
+			if (ExcludePluginsStr != null)
+			{
+				foreach (var PluginName in ExcludePluginsStr.Split(','))
+				{
+					ExcludePlugins.Add(PluginName.Trim());
+				}
+			}
+		}
+
 		var AdditionalCommandletArguments = ParseParamValue("AdditionalCommandletArguments");
 		if (AdditionalCommandletArguments == null)
 		{
 			AdditionalCommandletArguments = "";
 		}
 
-		var RootWorkingDirectory = CombinePaths(CmdEnv.LocalRoot, UEProjectDirectory);
+		var LocalizationBatches = new List<LocalizationBatch>();
+
+		// Add the static set of localization projects as a batch
+		if (LocalizationProjectNames.Count > 0)
+		{
+			LocalizationBatches.Add(new LocalizationBatch(UEProjectDirectory, "", LocalizationProjectNames));
+		}
+
+		// Build up any additional batches needed for plugins
+		if (ShouldGatherPlugins)
+		{
+			var PluginsRootDirectory = CombinePaths(UEProjectRoot, UEProjectDirectory, "Plugins");
+			IReadOnlyList<PluginInfo> AllPlugins = Plugins.ReadPluginsFromDirectory(new DirectoryReference(PluginsRootDirectory), UEProjectName.Length == 0 ? PluginLoadedFrom.Engine : PluginLoadedFrom.GameProject);
+
+			// Add a batch for each plugin that meets our criteria
+			foreach (var PluginInfo in AllPlugins)
+			{
+				bool ShouldIncludePlugin = (IncludePlugins.Count == 0 || IncludePlugins.Contains(PluginInfo.Name)) && !ExcludePlugins.Contains(PluginInfo.Name);
+				if (ShouldIncludePlugin && PluginInfo.Descriptor.LocalizationTargets != null && PluginInfo.Descriptor.LocalizationTargets.Length > 0)
+				{
+					var RootRelativePluginPath = PluginInfo.Directory.MakeRelativeTo(new DirectoryReference(UEProjectRoot));
+					RootRelativePluginPath = RootRelativePluginPath.Replace('\\', '/'); // Make sure we use / as these paths are used with P4
+
+					var PluginTargetNames = new List<string>();
+					foreach (var LocalizationTarget in PluginInfo.Descriptor.LocalizationTargets)
+					{
+						PluginTargetNames.Add(LocalizationTarget.Name);
+					}
+
+					LocalizationBatches.Add(new LocalizationBatch(RootRelativePluginPath, PluginInfo.Name, PluginTargetNames));
+				}
+			}
+		}
+
+		// Process each localization batch
+		foreach (var LocalizationBatch in LocalizationBatches)
+		{
+			ProcessLocalizationProjects(LocalizationBatch, UEProjectRoot, UEProjectName, LocalizationProviderName, LocalizationSteps, AdditionalCommandletArguments);
+		}
+	}
+
+	private void ProcessLocalizationProjects(LocalizationBatch LocalizationBatch, string UEProjectRoot, string UEProjectName, string LocalizationProviderName, List<string> LocalizationSteps, string AdditionalCommandletArguments)
+	{
+		var EditorExe = CombinePaths(CmdEnv.LocalRoot, @"Engine/Binaries/Win64/UE4Editor-Cmd.exe");
+		var RootWorkingDirectory = CombinePaths(UEProjectRoot, LocalizationBatch.UEProjectDirectory);
 
 		// Try and find our localization provider
 		LocalizationProvider LocProvider = null;
 		{
 			LocalizationProvider.LocalizationProviderArgs LocProviderArgs;
 			LocProviderArgs.RootWorkingDirectory = RootWorkingDirectory;
+			LocProviderArgs.RemoteFilenamePrefix = LocalizationBatch.RemoteFilenamePrefix;
 			LocProviderArgs.CommandUtils = this;
 			LocProvider = LocalizationProvider.GetLocalizationProvider(LocalizationProviderName, LocProviderArgs);
 		}
@@ -93,13 +181,13 @@ class Localise : BuildCommand
 		if (P4Enabled)
 		{
 			Log("Sync necessary content to head revision");
-			P4.Sync(P4Env.BuildRootP4 + "/" + UEProjectDirectory + "/Config/Localization/...");
-			P4.Sync(P4Env.BuildRootP4 + "/" + UEProjectDirectory + "/Content/Localization/...");
+			P4.Sync(P4Env.BuildRootP4 + "/" + LocalizationBatch.UEProjectDirectory + "/Config/Localization/...");
+			P4.Sync(P4Env.BuildRootP4 + "/" + LocalizationBatch.UEProjectDirectory + "/Content/Localization/...");
 		}
 
 		// Generate the info we need to gather for each project
 		var ProjectInfos = new List<ProjectInfo>();
-		foreach (var ProjectName in LocalizationProjectNames)
+		foreach (var ProjectName in LocalizationBatch.LocalizationProjectNames)
 		{
 			ProjectInfos.Add(GenerateProjectInfo(RootWorkingDirectory, ProjectName, LocalizationSteps));
 		}
@@ -140,7 +228,7 @@ class Localise : BuildCommand
 					continue;
 				}
 
-				var CommandletArguments = String.Format("-config={0}", LocalizationStep.LocalizationConfigFile) + (String.IsNullOrEmpty(CommandletSCCArguments) ? "" : " " + CommandletSCCArguments);
+				var CommandletArguments = String.Format("-config=\"{0}\"", LocalizationStep.LocalizationConfigFile) + (String.IsNullOrEmpty(CommandletSCCArguments) ? "" : " " + CommandletSCCArguments);
 
 				if (!String.IsNullOrEmpty(AdditionalCommandletArguments))
 				{
@@ -190,13 +278,12 @@ class Localise : BuildCommand
 		// Projects generated by the localization dashboard will use multiple config files that must be run in a specific order
 		// Older projects (such as the Engine) would use a single config file containing all the steps
 		// Work out which kind of project we're dealing with...
-		var MonolithicConfigFile = String.Format(@"Config/Localization/{0}.ini", ProjectName);
-		var IsMonolithicConfig = File.Exists(CombinePaths(RootWorkingDirectory, MonolithicConfigFile));
-		if (IsMonolithicConfig)
+		var MonolithicConfigFile = CombinePaths(RootWorkingDirectory, String.Format(@"Config/Localization/{0}.ini", ProjectName));
+		if (File.Exists(MonolithicConfigFile))
 		{
 			ProjectInfo.LocalizationSteps.Add(new ProjectStepInfo("Monolithic", MonolithicConfigFile));
 
-			ProjectInfo.ImportInfo = GenerateProjectImportExportInfo(RootWorkingDirectory, MonolithicConfigFile);
+			ProjectInfo.ImportInfo = GenerateProjectImportExportInfo(MonolithicConfigFile);
 			ProjectInfo.ExportInfo = ProjectInfo.ImportInfo;
 		}
 		else
@@ -211,19 +298,19 @@ class Localise : BuildCommand
 
 			foreach (var FileSuffix in FileSuffixes)
 			{
-				var ModularConfigFile = String.Format(@"Config/Localization/{0}_{1}.ini", ProjectName, FileSuffix.Suffix);
+				var ModularConfigFile = CombinePaths(RootWorkingDirectory, String.Format(@"Config/Localization/{0}_{1}.ini", ProjectName, FileSuffix.Suffix));
 
-				if (File.Exists(CombinePaths(RootWorkingDirectory, ModularConfigFile)))
+				if (File.Exists(ModularConfigFile))
 				{
 					ProjectInfo.LocalizationSteps.Add(new ProjectStepInfo(FileSuffix.Suffix, ModularConfigFile));
 
 					if (FileSuffix.Suffix == "Import")
 					{
-						ProjectInfo.ImportInfo = GenerateProjectImportExportInfo(RootWorkingDirectory, ModularConfigFile);
+						ProjectInfo.ImportInfo = GenerateProjectImportExportInfo(ModularConfigFile);
 					}
 					else if (FileSuffix.Suffix == "Export")
 					{
-						ProjectInfo.ExportInfo = GenerateProjectImportExportInfo(RootWorkingDirectory, ModularConfigFile);
+						ProjectInfo.ExportInfo = GenerateProjectImportExportInfo(ModularConfigFile);
 					}
 				}
 				else if (FileSuffix.Required)
@@ -236,11 +323,11 @@ class Localise : BuildCommand
 		return ProjectInfo;
 	}
 
-	private ProjectImportExportInfo GenerateProjectImportExportInfo(string RootWorkingDirectory, string LocalizationConfigFile)
+	private ProjectImportExportInfo GenerateProjectImportExportInfo(string LocalizationConfigFile)
 	{
 		var ProjectImportExportInfo = new ProjectImportExportInfo();
 
-		var LocalizationConfig = new ConfigCacheIni(new FileReference(CombinePaths(RootWorkingDirectory, LocalizationConfigFile)));
+		var LocalizationConfig = new ConfigCacheIni(new FileReference(LocalizationConfigFile));
 
 		if (!LocalizationConfig.GetString("CommonSettings", "DestinationPath", out ProjectImportExportInfo.DestinationPath))
 		{
diff --git a/Engine/Source/Programs/CrashReporter/CrashReportClient/Private/CrashUpload.cpp b/Engine/Source/Programs/CrashReporter/CrashReportClient/Private/CrashUpload.cpp
--- a/Engine/Source/Programs/CrashReporter/CrashReportClient/Private/CrashUpload.cpp
+++ b/Engine/Source/Programs/CrashReporter/CrashReportClient/Private/CrashUpload.cpp
@@ -332,13 +332,10 @@ void FCrashUploadBase::SetCurrentState(EUploadState::Type InState)
 
 void FCrashUploadBase::AddReportToFailedList() const
 {
-	if (PendingFiles.Num() == 0)
+	if (PendingFiles.Num() > 0)
 	{
-		// No files so don't bother
-		return;
+		FailedReportDirectories.AddUnique(ErrorReport.GetReportDirectory());
 	}
-
-	FailedReportDirectories.AddUnique(PendingReportDirectories[PendingReportDirectoryIndex]);
 }
 
 // FCrashUploadToReceiver //////////////////////////////////////////////////////
diff --git a/Engine/Source/Programs/UnrealBuildTool/System/PluginDescriptor.cs b/Engine/Source/Programs/UnrealBuildTool/System/PluginDescriptor.cs
--- a/Engine/Source/Programs/UnrealBuildTool/System/PluginDescriptor.cs
+++ b/Engine/Source/Programs/UnrealBuildTool/System/PluginDescriptor.cs
@@ -88,6 +88,11 @@ namespace UnrealBuildTool
 		/// </summary>
 		public ModuleDescriptor[] Modules;
 
+		/// <summary>
+		/// List of all localization targets associated with this plugin
+		/// </summary>
+		public LocalizationTargetDescriptor[] LocalizationTargets;
+
 		/// <summary>
 		/// Whether this plugin should be enabled by default for all projects
 		/// </summary>
@@ -195,6 +200,12 @@ namespace UnrealBuildTool
 					Descriptor.Modules = Array.ConvertAll(ModulesArray, x => ModuleDescriptor.FromJsonObject(x));
 				}
 
+				JsonObject[] LocalizationTargetsArray;
+				if (RawObject.TryGetObjectArrayField("LocalizationTargets", out LocalizationTargetsArray))
+				{
+					Descriptor.LocalizationTargets = Array.ConvertAll(LocalizationTargetsArray, x => LocalizationTargetDescriptor.FromJsonObject(x));
+				}
+
 				RawObject.TryGetBoolField("EnabledByDefault", out Descriptor.bEnabledByDefault);
 				RawObject.TryGetBoolField("CanContainContent", out Descriptor.bCanContainContent);
 				RawObject.TryGetBoolField("IsBetaVersion", out Descriptor.bIsBetaVersion);
diff --git a/Engine/Source/Programs/UnrealBuildTool/UnrealBuildTool.csproj b/Engine/Source/Programs/UnrealBuildTool/UnrealBuildTool.csproj
--- a/Engine/Source/Programs/UnrealBuildTool/UnrealBuildTool.csproj
+++ b/Engine/Source/Programs/UnrealBuildTool/UnrealBuildTool.csproj
@@ -121,6 +121,7 @@
     <Compile Include="System\CustomBuildSteps.cs" />
     <Compile Include="System\FileSystemReference.cs" />
     <Compile Include="System\InstalledPlatformInfo.cs" />
+    <Compile Include="System\LocalizationDescriptor.cs" />
     <Compile Include="System\TargetReceipt.cs" />
     <Compile Include="System\VersionManifest.cs" />
     <Compile Include="ToolChain\AppleToolChain.cs" />
diff --git a/Engine/Source/Runtime/AppFramework/Private/Framework/Testing/SWidgetGallery.cpp b/Engine/Source/Runtime/AppFramework/Private/Framework/Testing/SWidgetGallery.cpp
--- a/Engine/Source/Runtime/AppFramework/Private/Framework/Testing/SWidgetGallery.cpp
+++ b/Engine/Source/Runtime/AppFramework/Private/Framework/Testing/SWidgetGallery.cpp
@@ -714,6 +714,14 @@ public:
 										.Text(LOCTEXT("STextBlockRightExampleLabel", "This is a right-justified text block"))
 										.Justification(ETextJustify::Right)
 								]
+
+								+ SVerticalBox::Slot()
+								.AutoHeight()
+								[
+									SNew(STextBlock)
+										.TextStyle(FCoreStyle::Get(), "NormalUnderlinedText")
+										.Text(LOCTEXT("STextBlockUnderlineLabel", "This is text is underlined using a line highlighter!"))
+								]
 							]
 
 						// STextComboBox
diff --git a/Engine/Source/Runtime/AppFramework/Private/Widgets/Colors/SColorThemes.cpp b/Engine/Source/Runtime/AppFramework/Private/Widgets/Colors/SColorThemes.cpp
--- a/Engine/Source/Runtime/AppFramework/Private/Widgets/Colors/SColorThemes.cpp
+++ b/Engine/Source/Runtime/AppFramework/Private/Widgets/Colors/SColorThemes.cpp
@@ -393,7 +393,7 @@ FText SThemeColorBlock::GetRedText() const { return ColorPtr.IsValid() ? FormatT
 FText SThemeColorBlock::GetGreenText() const { return ColorPtr.IsValid() ? FormatToolTipText(LOCTEXT("Green", "G"), ColorPtr.Pin()->HSVToLinearRGB().G) : FText::GetEmpty(); }
 FText SThemeColorBlock::GetBlueText() const { return ColorPtr.IsValid() ? FormatToolTipText(LOCTEXT("Blue", "B"), ColorPtr.Pin()->HSVToLinearRGB().B) : FText::GetEmpty(); }
 FText SThemeColorBlock::GetAlphaText() const { return ColorPtr.IsValid() ? FormatToolTipText(LOCTEXT("Alpha", "A"), ColorPtr.Pin()->HSVToLinearRGB().A) : FText::GetEmpty(); }
-FText SThemeColorBlock::GetHueText() const { return ColorPtr.IsValid() ? FormatToolTipText(LOCTEXT("Hue", "H"), ColorPtr.Pin()->R) : FText::GetEmpty(); }
+FText SThemeColorBlock::GetHueText() const { return ColorPtr.IsValid() ? FormatToolTipText(LOCTEXT("Hue", "H"), FMath::RoundToFloat(ColorPtr.Pin()->R)) : FText::GetEmpty(); }	// Rounded to let the value match the value in the Hue spinbox in the color picker
 FText SThemeColorBlock::GetSaturationText() const { return ColorPtr.IsValid() ? FormatToolTipText(LOCTEXT("Saturation", "S"), ColorPtr.Pin()->G) : FText::GetEmpty(); }
 FText SThemeColorBlock::GetValueText() const { return ColorPtr.IsValid() ? FormatToolTipText(LOCTEXT("Value", "V"), ColorPtr.Pin()->B) : FText::GetEmpty(); }
 
diff --git a/Engine/Source/Runtime/AssetRegistry/Private/AssetRegistry.cpp b/Engine/Source/Runtime/AssetRegistry/Private/AssetRegistry.cpp
--- a/Engine/Source/Runtime/AssetRegistry/Private/AssetRegistry.cpp
+++ b/Engine/Source/Runtime/AssetRegistry/Private/AssetRegistry.cpp
@@ -863,10 +863,12 @@ bool FAssetRegistry::GetDependencies(FName PackageName, TArray<FName>& OutDepend
 	}
 }
 
-bool FAssetRegistry::GetReferencers(FName PackageName, TArray<FName>& OutReferencers) const
+bool FAssetRegistry::GetReferencers(FName PackageName, TArray<FName>& OutReferencers, EAssetRegistryDependencyType::Type InReferenceType) const
 {
 	const FDependsNode*const* NodePtr = CachedDependsNodes.Find(PackageName);
 	const FDependsNode* Node = nullptr;
+	bool bShowingAllReferences = InReferenceType == EAssetRegistryDependencyType::All;
+
 	if (NodePtr != nullptr )
 	{
 		Node = *NodePtr;
@@ -879,7 +881,24 @@ bool FAssetRegistry::GetReferencers(FName PackageName, TArray<FName>& OutReferen
 
 		for ( auto NodeIt = DependencyNodes.CreateConstIterator(); NodeIt; ++NodeIt )
 		{
-			OutReferencers.Add((*NodeIt)->GetPackageName());
+			if (!bShowingAllReferences)
+			{
+				TArray<FDependsNode*> DependenciesFromReferencer;
+				(*NodeIt)->GetDependencies(DependenciesFromReferencer, InReferenceType);
+				
+				for (FDependsNode* Dependency : DependenciesFromReferencer)
+				{
+					if (Dependency == Node)
+					{
+						OutReferencers.Add((*NodeIt)->GetPackageName());
+						break;
+					}
+				}
+			}
+			else
+			{
+				OutReferencers.Add((*NodeIt)->GetPackageName());
+			}
 		}
 
 		return true;
diff --git a/Engine/Source/Runtime/AssetRegistry/Private/AssetRegistry.h b/Engine/Source/Runtime/AssetRegistry/Private/AssetRegistry.h
--- a/Engine/Source/Runtime/AssetRegistry/Private/AssetRegistry.h
+++ b/Engine/Source/Runtime/AssetRegistry/Private/AssetRegistry.h
@@ -22,7 +22,7 @@ public:
 	virtual FAssetData GetAssetByObjectPath( const FName ObjectPath ) const override;
 	virtual bool GetAllAssets(TArray<FAssetData>& OutAssetData, bool bIncludeOnlyOnDiskAssets = false) const override;
 	virtual bool GetDependencies(FName PackageName, TArray<FName>& OutDependencies, EAssetRegistryDependencyType::Type InDependencyType = EAssetRegistryDependencyType::All, bool bResolveIniStringReferences = false) const override;
-	virtual bool GetReferencers(FName PackageName, TArray<FName>& OutReferencers) const override;
+	virtual bool GetReferencers(FName PackageName, TArray<FName>& OutReferencers, EAssetRegistryDependencyType::Type InReferenceType) const override;
 	virtual bool GetAncestorClassNames(FName ClassName, TArray<FName>& OutAncestorClassNames) const override;
 	virtual void GetDerivedClassNames(const TArray<FName>& ClassNames, const TSet<FName>& ExcludedClassNames, TSet<FName>& OutDerivedClassNames) const override;
 	virtual void GetAllCachedPaths(TArray<FString>& OutPathList) const override;
diff --git a/Engine/Source/Runtime/AssetRegistry/Public/IAssetRegistry.h b/Engine/Source/Runtime/AssetRegistry/Public/IAssetRegistry.h
--- a/Engine/Source/Runtime/AssetRegistry/Public/IAssetRegistry.h
+++ b/Engine/Source/Runtime/AssetRegistry/Public/IAssetRegistry.h
@@ -110,7 +110,7 @@ public:
 	 * @param PackageName		the name of the package for which to gather dependencies
 	 * @param OutReferencers	a list of paths to objects that reference the package whose path is PackageName
 	 */
-	virtual bool GetReferencers(FName PackageName, TArray<FName>& OutReferencers) const = 0;
+	virtual bool GetReferencers(FName PackageName, TArray<FName>& OutReferencers, EAssetRegistryDependencyType::Type InReferenceType = EAssetRegistryDependencyType::All) const = 0;
 
 	/** Returns true if the specified ClassName's ancestors could be found. If so, OutAncestorClassNames is a list of all its ancestors */
 	virtual bool GetAncestorClassNames(FName ClassName, TArray<FName>& OutAncestorClassNames) const = 0;
diff --git a/Engine/Source/Runtime/AutomationMessages/Classes/AutomationWorkerMessages.h b/Engine/Source/Runtime/AutomationMessages/Classes/AutomationWorkerMessages.h
--- a/Engine/Source/Runtime/AutomationMessages/Classes/AutomationWorkerMessages.h
+++ b/Engine/Source/Runtime/AutomationMessages/Classes/AutomationWorkerMessages.h
@@ -249,6 +249,45 @@ struct FAutomationWorkerRunTests
 	{ }
 };
 
+USTRUCT()
+struct FAutomationWorkerEvent
+{
+	GENERATED_USTRUCT_BODY()
+
+	/** */
+	UPROPERTY(EditAnywhere, Category="Event")
+	FString Message;
+
+	/** */
+	UPROPERTY(EditAnywhere, Category="Event")
+	FString Context;
+
+	/** */
+	UPROPERTY(EditAnywhere, Category="Event")
+	FString Filename;
+
+	/** */
+	UPROPERTY(EditAnywhere, Category="Event")
+	int32 LineNumber;
+
+	FAutomationWorkerEvent()
+		: LineNumber(0)
+	{
+	}
+
+	FAutomationWorkerEvent(const FAutomationEvent& Event)
+		: Message(Event.Message)
+		, Context(Event.Context)
+		, Filename(Event.Filename)
+		, LineNumber(Event.LineNumber)
+	{
+	}
+
+	FAutomationEvent ToAutomationEvent() const
+	{
+		return FAutomationEvent(Message, Context, Filename, LineNumber);
+	}
+};
 
 /**
  * Implements a message that is sent in response to FAutomationWorkerRunTests.
@@ -264,7 +303,7 @@ struct FAutomationWorkerRunTestsReply
 
 	/** */
 	UPROPERTY(EditAnywhere, Category="Message")
-	TArray<FString> Errors;
+	TArray<FAutomationWorkerEvent> Errors;
 
 	/** */
 	UPROPERTY(EditAnywhere, Category="Message")
diff --git a/Engine/Source/Runtime/AutomationWorker/Private/AutomationWorkerModule.cpp b/Engine/Source/Runtime/AutomationWorker/Private/AutomationWorkerModule.cpp
--- a/Engine/Source/Runtime/AutomationWorker/Private/AutomationWorkerModule.cpp
+++ b/Engine/Source/Runtime/AutomationWorker/Private/AutomationWorkerModule.cpp
@@ -210,7 +210,10 @@ void FAutomationWorkerModule::ReportTestComplete()
 			Message->ExecutionCount = ExecutionCount;
 			Message->Success = bSuccess;
 			Message->Duration = ExecutionInfo.Duration;
-			Message->Errors = ExecutionInfo.Errors;
+			for ( auto& Error : ExecutionInfo.Errors )
+			{
+				Message->Errors.Add(FAutomationWorkerEvent(Error));
+			}
 			Message->Warnings = ExecutionInfo.Warnings;
 			Message->Logs = ExecutionInfo.LogItems;
 
diff --git a/Engine/Source/Runtime/Core/Private/GenericPlatform/GenericPlatformStackWalk.cpp b/Engine/Source/Runtime/Core/Private/GenericPlatform/GenericPlatformStackWalk.cpp
--- a/Engine/Source/Runtime/Core/Private/GenericPlatform/GenericPlatformStackWalk.cpp
+++ b/Engine/Source/Runtime/Core/Private/GenericPlatform/GenericPlatformStackWalk.cpp
@@ -156,3 +156,34 @@ void FGenericPlatformStackWalk::StackWalkAndDumpEx(ANSICHAR* HumanReadableString
 	// generic implementation ignores extra flags
 	return FPlatformStackWalk::StackWalkAndDump(HumanReadableString, HumanReadableStringSize, IgnoreCount, Context);
 }
+
+TArray<FProgramCounterSymbolInfo> FGenericPlatformStackWalk::GetStack(int32 IgnoreCount, int32 MaxDepth, void* Context)
+{
+	TArray<FProgramCounterSymbolInfo> Stack;
+
+	// Temporary memory holding the stack trace.
+	static const int MAX_DEPTH = 100;
+	uint64 StackTrace[MAX_DEPTH];
+	FMemory::Memzero(StackTrace);
+
+	// Add 2 to account for CaptureStackBackTrace and GetStack.
+	IgnoreCount += 2;
+
+	MaxDepth = FMath::Min(MAX_DEPTH, IgnoreCount + MaxDepth);
+
+	// Capture stack backtrace.
+	FPlatformStackWalk::CaptureStackBackTrace(StackTrace, MaxDepth, Context);
+
+	// Skip the first two entries as they are inside the stack walking code.
+	int32 CurrentDepth = IgnoreCount;
+	// Allow the first entry to be NULL as the crash could have been caused by a call to a NULL function pointer,
+	// which would mean the top of the callstack is NULL.
+	while ( CurrentDepth < MaxDepth && ( StackTrace[CurrentDepth] || ( CurrentDepth == IgnoreCount ) ) )
+	{
+		int32 NewIndex = Stack.AddDefaulted();
+		FPlatformStackWalk::ProgramCounterToSymbolInfo(StackTrace[CurrentDepth], Stack[NewIndex]);
+		CurrentDepth++;
+	}
+
+	return Stack;
+}
diff --git a/Engine/Source/Runtime/Core/Private/HAL/ConsoleManager.cpp b/Engine/Source/Runtime/Core/Private/HAL/ConsoleManager.cpp
--- a/Engine/Source/Runtime/Core/Private/HAL/ConsoleManager.cpp
+++ b/Engine/Source/Runtime/Core/Private/HAL/ConsoleManager.cpp
@@ -2070,6 +2070,15 @@ static TAutoConsoleVariable<int32> CVarSetVSyncEnabled(
 	TEXT("1: VSync is enabled."),
 	ECVF_Scalability | ECVF_RenderThreadSafe);
 
+#if WITH_EDITOR
+static TAutoConsoleVariable<int32> CVarSetVSyncEditorEnabled(
+	TEXT("r.VSyncEditor"),
+	0,
+	TEXT("0: VSync is disabled in editor.(default)\n")
+	TEXT("1: VSync is enabled in editor."),
+	ECVF_RenderThreadSafe);
+#endif
+
 static TAutoConsoleVariable<int32> CVarFinishCurrentFrame(
 	TEXT("r.FinishCurrentFrame"),
 	0,
diff --git a/Engine/Source/Runtime/Core/Private/Internationalization/ICUInternationalization.cpp b/Engine/Source/Runtime/Core/Private/Internationalization/ICUInternationalization.cpp
--- a/Engine/Source/Runtime/Core/Private/Internationalization/ICUInternationalization.cpp
+++ b/Engine/Source/Runtime/Core/Private/Internationalization/ICUInternationalization.cpp
@@ -147,11 +147,16 @@ bool FICUInternationalization::Initialize()
 	}
 	I18N->DefaultCulture = FindOrMakeCulture(FPlatformMisc::GetDefaultLocale(), EAllowDefaultCultureFallback::Yes);
 	SetCurrentCulture( I18N->GetDefaultCulture()->GetName() );
+
+	InitializeInvariantGregorianCalendar();
+
 	return U_SUCCESS(ICUStatus) ? true : false;
 }
 
 void FICUInternationalization::Terminate()
 {
+	InvariantGregorianCalendar.Reset();
+
 	FICUBreakIteratorManager::Destroy();
 	CachedCultures.Empty();
 
@@ -690,6 +695,34 @@ FCulturePtr FICUInternationalization::FindOrMakeCulture(const FString& Name, con
 	return NewCulture;
 }
 
+void FICUInternationalization::InitializeInvariantGregorianCalendar()
+{
+	UErrorCode ICUStatus = U_ZERO_ERROR;
+	InvariantGregorianCalendar = MakeUnique<icu::GregorianCalendar>(ICUStatus);
+	InvariantGregorianCalendar->setTimeZone(icu::TimeZone::getUnknown());
+}
+
+UDate FICUInternationalization::UEDateTimeToICUDate(const FDateTime& DateTime)
+{
+	// UE4 and ICU have a different time scale for pre-Gregorian dates, so we can't just use the UNIX timestamp from the UE4 DateTime
+	// Instead we have to explode the UE4 DateTime into its component parts, and then use an ICU GregorianCalendar (set to the "unknown" 
+	// timezone so it doesn't apply any adjustment to the time) to reconstruct the DateTime as an ICU UDate in the correct scale
+	int32 Year, Month, Day;
+	DateTime.GetDate(Year, Month, Day);
+	const int32 Hour = DateTime.GetHour();
+	const int32 Minute = DateTime.GetMinute();
+	const int32 Second = DateTime.GetSecond();
+
+	{
+		FScopeLock Lock(&InvariantGregorianCalendarCS);
+
+		InvariantGregorianCalendar->set(Year, Month - 1, Day, Hour, Minute, Second);
+		
+		UErrorCode ICUStatus = U_ZERO_ERROR;
+		return InvariantGregorianCalendar->getTime(ICUStatus);
+	}
+}
+
 UBool FICUInternationalization::OpenDataFile(const void* context, void** fileContext, void** contents, const char* path)
 {
 	auto& PathToCachedFileDataMap = FInternationalization::Get().Implementation->PathToCachedFileDataMap;
diff --git a/Engine/Source/Runtime/Core/Private/Internationalization/ICUInternationalization.h b/Engine/Source/Runtime/Core/Private/Internationalization/ICUInternationalization.h
--- a/Engine/Source/Runtime/Core/Private/Internationalization/ICUInternationalization.h
+++ b/Engine/Source/Runtime/Core/Private/Internationalization/ICUInternationalization.h
@@ -10,6 +10,7 @@
 	#pragma warning(disable:28253)
 #endif
 	#include <unicode/umachine.h>
+	#include <unicode/gregocal.h>
 #if defined(_MSC_VER) && USING_CODE_ANALYSIS
 	#pragma warning(pop)
 #endif
@@ -37,6 +38,8 @@ public:
 	TArray<FString> GetPrioritizedCultureNames(const FString& Name);
 	FCulturePtr GetCulture(const FString& Name);
 
+	UDate UEDateTimeToICUDate(const FDateTime& DateTime);
+
 private:
 #if NEEDS_ICU_DLLS
 	void LoadDLLs();
@@ -50,6 +53,8 @@ private:
 	enum class EAllowDefaultCultureFallback : uint8 { No, Yes, };
 	FCulturePtr FindOrMakeCulture(const FString& Name, const EAllowDefaultCultureFallback AllowDefaultFallback);
 
+	void InitializeInvariantGregorianCalendar();
+
 private:
 	struct FICUCultureData
 	{
@@ -87,6 +92,9 @@ private:
 	TMap<FString, FCultureRef> CachedCultures;
 	FCriticalSection CachedCulturesCS;
 
+	TUniquePtr<icu::GregorianCalendar> InvariantGregorianCalendar;
+	FCriticalSection InvariantGregorianCalendarCS;
+
 	static UBool OpenDataFile(const void* context, void** fileContext, void** contents, const char* path);
 	static void CloseDataFile(const void* context, void* const fileContext, void* const contents);
 
diff --git a/Engine/Source/Runtime/Core/Private/Internationalization/ICUText.cpp b/Engine/Source/Runtime/Core/Private/Internationalization/ICUText.cpp
--- a/Engine/Source/Runtime/Core/Private/Internationalization/ICUText.cpp
+++ b/Engine/Source/Runtime/Core/Private/Internationalization/ICUText.cpp
@@ -28,6 +28,7 @@
 
 #include "ICUUtilities.h"
 #include "ICUCulture.h"
+#include "ICUInternationalization.h"
 #include "ICUTextCharacterIterator.h"
 
 bool FText::IsWhitespace( const TCHAR Char )
@@ -44,8 +45,7 @@ FText FText::AsDate(const FDateTime& DateTime, const EDateTimeStyle::Type DateSt
 	checkf(I18N.IsInitialized() == true, TEXT("FInternationalization is not initialized. An FText formatting method was likely used in static object initialization - this is not supported."));
 	FCulturePtr Culture = TargetCulture.IsValid() ? TargetCulture : I18N.GetCurrentCulture();
 
-	int64 UNIXTimestamp = DateTime.ToUnixTimestamp();
-	UDate ICUDate = static_cast<double>(UNIXTimestamp) * U_MILLIS_PER_SECOND;
+	const UDate ICUDate = I18N.Implementation->UEDateTimeToICUDate(DateTime);
 
 	UErrorCode ICUStatus = U_ZERO_ERROR;
 	const TSharedRef<const icu::DateFormat> ICUDateFormat( Culture->Implementation->GetDateFormatter(DateStyle, TimeZone) );
@@ -64,8 +64,7 @@ FText FText::AsTime(const FDateTime& DateTime, const EDateTimeStyle::Type TimeSt
 	checkf(I18N.IsInitialized() == true, TEXT("FInternationalization is not initialized. An FText formatting method was likely used in static object initialization - this is not supported."));
 	FCulturePtr Culture = TargetCulture.IsValid() ? TargetCulture : I18N.GetCurrentCulture();
 
-	int64 UNIXTimestamp = DateTime.ToUnixTimestamp();
-	UDate ICUDate = static_cast<double>(UNIXTimestamp) * U_MILLIS_PER_SECOND;
+	const UDate ICUDate = I18N.Implementation->UEDateTimeToICUDate(DateTime);
 
 	UErrorCode ICUStatus = U_ZERO_ERROR;
 	const TSharedRef<const icu::DateFormat> ICUDateFormat( Culture->Implementation->GetTimeFormatter(TimeStyle, TimeZone) );
@@ -118,8 +117,7 @@ FText FText::AsDateTime(const FDateTime& DateTime, const EDateTimeStyle::Type Da
 	checkf(I18N.IsInitialized() == true, TEXT("FInternationalization is not initialized. An FText formatting method was likely used in static object initialization - this is not supported."));
 	FCulturePtr Culture = TargetCulture.IsValid() ? TargetCulture : I18N.GetCurrentCulture();
 
-	int64 UNIXTimestamp = DateTime.ToUnixTimestamp();
-	UDate ICUDate = static_cast<double>(UNIXTimestamp) * U_MILLIS_PER_SECOND;
+	const UDate ICUDate = I18N.Implementation->UEDateTimeToICUDate(DateTime);
 
 	UErrorCode ICUStatus = U_ZERO_ERROR;
 	const TSharedRef<const icu::DateFormat> ICUDateFormat( Culture->Implementation->GetDateTimeFormatter(DateStyle, TimeStyle, TimeZone) );
diff --git a/Engine/Source/Runtime/Core/Private/Internationalization/InternationalizationArchive.cpp b/Engine/Source/Runtime/Core/Private/Internationalization/InternationalizationArchive.cpp
--- a/Engine/Source/Runtime/Core/Private/Internationalization/InternationalizationArchive.cpp
+++ b/Engine/Source/Runtime/Core/Private/Internationalization/InternationalizationArchive.cpp
@@ -6,94 +6,107 @@
 
 DEFINE_LOG_CATEGORY_STATIC(LogGatherArchiveObject, Log, All);
 
-FArchiveEntry::FArchiveEntry( const FString& InNamespace, const FLocItem& InSource, const FLocItem& InTranslation, TSharedPtr<FLocMetadataObject> InKeyMetadataObj /*= NULL*/, bool IsOptional /*= false */ ) : Namespace( InNamespace )
-	, Source( InSource )
-	, Translation( InTranslation )
-	, bIsOptional( IsOptional )
+FArchiveEntry::FArchiveEntry(const FString& InNamespace, const FString& InKey, const FLocItem& InSource, const FLocItem& InTranslation, TSharedPtr<FLocMetadataObject> InKeyMetadataObj, bool IsOptional)
+	: Namespace(InNamespace)
+	, Key(InKey)
+	, Source(InSource)
+	, Translation(InTranslation)
+	, bIsOptional(IsOptional)
 {
-	if( InKeyMetadataObj.IsValid() )
+	if (InKeyMetadataObj.IsValid())
 	{
-		KeyMetadataObj = MakeShareable( new FLocMetadataObject( *(InKeyMetadataObj) ) );
+		KeyMetadataObj = MakeShareable(new FLocMetadataObject(*InKeyMetadataObj));
 	}
 }
 
-bool FInternationalizationArchive::AddEntry( const FString& Namespace, const FLocItem& Source, const FLocItem& Translation, const TSharedPtr<FLocMetadataObject> KeyMetadataObj, const bool bOptional )
+bool FInternationalizationArchive::AddEntry(const FString& Namespace, const FString& Key, const FLocItem& Source, const FLocItem& Translation, const TSharedPtr<FLocMetadataObject> KeyMetadataObj, const bool bOptional)
 {
+	if (Key.IsEmpty())
+	{
+		return false;
+	}
 
-	TSharedPtr<FArchiveEntry> ExistingEntry = FindEntryBySource( Namespace, Source, KeyMetadataObj );
+	TSharedPtr<FArchiveEntry> ExistingEntry = FindEntryByKey(Namespace, Key, KeyMetadataObj);
+	if (ExistingEntry.IsValid() && ExistingEntry->Source == Source)
+	{
+		return (ExistingEntry->Translation == Translation);
+	}
 
-	if( ExistingEntry.IsValid() )
+	TSharedRef<FArchiveEntry> NewEntry = MakeShareable(new FArchiveEntry(Namespace, Key, Source, Translation, KeyMetadataObj, bOptional));
+	if (ExistingEntry.IsValid())
+	{
+		UpdateEntry(ExistingEntry.ToSharedRef(), NewEntry);
+	}
+	else
 	{
-		return ( ExistingEntry->Translation == Translation );
+		EntriesBySourceText.Add(Source.Text, NewEntry);
+		EntriesByKey.Add(Key, NewEntry);
 	}
 
-	TSharedRef< FArchiveEntry > NewEntry = MakeShareable( new FArchiveEntry( Namespace, Source, Translation, KeyMetadataObj, bOptional ) );
-	EntriesBySourceText.Add( Source.Text, NewEntry );
 	return true;
 }
 
-bool FInternationalizationArchive::AddEntry( const TSharedRef<FArchiveEntry>& InEntry )
+bool FInternationalizationArchive::AddEntry(const TSharedRef<FArchiveEntry>& InEntry)
+{
+	return AddEntry(InEntry->Namespace, InEntry->Key, InEntry->Source, InEntry->Translation, InEntry->KeyMetadataObj, InEntry->bIsOptional);
+}
+
+void FInternationalizationArchive::UpdateEntry(const TSharedRef<FArchiveEntry>& OldEntry, const TSharedRef<FArchiveEntry>& NewEntry)
 {
-	return AddEntry( InEntry->Namespace, InEntry->Source, InEntry->Translation, InEntry->KeyMetadataObj, InEntry->bIsOptional );
+	EntriesBySourceText.RemoveSingle(OldEntry->Source.Text, OldEntry);
+	EntriesBySourceText.Add(NewEntry->Source.Text, NewEntry);
+
+	EntriesByKey.RemoveSingle(OldEntry->Key, OldEntry);
+	EntriesByKey.Add(NewEntry->Key, NewEntry);
 }
 
-bool FInternationalizationArchive::SetTranslation( const FString& Namespace, const FLocItem& Source, const FLocItem& Translation, const TSharedPtr<FLocMetadataObject> KeyMetadataObj )
+bool FInternationalizationArchive::SetTranslation(const FString& Namespace, const FString& Key, const FLocItem& Source, const FLocItem& Translation, const TSharedPtr<FLocMetadataObject> KeyMetadataObj)
 {
-	TSharedPtr<FArchiveEntry> Entry = FindEntryBySource( Namespace, Source, KeyMetadataObj );
-	if( Entry.IsValid() )
+	TSharedPtr<FArchiveEntry> Entry = FindEntryByKey(Namespace, Key, KeyMetadataObj);
+	if (Entry.IsValid())
 	{
-		Entry->Translation = Translation;
+		if (Entry->Source != Source)
+		{
+			UpdateEntry(Entry.ToSharedRef(), MakeShareable(new FArchiveEntry(Namespace, Key, Source, Translation, KeyMetadataObj, Entry->bIsOptional)));
+		}
+		else
+		{
+			Entry->Translation = Translation;
+		}
+
 		return true;
 	}
+
 	return false;
 }
 
-TSharedPtr< FArchiveEntry > FInternationalizationArchive::FindEntryBySource( const FString& Namespace, const FLocItem& Source, const TSharedPtr<FLocMetadataObject> KeyMetadataObj ) const
+TSharedPtr<FArchiveEntry> FInternationalizationArchive::FindEntryByKey(const FString& Namespace, const FString& Key, const TSharedPtr<FLocMetadataObject> KeyMetadataObj) const
 {
-	TArray< TSharedRef< FArchiveEntry > > MatchingEntries;
-	EntriesBySourceText.MultiFind( Source.Text, MatchingEntries );
+	TArray<TSharedRef<FArchiveEntry>> MatchingEntries;
+	EntriesByKey.MultiFind(Key, MatchingEntries);
 
-	for (int Index = 0; Index < MatchingEntries.Num(); Index++)
+	for (const TSharedRef<FArchiveEntry>& Entry : MatchingEntries)
 	{
-		if ( MatchingEntries[Index]->Namespace.Equals( Namespace, ESearchCase::CaseSensitive ) )
+		if (Entry->Namespace.Equals(Namespace, ESearchCase::CaseSensitive))
 		{
-			if( Source == MatchingEntries[Index]->Source )
+			if (!Entry->KeyMetadataObj.IsValid() && !KeyMetadataObj.IsValid())
 			{
-				if( !MatchingEntries[Index]->KeyMetadataObj.IsValid() && !KeyMetadataObj.IsValid() )
-				{
-					return MatchingEntries[Index];
-				}
-				else if( (KeyMetadataObj.IsValid() != MatchingEntries[Index]->KeyMetadataObj.IsValid()) )
-				{
-					// If we are in here, we know that one of the metadata entries is null, if the other
-					//  contains zero entries we will still consider them equivalent.
-					if( (KeyMetadataObj.IsValid() && KeyMetadataObj->Values.Num() == 0) ||
-						(MatchingEntries[Index]->KeyMetadataObj.IsValid() && MatchingEntries[Index]->KeyMetadataObj->Values.Num() == 0) )
-					{
-						return MatchingEntries[Index];
-					}
-				}
-				else if( *(MatchingEntries[Index]->KeyMetadataObj) == *(KeyMetadataObj) )
+				return Entry;
+			}
+			else if ((KeyMetadataObj.IsValid() != Entry->KeyMetadataObj.IsValid()))
+			{
+				// If we are in here, we know that one of the metadata entries is null, if the other contains zero entries we will still consider them equivalent.
+				if ((KeyMetadataObj.IsValid() && KeyMetadataObj->Values.Num() == 0) || (Entry->KeyMetadataObj.IsValid() && Entry->KeyMetadataObj->Values.Num() == 0))
 				{
-					return MatchingEntries[Index];
+					return Entry;
 				}
 			}
+			else if (*Entry->KeyMetadataObj == *KeyMetadataObj)
+			{
+				return Entry;
+			}
 		}
 	}
 
-	return NULL;
+	return nullptr;
 }
-
-void FInternationalizationArchive::UpdateEntry(const TSharedRef<FArchiveEntry>& OldEntry, const TSharedRef<FArchiveEntry>& NewEntry)
-{
-	EntriesBySourceText.RemoveSingle(OldEntry->Source.Text, OldEntry);
-	EntriesBySourceText.Add(NewEntry->Source.Text, NewEntry);
-}
-
-TArchiveEntryContainer::TConstIterator FInternationalizationArchive::GetEntryIterator() const
-{
-	return EntriesBySourceText.CreateConstIterator();
-}
-
-
-
diff --git a/Engine/Source/Runtime/Core/Private/Internationalization/InternationalizationManifest.cpp b/Engine/Source/Runtime/Core/Private/Internationalization/InternationalizationManifest.cpp
--- a/Engine/Source/Runtime/Core/Private/Internationalization/InternationalizationManifest.cpp
+++ b/Engine/Source/Runtime/Core/Private/Internationalization/InternationalizationManifest.cpp
@@ -6,58 +6,56 @@
 
 DEFINE_LOG_CATEGORY_STATIC(LogInternationalizationManifestObject, Log, All);
 
-FManifestContext::FManifestContext( const FManifestContext& Other ) 
-	: Key( Other.Key )
-	, SourceLocation( Other.SourceLocation )
-	, bIsOptional( Other.bIsOptional )
+FManifestContext::FManifestContext(const FManifestContext& Other)
+	: Key(Other.Key)
+	, SourceLocation(Other.SourceLocation)
+	, bIsOptional(Other.bIsOptional)
 {
-	if( Other.InfoMetadataObj.IsValid() )
+	if (Other.InfoMetadataObj.IsValid())
 	{
-		InfoMetadataObj = MakeShareable( new FLocMetadataObject( *(Other.InfoMetadataObj) ) );
+		InfoMetadataObj = MakeShareable(new FLocMetadataObject(*Other.InfoMetadataObj));
 	}
 
-	if( Other.KeyMetadataObj.IsValid() )
+	if (Other.KeyMetadataObj.IsValid())
 	{
-		KeyMetadataObj = MakeShareable( new FLocMetadataObject( *(Other.KeyMetadataObj) ) );
+		KeyMetadataObj = MakeShareable(new FLocMetadataObject(*Other.KeyMetadataObj));
 	}
 }
 
-FManifestContext& FManifestContext::operator=( const FManifestContext& Other )
+FManifestContext& FManifestContext::operator=(const FManifestContext& Other)
 {
-	if( this != &Other )
+	if (this != &Other)
 	{
 		Key = Other.Key;
 		SourceLocation = Other.SourceLocation;
 		bIsOptional = Other.bIsOptional;
 		InfoMetadataObj.Reset();
 		KeyMetadataObj.Reset();
-		if( Other.InfoMetadataObj.IsValid() )
+		if (Other.InfoMetadataObj.IsValid())
 		{
-			InfoMetadataObj = MakeShareable( new FLocMetadataObject( *(Other.InfoMetadataObj) ) );
+			InfoMetadataObj = MakeShareable(new FLocMetadataObject(*(Other.InfoMetadataObj)));
 		}
 
-		if( Other.KeyMetadataObj.IsValid() )
+		if (Other.KeyMetadataObj.IsValid())
 		{
-			KeyMetadataObj = MakeShareable( new FLocMetadataObject( *(Other.KeyMetadataObj) ) );
+			KeyMetadataObj = MakeShareable(new FLocMetadataObject(*(Other.KeyMetadataObj)));
 		}
 	}
 	return *this;
 }
 
-bool FManifestContext::operator==( const FManifestContext& Other ) const
+bool FManifestContext::operator==(const FManifestContext& Other) const
 {
-	if( Key.Equals( Other.Key, ESearchCase::CaseSensitive ) )
+	if (Key.Equals(Other.Key, ESearchCase::CaseSensitive))
 	{
-		if( !KeyMetadataObj.IsValid() && !Other.KeyMetadataObj.IsValid() )
+		if (!KeyMetadataObj.IsValid() && !Other.KeyMetadataObj.IsValid())
 		{
 			return true;
 		}
-		else if( KeyMetadataObj.IsValid() != Other.KeyMetadataObj.IsValid() )
+		else if (KeyMetadataObj.IsValid() != Other.KeyMetadataObj.IsValid())
 		{
-			// If we are in here, we know that one of the metadata entries is null, if the other
-			//  contains zero entries we will still consider them equivalent.
-			if( (KeyMetadataObj.IsValid() && KeyMetadataObj->Values.Num() == 0) ||
-				(Other.KeyMetadataObj.IsValid() && Other.KeyMetadataObj->Values.Num() == 0) )
+			// If we are in here, we know that one of the metadata entries is null, if the other contains zero entries we will still consider them equivalent.
+			if ((KeyMetadataObj.IsValid() && KeyMetadataObj->Values.Num() == 0) || (Other.KeyMetadataObj.IsValid() && Other.KeyMetadataObj->Values.Num() == 0))
 			{
 				return true;
 			}
@@ -71,23 +69,21 @@ bool FManifestContext::operator==( const FManifestContext& Other ) const
 	return false;
 }
 
-bool FManifestContext::operator<( const FManifestContext& Other ) const
+bool FManifestContext::operator<(const FManifestContext& Other) const
 {
-	int32 Result = Key.Compare( Other.Key, ESearchCase::CaseSensitive );
+	int32 Result = Key.Compare(Other.Key, ESearchCase::CaseSensitive);
 
 	bool bResult = false;
-	if( Result == -1 ) //Is Less than
+	if (Result == -1) //Is Less than
 	{
 		bResult = true;
 	}
-	else if ( Result == 0 ) //Is Equal
+	else if (Result == 0) //Is Equal
 	{
-		if( KeyMetadataObj.IsValid() != Other.KeyMetadataObj.IsValid() )
+		if (KeyMetadataObj.IsValid() != Other.KeyMetadataObj.IsValid())
 		{
-			// If we are in here, we know that one of the metadata entries is null, if the other
-			//  contains zero entries we will still consider them equivalent.
-			if( (KeyMetadataObj.IsValid() && KeyMetadataObj->Values.Num() == 0) ||
-				(Other.KeyMetadataObj.IsValid() && Other.KeyMetadataObj->Values.Num() == 0) )
+			// If we are in here, we know that one of the metadata entries is null, if the other contains zero entries we will still consider them equivalent.
+			if ((KeyMetadataObj.IsValid() && KeyMetadataObj->Values.Num() == 0) || (Other.KeyMetadataObj.IsValid() && Other.KeyMetadataObj->Values.Num() == 0))
 			{
 				return false;
 			}
@@ -96,7 +92,7 @@ bool FManifestContext::operator<( const FManifestContext& Other ) const
 				bResult = Other.KeyMetadataObj.IsValid();
 			}
 		}
-		else if( KeyMetadataObj.IsValid() && Other.KeyMetadataObj.IsValid() )
+		else if (KeyMetadataObj.IsValid() && Other.KeyMetadataObj.IsValid())
 		{
 			bResult = (*(KeyMetadataObj) < *(Other.KeyMetadataObj));
 		}
@@ -104,74 +100,63 @@ bool FManifestContext::operator<( const FManifestContext& Other ) const
 	return bResult;
 }
 
-FLocItem::FLocItem( const FLocItem& Other ) 
-	: Text( Other.Text )
+FLocItem::FLocItem(const FLocItem& Other)
+	: Text(Other.Text)
 {
-	if( Other.MetadataObj.IsValid() )
+	if (Other.MetadataObj.IsValid())
 	{
-		MetadataObj = MakeShareable( new FLocMetadataObject( *(Other.MetadataObj) ) );
+		MetadataObj = MakeShareable(new FLocMetadataObject(*Other.MetadataObj));
 	}
 }
 
-FLocItem& FLocItem::operator=( const FLocItem& Other )
+FLocItem& FLocItem::operator=(const FLocItem& Other)
 {
-	if( this != &Other )
+	if (this != &Other)
 	{
 		Text = Other.Text;
 		MetadataObj.Reset();
 
-		if( Other.MetadataObj.IsValid() )
+		if (Other.MetadataObj.IsValid())
 		{
-			MetadataObj = MakeShareable( new FLocMetadataObject( *(Other.MetadataObj) ) );
+			MetadataObj = MakeShareable(new FLocMetadataObject(*Other.MetadataObj));
 		}
 	}
 	return *this;
 }
 
-bool FLocItem::operator==( const FLocItem& Other ) const
+bool FLocItem::operator==(const FLocItem& Other) const
 {
-	if( Text.Equals( Other.Text, ESearchCase::CaseSensitive ) )
+	if (Text.Equals(Other.Text, ESearchCase::CaseSensitive))
 	{
-		if( !MetadataObj.IsValid() && !Other.MetadataObj.IsValid() )
+		if (!MetadataObj.IsValid() && !Other.MetadataObj.IsValid())
 		{
 			return true;
 		}
-		else if( MetadataObj.IsValid() != Other.MetadataObj.IsValid() )
+		else if (MetadataObj.IsValid() != Other.MetadataObj.IsValid())
 		{
-			// If we are in here, we know that one of the metadata entries is null, if the other
-			//  contains zero entries we will still consider them equivalent.
-			if( (MetadataObj.IsValid() && MetadataObj->Values.Num() == 0) ||
-				(Other.MetadataObj.IsValid() && Other.MetadataObj->Values.Num() == 0) )
-			{
-				return true;
-			}
-			else
-			{
-				return false;
-			}
+			// If we are in here, we know that one of the metadata entries is null, if the other contains zero entries we will still consider them equivalent.
+			return ((MetadataObj.IsValid() && MetadataObj->Values.Num() == 0) || (Other.MetadataObj.IsValid() && Other.MetadataObj->Values.Num() == 0));
 		}
-		return *(MetadataObj) == *(Other.MetadataObj);
+		return *MetadataObj == *Other.MetadataObj;
 	}
 	return false;
 }
 
-bool FLocItem::operator<( const FLocItem& Other ) const
+bool FLocItem::operator<(const FLocItem& Other) const
 {
-	int32 Result = Text.Compare( Other.Text, ESearchCase::CaseSensitive );
+	int32 Result = Text.Compare(Other.Text, ESearchCase::CaseSensitive);
 
 	bool bResult = false;
-	if( Result == -1 ) //Is less than
+	if (Result == -1) //Is less than
 	{
 		bResult = true;
 	}
-	else if ( Result == 0 ) //Is equal
+	else if (Result == 0) //Is equal
 	{
-		if( MetadataObj.IsValid() != Other.MetadataObj.IsValid() )
+		if (MetadataObj.IsValid() != Other.MetadataObj.IsValid())
 		{
-			// If we are in here, we know that one of the metadata entries is null, if the other
-			//  contains zero entries we will still consider them equivalent.
-			if( (MetadataObj.IsValid() && MetadataObj->Values.Num() == 0) ||
-				(Other.MetadataObj.IsValid() && Other.MetadataObj->Values.Num() == 0) )
+			// If we are in here, we know that one of the metadata entries is null, if the other contains zero entries we will still consider them equivalent.
+			if ((MetadataObj.IsValid() && MetadataObj->Values.Num() == 0) || (Other.MetadataObj.IsValid() && Other.MetadataObj->Values.Num() == 0))
 			{
 				return false;
 			}
@@ -180,93 +165,114 @@ bool FLocItem::operator<( const FLocItem& Other ) const
 				bResult = Other.MetadataObj.IsValid();
 			}
 		}
-		else if( MetadataObj.IsValid() && Other.MetadataObj.IsValid() )
+		else if (MetadataObj.IsValid() && Other.MetadataObj.IsValid())
 		{
-			bResult = (*(MetadataObj) < *(Other.MetadataObj));
+			bResult = (*MetadataObj < *Other.MetadataObj);
 		}
 	}
 	return bResult;
 }
 
-bool FLocItem::IsExactMatch( const FLocItem& Other ) const
+bool FLocItem::IsExactMatch(const FLocItem& Other) const
 {
-	if( Text.Equals( Other.Text, ESearchCase::CaseSensitive ) )
+	if (Text.Equals(Other.Text, ESearchCase::CaseSensitive))
 	{
-		return FLocMetadataObject::IsMetadataExactMatch(MetadataObj.Get(), Other.MetadataObj.Get() );
+		return FLocMetadataObject::IsMetadataExactMatch(MetadataObj.Get(), Other.MetadataObj.Get());
 	}
 	return false;
 }
 
-bool FInternationalizationManifest::AddSource( const FString& Namespace, const FLocItem& Source, const FManifestContext& Context )
+bool FInternationalizationManifest::AddSource(const FString& Namespace, const FLocItem& Source, const FManifestContext& Context)
 {
-	TSharedPtr< FManifestEntry > ExistingEntry = FindEntryByContext( Namespace, Context );
+	if (Context.Key.IsEmpty())
+	{
+		return false;
+	}
 
-	if ( ExistingEntry.IsValid() )
+	TSharedPtr<FManifestEntry> ExistingEntry = FindEntryByContext(Namespace, Context);
+
+	if (ExistingEntry.IsValid())
 	{
-		return ( Source.IsExactMatch( ExistingEntry->Source ) );
+		return (Source.IsExactMatch(ExistingEntry->Source));
 	}
-			
-	ExistingEntry = FindEntryBySource( Namespace, Source );
 
-	if ( ExistingEntry.IsValid() && !Source.IsExactMatch( ExistingEntry->Source ) )
+	ExistingEntry = FindEntryBySource(Namespace, Source);
+
+	if (ExistingEntry.IsValid() && !Source.IsExactMatch(ExistingEntry->Source))
 	{
 		return false;
 	}
 
-	if ( !ExistingEntry.IsValid() )
+	if (!ExistingEntry.IsValid())
 	{
-		TSharedRef< FManifestEntry > NewEntry = MakeShareable( new FManifestEntry( Namespace, Source ) );
-		EntriesBySourceText.Add( NewEntry->Source.Text, NewEntry );
+		TSharedRef<FManifestEntry> NewEntry = MakeShareable(new FManifestEntry(Namespace, Source));
+		EntriesBySourceText.Add(NewEntry->Source.Text, NewEntry);
 
 		ExistingEntry = NewEntry;
 	}
 
-	ExistingEntry->Contexts.Add( Context );
-	EntriesByContextId.Add( Context.Key, ExistingEntry.ToSharedRef() );
-	return true;	
+	ExistingEntry->Contexts.Add(Context);
+	EntriesByKey.Add(Context.Key, ExistingEntry.ToSharedRef());
+	return true;
+}
+
+void FInternationalizationManifest::UpdateEntry(const TSharedRef<FManifestEntry>& OldEntry, TSharedRef<FManifestEntry>& NewEntry)
+{
+	for (const FManifestContext& Context : OldEntry->Contexts)
+	{
+		EntriesByKey.RemoveSingle(Context.Key, OldEntry);
+	}
+	for (const FManifestContext& Context : NewEntry->Contexts)
+	{
+		EntriesByKey.Add(Context.Key, NewEntry);
+	}
+
+	EntriesBySourceText.RemoveSingle(OldEntry->Source.Text, OldEntry);
+	EntriesBySourceText.Add(NewEntry->Source.Text, NewEntry);
 }
 
-TSharedPtr< FManifestEntry > FInternationalizationManifest::FindEntryBySource( const FString& Namespace, const FLocItem& Source ) const
-{	
-	TArray< TSharedRef< FManifestEntry > > MatchingEntries;
-	EntriesBySourceText.MultiFind( Source.Text, /*OUT*/MatchingEntries );
+TSharedPtr<FManifestEntry> FInternationalizationManifest::FindEntryBySource(const FString& Namespace, const FLocItem& Source) const
+{
+	TArray<TSharedRef<FManifestEntry>> MatchingEntries;
+	EntriesBySourceText.MultiFind(Source.Text, /*OUT*/MatchingEntries);
 
-	for (int Index = 0; Index < MatchingEntries.Num(); Index++)
+	for (const TSharedRef<FManifestEntry>& Entry : MatchingEntries)
 	{
-		if ( MatchingEntries[Index]->Source == Source && 
-			MatchingEntries[Index]->Namespace.Equals( Namespace, ESearchCase::CaseSensitive ) )
+		if (Entry->Source == Source && Entry->Namespace.Equals(Namespace, ESearchCase::CaseSensitive))
 		{
-			return MatchingEntries[Index];
+			return Entry;
 		}
 	}
-	return NULL;
+
+	return nullptr;
 }
 
-TSharedPtr< FManifestEntry > FInternationalizationManifest::FindEntryByContext( const FString& Namespace, const FManifestContext& Context ) const
+TSharedPtr<FManifestEntry> FInternationalizationManifest::FindEntryByContext(const FString& Namespace, const FManifestContext& Context) const
 {
-	TArray< TSharedRef< FManifestEntry > > MatchingEntries;
-	EntriesByContextId.MultiFind( Context.Key, MatchingEntries );
+	TArray<TSharedRef<FManifestEntry>> MatchingEntries;
+	EntriesByKey.MultiFind(Context.Key, MatchingEntries);
 
-	for (int Index = 0; Index < MatchingEntries.Num(); Index++)
+	for (const TSharedRef<FManifestEntry>& Entry : MatchingEntries)
 	{
-		if ( MatchingEntries[Index]->Namespace.Equals( Namespace, ESearchCase::CaseSensitive ) )
+		if (Entry->Namespace.Equals(Namespace, ESearchCase::CaseSensitive))
 		{
-			for( auto ContextIter = MatchingEntries[Index]->Contexts.CreateConstIterator(); ContextIter; ++ContextIter )
+			for (const FManifestContext& EntryContext : Entry->Contexts)
 			{
-				if( *ContextIter == Context )
+				if (EntryContext == Context)
 				{
-					return MatchingEntries[Index];
+					return Entry;
 				}
 			}
 		}
 	}
-	return NULL;
+
+	return nullptr;
 }
 
-TSharedPtr< FManifestEntry > FInternationalizationManifest::FindEntryByKey(const FString& Namespace, const FString& Key, const FString* SourceText) const
+TSharedPtr<FManifestEntry> FInternationalizationManifest::FindEntryByKey(const FString& Namespace, const FString& Key, const FString* SourceText) const
 {
 	TArray<TSharedRef<FManifestEntry>> MatchingEntries;
-	EntriesByContextId.MultiFind(Key, MatchingEntries);
+	EntriesByKey.MultiFind(Key, MatchingEntries);
 
 	for (const TSharedRef<FManifestEntry>& MatchingEntry : MatchingEntries)
 	{
@@ -279,42 +285,7 @@ TSharedPtr< FManifestEntry > FInternationalizationManifest::FindEntryByKey(const
 	return nullptr;
 }
 
-void FInternationalizationManifest::UpdateEntry(const TSharedRef<FManifestEntry>& OldEntry, TSharedRef<FManifestEntry>& NewEntry)
-{
-	for (const FManifestContext& Context : OldEntry->Contexts)
-	{
-		EntriesByContextId.RemoveSingle(Context.Key, OldEntry);
-	}
-	for (const FManifestContext& Context : NewEntry->Contexts)
-	{
-		EntriesByContextId.Add(Context.Key, NewEntry);
-	}
-
-	EntriesBySourceText.RemoveSingle(OldEntry->Source.Text, OldEntry);
-	EntriesBySourceText.Add(NewEntry->Source.Text, NewEntry);
-}
-
-FManifestContext* FManifestEntry::FindContext(const FString& ContextKey, const TSharedPtr<FLocMetadataObject>& KeyMetadata)
-{
-	return const_cast<FManifestContext*>(FindContextImpl(ContextKey, KeyMetadata));
-}
-
 const FManifestContext* FManifestEntry::FindContext(const FString& ContextKey, const TSharedPtr<FLocMetadataObject>& KeyMetadata) const
-{
-	return FindContextImpl(ContextKey, KeyMetadata);
-}
-
-FManifestContext* FManifestEntry::FindContextByKey(const FString& ContextKey)
-{
-	return const_cast<FManifestContext*>(FindContextByKeyImpl(ContextKey));
-}
-
-const FManifestContext* FManifestEntry::FindContextByKey(const FString& ContextKey) const
-{
-	return FindContextByKeyImpl(ContextKey);
-}
-
-const FManifestContext* FManifestEntry::FindContextImpl(const FString& ContextKey, const TSharedPtr<FLocMetadataObject>& KeyMetadata) const
 {
 	for (const FManifestContext& Context : Contexts)
 	{
@@ -334,7 +305,7 @@ const FManifestContext* FManifestEntry::FindContextImpl(const FString& ContextKe
 	return nullptr;
 }
 
-const FManifestContext* FManifestEntry::FindContextByKeyImpl(const FString& ContextKey) const
+const FManifestContext* FManifestEntry::FindContextByKey(const FString& ContextKey) const
 {
 	for (const FManifestContext& Context : Contexts)
 	{
diff --git a/Engine/Source/Runtime/Core/Private/Internationalization/TextLocalizationManager.cpp b/Engine/Source/Runtime/Core/Private/Internationalization/TextLocalizationManager.cpp
--- a/Engine/Source/Runtime/Core/Private/Internationalization/TextLocalizationManager.cpp
+++ b/Engine/Source/Runtime/Core/Private/Internationalization/TextLocalizationManager.cpp
@@ -2,9 +2,6 @@
 
 #include "CorePrivatePCH.h"
 #include "Misc/App.h"
-#include "InternationalizationManifest.h"
-#include "TextLocalizationResourceGenerator.h"
-
 
 DEFINE_LOG_CATEGORY_STATIC(LogTextLocalizationManager, Log, All);
 
@@ -54,7 +51,7 @@ void EndInitTextLocalization()
 		else
 #if WITH_EDITOR
 			// See if we've been provided a culture override in the editor
-			if(GConfig->GetString( TEXT("Internationalization"), TEXT("Culture"), RequestedCultureName, GEditorSettingsIni ))
+			if(GIsEditor && GConfig->GetString( TEXT("Internationalization"), TEXT("Culture"), RequestedCultureName, GEditorSettingsIni ))
 			{
 				UE_LOG(LogInit, Log, TEXT("Overriding culture %s w/ editor configuration."), *RequestedCultureName);
 			}
@@ -525,153 +522,24 @@ bool FTextLocalizationManager::UpdateDisplayString(const FTextDisplayStringRef&
 	return true;
 }
 
-void FTextLocalizationManager::LoadFromManifestAndArchives( const FString& ConfigFilePath, IInternationalizationArchiveSerializer& ArchiveSerializer, IInternationalizationManifestSerializer& ManifestSerializer )
+void FTextLocalizationManager::UpdateFromLocalizationResource(const FString& LocalizationResourceFilePath)
 {
-	FScopeLock ScopeLock( &SynchronizationObject );
-
-	FInternationalization& I18N = FInternationalization::Get();
-
-	FString SectionName = TEXT("RegenerateResources");
-
-	// Get native culture.
-	FString NativeCulture;
-	if( !( GConfig->GetString( *SectionName, TEXT("NativeCulture"), NativeCulture, ConfigFilePath ) ) )
-	{
-		UE_LOG(LogTextLocalizationManager, Error, TEXT("No native culture specified."));
-		return;
-	}
-
-	// Get source path.
-	FString SourcePath;
-	if( !( GConfig->GetString( *SectionName, TEXT("SourcePath"), SourcePath, ConfigFilePath ) ) )
-	{
-		UE_LOG(LogTextLocalizationManager, Error, TEXT("No source path specified."));
-		return;
-	}
-
-	// Get destination path.
-	FString DestinationPath;
-	if( !( GConfig->GetString( *SectionName, TEXT("DestinationPath"), DestinationPath, ConfigFilePath ) ) )
-	{
-		UE_LOG(LogTextLocalizationManager, Error, TEXT("No destination path specified."));
-		return;
-	}
-
-	// Get manifest name.
-	FString ManifestName;
-	if( !( GConfig->GetString( *SectionName, TEXT("ManifestName"), ManifestName, ConfigFilePath ) ) )
-	{
-		UE_LOG(LogTextLocalizationManager, Error, TEXT("No manifest name specified."));
-		return;
-	}
-
-	// Get resource name.
-	FString ResourceName;
-	if( !( GConfig->GetString( *SectionName, TEXT("ResourceName"), ResourceName, ConfigFilePath ) ) )
-	{
-		UE_LOG(LogTextLocalizationManager, Error, TEXT("No resource name specified."));
-		return;
-	}
-
-	TArray<FString> LocaleNames;
-	{
-		const FString CultureName = I18N.GetCurrentCulture()->GetName();
-		LocaleNames.Add(CultureName);
-		const FString BaseLanguageName = I18N.GetCurrentCulture()->GetTwoLetterISOLanguageName();
-		if(BaseLanguageName != CultureName)
-		{
-			LocaleNames.Add(BaseLanguageName);
-		}
-	}
-
-	// Source path needs to be relative to Engine or Game directory
-	FString ConfigFullPath = FPaths::ConvertRelativePathToFull(ConfigFilePath);
-	FString EngineFullPath = FPaths::ConvertRelativePathToFull(FPaths::EngineConfigDir());
-	bool IsEngineManifest = false;
-
-	if (ConfigFullPath.StartsWith(EngineFullPath))
-	{
-		IsEngineManifest = true;
-	}
-
-	if (IsEngineManifest)
-	{
-		SourcePath = FPaths::Combine(*(FPaths::EngineDir()), *SourcePath);
-		DestinationPath = FPaths::Combine(*(FPaths::EngineDir()), *DestinationPath);
-	}
-	else
-	{
-		SourcePath = FPaths::Combine(*(FPaths::GameDir()), *SourcePath);
-		DestinationPath = FPaths::Combine(*(FPaths::GameDir()), *DestinationPath);
-	}
-
-	TArray<TArray<uint8>> BackingBuffers;
-	BackingBuffers.SetNum(LocaleNames.Num());
-
-	for(int32 i = 0; i < BackingBuffers.Num(); ++i)
-	{
-		TArray<uint8>& BackingBuffer = BackingBuffers[i];
-		FMemoryWriter MemoryWriter(BackingBuffer, true);
-
-		// Read the manifest file from the source path.
-		FString ManifestFilePath = (SourcePath / ManifestName);
-		ManifestFilePath = FPaths::ConvertRelativePathToFull(ManifestFilePath);
-		TSharedRef<FInternationalizationManifest> InternationalizationManifest = MakeShareable(new FInternationalizationManifest);
-
-#if 0 // @todo Json: Serializing from FArchive is currently broken
-		FArchive* ManifestFile = IFileManager::Get().CreateFileReader(*ManifestFilePath);
-
-		if (ManifestFile == nullptr)
-		{
-			UE_LOG(LogTextLocalizationManager, Error, TEXT("No manifest found at %s."), *ManifestFilePath);
-			return;
-		}
-
-		ManifestSerializer.DeserializeManifest(*ManifestFile, InternationalizationManifest);
-#else
-		FString ManifestContent;
-
-		if (!FFileHelper::LoadFileToString(ManifestContent, *ManifestFilePath))
-		{
-			UE_LOG(LogTextLocalizationManager, Error, TEXT("Failed to load file %s."), *ManifestFilePath);
-			continue;
-		}
-
-		ManifestSerializer.DeserializeManifest(ManifestContent, InternationalizationManifest);
-#endif
-
-		// Write resource.
-		FTextLocalizationResourceGenerator::Generate(SourcePath, InternationalizationManifest, NativeCulture, LocaleNames[i], &(MemoryWriter), ArchiveSerializer);
-
-		MemoryWriter.Close();
-	}
-
-	// Prioritized array of localization entry trackers.
 	TArray<FLocalizationEntryTracker> LocalizationEntryTrackers;
-	for(int32 i = 0; i < BackingBuffers.Num(); ++i)
-	{
-		TArray<uint8>& BackingBuffer = BackingBuffers[i];
-		FMemoryReader MemoryReader(BackingBuffer, true);
-		const FString CulturePath = DestinationPath / LocaleNames[i];
-		const FString ResourceFilePath = FPaths::ConvertRelativePathToFull(CulturePath / ResourceName);
-
-		FLocalizationEntryTracker& CultureTracker = LocalizationEntryTrackers[LocalizationEntryTrackers.Add(FLocalizationEntryTracker())];
-		CultureTracker.LoadFromLocalizationResource(MemoryReader, ResourceFilePath);
-		CultureTracker.DetectAndLogConflicts();
 
-		MemoryReader.Close();
-	}
+	FLocalizationEntryTracker& LocalizationEntryTracker = LocalizationEntryTrackers[LocalizationEntryTrackers.Add(FLocalizationEntryTracker())];
+	LocalizationEntryTracker.LoadFromFile(LocalizationResourceFilePath);
+	LocalizationEntryTracker.DetectAndLogConflicts();
 
 	UpdateFromLocalizations(LocalizationEntryTrackers);
 }
 
-void FTextLocalizationManager::UpdateFromLocalizationResource(const FString& LocalizationResourceFilePath)
+void FTextLocalizationManager::UpdateFromLocalizationResource(FArchive& LocResArchive, const FString& LocResID)
 {
 	TArray<FLocalizationEntryTracker> LocalizationEntryTrackers;
 
-	FLocalizationEntryTracker& LocalizationEntryTracker = LocalizationEntryTrackers[LocalizationEntryTrackers.Add(FLocalizationEntryTracker())];
-	LocalizationEntryTracker.LoadFromFile(LocalizationResourceFilePath);
-	LocalizationEntryTracker.DetectAndLogConflicts();
+	FLocalizationEntryTracker& CultureTracker = LocalizationEntryTrackers[LocalizationEntryTrackers.Add(FLocalizationEntryTracker())];
+	CultureTracker.LoadFromLocalizationResource(LocResArchive, LocResID);
+	CultureTracker.DetectAndLogConflicts();
 
 	UpdateFromLocalizations(LocalizationEntryTrackers);
 }
@@ -750,6 +618,10 @@ void FTextLocalizationManager::LoadLocalizationResourcesForCulture(const FString
 	TArray<FString> EngineLocalizationPaths;
 	EngineLocalizationPaths += FPaths::GetEngineLocalizationPaths();
 
+	// Gather any additional paths that are unknown to the UE4 core (such as plugins)
+	TArray<FString> AdditionalLocalizationPaths;
+	GatherAdditionalLocResPathsCallback.Broadcast(AdditionalLocalizationPaths);
+
 	// Prioritized array of localization entry trackers.
 	TArray<FLocalizationEntryTracker> LocalizationEntryTrackers;
 
@@ -803,6 +675,10 @@ void FTextLocalizationManager::LoadLocalizationResourcesForCulture(const FString
 	{
 		LocalizationPathToCultureDirectoryMap.Add(LocalizationPath, MapCulturesToDirectories(LocalizationPath));
 	}
+	for (const FString& LocalizationPath : AdditionalLocalizationPaths)
+	{
+		LocalizationPathToCultureDirectoryMap.Add(LocalizationPath, MapCulturesToDirectories(LocalizationPath));
+	}
 
 	TArray<FString> PrioritizedLocalizationPaths;
 	if (!GIsEditor)
@@ -811,6 +687,7 @@ void FTextLocalizationManager::LoadLocalizationResourcesForCulture(const FString
 	}
 	PrioritizedLocalizationPaths += EditorLocalizationPaths;
 	PrioritizedLocalizationPaths += EngineLocalizationPaths;
+	PrioritizedLocalizationPaths += AdditionalLocalizationPaths;
 
 	// The editor cheats and loads the native culture's localizations.
 	if (GIsEditor)
@@ -835,7 +712,6 @@ void FTextLocalizationManager::LoadLocalizationResourcesForCulture(const FString
 		}
 	}
 
-
 	// Read culture localization resources.
 	const TArray<FString> PrioritizedCultureNames = FInternationalization::Get().GetPrioritizedCultureNames(CultureName);
 	for (const FString& ParentCultureName : PrioritizedCultureNames)
diff --git a/Engine/Source/Runtime/Core/Private/Misc/AutomationTest.cpp b/Engine/Source/Runtime/Core/Private/Misc/AutomationTest.cpp
--- a/Engine/Source/Runtime/Core/Private/Misc/AutomationTest.cpp
+++ b/Engine/Source/Runtime/Core/Private/Misc/AutomationTest.cpp
@@ -8,7 +8,7 @@ DEFINE_LOG_CATEGORY_STATIC(LogAutomationTest, Warning, All);
 
 FString FAutomationTestInfo::GetTestAsString() const
 {
-	return FString::Printf(TEXT("%s,%s,%i,%d,%s,%s,%i"), *DisplayName, *TestName, TestFlags, NumParticipantsRequired, *TestParameter, *SourceFile, SourceFileLine);
+	return FString::Printf(TEXT("%s,%s,%i,%d,%s,%s,%i,%s,%s"), *DisplayName, *TestName, TestFlags, NumParticipantsRequired, *TestParameter, *SourceFile, SourceFileLine, *AssetPath, *OpenCommand);
 }
 
 void FAutomationTestInfo::ParseStringInfo(const FString& InTestInfo)
@@ -41,6 +41,16 @@ void FAutomationTestInfo::ParseStringInfo(const FString& InTestInfo)
 	{
 		SourceFileLine = FCString::Atoi(*Pieces[6]);
 	}
+
+	if ( Pieces.Num() >= 8 )
+	{
+		AssetPath = Pieces[7];
+	}
+
+	if ( Pieces.Num() >= 9 )
+	{
+		OpenCommand = Pieces[8];
+	}
 }
 
 
@@ -64,7 +74,7 @@ void FAutomationTestFramework::FAutomationTestFeedbackContext::Serialize( const
 			// If warnings should be treated as errors, log the warnings as such in the current unit test
 			if ( TreatWarningsAsErrors )
 			{
-				CurTest->AddError( FString( V ) );
+				CurTest->AddError(FString(V), 2);
 			}
 			else
 			{
@@ -74,7 +84,7 @@ void FAutomationTestFramework::FAutomationTestFeedbackContext::Serialize( const
 		// Errors
 		else if ( Verbosity == ELogVerbosity::Error )
 		{
-			CurTest->AddError( FString( V ) );
+			CurTest->AddError(FString(V), 2);
 		}
 		// Log items
 		else
@@ -600,9 +610,9 @@ void FAutomationTestFramework::DumpAutomationTestExecutionInfo( const TMap<FStri
 		{
 			SET_WARN_COLOR(COLOR_RED);
 			CLEAR_WARN_COLOR();
-			for ( TArray<FString>::TConstIterator ErrorIter( CurExecutionInfo.Errors ); ErrorIter; ++ErrorIter )
+			for ( TArray<FAutomationEvent>::TConstIterator ErrorIter( CurExecutionInfo.Errors ); ErrorIter; ++ErrorIter )
 			{
-				UE_LOG(LogAutomationTest, Error, TEXT("%s"), **ErrorIter);
+				UE_LOG(LogAutomationTest, Error, TEXT("%s"), *(*ErrorIter).Message);
 			}
 		}
 
@@ -712,7 +722,6 @@ void FAutomationTestFramework::AddAnalyticsItemToCurrentTest( const FString& Ana
 	}
 }
 
-
 FAutomationTestFramework::FAutomationTestFramework()
 :	CachedContext( NULL )
 ,	RequestedTestFilter(EAutomationTestFlags::SmokeFilter)
@@ -734,17 +743,43 @@ FAutomationTestFramework::~FAutomationTestFramework()
 }
 
 
+FString FAutomationEvent::ToString() const
+{
+	FString ComplexString;
+
+	if ( !Filename.IsEmpty() && LineNumber > 0 )
+	{
+		ComplexString += Filename;
+		ComplexString += TEXT("(");
+		ComplexString += FString::FromInt(LineNumber);
+		ComplexString += TEXT("): ");
+	}
+
+	if ( !Context.IsEmpty() )
+	{
+		ComplexString += Context;
+		ComplexString += TEXT(": ");
+	}
+
+	ComplexString += Message;
+
+	return ComplexString;
+}
+
+
 void FAutomationTestBase::ClearExecutionInfo()
 {
 	ExecutionInfo.Clear();
 }
 
 
-void FAutomationTestBase::AddError( const FString& InError )
+void FAutomationTestBase::AddError(const FString& InError, int32 StackOffset)
 {
 	if( !bSuppressLogs )
 	{
-		ExecutionInfo.Errors.Add( InError );
+		TArray<FProgramCounterSymbolInfo> Stack = FPlatformStackWalk::GetStack(StackOffset + 1, 1);
+
+		ExecutionInfo.Errors.Add(FAutomationEvent(InError, ExecutionInfo.Context, Stack[0].Filename, Stack[0].LineNumber));
 	}
 }
 
@@ -811,7 +846,17 @@ void FAutomationTestBase::GenerateTestNames(TArray<FAutomationTestInfo>& TestInf
 		}
 
 		// Add the test info to our collection
-		FAutomationTestInfo NewTestInfo(CompleteBeautifiedNames, CompleteTestName, GetTestFlags(), GetRequiredDeviceNum(), ParameterNames[ParameterIndex], GetFileName(), GetFileLine());
+		FAutomationTestInfo NewTestInfo(
+			CompleteBeautifiedNames,
+			CompleteTestName,
+			GetTestFlags(),
+			GetRequiredDeviceNum(),
+			ParameterNames[ParameterIndex],
+			GetTestSourceFileName(),
+			GetTestSourceFileLine(),
+			GetTestAssetPath(ParameterNames[ParameterIndex]),
+			GetTestOpenCommand(ParameterNames[ParameterIndex])
+		);
 		
 		TestInfo.Add( NewTestInfo );
 	}
diff --git a/Engine/Source/Runtime/Core/Private/Tests/Internationalization/ArchiveTest.cpp b/Engine/Source/Runtime/Core/Private/Tests/Internationalization/ArchiveTest.cpp
--- a/Engine/Source/Runtime/Core/Private/Tests/Internationalization/ArchiveTest.cpp
+++ b/Engine/Source/Runtime/Core/Private/Tests/Internationalization/ArchiveTest.cpp
@@ -45,6 +45,7 @@ bool FArchiveTest::RunTest( const FString& Parameters )
 	Translation.Text = TEXT("TranslatedTextA");
 
 	FString TestNamespace = TEXT("TestNamespace");
+	FString SourceAKey = TEXT("TextA");
 
 	// Test entry add
 	{
@@ -53,22 +54,22 @@ bool FArchiveTest::RunTest( const FString& Parameters )
 
 		// bIsOptional is not used as a key.  We ensure adding entries where the bIsOptional member is the only difference works as expected.
 		FInternationalizationArchive TestArchive;
-		TestArchive.AddEntry( TestNamespace, SourceA, Translation, NULL, TestOptionalTrue );
+		TestArchive.AddEntry( TestNamespace, SourceAKey, SourceA, Translation, NULL, TestOptionalTrue );
 		// Add duplicate entry that differs in bIsOptional flag.  This add should report success because we already have an entry with matching 
 		//  namespace/source/keymetadata.  Differences in bIsOptional are not taken into consideration. 
-		bool bResult = TestArchive.AddEntry( TestNamespace, SourceA, Translation, NULL, TestOptionalFalse );
+		bool bResult = TestArchive.AddEntry( TestNamespace, SourceAKey, SourceA, Translation, NULL, TestOptionalFalse );
 		TestTrue( TEXT("AddEntry result = true"), bResult);
 
 		// We should only have one entry in the archive
 		int32 EntryCount = 0;
-		for(auto Iter = TestArchive.GetEntryIterator(); Iter; ++Iter, ++EntryCount);
+		for(auto Iter = TestArchive.GetEntriesBySourceTextIterator(); Iter; ++Iter, ++EntryCount);
 		TestEqual( TEXT("EntryCount == 1"), EntryCount, 1 );
 
 		// Make sure the original bIsOptional value is not overwritten by our second add.
-		TSharedPtr< FArchiveEntry > FoundEntry = TestArchive.FindEntryBySource( TestNamespace, SourceA, NULL );
+		TSharedPtr< FArchiveEntry > FoundEntry = TestArchive.FindEntryByKey( TestNamespace, SourceAKey, NULL );
 		if( !FoundEntry.IsValid() )
 		{
-			AddError(TEXT("FArchiveEntry could not find entry using FindEntryBySource."));
+			AddError(TEXT("FArchiveEntry could not find entry using FindEntryByKey."));
 		}
 		else
 		{
@@ -80,12 +81,12 @@ bool FArchiveTest::RunTest( const FString& Parameters )
 	{
 		{
 			FInternationalizationArchive TestArchive;
-			TestArchive.AddEntry( TestNamespace, SourceA, Translation, KeyMetadataA, false );
+			TestArchive.AddEntry( TestNamespace, SourceAKey, SourceA, Translation, KeyMetadataA, false );
 
-			TSharedPtr< FArchiveEntry > FoundEntry = TestArchive.FindEntryBySource( TestNamespace, SourceA, KeyMetadataA );
+			TSharedPtr< FArchiveEntry > FoundEntry = TestArchive.FindEntryByKey( TestNamespace, SourceAKey, KeyMetadataA );
 			if( !FoundEntry.IsValid() )
 			{
-				AddError(TEXT("FArchiveEntry could not find entry using FindEntryBySource."));
+				AddError(TEXT("FArchiveEntry could not find entry using FindEntryByKey."));
 			}
 			else
 			{
@@ -100,25 +101,25 @@ bool FArchiveTest::RunTest( const FString& Parameters )
 			}
 
 			// Passing in a mismatched key metadata will fail to find the entry.  Any fallback logic is intended to happen at runtime
-			FoundEntry = TestArchive.FindEntryBySource( TestNamespace, SourceA, NULL );
+			FoundEntry = TestArchive.FindEntryByKey( TestNamespace, SourceAKey, NULL );
 			TestInvalid( TEXT("!FoundEntry.IsValid()"), FoundEntry);
 
-			FoundEntry = TestArchive.FindEntryBySource( TestNamespace, SourceA, MakeShareable( new FLocMetadataObject() ) );
+			FoundEntry = TestArchive.FindEntryByKey( TestNamespace, SourceAKey, MakeShareable( new FLocMetadataObject() ) );
 			TestInvalid( TEXT("!FoundEntry.IsValid()"), FoundEntry);
 
-			FoundEntry = TestArchive.FindEntryBySource( TestNamespace, SourceA, KeyMetadataB );
+			FoundEntry = TestArchive.FindEntryByKey( TestNamespace, SourceAKey, KeyMetadataB );
 			TestInvalid( TEXT("!FoundEntry.IsValid()"), FoundEntry);
 		}
 
 		// Ensure we can properly lookup entries with non-null but empty key metadata
 		{
 			FInternationalizationArchive TestArchive;
-			TestArchive.AddEntry( TestNamespace, SourceA, Translation, MakeShareable( new FLocMetadataObject() ), false );
+			TestArchive.AddEntry( TestNamespace, SourceAKey, SourceA, Translation, MakeShareable( new FLocMetadataObject() ), false );
 
-			TSharedPtr< FArchiveEntry > FoundEntry = TestArchive.FindEntryBySource( TestNamespace, SourceA, NULL );
+			TSharedPtr< FArchiveEntry > FoundEntry = TestArchive.FindEntryByKey( TestNamespace, SourceAKey, NULL );
 			if( !FoundEntry.IsValid() )
 			{
-				AddError(TEXT("FArchiveEntry could not find entry using FindEntryBySource."));
+				AddError(TEXT("FArchiveEntry could not find entry using FindEntryByKey."));
 			}
 			else
 			{
@@ -130,12 +131,12 @@ bool FArchiveTest::RunTest( const FString& Parameters )
 		// Ensure we can properly lookup entries with null key metadata
 		{
 			FInternationalizationArchive TestArchive;
-			TestArchive.AddEntry( TestNamespace, SourceA, Translation, NULL, false );
+			TestArchive.AddEntry( TestNamespace, SourceAKey, SourceA, Translation, NULL, false );
 
-			TSharedPtr< FArchiveEntry > FoundEntry = TestArchive.FindEntryBySource( TestNamespace, SourceA, NULL );
+			TSharedPtr< FArchiveEntry > FoundEntry = TestArchive.FindEntryByKey( TestNamespace, SourceAKey, NULL );
 			if( !FoundEntry.IsValid() )
 			{
-				AddError(TEXT("FArchiveEntry could not find entry using FindEntryBySource."));
+				AddError(TEXT("FArchiveEntry could not find entry using FindEntryByKey."));
 			}
 			else
 			{
@@ -144,10 +145,10 @@ bool FArchiveTest::RunTest( const FString& Parameters )
 			}
 
 			// Test lookup with non-null but empty key metadata
-			FoundEntry = TestArchive.FindEntryBySource( TestNamespace, SourceA, MakeShareable( new FLocMetadataObject() ) );
+			FoundEntry = TestArchive.FindEntryByKey( TestNamespace, SourceAKey, MakeShareable( new FLocMetadataObject() ) );
 			if( !FoundEntry.IsValid() )
 			{
-				AddError(TEXT("FArchiveEntry could not find entry using FindEntryBySource."));
+				AddError(TEXT("FArchiveEntry could not find entry using FindEntryByKey."));
 			}
 			else
 			{
@@ -165,13 +166,13 @@ bool FArchiveTest::RunTest( const FString& Parameters )
 
 			FInternationalizationArchive TestArchive;
 			// Added entry with String *IsMature entry
-			TestArchive.AddEntry( TestNamespace, SourceCompare, Translation, KeyMetadataA, false );
+			TestArchive.AddEntry( TestNamespace, SourceAKey, SourceCompare, Translation, KeyMetadataA, false );
 
 			// Finding entry using a source that has Boolean *IsMature
-			TSharedPtr< FArchiveEntry > FoundEntry = TestArchive.FindEntryBySource( TestNamespace, SourceA, KeyMetadataA );
+			TSharedPtr< FArchiveEntry > FoundEntry = TestArchive.FindEntryByKey( TestNamespace, SourceAKey, KeyMetadataA );
 			if( !FoundEntry.IsValid() )
 			{
-				AddError(TEXT("FArchiveEntry could not find entry using FindEntryBySource."));
+				AddError(TEXT("FArchiveEntry could not find entry using FindEntryByKey."));
 			}
 			else
 			{
@@ -181,20 +182,20 @@ bool FArchiveTest::RunTest( const FString& Parameters )
 
 			// Attempting to add an entry that only differs by a * prefexed source metadata entry value or type will result in success since
 			//  a 'matching' entry already exists in the archive.  We should, however, only have one entry in the archive
-			bool bResult = TestArchive.AddEntry( TestNamespace, SourceA, Translation, KeyMetadataA, false );
+			bool bResult = TestArchive.AddEntry( TestNamespace, SourceAKey, SourceA, Translation, KeyMetadataA, false );
 
 			TestTrue( TEXT("AddEntry result = true"), bResult);
 
 			// We should only have one entry in the archive
 			int32 EntryCount = 0;
-			for(auto Iter = TestArchive.GetEntryIterator(); Iter; ++Iter, ++EntryCount);
+			for(auto Iter = TestArchive.GetEntriesBySourceTextIterator(); Iter; ++Iter, ++EntryCount);
 			TestEqual( TEXT("EntryCount == 1"), EntryCount, 1 );
 
 			// Check to see that the original type/value of the * prefixed entry did not get modified
-			FoundEntry = TestArchive.FindEntryBySource( TestNamespace, SourceA, KeyMetadataA );
+			FoundEntry = TestArchive.FindEntryByKey( TestNamespace, SourceAKey, KeyMetadataA );
 			if( !FoundEntry.IsValid() )
 			{
-				AddError(TEXT("FArchiveEntry could not find entry using FindEntryBySource."));
+				AddError(TEXT("FArchiveEntry could not find entry using FindEntryByKey."));
 			}
 			else
 			{
diff --git a/Engine/Source/Runtime/Core/Private/Tests/Internationalization/DateTimeFormattingRulesTest.cpp b/Engine/Source/Runtime/Core/Private/Tests/Internationalization/DateTimeFormattingRulesTest.cpp
--- a/Engine/Source/Runtime/Core/Private/Tests/Internationalization/DateTimeFormattingRulesTest.cpp
+++ b/Engine/Source/Runtime/Core/Private/Tests/Internationalization/DateTimeFormattingRulesTest.cpp
@@ -33,36 +33,48 @@ bool FDateTimeFormattingRulesTest::RunTest (const FString& Parameters)
 	const FDateTime UnixBillennium = FDateTime::FromUnixTimestamp(1000000000);
 	const FDateTime UnixOnes = FDateTime::FromUnixTimestamp(1111111111);
 	const FDateTime UnixDecimalSequence = FDateTime::FromUnixTimestamp(1234567890);
+	const FDateTime YearOne( 1, 1, 1, 00, 00, 00, 000 );
 	const FDateTime TestDateTime( 1990, 6, 13, 12, 34, 56, 789 );
 
+	const FDateTime LocalTime = FDateTime::Now();
+	const FDateTime UtcTime = FDateTime::UtcNow();
+
 	if (I18N.SetCurrentCulture("en-US"))
 	{
 		// Unix Time Values via Date Time
-		Test( this, TEXT("Testing Unix Epoch"), FText::AsDateTime(UnixEpoch, EDateTimeStyle::Short, EDateTimeStyle::Short, "GMT"), FText::FromString( TEXT("1/1/70, 12:00 AM") ) );
-		Test( this, TEXT("Testing Unix Epoch"), FText::AsDateTime(UnixEpoch, EDateTimeStyle::Medium, EDateTimeStyle::Medium, "GMT"), FText::FromString( TEXT("Jan 1, 1970, 12:00:00 AM") ) );
-		Test( this, TEXT("Testing Unix Epoch"), FText::AsDateTime(UnixEpoch, EDateTimeStyle::Long, EDateTimeStyle::Long, "GMT"), FText::FromString( TEXT("January 1, 1970 at 12:00:00 AM GMT") ) );
-		Test( this, TEXT("Testing Unix Epoch"), FText::AsDateTime(UnixEpoch, EDateTimeStyle::Full, EDateTimeStyle::Full, "GMT"), FText::FromString( TEXT("Thursday, January 1, 1970 at 12:00:00 AM GMT") ) );
-
-		Test( this, TEXT("Testing Unix Billennium"), FText::AsDateTime(UnixBillennium, EDateTimeStyle::Short, EDateTimeStyle::Short, "GMT"), FText::FromString( TEXT("9/9/01, 1:46 AM") ) );
-		Test( this, TEXT("Testing Unix Billennium"), FText::AsDateTime(UnixBillennium, EDateTimeStyle::Medium, EDateTimeStyle::Medium, "GMT"), FText::FromString( TEXT("Sep 9, 2001, 1:46:40 AM") ) );
-		Test( this, TEXT("Testing Unix Billennium"), FText::AsDateTime(UnixBillennium, EDateTimeStyle::Long, EDateTimeStyle::Long, "GMT"), FText::FromString( TEXT("September 9, 2001 at 1:46:40 AM GMT") ) );
-		Test( this, TEXT("Testing Unix Billennium"), FText::AsDateTime(UnixBillennium, EDateTimeStyle::Full, EDateTimeStyle::Full, "GMT"), FText::FromString( TEXT("Sunday, September 9, 2001 at 1:46:40 AM GMT") ) );
-
-		Test( this, TEXT("Testing Unix Ones"), FText::AsDateTime(UnixOnes, EDateTimeStyle::Short, EDateTimeStyle::Short, "GMT"), FText::FromString( TEXT("3/18/05, 1:58 AM") ) );
-		Test( this, TEXT("Testing Unix Ones"), FText::AsDateTime(UnixOnes, EDateTimeStyle::Medium, EDateTimeStyle::Medium, "GMT"), FText::FromString( TEXT("Mar 18, 2005, 1:58:31 AM") ) );
-		Test( this, TEXT("Testing Unix Ones"), FText::AsDateTime(UnixOnes, EDateTimeStyle::Long, EDateTimeStyle::Long, "GMT"), FText::FromString( TEXT("March 18, 2005 at 1:58:31 AM GMT") ) );
-		Test( this, TEXT("Testing Unix Ones"), FText::AsDateTime(UnixOnes, EDateTimeStyle::Full, EDateTimeStyle::Full, "GMT"), FText::FromString( TEXT("Friday, March 18, 2005 at 1:58:31 AM GMT") ) );
-
-		Test( this, TEXT("Testing Unix Decimal Sequence"), FText::AsDateTime(UnixDecimalSequence, EDateTimeStyle::Short, EDateTimeStyle::Short, "GMT"), FText::FromString( TEXT("2/13/09, 11:31 PM") ) );
-		Test( this, TEXT("Testing Unix Decimal Sequence"), FText::AsDateTime(UnixDecimalSequence, EDateTimeStyle::Medium, EDateTimeStyle::Medium, "GMT"), FText::FromString( TEXT("Feb 13, 2009, 11:31:30 PM") ) );
-		Test( this, TEXT("Testing Unix Decimal Sequence"), FText::AsDateTime(UnixDecimalSequence, EDateTimeStyle::Long, EDateTimeStyle::Long, "GMT"), FText::FromString( TEXT("February 13, 2009 at 11:31:30 PM GMT") ) );
-		Test( this, TEXT("Testing Unix Decimal Sequence"), FText::AsDateTime(UnixDecimalSequence, EDateTimeStyle::Full, EDateTimeStyle::Full, "GMT"), FText::FromString( TEXT("Friday, February 13, 2009 at 11:31:30 PM GMT") ) );
+		Test( this, TEXT("Testing Unix Epoch"), FText::AsDateTime(UnixEpoch, EDateTimeStyle::Short, EDateTimeStyle::Short, FText::GetInvariantTimeZone()), FText::FromString( TEXT("1/1/70, 12:00 AM") ) );
+		Test( this, TEXT("Testing Unix Epoch"), FText::AsDateTime(UnixEpoch, EDateTimeStyle::Medium, EDateTimeStyle::Medium, FText::GetInvariantTimeZone()), FText::FromString( TEXT("Jan 1, 1970, 12:00:00 AM") ) );
+		Test( this, TEXT("Testing Unix Epoch"), FText::AsDateTime(UnixEpoch, EDateTimeStyle::Long, EDateTimeStyle::Long, FText::GetInvariantTimeZone()), FText::FromString( TEXT("January 1, 1970 at 12:00:00 AM GMT") ) );
+		Test( this, TEXT("Testing Unix Epoch"), FText::AsDateTime(UnixEpoch, EDateTimeStyle::Full, EDateTimeStyle::Full, FText::GetInvariantTimeZone()), FText::FromString( TEXT("Thursday, January 1, 1970 at 12:00:00 AM GMT") ) );
+
+		Test( this, TEXT("Testing Unix Billennium"), FText::AsDateTime(UnixBillennium, EDateTimeStyle::Short, EDateTimeStyle::Short, FText::GetInvariantTimeZone()), FText::FromString( TEXT("9/9/01, 1:46 AM") ) );
+		Test( this, TEXT("Testing Unix Billennium"), FText::AsDateTime(UnixBillennium, EDateTimeStyle::Medium, EDateTimeStyle::Medium, FText::GetInvariantTimeZone()), FText::FromString( TEXT("Sep 9, 2001, 1:46:40 AM") ) );
+		Test( this, TEXT("Testing Unix Billennium"), FText::AsDateTime(UnixBillennium, EDateTimeStyle::Long, EDateTimeStyle::Long, FText::GetInvariantTimeZone()), FText::FromString( TEXT("September 9, 2001 at 1:46:40 AM GMT") ) );
+		Test( this, TEXT("Testing Unix Billennium"), FText::AsDateTime(UnixBillennium, EDateTimeStyle::Full, EDateTimeStyle::Full, FText::GetInvariantTimeZone()), FText::FromString( TEXT("Sunday, September 9, 2001 at 1:46:40 AM GMT") ) );
+
+		Test( this, TEXT("Testing Unix Ones"), FText::AsDateTime(UnixOnes, EDateTimeStyle::Short, EDateTimeStyle::Short, FText::GetInvariantTimeZone()), FText::FromString( TEXT("3/18/05, 1:58 AM") ) );
+		Test( this, TEXT("Testing Unix Ones"), FText::AsDateTime(UnixOnes, EDateTimeStyle::Medium, EDateTimeStyle::Medium, FText::GetInvariantTimeZone()), FText::FromString( TEXT("Mar 18, 2005, 1:58:31 AM") ) );
+		Test( this, TEXT("Testing Unix Ones"), FText::AsDateTime(UnixOnes, EDateTimeStyle::Long, EDateTimeStyle::Long, FText::GetInvariantTimeZone()), FText::FromString( TEXT("March 18, 2005 at 1:58:31 AM GMT") ) );
+		Test( this, TEXT("Testing Unix Ones"), FText::AsDateTime(UnixOnes, EDateTimeStyle::Full, EDateTimeStyle::Full, FText::GetInvariantTimeZone()), FText::FromString( TEXT("Friday, March 18, 2005 at 1:58:31 AM GMT") ) );
+
+		Test( this, TEXT("Testing Unix Decimal Sequence"), FText::AsDateTime(UnixDecimalSequence, EDateTimeStyle::Short, EDateTimeStyle::Short, FText::GetInvariantTimeZone()), FText::FromString( TEXT("2/13/09, 11:31 PM") ) );
+		Test( this, TEXT("Testing Unix Decimal Sequence"), FText::AsDateTime(UnixDecimalSequence, EDateTimeStyle::Medium, EDateTimeStyle::Medium, FText::GetInvariantTimeZone()), FText::FromString( TEXT("Feb 13, 2009, 11:31:30 PM") ) );
+		Test( this, TEXT("Testing Unix Decimal Sequence"), FText::AsDateTime(UnixDecimalSequence, EDateTimeStyle::Long, EDateTimeStyle::Long, FText::GetInvariantTimeZone()), FText::FromString( TEXT("February 13, 2009 at 11:31:30 PM GMT") ) );
+		Test( this, TEXT("Testing Unix Decimal Sequence"), FText::AsDateTime(UnixDecimalSequence, EDateTimeStyle::Full, EDateTimeStyle::Full, FText::GetInvariantTimeZone()), FText::FromString( TEXT("Friday, February 13, 2009 at 11:31:30 PM GMT") ) );
+
+		Test( this, TEXT("Testing Year One"), FText::AsDateTime(YearOne, EDateTimeStyle::Short,EDateTimeStyle::Short, FText::GetInvariantTimeZone()), FText::FromString( TEXT("1/1/01, 12:00 AM") ) );
+		Test( this, TEXT("Testing Year One"), FText::AsDateTime(YearOne, EDateTimeStyle::Medium, EDateTimeStyle::Medium, FText::GetInvariantTimeZone()), FText::FromString( TEXT("Jan 1, 1, 12:00:00 AM") ) );
+		Test( this, TEXT("Testing Year One"), FText::AsDateTime(YearOne, EDateTimeStyle::Long, EDateTimeStyle::Long, FText::GetInvariantTimeZone()), FText::FromString( TEXT("January 1, 1 at 12:00:00 AM GMT") ) );
+		Test( this, TEXT("Testing Year One"), FText::AsDateTime(YearOne, EDateTimeStyle::Full, EDateTimeStyle::Full, FText::GetInvariantTimeZone()), FText::FromString( TEXT("Saturday, January 1, 1 at 12:00:00 AM GMT") ) );
 
 		// Date Time
-		Test( this, TEXT("Testing Date-Time"), FText::AsDateTime(TestDateTime, EDateTimeStyle::Short,EDateTimeStyle::Short, "GMT"), FText::FromString( TEXT("6/13/90, 12:34 PM") ) );
-		Test( this, TEXT("Testing Date-Time"), FText::AsDateTime(TestDateTime, EDateTimeStyle::Medium, EDateTimeStyle::Medium, "GMT"), FText::FromString( TEXT("Jun 13, 1990, 12:34:56 PM") ) );
-		Test( this, TEXT("Testing Date-Time"), FText::AsDateTime(TestDateTime, EDateTimeStyle::Long, EDateTimeStyle::Long, "GMT"), FText::FromString( TEXT("June 13, 1990 at 12:34:56 PM GMT") ) );
-		Test( this, TEXT("Testing Date-Time"), FText::AsDateTime(TestDateTime, EDateTimeStyle::Full, EDateTimeStyle::Full, "GMT"), FText::FromString( TEXT("Wednesday, June 13, 1990 at 12:34:56 PM GMT") ) );
+		Test( this, TEXT("Testing Date-Time"), FText::AsDateTime(TestDateTime, EDateTimeStyle::Short,EDateTimeStyle::Short, FText::GetInvariantTimeZone()), FText::FromString( TEXT("6/13/90, 12:34 PM") ) );
+		Test( this, TEXT("Testing Date-Time"), FText::AsDateTime(TestDateTime, EDateTimeStyle::Medium, EDateTimeStyle::Medium, FText::GetInvariantTimeZone()), FText::FromString( TEXT("Jun 13, 1990, 12:34:56 PM") ) );
+		Test( this, TEXT("Testing Date-Time"), FText::AsDateTime(TestDateTime, EDateTimeStyle::Long, EDateTimeStyle::Long, FText::GetInvariantTimeZone()), FText::FromString( TEXT("June 13, 1990 at 12:34:56 PM GMT") ) );
+		Test( this, TEXT("Testing Date-Time"), FText::AsDateTime(TestDateTime, EDateTimeStyle::Full, EDateTimeStyle::Full, FText::GetInvariantTimeZone()), FText::FromString( TEXT("Wednesday, June 13, 1990 at 12:34:56 PM GMT") ) );
+
+		// Checks that the default ICU timezone is set correctly (including DST)
+		Test( this, TEXT("Testing Local Time"), FText::AsDateTime(UtcTime, EDateTimeStyle::Short, EDateTimeStyle::Short), FText::AsDateTime(LocalTime, EDateTimeStyle::Short, EDateTimeStyle::Short, FText::GetInvariantTimeZone()) );
 	}
 	else
 	{
@@ -72,31 +84,39 @@ bool FDateTimeFormattingRulesTest::RunTest (const FString& Parameters)
 	if (I18N.SetCurrentCulture("ja-JP"))
 	{
 		// Unix Time Values via Date Time
-		Test( this, TEXT("Testing Unix Epoch"), FText::AsDateTime(UnixEpoch, EDateTimeStyle::Short, EDateTimeStyle::Short, "GMT"), FText::FromString( TEXT("1970/01/01 0:00") ) );
-		Test( this, TEXT("Testing Unix Epoch"), FText::AsDateTime(UnixEpoch, EDateTimeStyle::Medium, EDateTimeStyle::Medium, "GMT"), FText::FromString( TEXT("1970/01/01 0:00:00") ) );
-		Test( this, TEXT("Testing Unix Epoch"), FText::AsDateTime(UnixEpoch, EDateTimeStyle::Long, EDateTimeStyle::Long, "GMT"), FText::FromString( TEXT("1970\x5E74") TEXT("1\x6708") TEXT("1\x65E5") TEXT(" 0:00:00 GMT") ) );
-		Test( this, TEXT("Testing Unix Epoch"), FText::AsDateTime(UnixEpoch, EDateTimeStyle::Full, EDateTimeStyle::Full, "GMT"), FText::FromString( TEXT("1970\x5E74") TEXT("1\x6708") TEXT("1\x65E5\x6728\x66DC\x65E5") TEXT(" ") TEXT("0\x6642") TEXT("00\x5206") TEXT("00\x79D2") TEXT(" GMT") ) );
-
-		Test( this, TEXT("Testing Unix Billennium"), FText::AsDateTime(UnixBillennium, EDateTimeStyle::Short, EDateTimeStyle::Short, "GMT"), FText::FromString( TEXT("2001/09/09 1:46") ) );
-		Test( this, TEXT("Testing Unix Billennium"), FText::AsDateTime(UnixBillennium, EDateTimeStyle::Medium, EDateTimeStyle::Medium, "GMT"), FText::FromString( TEXT("2001/09/09 1:46:40") ) );
-		Test( this, TEXT("Testing Unix Billennium"), FText::AsDateTime(UnixBillennium, EDateTimeStyle::Long, EDateTimeStyle::Long, "GMT"), FText::FromString( TEXT("2001\x5E74") TEXT("9\x6708") TEXT("9\x65E5") TEXT(" 1:46:40 GMT") ) );
-		Test( this, TEXT("Testing Unix Billennium"), FText::AsDateTime(UnixBillennium, EDateTimeStyle::Full, EDateTimeStyle::Full, "GMT"), FText::FromString( TEXT("2001\x5E74") TEXT("9\x6708") TEXT("9\x65E5\x65E5\x66DC\x65E5") TEXT(" ") TEXT("1\x6642") TEXT("46\x5206") TEXT("40\x79D2") TEXT(" GMT") ) );
-
-		Test( this, TEXT("Testing Unix Ones"), FText::AsDateTime(UnixOnes, EDateTimeStyle::Short, EDateTimeStyle::Short, "GMT"), FText::FromString( TEXT("2005/03/18 1:58") ) );
-		Test( this, TEXT("Testing Unix Ones"), FText::AsDateTime(UnixOnes, EDateTimeStyle::Medium, EDateTimeStyle::Medium, "GMT"), FText::FromString( TEXT("2005/03/18 1:58:31") ) );
-		Test( this, TEXT("Testing Unix Ones"), FText::AsDateTime(UnixOnes, EDateTimeStyle::Long, EDateTimeStyle::Long, "GMT"), FText::FromString( TEXT("2005\x5E74") TEXT("3\x6708") TEXT("18\x65E5") TEXT(" 1:58:31 GMT") ) );
-		Test( this, TEXT("Testing Unix Ones"), FText::AsDateTime(UnixOnes, EDateTimeStyle::Full, EDateTimeStyle::Full, "GMT"), FText::FromString( TEXT("2005\x5E74") TEXT("3\x6708") TEXT("18\x65E5\x91D1\x66DC\x65E5") TEXT(" ") TEXT("1\x6642") TEXT("58\x5206") TEXT("31\x79D2") TEXT(" GMT") ) );
-
-		Test( this, TEXT("Testing Unix Decimal Sequence"), FText::AsDateTime(UnixDecimalSequence, EDateTimeStyle::Short, EDateTimeStyle::Short, "GMT"), FText::FromString( TEXT("2009/02/13 23:31") ) );
-		Test( this, TEXT("Testing Unix Decimal Sequence"), FText::AsDateTime(UnixDecimalSequence, EDateTimeStyle::Medium, EDateTimeStyle::Medium, "GMT"), FText::FromString( TEXT("2009/02/13 23:31:30") ) );
-		Test( this, TEXT("Testing Unix Decimal Sequence"), FText::AsDateTime(UnixDecimalSequence, EDateTimeStyle::Long, EDateTimeStyle::Long, "GMT"), FText::FromString( TEXT("2009\x5E74") TEXT("2\x6708") TEXT("13\x65E5 23:31:30 GMT") ) );
-		Test( this, TEXT("Testing Unix Decimal Sequence"), FText::AsDateTime(UnixDecimalSequence, EDateTimeStyle::Full, EDateTimeStyle::Full, "GMT"), FText::FromString( TEXT("2009\x5E74") TEXT("2\x6708") TEXT("13\x65E5\x91D1\x66DC\x65E5") TEXT(" ") TEXT("23\x6642") TEXT("31\x5206") TEXT("30\x79D2") TEXT(" GMT") ) );
+		Test( this, TEXT("Testing Unix Epoch"), FText::AsDateTime(UnixEpoch, EDateTimeStyle::Short, EDateTimeStyle::Short, FText::GetInvariantTimeZone()), FText::FromString( TEXT("1970/01/01 0:00") ) );
+		Test( this, TEXT("Testing Unix Epoch"), FText::AsDateTime(UnixEpoch, EDateTimeStyle::Medium, EDateTimeStyle::Medium, FText::GetInvariantTimeZone()), FText::FromString( TEXT("1970/01/01 0:00:00") ) );
+		Test( this, TEXT("Testing Unix Epoch"), FText::AsDateTime(UnixEpoch, EDateTimeStyle::Long, EDateTimeStyle::Long, FText::GetInvariantTimeZone()), FText::FromString( TEXT("1970\x5E74") TEXT("1\x6708") TEXT("1\x65E5") TEXT(" 0:00:00 GMT") ) );
+		Test( this, TEXT("Testing Unix Epoch"), FText::AsDateTime(UnixEpoch, EDateTimeStyle::Full, EDateTimeStyle::Full, FText::GetInvariantTimeZone()), FText::FromString( TEXT("1970\x5E74") TEXT("1\x6708") TEXT("1\x65E5\x6728\x66DC\x65E5") TEXT(" ") TEXT("0\x6642") TEXT("00\x5206") TEXT("00\x79D2") TEXT(" GMT") ) );
+
+		Test( this, TEXT("Testing Unix Billennium"), FText::AsDateTime(UnixBillennium, EDateTimeStyle::Short, EDateTimeStyle::Short, FText::GetInvariantTimeZone()), FText::FromString( TEXT("2001/09/09 1:46") ) );
+		Test( this, TEXT("Testing Unix Billennium"), FText::AsDateTime(UnixBillennium, EDateTimeStyle::Medium, EDateTimeStyle::Medium, FText::GetInvariantTimeZone()), FText::FromString( TEXT("2001/09/09 1:46:40") ) );
+		Test( this, TEXT("Testing Unix Billennium"), FText::AsDateTime(UnixBillennium, EDateTimeStyle::Long, EDateTimeStyle::Long, FText::GetInvariantTimeZone()), FText::FromString( TEXT("2001\x5E74") TEXT("9\x6708") TEXT("9\x65E5") TEXT(" 1:46:40 GMT") ) );
+		Test( this, TEXT("Testing Unix Billennium"), FText::AsDateTime(UnixBillennium, EDateTimeStyle::Full, EDateTimeStyle::Full, FText::GetInvariantTimeZone()), FText::FromString( TEXT("2001\x5E74") TEXT("9\x6708") TEXT("9\x65E5\x65E5\x66DC\x65E5") TEXT(" ") TEXT("1\x6642") TEXT("46\x5206") TEXT("40\x79D2") TEXT(" GMT") ) );
+
+		Test( this, TEXT("Testing Unix Ones"), FText::AsDateTime(UnixOnes, EDateTimeStyle::Short, EDateTimeStyle::Short, FText::GetInvariantTimeZone()), FText::FromString( TEXT("2005/03/18 1:58") ) );
+		Test( this, TEXT("Testing Unix Ones"), FText::AsDateTime(UnixOnes, EDateTimeStyle::Medium, EDateTimeStyle::Medium, FText::GetInvariantTimeZone()), FText::FromString( TEXT("2005/03/18 1:58:31") ) );
+		Test( this, TEXT("Testing Unix Ones"), FText::AsDateTime(UnixOnes, EDateTimeStyle::Long, EDateTimeStyle::Long, FText::GetInvariantTimeZone()), FText::FromString( TEXT("2005\x5E74") TEXT("3\x6708") TEXT("18\x65E5") TEXT(" 1:58:31 GMT") ) );
+		Test( this, TEXT("Testing Unix Ones"), FText::AsDateTime(UnixOnes, EDateTimeStyle::Full, EDateTimeStyle::Full, FText::GetInvariantTimeZone()), FText::FromString( TEXT("2005\x5E74") TEXT("3\x6708") TEXT("18\x65E5\x91D1\x66DC\x65E5") TEXT(" ") TEXT("1\x6642") TEXT("58\x5206") TEXT("31\x79D2") TEXT(" GMT") ) );
+
+		Test( this, TEXT("Testing Unix Decimal Sequence"), FText::AsDateTime(UnixDecimalSequence, EDateTimeStyle::Short, EDateTimeStyle::Short, FText::GetInvariantTimeZone()), FText::FromString( TEXT("2009/02/13 23:31") ) );
+		Test( this, TEXT("Testing Unix Decimal Sequence"), FText::AsDateTime(UnixDecimalSequence, EDateTimeStyle::Medium, EDateTimeStyle::Medium, FText::GetInvariantTimeZone()), FText::FromString( TEXT("2009/02/13 23:31:30") ) );
+		Test( this, TEXT("Testing Unix Decimal Sequence"), FText::AsDateTime(UnixDecimalSequence, EDateTimeStyle::Long, EDateTimeStyle::Long, FText::GetInvariantTimeZone()), FText::FromString( TEXT("2009\x5E74") TEXT("2\x6708") TEXT("13\x65E5 23:31:30 GMT") ) );
+		Test( this, TEXT("Testing Unix Decimal Sequence"), FText::AsDateTime(UnixDecimalSequence, EDateTimeStyle::Full, EDateTimeStyle::Full, FText::GetInvariantTimeZone()), FText::FromString( TEXT("2009\x5E74") TEXT("2\x6708") TEXT("13\x65E5\x91D1\x66DC\x65E5") TEXT(" ") TEXT("23\x6642") TEXT("31\x5206") TEXT("30\x79D2") TEXT(" GMT") ) );
+
+		Test( this, TEXT("Testing Year One"), FText::AsDateTime(YearOne, EDateTimeStyle::Short,EDateTimeStyle::Short, FText::GetInvariantTimeZone()), FText::FromString( TEXT("1/01/01 0:00") ) );
+		Test( this, TEXT("Testing Year One"), FText::AsDateTime(YearOne, EDateTimeStyle::Medium, EDateTimeStyle::Medium, FText::GetInvariantTimeZone()), FText::FromString( TEXT("1/01/01 0:00:00") ) );
+		Test( this, TEXT("Testing Year One"), FText::AsDateTime(YearOne, EDateTimeStyle::Long, EDateTimeStyle::Long, FText::GetInvariantTimeZone()), FText::FromString( TEXT("1\x5E74") TEXT("1\x6708") TEXT("1\x65E5") TEXT(" 0:00:00 GMT") ) );
+		Test( this, TEXT("Testing Year One"), FText::AsDateTime(YearOne, EDateTimeStyle::Full, EDateTimeStyle::Full, FText::GetInvariantTimeZone()), FText::FromString( TEXT("1\x5E74") TEXT("1\x6708") TEXT("1\x65E5\x571F\x66DC\x65E5") TEXT(" ") TEXT("0\x6642") TEXT("00\x5206") TEXT("00\x79D2 GMT") ) );
 
 		// Date Time
-		Test( this, TEXT("Testing Date-Time"), FText::AsDateTime(TestDateTime, EDateTimeStyle::Short,EDateTimeStyle::Short, "GMT"), FText::FromString( TEXT("1990/06/13 12:34") ) );
-		Test( this, TEXT("Testing Date-Time"), FText::AsDateTime(TestDateTime, EDateTimeStyle::Medium, EDateTimeStyle::Medium, "GMT"), FText::FromString( TEXT("1990/06/13 12:34:56") ) );
-		Test( this, TEXT("Testing Date-Time"), FText::AsDateTime(TestDateTime, EDateTimeStyle::Long, EDateTimeStyle::Long, "GMT"), FText::FromString( TEXT("1990\x5E74") TEXT("6\x6708") TEXT("13\x65E5") TEXT(" 12:34:56 GMT") ) );
-		Test( this, TEXT("Testing Date-Time"), FText::AsDateTime(TestDateTime, EDateTimeStyle::Full, EDateTimeStyle::Full, "GMT"), FText::FromString( TEXT("1990\x5E74") TEXT("6\x6708") TEXT("13\x65E5\x6C34\x66DC\x65E5") TEXT(" ") TEXT("12\x6642") TEXT("34\x5206") TEXT("56\x79D2 GMT") ) );
+		Test( this, TEXT("Testing Date-Time"), FText::AsDateTime(TestDateTime, EDateTimeStyle::Short,EDateTimeStyle::Short, FText::GetInvariantTimeZone()), FText::FromString( TEXT("1990/06/13 12:34") ) );
+		Test( this, TEXT("Testing Date-Time"), FText::AsDateTime(TestDateTime, EDateTimeStyle::Medium, EDateTimeStyle::Medium, FText::GetInvariantTimeZone()), FText::FromString( TEXT("1990/06/13 12:34:56") ) );
+		Test( this, TEXT("Testing Date-Time"), FText::AsDateTime(TestDateTime, EDateTimeStyle::Long, EDateTimeStyle::Long, FText::GetInvariantTimeZone()), FText::FromString( TEXT("1990\x5E74") TEXT("6\x6708") TEXT("13\x65E5") TEXT(" 12:34:56 GMT") ) );
+		Test( this, TEXT("Testing Date-Time"), FText::AsDateTime(TestDateTime, EDateTimeStyle::Full, EDateTimeStyle::Full, FText::GetInvariantTimeZone()), FText::FromString( TEXT("1990\x5E74") TEXT("6\x6708") TEXT("13\x65E5\x6C34\x66DC\x65E5") TEXT(" ") TEXT("12\x6642") TEXT("34\x5206") TEXT("56\x79D2 GMT") ) );
+
+		// Checks that the default ICU timezone is set correctly (including DST)
+		Test(this, TEXT("Testing Local Time"), FText::AsDateTime(UtcTime, EDateTimeStyle::Short, EDateTimeStyle::Short), FText::AsDateTime(LocalTime, EDateTimeStyle::Short, EDateTimeStyle::Short, FText::GetInvariantTimeZone()));
 	}
 	else
 	{
