Copying //UE4/Dev-Editor to //UE4/Dev-Main (Source: //UE4/Dev-Editor @ 3082391)

diff --git a/Engine/Source/Runtime/Core/Private/Tests/Internationalization/ManifestTest.cpp b/Engine/Source/Runtime/Core/Private/Tests/Internationalization/ManifestTest.cpp
--- a/Engine/Source/Runtime/Core/Private/Tests/Internationalization/ManifestTest.cpp
+++ b/Engine/Source/Runtime/Core/Private/Tests/Internationalization/ManifestTest.cpp
@@ -11,7 +11,7 @@
 int32 CountManifestEntries( const FInternationalizationManifest& Manifest )
 {
 	int32 EntryCounter = 0;
-	for( auto Iter = Manifest.GetEntriesByContextIdIterator(); Iter; ++Iter, ++EntryCounter );
+	for( auto Iter = Manifest.GetEntriesByKeyIterator(); Iter; ++Iter, ++EntryCounter );
 	return EntryCounter;
 }
 
@@ -94,8 +94,7 @@ bool FLocContextTest::RunTest( const FString& Parameters )
 
 		// Test assignment operator
 		{
-			FManifestContext ContextAClone;
-			ContextAClone = ContextA;
+			FManifestContext ContextAClone = ContextA;
 
 			if( ContextAClone.InfoMetadataObj == ContextA.InfoMetadataObj )
 			{
diff --git a/Engine/Source/Runtime/Core/Private/Windows/WindowsPlatformStackWalk.cpp b/Engine/Source/Runtime/Core/Private/Windows/WindowsPlatformStackWalk.cpp
--- a/Engine/Source/Runtime/Core/Private/Windows/WindowsPlatformStackWalk.cpp
+++ b/Engine/Source/Runtime/Core/Private/Windows/WindowsPlatformStackWalk.cpp
@@ -182,6 +182,12 @@ void FWindowsPlatformStackWalk::StackWalkAndDump( ANSICHAR* HumanReadableString,
 	FGenericPlatformStackWalk::StackWalkAndDump(HumanReadableString, HumanReadableStringSize, IgnoreCount, Context);
 }
 
+TArray<FProgramCounterSymbolInfo> FWindowsPlatformStackWalk::GetStack(int32 IgnoreCount, int32 MaxDepth, void* Context)
+{
+	InitStackWalking();
+	return FGenericPlatformStackWalk::GetStack(IgnoreCount + 1, MaxDepth, Context);
+}
+
 void FWindowsPlatformStackWalk::ThreadStackWalkAndDump(ANSICHAR* HumanReadableString, SIZE_T HumanReadableStringSize, int32 IgnoreCount, uint32 ThreadId)
 {
 	InitStackWalking();
diff --git a/Engine/Source/Runtime/Core/Public/GenericPlatform/GenericPlatformStackWalk.h b/Engine/Source/Runtime/Core/Public/GenericPlatform/GenericPlatformStackWalk.h
--- a/Engine/Source/Runtime/Core/Public/GenericPlatform/GenericPlatformStackWalk.h
+++ b/Engine/Source/Runtime/Core/Public/GenericPlatform/GenericPlatformStackWalk.h
@@ -70,7 +70,6 @@ struct FProgramCounterSymbolInfo final
 
 struct FProgramCounterSymbolInfoEx;
 
-
 /**
  * Generic implementation for most platforms
  */
@@ -181,6 +180,17 @@ struct CORE_API FGenericPlatformStackWalk
 	 * @param	Context				Optional thread context information
 	 */ 
 	static void StackWalkAndDump( ANSICHAR* HumanReadableString, SIZE_T HumanReadableStringSize, int32 IgnoreCount, void* Context = nullptr );
+	
+	/**
+	 * Walks the stack and updates the Stack array with the symbol information for each line in the stack.
+	 *
+	 * @param	IgnoreCount			Number of stack entries to ignore (some are guaranteed to be in the stack walking code)
+	 * @param	MaxDepth			The maximum depth to trace, can't be more than 100, offset from IgnoreCount.
+	 * @param	Context				Optional thread context information
+	 * 
+	 * @return	The stack of symbols to return.
+	 */ 
+	static TArray<FProgramCounterSymbolInfo> GetStack(int32 IgnoreCount, int32 MaxDepth = 100, void* Context = nullptr);
 
 	/**
 	* Walks the stack for the specified thread and appends the human readable string to the passed in one.
diff --git a/Engine/Source/Runtime/Core/Public/Internationalization/Internationalization.h b/Engine/Source/Runtime/Core/Public/Internationalization/Internationalization.h
--- a/Engine/Source/Runtime/Core/Public/Internationalization/Internationalization.h
+++ b/Engine/Source/Runtime/Core/Public/Internationalization/Internationalization.h
@@ -16,6 +16,8 @@ class FLegacyInternationalization;
 
 class FInternationalization
 {
+	friend class FText;
+
 public:
 	static CORE_API FInternationalization& Get();
 
diff --git a/Engine/Source/Runtime/Core/Public/Internationalization/InternationalizationArchive.h b/Engine/Source/Runtime/Core/Public/Internationalization/InternationalizationArchive.h
--- a/Engine/Source/Runtime/Core/Public/Internationalization/InternationalizationArchive.h
+++ b/Engine/Source/Runtime/Core/Public/Internationalization/InternationalizationArchive.h
@@ -7,24 +7,25 @@
 class CORE_API FArchiveEntry
 {
 public:
-	FArchiveEntry( const FString& InNamespace, const FLocItem& InSource, const FLocItem& InTranslation, TSharedPtr<FLocMetadataObject> InKeyMetadataObj = NULL, bool IsOptional = false );
+	FArchiveEntry(const FString& InNamespace, const FString& InKey, const FLocItem& InSource, const FLocItem& InTranslation, TSharedPtr<FLocMetadataObject> InKeyMetadataObj = nullptr, bool IsOptional = false);
 
 	const FString Namespace;
-	FLocItem Source;
+	const FString Key;
+	const FLocItem Source;
 	FLocItem Translation;
 	bool bIsOptional;
 	TSharedPtr<FLocMetadataObject> KeyMetadataObj;
 };
 
-struct FInternationalizationArchiveSourceTextKeyFuncs : BaseKeyFuncs<TSharedRef< FArchiveEntry >, FString, true>
+struct FInternationalizationArchiveStringKeyFuncs : BaseKeyFuncs<TSharedRef< FArchiveEntry >, FString, true>
 {
-	static FORCEINLINE const FString& GetSetKey( const TPair<FString,TSharedRef< FArchiveEntry >>& Element)
+	static FORCEINLINE const FString& GetSetKey(const TPair<FString, TSharedRef< FArchiveEntry >>& Element)
 	{
 		return Element.Key;
 	}
 	static FORCEINLINE bool Matches(const FString& A, const FString& B)
 	{
-		return A.Equals( B, ESearchCase::CaseSensitive );
+		return A.Equals(B, ESearchCase::CaseSensitive);
 	}
 	static FORCEINLINE uint32 GetKeyHash(const FString& Key)
 	{
@@ -32,17 +33,16 @@ struct FInternationalizationArchiveSourceTextKeyFuncs : BaseKeyFuncs<TSharedRef<
 	}
 };
 
-typedef TMultiMap< FString, TSharedRef< FArchiveEntry >, FDefaultSetAllocator, FInternationalizationArchiveSourceTextKeyFuncs > TArchiveEntryContainer;
+typedef TMultiMap< FString, TSharedRef< FArchiveEntry >, FDefaultSetAllocator, FInternationalizationArchiveStringKeyFuncs > FArchiveEntryByStringContainer;
 
 class CORE_API FInternationalizationArchive 
 {
-	friend class IInternationalizationArchiveSerializer;
-
 public:
-	enum EFormatVersion
+	enum class EFormatVersion : uint8
 	{
 		Initial = 0,
 		EscapeFixes,
+		AddedKeys,
 
 		LatestPlusOne,
 		Latest = LatestPlusOne - 1
@@ -50,24 +50,38 @@ public:
 
 	FInternationalizationArchive()
 		: FormatVersion(EFormatVersion::Latest)
-	{ }
+	{
+	}
+
+	bool AddEntry(const FString& Namespace, const FString& Key, const FLocItem& Source, const FLocItem& Translation, const TSharedPtr<FLocMetadataObject> KeyMetadataObj, const bool bOptional);
+	bool AddEntry(const TSharedRef<FArchiveEntry>& InEntry);
 
-	TArchiveEntryContainer::TConstIterator GetEntryIterator() const;
+	void UpdateEntry(const TSharedRef<FArchiveEntry>& OldEntry, const TSharedRef<FArchiveEntry>& NewEntry);
+
+	bool SetTranslation(const FString& Namespace, const FString& Key, const FLocItem& Source, const FLocItem& Translation, const TSharedPtr<FLocMetadataObject> KeyMetadataObj);
 
-	bool AddEntry( const FString& Namespace, const FLocItem& Source, const FLocItem& Translation, const TSharedPtr<FLocMetadataObject> KeyMetadataObj, const bool bOptional );
-	bool AddEntry( const TSharedRef<FArchiveEntry>& InEntry );
+	TSharedPtr<FArchiveEntry> FindEntryByKey(const FString& Namespace, const FString& Key, const TSharedPtr<FLocMetadataObject> KeyMetadataObj) const;
+
+	FArchiveEntryByStringContainer::TConstIterator GetEntriesByKeyIterator() const
+	{
+		return EntriesByKey.CreateConstIterator();
+	}
 
-	bool SetTranslation( const FString& Namespace, const FLocItem& Source, const FLocItem& Translation, const TSharedPtr<FLocMetadataObject> KeyMetadataObj );
+	int32 GetNumEntriesByKey() const
+	{
+		return EntriesByKey.Num();
+	}
 
-	TSharedPtr< FArchiveEntry > FindEntryBySource( const FString& Namespace, const FLocItem& Source, const TSharedPtr<FLocMetadataObject> KeyMetadataObj ) const;
+	FArchiveEntryByStringContainer::TConstIterator GetEntriesBySourceTextIterator() const
+	{
+		return EntriesBySourceText.CreateConstIterator();
+	}
 
 	int32 GetNumEntriesBySourceText() const
 	{
 		return EntriesBySourceText.Num();
 	}
 
-	void UpdateEntry(const TSharedRef<FArchiveEntry>& OldEntry, const TSharedRef<FArchiveEntry>& NewEntry);
-
 	void SetFormatVersion(const EFormatVersion Version)
 	{
 		FormatVersion = Version;
@@ -75,11 +89,11 @@ public:
 
 	EFormatVersion GetFormatVersion() const
 	{
-		return EFormatVersion(FormatVersion);
+		return FormatVersion;
 	}
 
 private:
-
 	EFormatVersion FormatVersion;
-	TArchiveEntryContainer EntriesBySourceText;
+	FArchiveEntryByStringContainer EntriesBySourceText;
+	FArchiveEntryByStringContainer EntriesByKey;
 };
diff --git a/Engine/Source/Runtime/Core/Public/Internationalization/InternationalizationManifest.h b/Engine/Source/Runtime/Core/Public/Internationalization/InternationalizationManifest.h
--- a/Engine/Source/Runtime/Core/Public/Internationalization/InternationalizationManifest.h
+++ b/Engine/Source/Runtime/Core/Public/Internationalization/InternationalizationManifest.h
@@ -9,54 +9,76 @@ struct CORE_API FManifestContext
 {
 public:
 	FManifestContext()
-		: bIsOptional(false)
-	{}
+		: Key()
+		, SourceLocation()
+		, bIsOptional(false)
+	{
+	}
+
+	explicit FManifestContext(FString InKey)
+		: Key(MoveTemp(InKey))
+		, SourceLocation()
+		, bIsOptional(false)
+	{
+	}
 
 	/** Copy ctor */
-	FManifestContext( const FManifestContext& Other );
+	FManifestContext(const FManifestContext& Other);
 
-	FManifestContext& operator=( const FManifestContext& Other );
+	FManifestContext& operator=(const FManifestContext& Other);
 	bool operator==(const FManifestContext& Other) const;
 	inline bool operator!=(const FManifestContext& Other) const { return !(*this == Other); }
 	bool operator<(const FManifestContext& Other) const;
 
 public:
-
 	FString Key;
 	FString SourceLocation;
 	bool bIsOptional;
 
 	TSharedPtr<FLocMetadataObject> InfoMetadataObj;
 	TSharedPtr<FLocMetadataObject> KeyMetadataObj;
-
 };
 
 
 struct CORE_API FLocItem
 {
 public:
-	FLocItem( const FString& InSourceText )
-		: Text( InSourceText )
-		, MetadataObj( NULL )
-	{ }
+	FLocItem()
+		: Text()
+		, MetadataObj(nullptr)
+	{
+	}
+
+	explicit FLocItem(FString InSourceText)
+		: Text(MoveTemp(InSourceText))
+		, MetadataObj(nullptr)
+	{
+	}
+
+	FLocItem(FString InSourceText, TSharedPtr<FLocMetadataObject> InMetadataObj)
+		: Text(MoveTemp(InSourceText))
+		, MetadataObj(MoveTemp(InMetadataObj))
+	{
+	}
 
 	/** Copy ctor */
-	FLocItem( const FLocItem& Other );
+	FLocItem(const FLocItem& Other);
 
-	FLocItem& operator=( const FLocItem& Other );
-	bool operator==( const FLocItem& Other ) const;
+	FLocItem& operator=(const FLocItem& Other);
+	bool operator==(const FLocItem& Other) const;
 	inline bool operator!=(const FLocItem& Other) const { return !(*this == Other); }
-	bool operator<( const FLocItem& Other ) const;
+	bool operator<(const FLocItem& Other) const;
 
 	/** Similar functionality to == operator but ensures everything matches(ex. ignores COMPARISON_MODIFIER_PREFIX on metadata). */
-	bool IsExactMatch( const FLocItem& Other ) const;
+	bool IsExactMatch(const FLocItem& Other) const;
+
 public:
 	FString Text;
 	TSharedPtr<FLocMetadataObject> MetadataObj;
 };
 
 
-class FManifestEntry
+class CORE_API FManifestEntry
 {
 public:
 	FManifestEntry(const FString& InNamespace, const FLocItem& InSource)
@@ -66,43 +88,18 @@ public:
 	{
 	}
 
-	CORE_API FManifestContext* FindContext(const FString& ContextKey, const TSharedPtr<FLocMetadataObject>& KeyMetadata = nullptr);
-	CORE_API const FManifestContext* FindContext(const FString& ContextKey, const TSharedPtr<FLocMetadataObject>& KeyMetadata = nullptr) const;
-
-	CORE_API FManifestContext* FindContextByKey(const FString& ContextKey);
-	CORE_API const FManifestContext* FindContextByKey(const FString& ContextKey) const;
+	const FManifestContext* FindContext(const FString& ContextKey, const TSharedPtr<FLocMetadataObject>& KeyMetadata = nullptr) const;
+	const FManifestContext* FindContextByKey(const FString& ContextKey) const;
 
 	const FString Namespace;
 	const FLocItem Source;
 	TArray<FManifestContext> Contexts;
-
-private:
-	const FManifestContext* FindContextImpl(const FString& ContextKey, const TSharedPtr<FLocMetadataObject>& KeyMetadata) const;
-	const FManifestContext* FindContextByKeyImpl(const FString& ContextKey) const;
-};
-
-
-struct FInternationalizationManifestSourceTextKeyFuncs : BaseKeyFuncs<TSharedRef< FManifestEntry >, FString, true>
-{
-	static FORCEINLINE const FString& GetSetKey(const TPair<FString,TSharedRef< FManifestEntry >>& Element)
-	{
-		return Element.Key;
-	}
-	static FORCEINLINE bool Matches(const FString& A, const FString& B)
-	{
-		return A.Equals( B, ESearchCase::CaseSensitive );
-	}
-	static FORCEINLINE uint32 GetKeyHash(const FString& Key)
-	{
-		return FCrc::StrCrc32<TCHAR>(*Key);
-	}
 };
-typedef TMultiMap< FString, TSharedRef< FManifestEntry >, FDefaultSetAllocator, FInternationalizationManifestSourceTextKeyFuncs > TManifestEntryBySourceTextContainer;
 
 
-struct FInternationalizationManifestContextIdKeyFuncs : BaseKeyFuncs<TSharedRef< FManifestEntry >, FString, true>
+struct FInternationalizationManifestStringKeyFuncs : BaseKeyFuncs<TSharedRef< FManifestEntry >, FString, true>
 {
-	static FORCEINLINE const FString& GetSetKey(const TPair<FString,TSharedRef< FManifestEntry >>& Element)
+	static FORCEINLINE const FString& GetSetKey(const TPair<FString, TSharedRef< FManifestEntry >>& Element)
 	{
 		return Element.Key;
 	}
@@ -115,13 +112,13 @@ struct FInternationalizationManifestContextIdKeyFuncs : BaseKeyFuncs<TSharedRef<
 		return FCrc::StrCrc32<TCHAR>(*Key);
 	}
 };
-typedef TMultiMap< FString, TSharedRef< FManifestEntry >, FDefaultSetAllocator, FInternationalizationManifestContextIdKeyFuncs > TManifestEntryByContextIdContainer;
+typedef TMultiMap< FString, TSharedRef< FManifestEntry >, FDefaultSetAllocator, FInternationalizationManifestStringKeyFuncs > FManifestEntryByStringContainer;
 
 
 class CORE_API FInternationalizationManifest 
 {
 public:
-	enum EFormatVersion
+	enum class EFormatVersion : uint8
 	{
 		Initial = 0,
 		EscapeFixes,
@@ -131,35 +128,37 @@ public:
 	};
 
 	//Default constructor
-	FInternationalizationManifest( )
-		: FormatVersion(static_cast<int32>(EFormatVersion::Latest))
+	FInternationalizationManifest()
+		: FormatVersion(EFormatVersion::Latest)
 	{
 	}
 
 	/**
-	 * Adds a manifest entry.
-	 *
-	 * @return Returns true if add was successful or a matching entry already exists, false is only returned in the case where a duplicate context was found with different text.
-	 */
-	bool AddSource( const FString& Namespace, const FLocItem& Source, const FManifestContext& Context );
+	* Adds a manifest entry.
+	*
+	* @return Returns true if add was successful or a matching entry already exists, false is only returned in the case where a duplicate context was found with different text.
+	*/
+	bool AddSource(const FString& Namespace, const FLocItem& Source, const FManifestContext& Context);
+
+	void UpdateEntry(const TSharedRef<FManifestEntry>& OldEntry, TSharedRef<FManifestEntry>& NewEntry);
 
-	TSharedPtr< FManifestEntry > FindEntryBySource( const FString& Namespace, const FLocItem& Source ) const;
+	TSharedPtr<FManifestEntry> FindEntryBySource(const FString& Namespace, const FLocItem& Source) const;
 
-	TSharedPtr< FManifestEntry > FindEntryByContext( const FString& Namespace, const FManifestContext& Context ) const;
+	TSharedPtr<FManifestEntry> FindEntryByContext(const FString& Namespace, const FManifestContext& Context) const;
 
-	TSharedPtr< FManifestEntry > FindEntryByKey( const FString& Namespace, const FString& Key, const FString* SourceText = nullptr ) const;
+	TSharedPtr<FManifestEntry> FindEntryByKey(const FString& Namespace, const FString& Key, const FString* SourceText = nullptr) const;
 
-	TManifestEntryByContextIdContainer::TConstIterator GetEntriesByContextIdIterator() const
+	FManifestEntryByStringContainer::TConstIterator GetEntriesByKeyIterator() const
 	{
-		return EntriesByContextId.CreateConstIterator();
+		return EntriesByKey.CreateConstIterator();
 	}
 
-	int32 GetNumEntriesByContextId() const
+	int32 GetNumEntriesByKey() const
 	{
-		return EntriesByContextId.Num();
+		return EntriesByKey.Num();
 	}
 
-	TManifestEntryBySourceTextContainer::TConstIterator GetEntriesBySourceTextIterator() const
+	FManifestEntryByStringContainer::TConstIterator GetEntriesBySourceTextIterator() const
 	{
 		return EntriesBySourceText.CreateConstIterator();
 	}
@@ -169,23 +168,18 @@ public:
 		return EntriesBySourceText.Num();
 	}
 
-	void UpdateEntry(const TSharedRef<FManifestEntry>& OldEntry, TSharedRef<FManifestEntry>& NewEntry);
-
 	void SetFormatVersion(const EFormatVersion Version)
 	{
-		FormatVersion = static_cast<int>(Version);
+		FormatVersion = Version;
 	}
 
 	EFormatVersion GetFormatVersion() const
 	{
-		return EFormatVersion(FormatVersion);
+		return FormatVersion;
 	}
 
-	friend class IInternationalizationManifestSerializer;
-
 private:
-
-	int32 FormatVersion;
-	TManifestEntryBySourceTextContainer EntriesBySourceText;
-	TManifestEntryByContextIdContainer EntriesByContextId;
+	EFormatVersion FormatVersion;
+	FManifestEntryByStringContainer EntriesBySourceText;
+	FManifestEntryByStringContainer EntriesByKey;
 };
diff --git a/Engine/Source/Runtime/Core/Public/Internationalization/TextLocalizationManager.h b/Engine/Source/Runtime/Core/Public/Internationalization/TextLocalizationManager.h
--- a/Engine/Source/Runtime/Core/Public/Internationalization/TextLocalizationManager.h
+++ b/Engine/Source/Runtime/Core/Public/Internationalization/TextLocalizationManager.h
@@ -2,9 +2,6 @@
 
 #pragma once
 
-#include "IInternationalizationArchiveSerializer.h"
-#include "IInternationalizationManifestSerializer.h"
-
 typedef TSharedRef<FString, ESPMode::ThreadSafe> FTextDisplayStringRef;
 typedef TSharedPtr<FString, ESPMode::ThreadSafe> FTextDisplayStringPtr;
 
@@ -213,12 +210,12 @@ public:
 	 */
 	bool UpdateDisplayString(const FTextDisplayStringRef& DisplayString, const FString& Value, const FString& Namespace, const FString& Key);
 
-	/** Loads localizations for the current culture based on a configuration file specifying which manifests and archives to use. Effectively generates a localization resource in memory. */
-	void LoadFromManifestAndArchives(const FString& ConfigFilePath, IInternationalizationArchiveSerializer& ArchiveSerializer, IInternationalizationManifestSerializer& ManifestSerializer);
-
 	/** Updates display string entries and adds new display string entries based on localizations found in a specified localization resource. */
 	void UpdateFromLocalizationResource(const FString& LocalizationResourceFilePath);
 
+	/** Updates display string entries and adds new display string entries based on localizations found in a specified localization resource. */
+	void UpdateFromLocalizationResource(FArchive& LocResArchive, const FString& LocResID);
+
 	/** Reloads resources for the current culture. */
 	void RefreshResources();
 
@@ -230,6 +227,10 @@ public:
 	DECLARE_EVENT(FTextLocalizationManager, FTextRevisionChangedEvent)
 	FTextRevisionChangedEvent OnTextRevisionChangedEvent;
 
+	/** Delegate for gathering up additional localization paths that are unknown to the UE4 core (such as plugins) */
+	DECLARE_MULTICAST_DELEGATE_OneParam(FGatherAdditionalLocResPathsDelegate, TArray<FString>&);
+	FGatherAdditionalLocResPathsDelegate GatherAdditionalLocResPathsCallback;
+
 private:
 	/** Callback for changes in culture. Loads the new culture's localization resources. */
 	void OnCultureChanged();
diff --git a/Engine/Source/Runtime/Core/Public/Internationalization/TextNamespaceFwd.h b/Engine/Source/Runtime/Core/Public/Internationalization/TextNamespaceFwd.h
--- a/Engine/Source/Runtime/Core/Public/Internationalization/TextNamespaceFwd.h
+++ b/Engine/Source/Runtime/Core/Public/Internationalization/TextNamespaceFwd.h
@@ -5,7 +5,7 @@
 // USE_STABLE_LOCALIZATION_KEYS controls whether we attempt to keep localization keys stable when editing text properties.
 // This requires WITH_EDITORONLY_DATA, but can also be optionally disabled in editor builds by changing its value to 0.
 #if WITH_EDITORONLY_DATA
-	#define USE_STABLE_LOCALIZATION_KEYS (0)
+	#define USE_STABLE_LOCALIZATION_KEYS (1)
 #else	// WITH_EDITORONLY_DATA
 	#define USE_STABLE_LOCALIZATION_KEYS (0)
 #endif	// WITH_EDITORONLY_DATA
diff --git a/Engine/Source/Runtime/Core/Public/Misc/AutomationTest.h b/Engine/Source/Runtime/Core/Public/Misc/AutomationTest.h
--- a/Engine/Source/Runtime/Core/Public/Misc/AutomationTest.h
+++ b/Engine/Source/Runtime/Core/Public/Misc/AutomationTest.h
@@ -46,6 +46,33 @@ namespace EAutomationTestFlags
 };
 
 
+struct CORE_API FAutomationEvent
+{
+	FString Message;
+	FString Context;
+	FString Filename;
+	int32 LineNumber;
+
+	FAutomationEvent(FString InMessage)
+		: Message(InMessage)
+		, Context()
+		, Filename()
+		, LineNumber(-1)
+	{
+	}
+
+	FAutomationEvent(FString InMessage, FString InContext, FString InFilename, int32 InLineNumber)
+		: Message(InMessage)
+		, Context(InContext)
+		, Filename(InFilename)
+		, LineNumber(InLineNumber)
+	{
+	}
+
+	FString ToString() const;
+};
+
+
 /** Simple class to store the results of the execution of a automation test */
 class FAutomationTestExecutionInfo
 {
@@ -65,16 +92,24 @@ public:
 	/** Helper method to clear out the results from a previous execution */
 	void Clear()
 	{
+		Context.Reset();
+
 		Errors.Empty();
 		Warnings.Empty();
 		LogItems.Empty();
 		AnalyticsItems.Empty();
 	}
 
+	/**
+	 * Allows tests to set the current logging context set on the events, like if we're currently
+	 * importing a specific file, you might set the file as the context.
+	 */
+	FString Context;
+
 	/** Whether the automation test completed successfully or not */
 	bool bSuccessful;
 	/** Any errors that occurred during execution */
-	TArray<FString> Errors;
+	TArray<FAutomationEvent> Errors;
 	/** Any warnings that occurred during execution */
 	TArray<FString> Warnings;
 	/** Any log items that occurred during execution */
@@ -118,12 +153,14 @@ public:
 	 * @param	InTestFlag - Test flags
 	 * @param	InParameterName - optional parameter. e.g. asset name
 	 */
-	FAutomationTestInfo(const FString& InDisplayName, const FString& InTestName, const uint32 InTestFlags, const int32 InNumParticipantsRequired, const FString& InParameterName = FString(), const FString& InSourceFile = FString(), int32 InSourceFileLine = 0)
+	FAutomationTestInfo(const FString& InDisplayName, const FString& InTestName, const uint32 InTestFlags, const int32 InNumParticipantsRequired, const FString& InParameterName = FString(), const FString& InSourceFile = FString(), int32 InSourceFileLine = 0, const FString& InAssetPath = FString(), const FString& InOpenCommand = FString())
 		: DisplayName( InDisplayName )
 		, TestName( InTestName )
 		, TestParameter( InParameterName )
 		, SourceFile( InSourceFile )
 		, SourceFileLine( InSourceFileLine )
+		, AssetPath( InAssetPath )
+		, OpenCommand( InOpenCommand )
 		, TestFlags( InTestFlags )
 		, NumParticipantsRequired( InNumParticipantsRequired )
 		, NumDevicesCurrentlyRunningTest( 0 )
@@ -202,6 +239,26 @@ public:
 		return SourceFileLine;
 	}
 
+	/**
+	 * Gets the asset potentially associated with the test.
+	 *
+	 * @return the source line number.
+	 */
+	const FString GetAssetPath() const
+	{
+		return AssetPath;
+	}
+
+	/**
+	 * Gets the open command potentially associated with the test.
+	 *
+	 * @return the source line number.
+	 */
+	const FString GetOpenCommand() const
+	{
+		return OpenCommand;
+	}
+
 	/**
 	 * Get the type of test.
 	 *
@@ -295,6 +352,12 @@ private:
 	/** The line number in the source file this test originated on. */
 	int32 SourceFileLine;
 
+	/** The asset path associated with the test. */
+	FString AssetPath;
+
+	/** A custom open command for the test. */
+	FString OpenCommand;
+
 	/** The test flags. */
 	uint32 TestFlags;
 
@@ -726,7 +789,7 @@ public:
 	 *
 	 * @param	InError	Error message to add to this test
 	 */
-	void AddError( const FString& InError );
+	void AddError( const FString& InError, int32 StackOffset = 0 );
 
 	/**
 	 * Adds a warning to this test
@@ -799,24 +862,32 @@ public:
 	/**
 	 * Enqueues a new latent command.
 	 */
-	FORCEINLINE void AddCommand(TSharedPtr<IAutomationLatentCommand> NewCommand)
+	FORCEINLINE void AddCommand(IAutomationLatentCommand* NewCommand)
 	{
-		FAutomationTestFramework::GetInstance().EnqueueLatentCommand(NewCommand);
+		TSharedRef<IAutomationLatentCommand> CommandPtr = MakeShareable(NewCommand);
+		FAutomationTestFramework::GetInstance().EnqueueLatentCommand(CommandPtr);
 	}
 
 	/**
 	 * Enqueues a new latent network command.
 	 */
-	FORCEINLINE void AddCommand(TSharedPtr<IAutomationNetworkCommand> NewCommand)
+	FORCEINLINE void AddCommand(IAutomationNetworkCommand* NewCommand)
 	{
-		FAutomationTestFramework::GetInstance().EnqueueNetworkCommand(NewCommand);
+		TSharedRef<IAutomationNetworkCommand> CommandPtr = MakeShareable(NewCommand);
+		FAutomationTestFramework::GetInstance().EnqueueNetworkCommand(CommandPtr);
 	}
 
 	/** Gets the filename where this test was defined. */
-	virtual FString GetFileName() const { return ""; }
+	virtual FString GetTestSourceFileName() const { return TEXT(""); }
 
 	/** Gets the line number where this test was defined. */
-	virtual int32 GetFileLine() const { return 0; }
+	virtual int32 GetTestSourceFileLine() const { return 0; }
+
+	/** Allows navigation to the asset associated with the test if there is one. */
+	virtual FString GetTestAssetPath(const FString& Parameter) const { return TEXT(""); }
+
+	/** Return an exec command to open the test associated with this parameter. */
+	virtual FString GetTestOpenCommand(const FString& Parameter) const { return TEXT(""); }
 
 public:
 
@@ -833,7 +904,7 @@ public:
 	{
 		if (A != B)
 		{
-			AddError(FString::Printf(TEXT("%s: The two values are not equal."), *Description));
+			AddError(FString::Printf(TEXT("%s: The two values are not equal."), *Description), 1);
 		}
 	}
 
@@ -849,7 +920,7 @@ public:
 	{
 		if (Value)
 		{
-			AddError(FString::Printf(TEXT("%s: The value is not false."), *Description));
+			AddError(FString::Printf(TEXT("%s: The value is not false."), *Description), 1);
 		}
 	}
 
@@ -865,7 +936,7 @@ public:
 	{
 		if (SharedPointer.IsValid())
 		{
-			AddError(FString::Printf(TEXT("%s: The shared pointer is valid."), *Description));
+			AddError(FString::Printf(TEXT("%s: The shared pointer is valid."), *Description), 1);
 		}
 	}
 
@@ -882,7 +953,7 @@ public:
 	{
 		if (A == B)
 		{
-			AddError(FString::Printf(TEXT("%s: The two values are not equal."), *Description));
+			AddError(FString::Printf(TEXT("%s: The two values are not equal."), *Description), 1);
 		}
 	}
 
@@ -898,7 +969,7 @@ public:
 	{
 		if (Pointer == NULL)
 		{
-			AddError(FString::Printf(TEXT("%s: The pointer is NULL."), *Description));
+			AddError(FString::Printf(TEXT("%s: The pointer is NULL."), *Description), 1);
 		}
 	}
 
@@ -915,7 +986,7 @@ public:
 	{
 		if (&A == &B)
 		{
-			AddError(FString::Printf(TEXT("%s: The two values are the same."), *Description));
+			AddError(FString::Printf(TEXT("%s: The two values are the same."), *Description), 1);
 		}
 	}
 
@@ -931,7 +1002,7 @@ public:
 	{
 		if (Pointer != NULL)
 		{
-			AddError(FString::Printf(TEXT("%s: The pointer is not NULL."), *Description));
+			AddError(FString::Printf(TEXT("%s: The pointer is not NULL."), *Description), 1);
 		}
 	}
 
@@ -948,7 +1019,7 @@ public:
 	{
 		if (&A != &B)
 		{
-			AddError(FString::Printf(TEXT("%s: The two values are not the same."), *Description));
+			AddError(FString::Printf(TEXT("%s: The two values are not the same."), *Description), 1);
 		}
 	}
 
@@ -964,7 +1035,7 @@ public:
 	{
 		if (!Value)
 		{
-			AddError(FString::Printf(TEXT("%s: The value is not true."), *Description));
+			AddError(FString::Printf(TEXT("%s: The value is not true."), *Description), 1);
 		}
 	}
 
@@ -980,7 +1051,7 @@ public:
 	{
 		if (!SharedPointer.IsValid())
 		{
-			AddError(FString::Printf(TEXT("%s: The shared pointer is not valid."), *Description));
+			AddError(FString::Printf(TEXT("%s: The shared pointer is not valid."), *Description), 1);
 		}
 	}
 
@@ -1052,8 +1123,8 @@ class CommandName : public IAutomationLatentCommand \
 	ParamType ParamName; \
 }
 
-#define DEFINE_ENGINE_LATENT_AUTOMATION_COMMAND(CommandName)	\
-class ENGINE_API CommandName : public IAutomationLatentCommand \
+#define DEFINE_EXPORTED_LATENT_AUTOMATION_COMMAND(EXPORT_API, CommandName)	\
+class EXPORT_API CommandName : public IAutomationLatentCommand \
 	{ \
 	public: \
 	virtual ~CommandName() \
@@ -1061,8 +1132,8 @@ class ENGINE_API CommandName : public IAutomationLatentCommand \
 		virtual bool Update() override; \
 }
 
-#define DEFINE_ENGINE_LATENT_AUTOMATION_COMMAND_ONE_PARAMETER(CommandName,ParamType,ParamName)	\
-class ENGINE_API CommandName : public IAutomationLatentCommand \
+#define DEFINE_EXPORTED_LATENT_AUTOMATION_COMMAND_ONE_PARAMETER(EXPORT_API, CommandName,ParamType,ParamName)	\
+class EXPORT_API CommandName : public IAutomationLatentCommand \
 	{ \
 	public: \
 	CommandName(ParamType InputParam) \
@@ -1075,6 +1146,12 @@ class ENGINE_API CommandName : public IAutomationLatentCommand \
 	ParamType ParamName; \
 }
 
+#define DEFINE_ENGINE_LATENT_AUTOMATION_COMMAND(CommandName)	\
+	DEFINE_EXPORTED_LATENT_AUTOMATION_COMMAND(ENGINE_API, CommandName)
+
+#define DEFINE_ENGINE_LATENT_AUTOMATION_COMMAND_ONE_PARAMETER(CommandName,ParamType,ParamName)	\
+	DEFINE_EXPORTED_LATENT_AUTOMATION_COMMAND_ONE_PARAMETER(ENGINE_API, CommandName, ParamType, ParamName)
+
 //macro to simply the syntax for enqueueing a latent command
 #define ADD_LATENT_AUTOMATION_COMMAND(ClassDeclaration) FAutomationTestFramework::GetInstance().EnqueueLatentCommand(MakeShareable(new ClassDeclaration));
 
@@ -1108,12 +1185,12 @@ public: \
  * the automation test framework as a result of the macro.
  */
 
-#define IMPLEMENT_SIMPLE_AUTOMATION_TEST_PRIVATE( TClass, PrettyName, TFlags, FileName, LineNumber ) \
-	class TClass : public FAutomationTestBase \
+#define IMPLEMENT_SIMPLE_AUTOMATION_TEST_PRIVATE( TClass, TBaseClass, PrettyName, TFlags, FileName, LineNumber ) \
+	class TClass : public TBaseClass \
 	{ \
 	public: \
 		TClass( const FString& InName ) \
-		:FAutomationTestBase( InName, false ) {\
+		:TBaseClass( InName, false ) {\
 			static_assert((TFlags)&EAutomationTestFlags::ApplicationContextMask, "AutomationTest has no application flag.  It shouldn't run.  See AutomationTest.h."); \
 			static_assert(	(((TFlags)&EAutomationTestFlags::FilterMask) == EAutomationTestFlags::SmokeFilter) || \
 							(((TFlags)&EAutomationTestFlags::FilterMask) == EAutomationTestFlags::EngineFilter) || \
@@ -1125,8 +1202,8 @@ public: \
 		virtual uint32 GetTestFlags() const override { return TFlags; } \
 		virtual bool IsStressTest() const { return false; } \
 		virtual uint32 GetRequiredDeviceNum() const override { return 1; } \
-		virtual FString GetFileName() const { return FileName; } \
-		virtual int32 GetFileLine() const { return LineNumber; } \
+		virtual FString GetTestSourceFileName() const override { return FileName; } \
+		virtual int32 GetTestSourceFileLine() const override { return LineNumber; } \
 	protected: \
 		virtual void GetTests(TArray<FString>& OutBeautifiedNames, TArray <FString>& OutTestCommands) const override \
 		{ \
@@ -1137,12 +1214,12 @@ public: \
 		virtual FString GetBeautifiedTestName() const override { return PrettyName; } \
 	};
 
-#define IMPLEMENT_COMPLEX_AUTOMATION_TEST_PRIVATE( TClass, PrettyName, TFlags, FileName, LineNumber ) \
-	class TClass : public FAutomationTestBase \
+#define IMPLEMENT_COMPLEX_AUTOMATION_TEST_PRIVATE( TClass, TBaseClass, PrettyName, TFlags, FileName, LineNumber ) \
+	class TClass : public TBaseClass \
 	{ \
 	public: \
 		TClass( const FString& InName ) \
-		:FAutomationTestBase( InName, true ) { \
+		:TBaseClass( InName, true ) { \
 			static_assert((TFlags)&EAutomationTestFlags::ApplicationContextMask, "AutomationTest has no application flag.  It shouldn't run.  See AutomationTest.h."); \
 			static_assert(	(((TFlags)&EAutomationTestFlags::FilterMask) == EAutomationTestFlags::SmokeFilter) || \
 							(((TFlags)&EAutomationTestFlags::FilterMask) == EAutomationTestFlags::EngineFilter) || \
@@ -1154,20 +1231,20 @@ public: \
 		virtual uint32 GetTestFlags() const override { return ((TFlags) & ~(EAutomationTestFlags::SmokeFilter)); } \
 		virtual bool IsStressTest() const { return true; } \
 		virtual uint32 GetRequiredDeviceNum() const override { return 1; } \
-		virtual FString GetFileName() const { return FileName; } \
-		virtual int32 GetFileLine() const { return LineNumber; } \
+		virtual FString GetTestSourceFileName() const override { return FileName; } \
+		virtual int32 GetTestSourceFileLine() const override { return LineNumber; } \
 	protected: \
 		virtual void GetTests(TArray<FString>& OutBeautifiedNames, TArray <FString>& OutTestCommands) const override; \
 		virtual bool RunTest(const FString& Parameters) override; \
 		virtual FString GetBeautifiedTestName() const override { return PrettyName; } \
 	};
 
-#define IMPLEMENT_NETWORKED_AUTOMATION_TEST_PRIVATE(TClass, PrettyName, TFlags, NumParticipants, FileName, LineNumber) \
-	class TClass : public FAutomationTestBase \
+#define IMPLEMENT_NETWORKED_AUTOMATION_TEST_PRIVATE(TClass, TBaseClass, PrettyName, TFlags, NumParticipants, FileName, LineNumber) \
+	class TClass : public TBaseClass \
 	{ \
 	public: \
 		TClass( const FString& InName ) \
-		:FAutomationTestBase( InName, false ) { \
+		:TBaseClass( InName, false ) { \
 			static_assert((TFlags)&EAutomationTestFlags::ApplicationContextMask, "AutomationTest has no application flag.  It shouldn't run.  See AutomationTest.h."); \
 			static_assert(	(((TFlags)&EAutomationTestFlags::FilterMask) == EAutomationTestFlags::SmokeFilter) || \
 							(((TFlags)&EAutomationTestFlags::FilterMask) == EAutomationTestFlags::EngineFilter) || \
@@ -1178,8 +1255,8 @@ public: \
 		} \
 		virtual uint32 GetTestFlags() const override { return ((TFlags) & ~(EAutomationTestFlags::EditorContext | EAutomationTestFlags::CommandletContext | EAutomationTestFlags::SmokeFilter)); } \
 		virtual uint32 GetRequiredDeviceNum() const override { return NumParticipants; } \
-		virtual FString GetFileName() const { return FileName; } \
-		virtual int32 GetFileLine() const { return LineNumber; } \
+		virtual FString GetTestSourceFileName() const override { return FileName; } \
+		virtual int32 GetTestSourceFileLine() const override { return LineNumber; } \
 	protected: \
 		virtual void GetTests(TArray<FString>& OutBeautifiedNames, TArray <FString>& OutTestCommands) const override \
 		{ \
@@ -1193,30 +1270,50 @@ public: \
 
 #if WITH_AUTOMATION_WORKER
 	#define IMPLEMENT_SIMPLE_AUTOMATION_TEST( TClass, PrettyName, TFlags ) \
-		IMPLEMENT_SIMPLE_AUTOMATION_TEST_PRIVATE(TClass, PrettyName, TFlags, __FILE__, __LINE__) \
+		IMPLEMENT_SIMPLE_AUTOMATION_TEST_PRIVATE(TClass, FAutomationTestBase, PrettyName, TFlags, __FILE__, __LINE__) \
 		namespace\
 		{\
 			TClass TClass##AutomationTestInstance( TEXT(#TClass) );\
 		}
 	#define IMPLEMENT_COMPLEX_AUTOMATION_TEST( TClass, PrettyName, TFlags ) \
-		IMPLEMENT_COMPLEX_AUTOMATION_TEST_PRIVATE(TClass, PrettyName, TFlags, __FILE__, __LINE__) \
+		IMPLEMENT_COMPLEX_AUTOMATION_TEST_PRIVATE(TClass, FAutomationTestBase, PrettyName, TFlags, __FILE__, __LINE__) \
 		namespace\
 		{\
 			TClass TClass##AutomationTestInstance( TEXT(#TClass) );\
 		}
 	#define IMPLEMENT_NETWORKED_AUTOMATION_TEST(TClass, PrettyName, TFlags, NumParticipants) \
-		IMPLEMENT_NETWORKED_AUTOMATION_TEST_PRIVATE(TClass, PrettyName, TFlags, NumParticipants, __FILE__, __LINE__) \
+		IMPLEMENT_NETWORKED_AUTOMATION_TEST_PRIVATE(TClass, FAutomationTestBase, PrettyName, TFlags, NumParticipants, __FILE__, __LINE__) \
+		namespace\
+		{\
+			TClass TClass##AutomationTestInstance( TEXT(#TClass) );\
+		}
+
+	#define IMPLEMENT_CUSTOM_SIMPLE_AUTOMATION_TEST( TClass, TBaseClass, PrettyName, TFlags ) \
+		IMPLEMENT_SIMPLE_AUTOMATION_TEST_PRIVATE(TClass, TBaseClass, PrettyName, TFlags, __FILE__, __LINE__) \
+		namespace\
+		{\
+			TClass TClass##AutomationTestInstance( TEXT(#TClass) );\
+		}
+
+	#define IMPLEMENT_CUSTOM_COMPLEX_AUTOMATION_TEST( TClass, TBaseClass, PrettyName, TFlags ) \
+		IMPLEMENT_COMPLEX_AUTOMATION_TEST_PRIVATE(TClass, TBaseClass, PrettyName, TFlags, __FILE__, __LINE__) \
 		namespace\
 		{\
 			TClass TClass##AutomationTestInstance( TEXT(#TClass) );\
 		}
+
 #else
 	#define IMPLEMENT_SIMPLE_AUTOMATION_TEST( TClass, PrettyName, TFlags ) \
-		IMPLEMENT_SIMPLE_AUTOMATION_TEST_PRIVATE(TClass, PrettyName, TFlags, __FILE__, __LINE__)
+		IMPLEMENT_SIMPLE_AUTOMATION_TEST_PRIVATE(TClass, FAutomationTestBase, PrettyName, TFlags, __FILE__, __LINE__)
 	#define IMPLEMENT_COMPLEX_AUTOMATION_TEST( TClass, PrettyName, TFlags ) \
-		IMPLEMENT_COMPLEX_AUTOMATION_TEST_PRIVATE(TClass, PrettyName, TFlags, __FILE__, __LINE__)
+		IMPLEMENT_COMPLEX_AUTOMATION_TEST_PRIVATE(TClass, FAutomationTestBase, PrettyName, TFlags, __FILE__, __LINE__)
 	#define IMPLEMENT_NETWORKED_AUTOMATION_TEST(TClass, PrettyName, TFlags, NumParticipants) \
-		IMPLEMENT_NETWORKED_AUTOMATION_TEST_PRIVATE(TClass, PrettyName, TFlags, NumParticipants, __FILE__, __LINE__)
+		IMPLEMENT_NETWORKED_AUTOMATION_TEST_PRIVATE(TClass, FAutomationTestBase, PrettyName, TFlags, NumParticipants, __FILE__, __LINE__)
+
+	#define IMPLEMENT_CUSTOM_SIMPLE_AUTOMATION_TEST( TClass, TBaseClass, PrettyName, TFlags ) \
+		IMPLEMENT_SIMPLE_AUTOMATION_TEST_PRIVATE(TClass, TBaseClass, PrettyName, TFlags, __FILE__, __LINE__)
+	#define IMPLEMENT_CUSTOM_COMPLEX_AUTOMATION_TEST( TClass, TBaseClass, PrettyName, TFlags ) \
+		IMPLEMENT_COMPLEX_AUTOMATION_TEST_PRIVATE(TClass, TBaseClass, PrettyName, TFlags, __FILE__, __LINE__)
 #endif // #if WITH_AUTOMATION_WORKER
 
 
@@ -1227,11 +1324,11 @@ public: \
  * Run some code latently with a predicate lambda.  If the predicate returns true, the latent action will be called 
  * again next frame.  If it returns false, the command will stop running.
  */
-class CORE_API FFunctionLatentCommand : public IAutomationLatentCommand
+class FFunctionLatentCommand : public IAutomationLatentCommand
 {
 public:
 	FFunctionLatentCommand(TFunction<bool()> InLatentPredicate)
-		: LatentPredicate(InLatentPredicate)
+		: LatentPredicate(MoveTemp(InLatentPredicate))
 	{
 	}
 
@@ -1247,3 +1344,61 @@ public:
 private:
 	TFunction<bool()> LatentPredicate;
 };
+
+
+class FDelayedFunctionLatentCommand : public IAutomationLatentCommand
+{
+public:
+	FDelayedFunctionLatentCommand(TFunction<void()> InCallback, float InDelay = 0.1f)
+		: Callback(MoveTemp(InCallback))
+		, Delay(InDelay)
+	{}
+
+	virtual bool Update() override
+	{
+		float NewTime = FPlatformTime::Seconds();
+		if ( NewTime - StartTime >= Delay )
+		{
+			Callback();
+			return true;
+		}
+		return false;
+	}
+
+private:
+	TFunction<void()> Callback;
+	float Delay;
+};
+
+
+class FUntilCommand : public IAutomationLatentCommand
+{
+public:
+	FUntilCommand(TFunction<bool()> InCallback, TFunction<bool()> InTimeoutCallback, float InTimeout = 5.0f)
+		: Callback(MoveTemp(InCallback))
+		, TimeoutCallback(MoveTemp(InTimeoutCallback))
+		, Timeout(InTimeout)
+	{}
+
+	virtual bool Update() override
+	{
+		if ( !Callback() )
+		{
+			float NewTime = FPlatformTime::Seconds();
+			if ( NewTime - StartTime >= Timeout )
+			{
+				TimeoutCallback();
+				return false;
+			}
+
+			return true;
+		}
+
+		return false;
+	}
+
+private:
+	TFunction<bool()> Callback;
+	TFunction<bool()> TimeoutCallback;
+	float Timeout;
+};
\ No newline at end of file
diff --git a/Engine/Source/Runtime/Core/Public/Misc/Paths.h b/Engine/Source/Runtime/Core/Public/Misc/Paths.h
--- a/Engine/Source/Runtime/Core/Public/Misc/Paths.h
+++ b/Engine/Source/Runtime/Core/Public/Misc/Paths.h
@@ -233,7 +233,7 @@ public:
 	static FString GameUserDeveloperDir();
 
 	/**
-	 * @return The directory for temp files used for diff'ing
+	 * @return The directory for temp files used for diffing
 	 */
 	static FString DiffDir();
 
diff --git a/Engine/Source/Runtime/Core/Public/UObject/EditorObjectVersion.h b/Engine/Source/Runtime/Core/Public/UObject/EditorObjectVersion.h
--- a/Engine/Source/Runtime/Core/Public/UObject/EditorObjectVersion.h
+++ b/Engine/Source/Runtime/Core/Public/UObject/EditorObjectVersion.h
@@ -18,6 +18,8 @@ struct CORE_API FEditorObjectVersion
 		GatheredTextPackageCacheFixesV2,
 		// Updated FFormatArgumentData to allow variant data to be marshaled from a BP into C++
 		TextFormatArgumentDataIsVariant,
+		// Changes to SplineComponent
+		SplineComponentCurvesInStruct,
 		// -----<new versions can be added above this line>-------------------------------------------------
 		VersionPlusOne,
 		LatestVersion = VersionPlusOne - 1
diff --git a/Engine/Source/Runtime/Core/Public/Windows/WindowsPlatformStackWalk.h b/Engine/Source/Runtime/Core/Public/Windows/WindowsPlatformStackWalk.h
--- a/Engine/Source/Runtime/Core/Public/Windows/WindowsPlatformStackWalk.h
+++ b/Engine/Source/Runtime/Core/Public/Windows/WindowsPlatformStackWalk.h
@@ -12,6 +12,8 @@ struct CORE_API FWindowsPlatformStackWalk
 {
 	static bool InitStackWalking();
 	
+	static TArray<FProgramCounterSymbolInfo> GetStack(int32 IgnoreCount, int32 MaxDepth = 100, void* Context = nullptr);
+
 	static void ProgramCounterToSymbolInfo( uint64 ProgramCounter, FProgramCounterSymbolInfo& out_SymbolInfo );
 	static void CaptureStackBackTrace( uint64* BackTrace, uint32 MaxDepth, void* Context = nullptr );
 	static void StackWalkAndDump( ANSICHAR* HumanReadableString, SIZE_T HumanReadableStringSize, int32 IgnoreCount, void* Context = nullptr );
diff --git a/Engine/Source/Runtime/CoreUObject/Private/UObject/LazyObjectPtr.cpp b/Engine/Source/Runtime/CoreUObject/Private/UObject/LazyObjectPtr.cpp
--- a/Engine/Source/Runtime/CoreUObject/Private/UObject/LazyObjectPtr.cpp
+++ b/Engine/Source/Runtime/CoreUObject/Private/UObject/LazyObjectPtr.cpp
@@ -119,7 +119,7 @@ void FLazyObjectPtr::PossiblySerializeObjectGuid(UObject *Object, FArchive& Ar)
 			FUniqueObjectGuid Guid;
 			Ar << Guid;
 
-			// Don't try and resolve GUIDs when loading a package for diff'ing
+			// Don't try and resolve GUIDs when loading a package for diffing
 			const UPackage* Package = Object->GetOutermost();
 			const bool bLoadedForDiff = (Package && Package->HasAnyPackageFlags(PKG_ForDiffing));
 			if (!bLoadedForDiff && (!(Ar.GetPortFlags() & PPF_Duplicate) || (Ar.GetPortFlags() & PPF_DuplicateForPIE)))
diff --git a/Engine/Source/Runtime/CoreUObject/Private/UObject/UObjectGlobals.cpp b/Engine/Source/Runtime/CoreUObject/Private/UObject/UObjectGlobals.cpp
--- a/Engine/Source/Runtime/CoreUObject/Private/UObject/UObjectGlobals.cpp
+++ b/Engine/Source/Runtime/CoreUObject/Private/UObject/UObjectGlobals.cpp
@@ -800,7 +800,7 @@ UObject* StaticLoadObjectInternal(UClass* ObjectClass, UObject* InOuter, const T
 	const bool bContainsObjectName = !!FCString::Strstr(InName, TEXT("."));
 
 	// break up the name into packages, returning the innermost name and its outer
-	ResolveName(InOuter, StrName, true, true, LoadFlags & LOAD_EditorOnly);
+	ResolveName(InOuter, StrName, true, true, LoadFlags & (LOAD_EditorOnly | LOAD_Quiet | LOAD_NoWarn));
 	if (InOuter)
 	{
 		// If we have a full UObject name then attempt to find the object in memory first,
@@ -1181,8 +1181,7 @@ static UPackage* LoadPackageInternalInner(UPackage* InOuter, const TCHAR* InLong
 		// The time tracker keeps track of time spent in LoadPackage.
 		FExclusiveLoadPackageTimeTracker::FScopedPackageTracker Tracker(Result);
 
-
-		// If we are loading a package for diff'ing, set the package flag
+		// If we are loading a package for diffing, set the package flag
 		if(LoadFlags & LOAD_ForDiff)
 		{
 			Result->SetPackageFlags(PKG_ForDiffing);
diff --git a/Engine/Source/Runtime/CoreUObject/Public/UObject/ObjectBase.h b/Engine/Source/Runtime/CoreUObject/Public/UObject/ObjectBase.h
--- a/Engine/Source/Runtime/CoreUObject/Public/UObject/ObjectBase.h
+++ b/Engine/Source/Runtime/CoreUObject/Public/UObject/ObjectBase.h
@@ -91,7 +91,7 @@ enum EPackageFlags
 	PKG_ClientOptional				= 0x00000002,	// Purely optional for clients.
 	PKG_ServerSideOnly				= 0x00000004,   // Only needed on the server side.
 	PKG_CompiledIn					= 0x00000010,   // This package is from "compiled in" classes.
-	PKG_ForDiffing					= 0x00000020,	// This package was loaded just for the purposes of diff'ing
+	PKG_ForDiffing					= 0x00000020,	// This package was loaded just for the purposes of diffing
 	PKG_EditorOnly					= 0x00000040, // This is editor-only package (for example: editor module script package)
 	PKG_Developer					= 0x00000080,	// Developer module
 //	PKG_Unused						= 0x00000100,
diff --git a/Engine/Source/Runtime/CoreUObject/Public/UObject/UnrealType.h b/Engine/Source/Runtime/CoreUObject/Public/UObject/UnrealType.h
--- a/Engine/Source/Runtime/CoreUObject/Public/UObject/UnrealType.h
+++ b/Engine/Source/Runtime/CoreUObject/Public/UObject/UnrealType.h
@@ -4052,18 +4052,21 @@ struct FPropertyChangedEvent
 		, ChangeType(EPropertyChangeType::Unspecified)
 		, ObjectIteratorIndex(INDEX_NONE)
 		, ArrayIndicesPerObject(nullptr)
+		, TopLevelObjects(nullptr)
 	{
 	}
 
-	FPropertyChangedEvent(UProperty* InProperty, EPropertyChangeType::Type InChangeType)
+	FPropertyChangedEvent(UProperty* InProperty, EPropertyChangeType::Type InChangeType, const TArray<const UObject*>* InTopLevelObjects = nullptr )
 		: Property(InProperty)
 		, MemberProperty(InProperty)
 		, ChangeType(InChangeType)
 		, ObjectIteratorIndex(INDEX_NONE)
 		, ArrayIndicesPerObject(nullptr)
+		, TopLevelObjects(InTopLevelObjects)
 	{
 	}
 
+
 	DEPRECATED(4.7, "The bInChangesTopology parameter has been removed, use the two-argument constructor of FPropertyChangedEvent instead")
 	FPropertyChangedEvent(UProperty* InProperty, const bool /*bInChangesTopology*/, EPropertyChangeType::Type InChangeType)
 		: Property(InProperty)
@@ -4071,6 +4074,7 @@ struct FPropertyChangedEvent
 		, ChangeType(InChangeType)
 		, ObjectIteratorIndex(INDEX_NONE)
 		, ArrayIndicesPerObject(nullptr)
+		, TopLevelObjects(nullptr)
 	{
 	}
 
@@ -4106,6 +4110,19 @@ struct FPropertyChangedEvent
 		return Retval;
 	}
 
+	/**
+	 * @return The number of objects being edited during this change event
+	 */
+	int32 GetNumObjectsBeingEdited() const { return TopLevelObjects ? TopLevelObjects->Num() : 0; }
+
+	/**
+	 * Gets an object being edited by this change event.  Multiple objects could be edited at once
+	 *
+	 * @param Index	The index of the object being edited. Assumes index is valid.  Call GetNumObjectsBeingEdited first to check if there are valid objects
+	 * @return The object being edited or nullptr if no object was found
+	 */
+	const UObject* GetObjectBeingEdited(int32 Index) const { return TopLevelObjects ? (*TopLevelObjects)[Index] : nullptr; }
+
 	/**
 	 * The actual property that changed
 	 */
@@ -4125,6 +4142,9 @@ struct FPropertyChangedEvent
 private:
 	//In the property window, multiple objects can be selected at once.  In the case of adding/inserting to an array, each object COULD have different indices for the new entries in the array
 	const TArray< TMap<FString,int32> >* ArrayIndicesPerObject;
+
+	/** List of top level objects being changed */
+	const TArray<const UObject*>* TopLevelObjects;
 };
 
 /**
diff --git a/Engine/Source/Runtime/Engine/Classes/Components/ActorComponent.h b/Engine/Source/Runtime/Engine/Classes/Components/ActorComponent.h
--- a/Engine/Source/Runtime/Engine/Classes/Components/ActorComponent.h
+++ b/Engine/Source/Runtime/Engine/Classes/Components/ActorComponent.h
@@ -223,6 +223,7 @@ public:
 
 	void DetermineUCSModifiedProperties();
 	void GetUCSModifiedProperties(TSet<const UProperty*>& ModifiedProperties) const;
+	void RemoveUCSModifiedProperties(const TArray<UProperty*>& Properties);
 
 	bool IsEditableWhenInherited() const;
 
diff --git a/Engine/Source/Runtime/Engine/Classes/Components/SplineComponent.h b/Engine/Source/Runtime/Engine/Classes/Components/SplineComponent.h
--- a/Engine/Source/Runtime/Engine/Classes/Components/SplineComponent.h
+++ b/Engine/Source/Runtime/Engine/Classes/Components/SplineComponent.h
@@ -33,31 +33,117 @@ namespace ESplineCoordinateSpace
 	};
 }
 
+USTRUCT()
+struct FSplineCurves
+{
+	GENERATED_BODY()
+
+	/** Spline built from position data. */
+	UPROPERTY()
+	FInterpCurveVector Position;
+
+	/** Spline built from rotation data. */
+	UPROPERTY()
+	FInterpCurveQuat Rotation;
+
+	/** Spline built from scale data. */
+	UPROPERTY()
+	FInterpCurveVector Scale;
+
+	/** Input: distance along curve, output: parameter that puts you there. */
+	UPROPERTY()
+	FInterpCurveFloat ReparamTable;
+
+	bool operator==(const FSplineCurves& Other) const
+	{
+		return Position == Other.Position && Rotation == Other.Rotation && Scale == Other.Scale;
+	}
+
+	bool operator!=(const FSplineCurves& Other) const
+	{
+		return !(*this == Other);
+	}
+};
+
+USTRUCT(BlueprintType)
+struct FSplinePoint
+{
+	GENERATED_BODY()
+
+	UPROPERTY(EditAnywhere, Category = SplinePoint)
+	float InputKey;
+
+	UPROPERTY(EditAnywhere, Category = SplinePoint)
+	FVector Position;
+
+	UPROPERTY(EditAnywhere, Category = SplinePoint)
+	FVector ArriveTangent;
+
+	UPROPERTY(EditAnywhere, Category = SplinePoint)
+	FVector LeaveTangent;
+
+	UPROPERTY(EditAnywhere, Category = SplinePoint)
+	FRotator Rotation;
+
+	UPROPERTY(EditAnywhere, Category = SplinePoint)
+	FVector Scale;
+
+	UPROPERTY(EditAnywhere, Category = SplinePoint)
+	TEnumAsByte<ESplinePointType::Type> Type;
+
+	/** Default constructor */
+	FSplinePoint()
+		: InputKey(0.0f), Position(0.0f), ArriveTangent(0.0f), LeaveTangent(0.0f), Rotation(0.0f), Scale(1.0f), Type(ESplinePointType::Curve)
+	{}
+
+	/** Constructor taking a point position */
+	FSplinePoint(float InInputKey, const FVector& InPosition)
+		: InputKey(InInputKey), Position(InPosition), ArriveTangent(0.0f), LeaveTangent(0.0f), Rotation(0.0f), Scale(1.0f), Type(ESplinePointType::Curve)
+	{}
+
+	/** Constructor taking a point position and type, and optionally rotation and scale */
+	FSplinePoint(float InInputKey, const FVector& InPosition, ESplinePointType::Type InType, const FRotator& InRotation = FRotator(0.0f), const FVector& InScale = FVector(1.0f))
+		: InputKey(InInputKey), Position(InPosition), ArriveTangent(0.0f), LeaveTangent(0.0f), Rotation(InRotation), Scale(InScale), Type(InType)
+	{}
+
+	/** Constructor taking a point position and tangent, and optionally rotation and scale */
+	FSplinePoint(float InInputKey,
+				 const FVector& InPosition,
+				 const FVector& InArriveTangent,
+				 const FVector& InLeaveTangent,
+				 const FRotator& InRotation = FRotator(0.0f),
+				 const FVector& InScale = FVector(1.0f),
+				 ESplinePointType::Type InType = ESplinePointType::CurveCustomTangent)
+		: InputKey(InInputKey), Position(InPosition), ArriveTangent(InArriveTangent), LeaveTangent(InLeaveTangent), Rotation(InRotation), Scale(InScale), Type(InType)
+	{}
+};
 
 /** 
- *	A spline component is a spline shape which can be used for other purposes (e.g. animating objects). It does not contain rendering capabilities itself (outside the editor) 
+ *	A spline component is a spline shape which can be used for other purposes (e.g. animating objects). It contains debug rendering capabilities.
  *	@see https://docs.unrealengine.com/latest/INT/Resources/ContentExamples/Blueprint_Splines
  */
-UCLASS(ClassGroup=Utility, meta=(BlueprintSpawnableComponent))
+UCLASS(ClassGroup=Utility, ShowCategories = (Mobility), HideCategories = (Physics, Collision, Lighting, Rendering, Mobile), meta=(BlueprintSpawnableComponent))
 class ENGINE_API USplineComponent : public UPrimitiveComponent
 {
 	GENERATED_UCLASS_BODY()
 
-	/** Spline built from position data. */
-	UPROPERTY()//EditAnywhere, Category = Points)
-	FInterpCurveVector SplineInfo;
+	UPROPERTY(EditAnywhere, Category=Points)
+	FSplineCurves SplineCurves;
 
-	/** Spline built from rotation data. */
+	/** Deprecated - please use GetSplinePointsPosition() to fetch this FInterpCurve */
 	UPROPERTY()
-	FInterpCurveQuat SplineRotInfo;
+	FInterpCurveVector SplineInfo_DEPRECATED;
 
-	/** Spline built from scale data. */
+	/** Deprecated - please use GetSplinePointsRotation() to fetch this FInterpCurve */
+	UPROPERTY()
+	FInterpCurveQuat SplineRotInfo_DEPRECATED;
+
+	/** Deprecated - please use GetSplinePointsScale() to fetch this FInterpCurve */
 	UPROPERTY()
-	FInterpCurveVector SplineScaleInfo;
+	FInterpCurveVector SplineScaleInfo_DEPRECATED;
 
-	/** Input, distance along curve, output, parameter that puts you there. */
 	UPROPERTY()
-	FInterpCurveFloat SplineReparamTable;
+	FInterpCurveFloat SplineReparamTable_DEPRECATED;
 
 	UPROPERTY()
 	bool bAllowSplineEditingPerInstance_DEPRECATED;
@@ -78,6 +164,10 @@ class ENGINE_API USplineComponent : public UPrimitiveComponent
 	UPROPERTY(EditAnywhere, Category = Spline, meta=(DisplayName="Override Construction Script"))
 	bool bSplineHasBeenEdited;
 
+	UPROPERTY()
+	/** Whether the UCS has made changes to the spline points */
+	bool bModifiedByConstructionScript;
+
 	/**
 	 * Whether the spline points should be passed to the User Construction Script so they can be further manipulated by it.
 	 * If false, they will not be visible to it, and it will not be able to influence the per-instance positions set in the editor.
@@ -97,6 +187,12 @@ private:
 	UPROPERTY(EditAnywhere, Category = Spline)
 	bool bClosedLoop;
 
+	UPROPERTY(EditAnywhere, Category = Spline, meta=(InlineEditConditionToggle=true))
+	bool bLoopPositionOverride;
+
+	UPROPERTY(EditAnywhere, Category = Spline, meta=(EditCondition="bLoopPositionOverride"))
+	float LoopPosition;
+
 public:
 	/** Default up vector in local space to be used when calculating transforms along the spline */
 	UPROPERTY(EditAnywhere, BlueprintReadWrite, Category = Spline)
@@ -111,20 +207,24 @@ public:
 	UPROPERTY(EditAnywhere, Category = Editor)
 	FLinearColor EditorSelectedSplineSegmentColor;
 
-	/** Whether scale visualization should be displayed */
+	/** Whether the spline's leave and arrive tangents can be different */
 	UPROPERTY(EditAnywhere, Category = Editor)
+	bool bAllowDiscontinuousSpline;
+
+	/** Whether scale visualization should be displayed */
+	UPROPERTY(EditAnywhere, Category = Editor, meta=(InlineEditConditionToggle=true))
 	bool bShouldVisualizeScale;
 
 	/** Width of spline in editor for use with scale visualization */
-	UPROPERTY(EditAnywhere, Category = Editor)
+	UPROPERTY(EditAnywhere, Category = Editor, meta=(EditCondition="bShouldVisualizeScale"))
 	float ScaleVisualizationWidth;
 #endif
 
 	//~ Begin UObject Interface
 	virtual void PostLoad() override;
 	virtual void PostEditImport() override;
-#if WITH_EDITOR
 	virtual void Serialize(FArchive& Ar) override;
+#if WITH_EDITOR
 	virtual void PostEditChangeChainProperty(FPropertyChangedChainEvent& PropertyChangedEvent) override;
 #endif
 	//~ End UObject Interface
@@ -143,6 +243,13 @@ public:
 	//~ End USceneComponent Interface
 #endif
 
+	FInterpCurveVector& GetSplinePointsPosition() { return SplineCurves.Position; }
+	const FInterpCurveVector& GetSplinePointsPosition() const { return SplineCurves.Position; }
+	FInterpCurveQuat& GetSplinePointsRotation() { return SplineCurves.Rotation; }
+	const FInterpCurveQuat& GetSplinePointsRotation() const { return SplineCurves.Rotation; }
+	FInterpCurveVector& GetSplinePointsScale() { return SplineCurves.Scale; }
+	const FInterpCurveVector& GetSplinePointsScale() const { return SplineCurves.Scale; }
+
 	void ApplyComponentInstanceData(class FSplineInstanceData* ComponentInstanceData, const bool bPostUCS);
 
 	/** Update the spline tangents and SplineReparamTable */
@@ -194,9 +301,13 @@ public:
 	UFUNCTION(BlueprintCallable, Category = Spline)
 	void SetDrawDebug(bool bShow);
 
-	/** Specify whether the spline is a closed loop or not */
+	/** Specify whether the spline is a closed loop or not. The loop position will be at 1.0 after the last point's input key */
+	UFUNCTION(BlueprintCallable, Category = Spline)
+	void SetClosedLoop(bool bInClosedLoop, bool bUpdateSpline = true);
+
+	/** Specify whether the spline is a closed loop or not, and if so, the input key corresponding to the loop point */
 	UFUNCTION(BlueprintCallable, Category = Spline)
-	void SetClosedLoop(bool bInClosedLoop);
+	void SetClosedLoopAtPosition(bool bInClosedLoop, float Key, bool bUpdateSpline = true);
 
 	/** Check whether the spline is a closed loop or not */
 	UFUNCTION(BlueprintCallable, Category = Spline)
@@ -204,19 +315,27 @@ public:
 
 	/** Clears all the points in the spline */
 	UFUNCTION(BlueprintCallable, Category = Spline)
-	void ClearSplinePoints();
+	void ClearSplinePoints(bool bUpdateSpline = true);
+
+	/** Adds an FSplinePoint to the spline. This contains its input key, position, tangent, rotation and scale. */
+	UFUNCTION(BlueprintCallable, Category = Spline)
+	void AddPoint(const FSplinePoint& Point, bool bUpdateSpline = true);
+
+	/** Adds an array of FSplinePoints to the spline. */
+	UFUNCTION(BlueprintCallable, Category = Spline)
+	void AddPoints(const TArray<FSplinePoint>& Points, bool bUpdateSpline = true);
 
 	/** Adds a point to the spline */
 	UFUNCTION(BlueprintCallable, Category = Spline)
-	void AddSplinePoint(const FVector& Position, ESplineCoordinateSpace::Type CoordinateSpace);
+	void AddSplinePoint(const FVector& Position, ESplineCoordinateSpace::Type CoordinateSpace, bool bUpdateSpline = true);
 
-	/** Adds a point to the spline at the specified index*/
+	/** Adds a point to the spline at the specified index */
 	UFUNCTION(BlueprintCallable, Category = Spline)
-	void AddSplinePointAtIndex(const FVector& Position, int32 Index, ESplineCoordinateSpace::Type CoordinateSpace);
+	void AddSplinePointAtIndex(const FVector& Position, int32 Index, ESplineCoordinateSpace::Type CoordinateSpace, bool bUpdateSpline = true);
 
 	/** Removes point at specified index from the spline */
 	UFUNCTION(BlueprintCallable, Category = Spline)
-	void RemoveSplinePoint(const int32 Index);
+	void RemoveSplinePoint(int32 Index, bool bUpdateSpline = true);
 
 	/** Adds a world space point to the spline */
 	UFUNCTION(BlueprintCallable, Category = Spline, meta = (DeprecatedFunction, DeprecationMessage = "Please use AddSplinePoint, specifying SplineCoordinateSpace::World"))
@@ -228,7 +347,7 @@ public:
 
 	/** Sets the spline to an array of points */
 	UFUNCTION(BlueprintCallable, Category = Spline)
-	void SetSplinePoints(const TArray<FVector>& Points, ESplineCoordinateSpace::Type CoordinateSpace);
+	void SetSplinePoints(const TArray<FVector>& Points, ESplineCoordinateSpace::Type CoordinateSpace, bool bUpdateSpline = true);
 
 	/** Sets the spline to an array of world space points */
 	UFUNCTION(BlueprintCallable, Category = Spline, meta = (DeprecatedFunction, DeprecationMessage = "Please use SetSplinePoints, specifying SplineCoordinateSpace::World"))
@@ -240,7 +359,7 @@ public:
 
 	/** Move an existing point to a new location */
 	UFUNCTION(BlueprintCallable, Category = Spline)
-	void SetLocationAtSplinePoint(int32 PointIndex, const FVector& InLocation, ESplineCoordinateSpace::Type CoordinateSpace);
+	void SetLocationAtSplinePoint(int32 PointIndex, const FVector& InLocation, ESplineCoordinateSpace::Type CoordinateSpace, bool bUpdateSpline = true);
 
 	/** Move an existing point to a new world location */
 	UFUNCTION(BlueprintCallable, Category = Spline, meta = (DeprecatedFunction, DeprecationMessage = "Please use SetLocationAtSplinePoint, specifying SplineCoordinateSpace::World"))
@@ -248,11 +367,15 @@ public:
 
 	/** Specify the tangent at a given spline point */
 	UFUNCTION(BlueprintCallable, Category = Spline)
-	void SetTangentAtSplinePoint(int32 PointIndex, const FVector& InTangent, ESplineCoordinateSpace::Type CoordinateSpace);
+	void SetTangentAtSplinePoint(int32 PointIndex, const FVector& InTangent, ESplineCoordinateSpace::Type CoordinateSpace, bool bUpdateSpline = true);
+
+	/** Specify the tangents at a given spline point */
+	UFUNCTION(BlueprintCallable, Category = Spline)
+	void SetTangentsAtSplinePoint(int32 PointIndex, const FVector& InArriveTangent, const FVector& InLeaveTangent, ESplineCoordinateSpace::Type CoordinateSpace, bool bUpdateSpline = true);
 
 	/** Specify the up vector at a given spline point */
 	UFUNCTION(BlueprintCallable, Category = Spline)
-	void SetUpVectorAtSplinePoint(int32 PointIndex, const FVector& InUpVector, ESplineCoordinateSpace::Type CoordinateSpace);
+	void SetUpVectorAtSplinePoint(int32 PointIndex, const FVector& InUpVector, ESplineCoordinateSpace::Type CoordinateSpace, bool bUpdateSpline = true);
 
 	/** Get the type of a spline point */
 	UFUNCTION(BlueprintCallable, Category = Spline)
@@ -260,7 +383,7 @@ public:
 
 	/** Specify the type of a spline point */
 	UFUNCTION(BlueprintCallable, Category = Spline)
-	void SetSplinePointType(int32 PointIndex, ESplinePointType::Type Type);
+	void SetSplinePointType(int32 PointIndex, ESplinePointType::Type Type, bool bUpdateSpline = true);
 
 	/** Get the number of points that make up this spline */
 	UFUNCTION(BlueprintCallable, Category = Spline)
@@ -278,10 +401,18 @@ public:
 	UFUNCTION(BlueprintCallable, Category = Spline)
 	FVector GetDirectionAtSplinePoint(int32 PointIndex, ESplineCoordinateSpace::Type CoordinateSpace) const;
 
-	/** Get the tangent at spline point */
+	/** Get the tangent at spline point. This fetches the Leave tangent of the point. */
 	UFUNCTION(BlueprintCallable, Category = Spline)
 	FVector GetTangentAtSplinePoint(int32 PointIndex, ESplineCoordinateSpace::Type CoordinateSpace) const;
 
+	/** Get the arrive tangent at spline point */
+	UFUNCTION(BlueprintCallable, Category = Spline)
+	FVector GetArriveTangentAtSplinePoint(int32 PointIndex, ESplineCoordinateSpace::Type CoordinateSpace) const;
+
+	/** Get the leave tangent at spline point */
+	UFUNCTION(BlueprintCallable, Category = Spline)
+	FVector GetLeaveTangentAtSplinePoint(int32 PointIndex, ESplineCoordinateSpace::Type CoordinateSpace) const;
+
 	/** Get the rotation at spline point as a quaternion */
 	FQuat GetQuaternionAtSplinePoint(int32 PointIndex, ESplineCoordinateSpace::Type CoordinateSpace) const;
 
@@ -495,6 +626,9 @@ private:
 	float GetSegmentParamFromLength(const int32 Index, const float Length, const float SegmentLength) const;
 };
 
+ENGINE_API EInterpCurveMode ConvertSplinePointTypeToInterpCurveMode(ESplinePointType::Type SplinePointType);
+ENGINE_API ESplinePointType::Type ConvertInterpCurveModeToSplinePointType(EInterpCurveMode InterpCurveMode);
+
 
 // Deprecated method definitions
 inline void USplineComponent::AddSplineWorldPoint(const FVector& Position) { AddSplinePoint(Position, ESplineCoordinateSpace::World); }
diff --git a/Engine/Source/Runtime/Engine/Classes/Components/StaticMeshComponent.h b/Engine/Source/Runtime/Engine/Classes/Components/StaticMeshComponent.h
--- a/Engine/Source/Runtime/Engine/Classes/Components/StaticMeshComponent.h
+++ b/Engine/Source/Runtime/Engine/Classes/Components/StaticMeshComponent.h
@@ -551,6 +551,14 @@ public:
 
 	/** Unregister this component's render data with the scene for SpeedTree wind */
 	void RemoveSpeedTreeWind();
+
+#if WITH_EDITOR
+	/** Called when the static mesh changes  */
+	DECLARE_EVENT_OneParam(UStaticMeshComponent, FOnStaticMeshChanged, UStaticMeshComponent*);
+	virtual FOnStaticMeshChanged& OnStaticMeshChanged() { return OnStaticMeshChangedEvent; }
+private:
+	FOnStaticMeshChanged OnStaticMeshChangedEvent;
+#endif
 };
 
 
diff --git a/Engine/Source/Runtime/Engine/Classes/EdGraph/EdGraphNode.h b/Engine/Source/Runtime/Engine/Classes/EdGraph/EdGraphNode.h
--- a/Engine/Source/Runtime/Engine/Classes/EdGraph/EdGraphNode.h
+++ b/Engine/Source/Runtime/Engine/Classes/EdGraph/EdGraphNode.h
@@ -164,7 +164,7 @@ class ENGINE_API UEdGraphNode : public UObject
 	UPROPERTY()
 	FString ErrorMsg;
 	
-	/** GUID to uniquely identify this node, to facilitate diff'ing versions of this graph */
+	/** GUID to uniquely identify this node, to facilitate diffing versions of this graph */
 	UPROPERTY()
 	FGuid NodeGuid;
 
@@ -230,6 +230,9 @@ public:
 	/** widget representing this node if it exists; Note: This is not safe to use in general and will be removed in the future, as there is no guarantee that only one graph editor/panel is viewing a given graph */
 	TWeakPtr<SGraphNode> DEPRECATED_NodeWidget;
 
+	/** Get all pins this node owns */
+	TArray<UEdGraphPin*> GetAllPins() { return Pins; }
+
 	/** Create a new pin on this node using the supplied info, and return the new pin */
 	UEdGraphPin* CreatePin(EEdGraphPinDirection Dir, const FString& PinCategory, const FString& PinSubCategory, UObject* PinSubCategoryObject, bool bIsArray, bool bIsReference, const FString& PinName, bool bIsConst = false, int32 Index = INDEX_NONE);
 
diff --git a/Engine/Source/Runtime/Engine/Classes/EditorFramework/AssetImportData.h b/Engine/Source/Runtime/Engine/Classes/EditorFramework/AssetImportData.h
--- a/Engine/Source/Runtime/Engine/Classes/EditorFramework/AssetImportData.h
+++ b/Engine/Source/Runtime/Engine/Classes/EditorFramework/AssetImportData.h
@@ -82,7 +82,7 @@ public:
 	static FOnImportDataChanged OnImportDataChanged;
 
 	/** Update this import data using the specified file. Called when an asset has been imported from a file. */
-	void Update(const FString& AbsoluteFilename);
+	void Update(const FString& AbsoluteFilename, FMD5Hash *Md5Hash = nullptr);
 
 	/** Update this import data using the specified filename. Will not update the imported timestamp or MD5 (so we can update files when they move). */
 	void UpdateFilenameOnly(const FString& InPath);
diff --git a/Engine/Source/Runtime/Engine/Classes/Engine/GameViewportDelegates.h b/Engine/Source/Runtime/Engine/Classes/Engine/GameViewportDelegates.h
--- a/Engine/Source/Runtime/Engine/Classes/Engine/GameViewportDelegates.h
+++ b/Engine/Source/Runtime/Engine/Classes/Engine/GameViewportDelegates.h
@@ -36,10 +36,4 @@ DECLARE_MULTICAST_DELEGATE_OneParam(FOnGameViewportClientPlayerAction, int32);
  * @see UGameViewportClient
  */
 DECLARE_MULTICAST_DELEGATE_OneParam(FOnGameViewportTick, float);
-
-/**
-* Delegate type used by UGameViewportClient when received focus
-*
-* @see UGameViewportClient
-*/
 DECLARE_DELEGATE_RetVal_TwoParams(bool, FOnGameViewportInputKey, FKey, FModifierKeysState);
diff --git a/Engine/Source/Runtime/Engine/Classes/Engine/Polys.h b/Engine/Source/Runtime/Engine/Classes/Engine/Polys.h
--- a/Engine/Source/Runtime/Engine/Classes/Engine/Polys.h
+++ b/Engine/Source/Runtime/Engine/Classes/Engine/Polys.h
@@ -42,6 +42,7 @@ public:
 	FName				RulesetVariation;		// Name of variation within a ProcBuilding Ruleset for this face
 	FName				ItemName;				// Item name.
 	int32					iLink;					// iBspSurf, or brush fpoly index of first identical polygon, or MAX_uint16.
+	int32					iLinkSurf;
 	int32					iBrushPoly;				// Index of editor solid's polygon this originated from.
 	uint32				SmoothingMask;			// A mask used to determine which smoothing groups this polygon is in.  SmoothingMask & (1 << GroupNumber)
 	float				LightMapScale;			// The number of units/shadowmap texel on this surface.
diff --git a/Engine/Source/Runtime/Engine/Classes/Engine/Selection.h b/Engine/Source/Runtime/Engine/Classes/Engine/Selection.h
--- a/Engine/Source/Runtime/Engine/Classes/Engine/Selection.h
+++ b/Engine/Source/Runtime/Engine/Classes/Engine/Selection.h
@@ -18,8 +18,36 @@ class ENGINE_API USelection : public UObject
 {
 	GENERATED_UCLASS_BODY()
 private:
+	/** Contains info about each class and how many objects of that class are selected */
+	struct FSelectedClassInfo
+	{
+		/** The selected class */
+		const UClass* Class;
+		/** How many objects of that class are selected */
+		int32 SelectionCount;
+
+		FSelectedClassInfo(const UClass* InClass)
+			: Class(InClass)
+			, SelectionCount(0)
+		{}
+
+		FSelectedClassInfo(const UClass* InClass, int32 InSelectionCount)
+			: Class(InClass)
+			, SelectionCount(InSelectionCount)
+		{}
+		bool operator==(const FSelectedClassInfo& Info) const
+		{
+			return Class == Info.Class;
+		}
+
+		friend uint32 GetTypeHash(const FSelectedClassInfo& Info)
+		{
+			return GetTypeHash(Info.Class);
+		}
+	};
+
 	typedef TArray<TWeakObjectPtr<UObject> >	ObjectArray;
-	typedef TMRUArray<UClass*>	ClassArray;
+	typedef TSet<FSelectedClassInfo>	ClassArray;
 
 	template<typename SelectionFilter>
 	friend class TSelectionIterator;
@@ -61,7 +89,7 @@ public:
 	 */
 	UObject* GetSelectedObject(const int32 InIndex)
 	{
-		return (SelectedObjects.IsValidIndex(InIndex) && SelectedObjects[InIndex].IsValid() ? SelectedObjects[InIndex].Get() : NULL);
+		return (SelectedObjects.IsValidIndex(InIndex) && SelectedObjects[InIndex].IsValid() ? SelectedObjects[InIndex].Get() : nullptr);
 	}
 
 	/**
@@ -69,7 +97,7 @@ public:
 	 */
 	const UObject* GetSelectedObject(const int32 InIndex) const
 	{
-		return (SelectedObjects.IsValidIndex(InIndex) && SelectedObjects[InIndex].IsValid() ? SelectedObjects[InIndex].Get() : NULL);
+		return (SelectedObjects.IsValidIndex(InIndex) && SelectedObjects[InIndex].IsValid() ? SelectedObjects[InIndex].Get() : nullptr);
 	}
 
 	/**
@@ -156,7 +184,7 @@ public:
 	 * @param	bArchetypesOnly		[opt] true to only return archetype objects, false otherwise
 	 * @return						The first selected object of the specified class.
 	 */
-	UObject* GetTop(UClass* InClass, UClass* RequiredInterface=NULL, bool bArchetypesOnly=false)
+	UObject* GetTop(UClass* InClass, UClass* RequiredInterface=nullptr, bool bArchetypesOnly=false)
 	{
 		check( InClass );
 		for( int32 i=0; i<SelectedObjects.Num(); ++i )
@@ -188,7 +216,7 @@ public:
 				}
 			}
 		}
-		return NULL;
+		return nullptr;
 	}
 
 	/**
@@ -208,7 +236,7 @@ public:
 				return SelectedObject;
 			}
 		}
-		return NULL;
+		return nullptr;
 	}
 
 	/**
@@ -219,7 +247,7 @@ public:
 	template< class T > T* GetTop()
 	{
 		UObject* Selected = GetTop(T::StaticClass());
-		return Selected ? CastChecked<T>(Selected) : NULL;
+		return Selected ? CastChecked<T>(Selected) : nullptr;
 	}
 
 	/**
@@ -230,7 +258,7 @@ public:
 	template< class T > T* GetBottom()
 	{
 		UObject* Selected = GetBottom(T::StaticClass());
-		return Selected ? CastChecked<T>(Selected) : NULL;
+		return Selected ? CastChecked<T>(Selected) : nullptr;
 	}
 
 	/**
@@ -284,11 +312,18 @@ public:
 	 *
 	 * @return			The selected class at the specified index.
 	 */
+	DEPRECATED(4.14, "GetSelectedClass is deprecated.  Use IsClassSelected or a class iterator to search through classes")
 	UClass* GetSelectedClass(int32 InIndex) const
 	{
-		return SelectedClasses[ InIndex ];
+		FSetElementId Id = FSetElementId::FromInteger(InIndex);
+		return const_cast<UClass*>(SelectedClasses[Id].Class);
 	}
 
+	bool IsClassSelected(UClass* Class) const
+	{
+		const FSelectedClassInfo* Info = SelectedClasses.Find(Class);
+		return Info && Info->SelectionCount > 0;
+	}
 
 	//~ Begin UObject Interface
 	virtual void Serialize(FArchive& Ar) override;
diff --git a/Engine/Source/Runtime/Engine/Classes/Kismet/KismetTextLibrary.h b/Engine/Source/Runtime/Engine/Classes/Kismet/KismetTextLibrary.h
--- a/Engine/Source/Runtime/Engine/Classes/Kismet/KismetTextLibrary.h
+++ b/Engine/Source/Runtime/Engine/Classes/Kismet/KismetTextLibrary.h
@@ -206,18 +206,30 @@ class ENGINE_API UKismetTextLibrary : public UBlueprintFunctionLibrary
 	UFUNCTION(BlueprintPure, meta=(DisplayName = "AsPercent", AdvancedDisplay = "1"), Category="Utilities|Text")
 	static FText AsPercent_Float(float Value, TEnumAsByte<ERoundingMode> RoundingMode, bool bUseGrouping = true, int32 MinimumIntegralDigits = 1, int32 MaximumIntegralDigits = 324, int32 MinimumFractionalDigits = 0, int32 MaximumFractionalDigits = 3);
 
-	/* Converts a passed in date & time to a text, formatted as a date */
+	/* Converts a passed in date & time to a text, formatted as a date using an invariant timezone. This will use the given date & time as-is, so it's assumed to already be in the correct timezone. */
 	UFUNCTION(BlueprintPure, meta=(DisplayName = "AsDate", AdvancedDisplay = "1"), Category="Utilities|Text")
 	static FText AsDate_DateTime(const FDateTime& InDateTime);
 
-	/* Converts a passed in date & time to a text, formatted as a date & time */
+	/* Converts a passed in date & time to a text, formatted as a date using the given timezone (default is the local timezone). This will convert the given date & time from UTC to the given timezone (taking into account DST). */
+	UFUNCTION(BlueprintPure, meta=(DisplayName = "AsDate (from UTC)", AdvancedDisplay = "1"), Category="Utilities|Text")
+	static FText AsTimeZoneDate_DateTime(const FDateTime& InDateTime, const FString& InTimeZone = TEXT(""));
+
+	/* Converts a passed in date & time to a text, formatted as a date & time using an invariant timezone. This will use the given date & time as-is, so it's assumed to already be in the correct timezone. */
 	UFUNCTION(BlueprintPure, meta=(DisplayName = "AsDateTime", AdvancedDisplay = "1"), Category="Utilities|Text")
 	static FText AsDateTime_DateTime(const FDateTime& In);
 
-	/* Converts a passed in date & time to a text, formatted as a time */
+	/* Converts a passed in date & time to a text, formatted as a date & time using the given timezone (default is the local timezone). This will convert the given date & time from UTC to the given timezone (taking into account DST). */
+	UFUNCTION(BlueprintPure, meta=(DisplayName = "AsDateTime (from UTC)", AdvancedDisplay = "1"), Category="Utilities|Text")
+	static FText AsTimeZoneDateTime_DateTime(const FDateTime& InDateTime, const FString& InTimeZone = TEXT(""));
+
+	/* Converts a passed in date & time to a text, formatted as a time using an invariant timezone. This will use the given date & time as-is, so it's assumed to already be in the correct timezone. */
 	UFUNCTION(BlueprintPure, meta=(DisplayName = "AsTime", AdvancedDisplay = "1"), Category="Utilities|Text")
 	static FText AsTime_DateTime(const FDateTime& In);
 
+	/* Converts a passed in date & time to a text, formatted as a time using the given timezone (default is the local timezone). This will convert the given date & time from UTC to the given timezone (taking into account DST). */
+	UFUNCTION(BlueprintPure, meta=(DisplayName = "AsTime (from UTC)", AdvancedDisplay = "1"), Category="Utilities|Text")
+	static FText AsTimeZoneTime_DateTime(const FDateTime& InDateTime, const FString& InTimeZone = TEXT(""));
+
 	/* Converts a passed in time span to a text, formatted as a time span */
 	UFUNCTION(BlueprintPure, meta=(DisplayName = "AsTimespan", AdvancedDisplay = "1"), Category="Utilities|Text")
 	static FText AsTimespan_Timespan(const FTimespan& InTimespan);
diff --git a/Engine/Source/Runtime/Engine/Private/ActorConstruction.cpp b/Engine/Source/Runtime/Engine/Private/ActorConstruction.cpp
--- a/Engine/Source/Runtime/Engine/Private/ActorConstruction.cpp
+++ b/Engine/Source/Runtime/Engine/Private/ActorConstruction.cpp
@@ -348,32 +348,106 @@ void AActor::RerunConstructionScripts()
 		}
 
 #if WITH_EDITOR
-		// Save the current construction script-created components by name
-		TMap<const FName, UObject*> DestroyedComponentsByName;
+
+		// Return the component which was added by the construction script.
+		// It may be the same as the argument, or a parent component if the argument was a native subobject.
+		auto GetComponentAddedByConstructionScript = [](UActorComponent* Component) -> UActorComponent*
+		{
+			while (Component)
+			{
+				if (Component->IsCreatedByConstructionScript())
+				{
+					return Component;
+				}
+
+				Component = Component->GetTypedOuter<UActorComponent>();
+			}
+
+			return nullptr;
+		};
+
+		// Build a list of previously attached components which will be matched with their newly instanced counterparts.
+		// Components which will be reinstanced may be created by the SCS or the UCS.
+		// SCS components can only be matched by name, and outermost parent to resolve duplicated names.
+		// UCS components remember a serialized index which is used to identify them in the case that the UCS adds many of the same type.
+
 		TInlineComponentArray<UActorComponent*> PreviouslyAttachedComponents;
 		GetComponents(PreviouslyAttachedComponents);
+
+		struct FComponentData
+		{
+			UActorComponent* OldComponent;
+			UActorComponent* OldOuter;
+			UObject* OldArchetype;
+			FName OldName;
+			int32 UCSComponentIndex;
+		};
+
+		TArray<FComponentData> ComponentMapping;
+		ComponentMapping.Reserve(PreviouslyAttachedComponents.Num());
+		int32 IndexOffset = 0;
+
 		for (UActorComponent* Component : PreviouslyAttachedComponents)
 		{
-			if (Component)
+			// Look for the outermost component object.
+			// Normally components have their parent actor as their outer, but it's possible that a native component may construct a subobject component.
+			// In this case we need to "tunnel out" to find the parent component which has been created by the construction script.
+			if (UActorComponent* CSAddedComponent = GetComponentAddedByConstructionScript(Component))
 			{
-				if (Component->IsCreatedByConstructionScript())
-				{
+				// Determine if this component is an inner of a component added by the construction script
+				const bool bIsInnerComponent = (CSAddedComponent != Component);
 
-					DestroyedComponentsByName.Add(Component->GetFName(), Component);
+				// Poor man's topological sort - try to ensure that children are added to the list after the parents
+				// IndexOffset specifies how many items from the end new items are added.
+				const int32 Index = ComponentMapping.Num() - IndexOffset;
+				if (bIsInnerComponent)
+				{
+					int32 OuterIndex = ComponentMapping.IndexOfByPredicate([CSAddedComponent](const FComponentData& CD) { return CD.OldComponent == CSAddedComponent; });
+					if (OuterIndex == INDEX_NONE)
+					{
+						// If we find an item whose parent isn't yet in the list, we put it at the end, and then force all subsequent items to be added before.
+						// TODO: improve this, it may fail in certain circumstances, but a full topological ordering is far more complicated a problem.
+						IndexOffset++;
+					}
 				}
-				else
+
+				// Add a new item
+				ComponentMapping.Insert(FComponentData(), Index);
+				ComponentMapping[Index].OldComponent = Component;
+				ComponentMapping[Index].OldOuter = bIsInnerComponent ? CSAddedComponent : nullptr;
+				ComponentMapping[Index].OldArchetype = Component->GetArchetype();
+				ComponentMapping[Index].OldName = Component->GetFName();
+
+				// If it's a UCS-created component, store a serialized index which will be used to match it to the reinstanced counterpart later
+				int32 SerializedIndex = -1;
+				if (Component->CreationMethod == EComponentCreationMethod::UserConstructionScript)
 				{
-					UActorComponent* OuterComponent = Component->GetTypedOuter<UActorComponent>();
-					while (OuterComponent)
+					bool bFound = false;
+					for (const UActorComponent* BlueprintCreatedComponent : BlueprintCreatedComponents)
 					{
-						if (OuterComponent->IsCreatedByConstructionScript())
+						if (BlueprintCreatedComponent)
 						{
-							DestroyedComponentsByName.Add(Component->GetFName(), Component);
-							break;
+							if (BlueprintCreatedComponent == Component)
+							{
+								SerializedIndex++;
+								bFound = true;
+								break;
+							}
+							else if (BlueprintCreatedComponent->CreationMethod == EComponentCreationMethod::UserConstructionScript &&
+									 BlueprintCreatedComponent->GetArchetype() == ComponentMapping[Index].OldArchetype)
+							{
+								SerializedIndex++;
+							}
 						}
-						OuterComponent = OuterComponent->GetTypedOuter<UActorComponent>();
+					}
+
+					if (!bFound)
+					{
+						SerializedIndex = -1;
 					}
 				}
+
+				ComponentMapping[Index].UCSComponentIndex = SerializedIndex;
 			}
 		}
 #endif
@@ -431,16 +505,100 @@ void AActor::RerunConstructionScripts()
 
 #if WITH_EDITOR
 		// Create the mapping of old->new components and notify the editor of the replacements
-		TMap<UObject*, UObject*> OldToNewComponentMapping;
-
 		TInlineComponentArray<UActorComponent*> NewComponents;
 		GetComponents(NewComponents);
-		for (UActorComponent* NewComp : NewComponents)
+
+		TMap<UObject*, UObject*> OldToNewComponentMapping;
+		OldToNewComponentMapping.Reserve(NewComponents.Num());
+
+		// Build some quick lookup maps for speedy access.
+		// The NameToNewComponent map is a multimap because names are not necessarily unique.
+		// For example, there may be two components, subobjects of components added by the construction script, which have the same name, because they are unique in their scope.
+		TMultiMap<FName, UActorComponent*> NameToNewComponent;
+		TMap<UActorComponent*, UObject*> ComponentToArchetypeMap;
+		NameToNewComponent.Reserve(NewComponents.Num());
+		ComponentToArchetypeMap.Reserve(NewComponents.Num());
+
+		for (UActorComponent* NewComponent : NewComponents)
 		{
-			const FName NewCompName = NewComp->GetFName();
-			if (DestroyedComponentsByName.Contains(NewCompName))
+			if (GetComponentAddedByConstructionScript(NewComponent))
+			{
+				NameToNewComponent.Add(NewComponent->GetFName(), NewComponent);
+				ComponentToArchetypeMap.Add(NewComponent, NewComponent->GetArchetype());
+			}
+		}
+
+		// Now iterate through all previous construction script created components, looking for a match with reinstanced components.
+		for (const FComponentData& ComponentData : ComponentMapping)
+		{
+			if (ComponentData.OldComponent->CreationMethod == EComponentCreationMethod::UserConstructionScript)
+			{
+				// If created by the UCS, look for a component whose class, archetype and serialized index matches
+				for (UActorComponent* NewComponent : NewComponents)
+				{
+					if (NewComponent->CreationMethod == EComponentCreationMethod::UserConstructionScript &&
+						ComponentData.OldComponent->GetClass() == NewComponent->GetClass() &&
+						ComponentData.OldArchetype == NewComponent->GetArchetype() &&
+						ComponentData.UCSComponentIndex >= 0)
+					{
+						int32 FoundSerializedIndex = -1;
+						bool bMatches = false;
+						for (UActorComponent* BlueprintCreatedComponent : BlueprintCreatedComponents)
+						{
+							if (BlueprintCreatedComponent && BlueprintCreatedComponent->CreationMethod == EComponentCreationMethod::UserConstructionScript)
+							{
+								UObject* BlueprintComponentTemplate = ComponentToArchetypeMap.FindChecked(BlueprintCreatedComponent);
+								if (ComponentData.OldArchetype == BlueprintComponentTemplate &&
+									++FoundSerializedIndex == ComponentData.UCSComponentIndex)
+								{
+									bMatches = (BlueprintCreatedComponent == NewComponent);
+									break;
+								}
+							}
+						}
+
+						if (bMatches)
+						{
+							OldToNewComponentMapping.Add(ComponentData.OldComponent, NewComponent);
+							break;
+						}
+					}
+				}
+			}
+			else
 			{
-				OldToNewComponentMapping.Add(DestroyedComponentsByName[NewCompName], NewComp);
+				// Component added by the SCS. We can't rely on serialization order as this can change.
+				// Instead look for matching names, and, if there's an outer component, look for a match there.
+				TArray<UActorComponent*> MatchedComponents;
+				NameToNewComponent.MultiFind(ComponentData.OldName, MatchedComponents);
+				if (MatchedComponents.Num() > 0)
+				{
+					UActorComponent* OuterToMatch = ComponentData.OldOuter;
+					if (OuterToMatch)
+					{
+						// The saved outer component is the previous component, hence we transform it to the new one through the OldToNewComponentMapping
+						// before comparing with the new outer to match.
+						// We can rely on this because the ComponentMapping list is ordered topologically, such that parents appear before children.
+						if (UObject** NewOuterToMatch = OldToNewComponentMapping.Find(OuterToMatch))
+						{
+							OuterToMatch = Cast<UActorComponent>(*NewOuterToMatch);
+						}
+						else
+						{
+							OuterToMatch = nullptr;
+						}
+					}
+
+					// Now look for a match within the set of possible matches
+					for (UActorComponent* MatchedComponent : MatchedComponents)
+					{
+						if (!OuterToMatch || GetComponentAddedByConstructionScript(MatchedComponent) == OuterToMatch)
+						{
+							OldToNewComponentMapping.Add(ComponentData.OldComponent, MatchedComponent);
+							break;
+						}
+					}
+				}
 			}
 		}
 
diff --git a/Engine/Source/Runtime/Engine/Private/Brush.cpp b/Engine/Source/Runtime/Engine/Private/Brush.cpp
--- a/Engine/Source/Runtime/Engine/Private/Brush.cpp
+++ b/Engine/Source/Runtime/Engine/Private/Brush.cpp
@@ -25,6 +25,8 @@ TArray< TWeakObjectPtr< ULevel > > ABrush::LevelsToRebuild;
 bool ABrush::bSuppressBSPRegeneration = false;
 #endif
 
+DEFINE_LOG_CATEGORY_STATIC(LogBrush, Log, All);
+
 ABrush::ABrush(const FObjectInitializer& ObjectInitializer)
 	: Super(ObjectInitializer)
 {
@@ -190,6 +192,30 @@ void ABrush::PostLoad()
 			}
 		}
 
+		// Fix up any broken poly normals.
+		// They have not been getting fixed up after vertex editing since at least UE2!
+		if(0)
+		{
+			bool bHasBrokenNormals = false;
+			for(FPoly& Poly : Brush->Polys->Element)
+			{
+				FVector Normal = Poly.Normal;
+				if(!Poly.CalcNormal())
+				{
+					if(!Poly.Normal.Equals(Normal))
+					{
+						UE_LOG(LogBrush, Warning, TEXT("%s had invalid poly normals which have been fixed. Resave the level to remove this warning."), *Brush->GetName());
+						if(IsStaticBrush())
+						{
+							// Flag BSP as needing rebuild
+							SetNeedRebuild(GetLevel());
+							bHasBrokenNormals = true;
+						}
+					}
+				}
+			}
+		}
+
 		// if the polys of the brush have the wrong outer, fix it up to be the UModel (my Brush member)
 		// UModelFactory::FactoryCreateText was passing in the ABrush as the Outer instead of the UModel
 		if (Brush->Polys->GetOuter() == this)
diff --git a/Engine/Source/Runtime/Engine/Private/Components/ActorComponent.cpp b/Engine/Source/Runtime/Engine/Private/Components/ActorComponent.cpp
--- a/Engine/Source/Runtime/Engine/Private/Components/ActorComponent.cpp
+++ b/Engine/Source/Runtime/Engine/Private/Components/ActorComponent.cpp
@@ -1612,6 +1612,16 @@ void UActorComponent::GetUCSModifiedProperties(TSet<const UProperty*>& ModifiedP
 	}
 }
 
+void UActorComponent::RemoveUCSModifiedProperties(const TArray<UProperty*>& Properties)
+{
+	for (UProperty* Property : Properties)
+	{
+		FSimpleMemberReference MemberReference;
+		FMemberReference::FillSimpleMemberReference<UProperty>(Property, MemberReference);
+		UCSModifiedProperties.RemoveSwap(MemberReference);
+	}
+}
+
 void UActorComponent::SetCanEverAffectNavigation(bool bRelevant)
 {
 	if (bCanEverAffectNavigation != bRelevant)
diff --git a/Engine/Source/Runtime/Engine/Private/Components/PrimitiveComponent.cpp b/Engine/Source/Runtime/Engine/Private/Components/PrimitiveComponent.cpp
--- a/Engine/Source/Runtime/Engine/Private/Components/PrimitiveComponent.cpp
+++ b/Engine/Source/Runtime/Engine/Private/Components/PrimitiveComponent.cpp
@@ -719,35 +719,43 @@ bool UPrimitiveComponent::CanEditChange(const UProperty* InProperty) const
 	{
 		const FName PropertyName = InProperty->GetFName();
 
-		if (PropertyName == GET_MEMBER_NAME_CHECKED(UPrimitiveComponent, bLightAsIfStatic))
+		static FName LightAsIfStaticName = GET_MEMBER_NAME_CHECKED(UPrimitiveComponent, bLightAsIfStatic);
+		static FName LightmassSettingsName = TEXT("LightmassSettings");
+		static FName LightingChannelsName = GET_MEMBER_NAME_CHECKED(UPrimitiveComponent, LightingChannels);
+		static FName SingleSampleShadowFromStationaryLightsName = GET_MEMBER_NAME_CHECKED(UPrimitiveComponent, bSingleSampleShadowFromStationaryLights);
+		static FName IndirectLightingCacheQualityName = GET_MEMBER_NAME_CHECKED(UPrimitiveComponent, IndirectLightingCacheQuality);
+		static FName CastCinematicShadowName = GET_MEMBER_NAME_CHECKED(UPrimitiveComponent, bCastCinematicShadow);
+		static FName CastInsetShadowName = GET_MEMBER_NAME_CHECKED(UPrimitiveComponent, bCastInsetShadow);
+		static FName CastShadowName = GET_MEMBER_NAME_CHECKED(UPrimitiveComponent, CastShadow);
+
+		if (PropertyName == LightAsIfStaticName)
 		{
 			// Disable editing bLightAsIfStatic on static components, since it has no effect
 			return Mobility != EComponentMobility::Static;
 		}
 
-		if (PropertyName == TEXT("LightmassSettings"))
+		if (PropertyName == LightmassSettingsName)
 		{
 			return Mobility != EComponentMobility::Movable || bLightAsIfStatic;
 		}
 
-		if (PropertyName == GET_MEMBER_NAME_STRING_CHECKED(UPrimitiveComponent, LightingChannels)
-			|| PropertyName == GET_MEMBER_NAME_STRING_CHECKED(UPrimitiveComponent, bSingleSampleShadowFromStationaryLights))
+		if (PropertyName == LightingChannelsName
+			|| PropertyName ==SingleSampleShadowFromStationaryLightsName)
 		{
 			return Mobility != EComponentMobility::Static;
 		}
 
-		if (PropertyName == GET_MEMBER_NAME_CHECKED(UPrimitiveComponent, IndirectLightingCacheQuality)
-			|| PropertyName == GET_MEMBER_NAME_CHECKED(UPrimitiveComponent, bCastCinematicShadow))
+		if (PropertyName == IndirectLightingCacheQualityName || PropertyName == CastCinematicShadowName)
 		{
 			return Mobility == EComponentMobility::Movable;
 		}
 
-		if (PropertyName == GET_MEMBER_NAME_CHECKED(UPrimitiveComponent, bCastInsetShadow))
+		if (PropertyName == CastInsetShadowName)
 		{
 			return !bSelfShadowOnly;
 		}
 
-		if (PropertyName == GET_MEMBER_NAME_CHECKED(UPrimitiveComponent, CastShadow))
+		if (PropertyName == CastShadowName)
 		{
 			// Look for any lit materials
 			bool bHasAnyLitMaterials = false;
