Copying //UE4/Dev-Build to //UE4/Dev-Main (Source: //UE4/Dev-Build @ 3092544)

diff --git a/Engine/Source/Programs/UnrealBuildTool/Configuration/UEBuildTarget.cs b/Engine/Source/Programs/UnrealBuildTool/Configuration/UEBuildTarget.cs
--- a/Engine/Source/Programs/UnrealBuildTool/Configuration/UEBuildTarget.cs
+++ b/Engine/Source/Programs/UnrealBuildTool/Configuration/UEBuildTarget.cs
@@ -2711,6 +2711,12 @@ namespace UnrealBuildTool
 					}
 				}
 				AppBinaries = FilteredBinaries;
+
+				// Check we have at least one match
+				if(AppBinaries.Count == 0)
+				{
+					throw new BuildException("Couldn't find any module containing {0} in {1}.", SingleFileItem.Reference, TargetName);
+				}
 			}
 
 			if (ShouldCheckOutputDistributionLevel() && !ProjectFileGenerator.bGenerateProjectFiles)
diff --git a/Engine/Source/Programs/UnrealBuildTool/System/Plugins.cs b/Engine/Source/Programs/UnrealBuildTool/System/Plugins.cs
--- a/Engine/Source/Programs/UnrealBuildTool/System/Plugins.cs
+++ b/Engine/Source/Programs/UnrealBuildTool/System/Plugins.cs
@@ -62,6 +62,34 @@ namespace UnrealBuildTool
 		/// </summary>
 		static Dictionary<DirectoryReference, List<FileReference>> PluginFileCache = new Dictionary<DirectoryReference, List<FileReference>>();
 
+		/// <summary>
+		/// Filters the list of plugins to ensure that any game plugins override engine plugins with the same name, and otherwise that no two
+		/// plugins with the same name exist. 
+		/// </summary>
+		/// <param name="Plugins">List of plugins to filter</param>
+		/// <returns>Filtered list of plugins in the original order</returns>
+		public static IEnumerable<PluginInfo> FilterPlugins(IEnumerable<PluginInfo> Plugins)
+		{
+			Dictionary<string, PluginInfo> NameToPluginInfo = new Dictionary<string, PluginInfo>(StringComparer.InvariantCultureIgnoreCase);
+			foreach(PluginInfo Plugin in Plugins)
+			{
+				PluginInfo ExistingPluginInfo;
+				if(!NameToPluginInfo.TryGetValue(Plugin.Name, out ExistingPluginInfo))
+				{
+					NameToPluginInfo.Add(Plugin.Name, Plugin);
+				}
+				else if(ExistingPluginInfo.LoadedFrom == PluginLoadedFrom.Engine && Plugin.LoadedFrom == PluginLoadedFrom.GameProject)
+				{
+					NameToPluginInfo[Plugin.Name] = Plugin;
+				}
+				else if(ExistingPluginInfo.LoadedFrom != PluginLoadedFrom.GameProject || Plugin.LoadedFrom != PluginLoadedFrom.Engine)
+				{
+					throw new BuildException(String.Format("Found '{0}' plugin in two locations ({1} and {2}). Plugin names must be unique.", Plugin.Name, ExistingPluginInfo.File, Plugin.File));
+				}
+			}
+			return Plugins.Where(x => NameToPluginInfo[x.Name] == x);
+		}
+
 		/// <summary>
 		/// Read all the plugins available to a given project
 		/// </summary>
@@ -83,7 +111,7 @@ namespace UnrealBuildTool
 				Plugins.AddRange(ReadPluginsFromDirectory(ProjectPluginsDir, PluginLoadedFrom.GameProject));
 			}
 
-			return Plugins;
+			return FilterPlugins(Plugins).ToList();
 		}
 
 		/// <summary>
diff --git a/Engine/Source/Programs/UnrealBuildTool/System/ProjectFileGenerator.cs b/Engine/Source/Programs/UnrealBuildTool/System/ProjectFileGenerator.cs
--- a/Engine/Source/Programs/UnrealBuildTool/System/ProjectFileGenerator.cs
+++ b/Engine/Source/Programs/UnrealBuildTool/System/ProjectFileGenerator.cs
@@ -1415,9 +1415,10 @@ namespace UnrealBuildTool
 
 						Log.TraceVerbose( "Found target: " + TargetName );
 
-						string[] ArgumentsCopy = new string[ Arguments.Length + 1 ];
+						string[] ArgumentsCopy = new string[ Arguments.Length + 2 ];
 						ArgumentsCopy[ 0 ] = TargetName;
-						Array.Copy(Arguments, 0, ArgumentsCopy, 1, Arguments.Length);
+						ArgumentsCopy[ 1 ] = "-precompile";
+						Array.Copy(Arguments, 0, ArgumentsCopy, 2, Arguments.Length);
 
 						// We only want to update definitions and include paths for modules that are part of this target's project file.
 						ProjectFileGenerator.OnlyGenerateIntelliSenseDataForProject = TargetProjectFile;
diff --git a/Engine/Source/Programs/UnrealBuildTool/System/RulesCompiler.cs b/Engine/Source/Programs/UnrealBuildTool/System/RulesCompiler.cs
--- a/Engine/Source/Programs/UnrealBuildTool/System/RulesCompiler.cs
+++ b/Engine/Source/Programs/UnrealBuildTool/System/RulesCompiler.cs
@@ -1625,6 +1625,15 @@ namespace UnrealBuildTool
 		/// </summary>
 		/// <returns></returns>
 		public IEnumerable<PluginInfo> EnumeratePlugins()
+		{
+			return global::UnrealBuildTool.Plugins.FilterPlugins(EnumeratePluginsInternal());
+		}
+
+		/// <summary>
+		/// Enumerates all the plugins that are available
+		/// </summary>
+		/// <returns></returns>
+		protected IEnumerable<PluginInfo> EnumeratePluginsInternal()
 		{
 			if (Parent == null)
 			{
@@ -1632,7 +1641,7 @@ namespace UnrealBuildTool
 			}
 			else
 			{
-				return Plugins.Concat(Parent.EnumeratePlugins());
+				return Plugins.Concat(Parent.EnumeratePluginsInternal());
 			}
 		}
 
diff --git a/Engine/Source/Programs/UnrealBuildTool/System/UnrealBuildTool.cs b/Engine/Source/Programs/UnrealBuildTool/System/UnrealBuildTool.cs
--- a/Engine/Source/Programs/UnrealBuildTool/System/UnrealBuildTool.cs
+++ b/Engine/Source/Programs/UnrealBuildTool/System/UnrealBuildTool.cs
@@ -2581,6 +2581,18 @@ namespace UnrealBuildTool
 				// Get all H files in all modules processed in the last makefile build
 				HashSet<string> AllUHTHeaders = new HashSet<string>(Target.FlatModuleCsData.Select(x => x.Value).SelectMany(x => x.UHTHeaderNames));
 
+				// Check whether any headers have been deleted. If they have, we need to regenerate the makefile since the module might now be empty. If we don't,
+				// and the file has been moved to a different module, we may include stale generated headers.
+				foreach(string FileName in AllUHTHeaders)
+				{
+					if(!File.Exists(FileName))
+					{
+						Log.TraceVerbose("File processed by UHT was deleted ({0}); invalidating makefile", FileName);
+						ReasonNotLoaded = string.Format("UHT file was deleted");
+						return null;
+					}
+				}
+
 				// Makefile is invalid if:
 				// * There are any newer files which contain no UHT data, but were previously in the makefile
 				// * There are any newer files contain data which needs processing by UHT, but weren't not previously in the makefile
diff --git a/Engine/Source/Programs/UnrealBuildTool/Windows/UEBuildWindows.cs b/Engine/Source/Programs/UnrealBuildTool/Windows/UEBuildWindows.cs
--- a/Engine/Source/Programs/UnrealBuildTool/Windows/UEBuildWindows.cs
+++ b/Engine/Source/Programs/UnrealBuildTool/Windows/UEBuildWindows.cs
@@ -159,6 +159,16 @@ namespace UnrealBuildTool
 				}
 			}
 
+			if (WindowsPlatform.bCompileWithICL)
+			{
+				BuildConfiguration.bUseSharedPCHs = false;
+
+				if (WindowsPlatform.bUseVCCompilerArgs)
+				{
+					BuildConfiguration.bUsePCHFiles = false;
+				}
+			}
+
 			// A bug in the UCRT can cause XGE to hang on VS2015 builds. Figure out if this hang is likely to effect this build and workaround it if able.
 			if (WindowsPlatform.Compiler == WindowsCompiler.VisualStudio2015)
 			{
@@ -245,6 +255,18 @@ namespace UnrealBuildTool
 				InBuildTarget.GlobalCompileEnvironment.Config.Definitions.Add("MORPHEUS_ENGINE_DISTORTION=1");
 			}
 
+			// Add path to Intel math libraries when using ICL based on target platform
+			if (WindowsPlatform.bCompileWithICL)
+			{
+				var Result = Path.Combine(Environment.GetFolderPath(Environment.SpecialFolder.ProgramFilesX86), "IntelSWTools", "compilers_and_libraries", "windows", "compiler", "lib", InBuildTarget.Platform == UnrealTargetPlatform.Win32 ? "ia32" : "intel64");
+				if (!Directory.Exists(Result))
+				{
+					throw new BuildException("ICL was selected but the required math libraries were not found.  Could not find: " + Result);
+				}
+
+				InBuildTarget.GlobalLinkEnvironment.Config.LibraryPaths.Add(Result);
+			}
+
 			if (InBuildTarget.Rules != null)
 			{
 				// Explicitly exclude the MS C++ runtime libraries we're not using, to ensure other libraries we link with use the same
@@ -535,6 +557,12 @@ namespace UnrealBuildTool
 		/// True if we should use the Clang linker (LLD) when bCompileWithClang is enabled, otherwise we use the MSVC linker
 		public static readonly bool bAllowClangLinker = bCompileWithClang && false;
 
+		/// True if we should use the Intel Compiler instead of MSVC to compile code on Windows platform
+		public static readonly bool bCompileWithICL = false;
+
+		/// True if we should use the Intel linker (xilink) when bCompileWithICL is enabled, otherwise we use the MSVC linker
+		public static readonly bool bAllowICLLinker = bCompileWithICL && true;
+
 		/// Whether to compile against the Windows 10 SDK, instead of the Windows 8.1 SDK.  This requires the Visual Studio 2015
 		/// compiler or later, and the Windows 10 SDK must be installed.  The application will require at least Windows 8.x to run.
 		// @todo UWP: Expose this to be enabled more easily for building Windows 10 desktop apps
diff --git a/Engine/Source/Programs/UnrealBuildTool/Windows/VCEnvironment.cs b/Engine/Source/Programs/UnrealBuildTool/Windows/VCEnvironment.cs
--- a/Engine/Source/Programs/UnrealBuildTool/Windows/VCEnvironment.cs
+++ b/Engine/Source/Programs/UnrealBuildTool/Windows/VCEnvironment.cs
@@ -91,10 +91,10 @@ namespace UnrealBuildTool
 			// Regardless of the target, if we're linking on a 64 bit machine, we want to use the 64 bit linker (it's faster than the 32 bit linker and can handle large linking jobs)
 			string LinkerVSToolPath = VSToolPath64Bit;
 
-			CompilerPath = GetCompilerToolPath(CompilerVSToolPath);
+			CompilerPath = GetCompilerToolPath(InPlatform, CompilerVSToolPath);
 			CLExeVersion = FindCLExeVersion(CompilerPath);
-			LinkerPath = GetLinkerToolPath(LinkerVSToolPath);
-			LibraryLinkerPath = GetLibraryLinkerToolPath(LinkerVSToolPath);
+			LinkerPath = GetLinkerToolPath(InPlatform, LinkerVSToolPath);
+			LibraryLinkerPath = GetLibraryLinkerToolPath(InPlatform, LinkerVSToolPath);
 			ResourceCompilerPath = GetResourceCompilerToolPath(Platform, bSupportWindowsXP);
 
             // Make sure the base 32-bit VS tool path is in the PATH, regardless of which configuration we're using. The toolchain may need to reference support DLLs from this directory (eg. mspdb120.dll).
@@ -327,7 +327,7 @@ namespace UnrealBuildTool
 		/// <summary>
 		/// Gets the path to the compiler.
 		/// </summary>
-		static string GetCompilerToolPath(string PlatformVSToolPath)
+		static string GetCompilerToolPath(CPPTargetPlatform Platform, string PlatformVSToolPath)
 		{
 			// If we were asked to use Clang, then we'll redirect the path to the compiler to the LLVM installation directory
 			if (WindowsPlatform.bCompileWithClang)
@@ -359,6 +359,17 @@ namespace UnrealBuildTool
 				return Result;
 			}
 
+			if (WindowsPlatform.bCompileWithICL)
+			{
+				var Result = Path.Combine(Environment.GetFolderPath(Environment.SpecialFolder.ProgramFilesX86), "IntelSWTools", "compilers_and_libraries", "windows", "bin", Platform == CPPTargetPlatform.Win32 ? "ia32" : "intel64", "icl.exe");
+				if (!File.Exists(Result))
+				{
+					throw new BuildException("ICL was selected as the Windows compiler, but does not appear to be installed.  Could not find: " + Result);
+				}
+
+				return Result;
+			}
+
 			return Path.Combine(PlatformVSToolPath, "cl.exe");
 		}
 
@@ -384,7 +395,7 @@ namespace UnrealBuildTool
 		/// <summary>
 		/// Gets the path to the linker.
 		/// </summary>
-		static string GetLinkerToolPath(string PlatformVSToolPath)
+		static string GetLinkerToolPath(CPPTargetPlatform Platform, string PlatformVSToolPath)
 		{
 			// If we were asked to use Clang, then we'll redirect the path to the compiler to the LLVM installation directory
 			if (WindowsPlatform.bCompileWithClang && WindowsPlatform.bAllowClangLinker)
@@ -402,16 +413,39 @@ namespace UnrealBuildTool
 				return Result;
 			}
 
+			if (WindowsPlatform.bCompileWithICL && WindowsPlatform.bAllowICLLinker)
+			{
+				var Result = Path.Combine(Environment.GetFolderPath(Environment.SpecialFolder.ProgramFilesX86), "IntelSWTools", "compilers_and_libraries", "windows", "bin", Platform == CPPTargetPlatform.Win32 ? "ia32" : "intel64", "xilink.exe");
+				if (!File.Exists(Result))
+				{
+					throw new BuildException("ICL was selected as the Windows compiler, but does not appear to be installed.  Could not find: " + Result);
+				}
+
+				return Result;
+			}
+
 			return Path.Combine(PlatformVSToolPath, "link.exe");
 		}
 
 		/// <summary>
 		/// Gets the path to the library linker.
 		/// </summary>
-		static string GetLibraryLinkerToolPath(string PlatformVSToolPath)
+		static string GetLibraryLinkerToolPath(CPPTargetPlatform Platform, string PlatformVSToolPath)
 		{
 			// Regardless of the target, if we're linking on a 64 bit machine, we want to use the 64 bit linker (it's faster than the 32 bit linker)
 			[REDACTED_EMAIL]: Using the 64-bit linker appears to be broken at the moment.
+
+			if (WindowsPlatform.bCompileWithICL && WindowsPlatform.bAllowICLLinker)
+			{
+				var Result = Path.Combine(Environment.GetFolderPath(Environment.SpecialFolder.ProgramFilesX86), "IntelSWTools", "compilers_and_libraries", "windows", "bin", Platform == CPPTargetPlatform.Win32 ? "ia32" : "intel64", "xilib.exe");
+				if (!File.Exists(Result))
+				{
+					throw new BuildException("ICL was selected as the Windows compiler, but does not appear to be installed.  Could not find: " + Result);
+				}
+
+				return Result;
+			}
+
 			return Path.Combine(PlatformVSToolPath, "lib.exe");
 		}
 
diff --git a/Engine/Source/Programs/UnrealBuildTool/Windows/VCToolChain.cs b/Engine/Source/Programs/UnrealBuildTool/Windows/VCToolChain.cs
--- a/Engine/Source/Programs/UnrealBuildTool/Windows/VCToolChain.cs
+++ b/Engine/Source/Programs/UnrealBuildTool/Windows/VCToolChain.cs
@@ -406,7 +406,7 @@ namespace UnrealBuildTool
 					}
 				}
 
-				if (WindowsPlatform.bUseVCCompilerArgs)
+				if (WindowsPlatform.bUseVCCompilerArgs && !WindowsPlatform.bCompileWithICL)
 				{
 					// Prompt the user before reporting internal errors to Microsoft.
 					Arguments.Append(" /errorReport:prompt");
@@ -545,8 +545,11 @@ namespace UnrealBuildTool
 					Arguments.Append(" /bigobj");
 				}
 
-				// Relaxes floating point precision semantics to allow more optimization.
-				Arguments.Append(" /fp:fast");
+				if (!WindowsPlatform.bCompileWithICL)
+				{
+					// Relaxes floating point precision semantics to allow more optimization.
+					Arguments.Append(" /fp:fast");
+				}
 
 				if (CompileEnvironment.Config.OptimizeCode >= ModuleRules.CodeOptimization.InNonDebugBuilds)
 				{
@@ -559,7 +562,10 @@ namespace UnrealBuildTool
 					}
 					else
 					{
-						Arguments.Append(" /d2Zi+");
+						if (!WindowsPlatform.bCompileWithICL)
+						{
+							Arguments.Append(" /d2Zi+");
+						}
 					}
 				}
 			}
@@ -700,14 +706,14 @@ namespace UnrealBuildTool
 			// Set warning level.
 			if (WindowsPlatform.bUseVCCompilerArgs)
 			{
-				if (!BuildConfiguration.bRunUnrealCodeAnalyzer)
+				if (!BuildConfiguration.bRunUnrealCodeAnalyzer && !WindowsPlatform.bCompileWithICL)
 				{
 					// Restrictive during regular compilation.
 					Arguments.Append(" /W4");
 				}
 				else
 				{
-					// If we had /W4 with clang on windows we would be flooded with warnings. This will be fixed incrementally.
+					// If we had /W4 with clang or Intel on windows we would be flooded with warnings. This will be fixed incrementally.
 					Arguments.Append(" /W0");
 				}
 			}
@@ -716,6 +722,14 @@ namespace UnrealBuildTool
 				Arguments.Append(" -Wall");
 			}
 
+			// Intel compiler options.
+			if (WindowsPlatform.bCompileWithICL)
+			{
+				Arguments.Append(" /Qstd=c++11");
+				Arguments.Append(" /fp:precise");
+				Arguments.Append(" /nologo");
+			}
+
 			if (WindowsPlatform.bCompileWithClang)
 			{
 				// Disable specific warnings that cause problems with Clang
@@ -1155,7 +1169,7 @@ namespace UnrealBuildTool
 					}
 				}
 
-				if (WindowsPlatform.bCompileWithClang)
+				if (WindowsPlatform.bCompileWithClang || WindowsPlatform.bCompileWithICL)
 				{
 					CompileAction.OutputEventHandler = new DataReceivedEventHandler(ClangCompilerOutputFormatter);
 				}
@@ -1745,7 +1759,7 @@ namespace UnrealBuildTool
 				LinkAction.bPrintDebugInfo = true;
 			}
 
-			if (WindowsPlatform.bCompileWithClang)
+			if (WindowsPlatform.bCompileWithClang || WindowsPlatform.bCompileWithICL)
 			{
 				LinkAction.OutputEventHandler = new DataReceivedEventHandler(ClangCompilerOutputFormatter);
 			}
diff --git a/Engine/Source/Programs/UnrealHeaderTool/Private/CodeGenerator.cpp b/Engine/Source/Programs/UnrealHeaderTool/Private/CodeGenerator.cpp
--- a/Engine/Source/Programs/UnrealHeaderTool/Private/CodeGenerator.cpp
+++ b/Engine/Source/Programs/UnrealHeaderTool/Private/CodeGenerator.cpp
@@ -5432,6 +5432,9 @@ void FNativeClassHeaderGenerator::ExportGeneratedCPP()
 	FString DepHeaderPathname = ModuleInfo->GeneratedCPPFilenameBase + TEXT(".dep.h");
 	SaveHeaderIfChanged(*DepHeaderPathname, *(GeneratedCPPPreamble + ListOfPublicClassesUObjectHeaderModuleIncludes));
 
+	// Write out an include for the generated header boilerplate code
+	GeneratedCPPClassesIncludes.Logf(TEXT("#include \"GeneratedCppIncludes.h\"") LINE_TERMINATOR, *FPaths::GetCleanFilename(DepHeaderPathname));
+
 	// Write out our include to the .dep.h file
 	GeneratedCPPClassesIncludes.Logf(TEXT("#include \"%s\"") LINE_TERMINATOR, *FPaths::GetCleanFilename(DepHeaderPathname));
 
diff --git a/Engine/Source/Programs/UnrealHeaderTool/Private/HeaderParser.cpp b/Engine/Source/Programs/UnrealHeaderTool/Private/HeaderParser.cpp
--- a/Engine/Source/Programs/UnrealHeaderTool/Private/HeaderParser.cpp
+++ b/Engine/Source/Programs/UnrealHeaderTool/Private/HeaderParser.cpp
@@ -6897,7 +6897,7 @@ ECompilationResult::Type FHeaderParser::ParseHeaders(FClasses& AllClasses, FHead
 
 	TArray<FUnrealSourceFile*> SourceFilesRequired;
 
-	static const FString ObjectHeader = FString(TEXT("Object.h"));
+	static const FString ObjectHeader = FString(TEXT("NoExportTypes.h"));
 	for (auto& Include : SourceFile->GetIncludes())
 	{
 		if (Include.GetId() == ObjectHeader)
diff --git a/Engine/Source/Programs/UnrealHeaderTool/Private/UnrealSourceFile.h b/Engine/Source/Programs/UnrealHeaderTool/Private/UnrealSourceFile.h
--- a/Engine/Source/Programs/UnrealHeaderTool/Private/UnrealSourceFile.h
+++ b/Engine/Source/Programs/UnrealHeaderTool/Private/UnrealSourceFile.h
@@ -26,9 +26,9 @@ public:
 		, bParsed              (false)
 		, bDependenciesResolved(false)
 	{
-		if (GetStrippedFilename() != "Object")
+		if (GetStrippedFilename() != "NoExportTypes")
 		{
-			Includes.Add(FHeaderProvider(EHeaderProviderSourceType::FileName, "Object.h"));
+			Includes.Add(FHeaderProvider(EHeaderProviderSourceType::FileName, "NoExportTypes.h"));
 		}
 	}
 
diff --git a/Engine/Source/Runtime/Advertising/Advertising/Private/AdvertisingPrivatePCH.h b/Engine/Source/Runtime/Advertising/Advertising/Private/AdvertisingPrivatePCH.h
--- a/Engine/Source/Runtime/Advertising/Advertising/Private/AdvertisingPrivatePCH.h
+++ b/Engine/Source/Runtime/Advertising/Advertising/Private/AdvertisingPrivatePCH.h
@@ -1,6 +1,5 @@
 // Copyright 1998-2016 Epic Games, Inc. All Rights Reserved.
 
-#include "Advertising.h"
+#pragma once
 
-// You should place include statements to your module's private header files here.  You only need to
-// add includes for headers that are used in most of your module's source files though.
+#include "Advertising.h"
diff --git a/Engine/Source/Runtime/Analytics/Analytics/Private/Analytics.cpp b/Engine/Source/Runtime/Analytics/Analytics/Private/Analytics.cpp
--- a/Engine/Source/Runtime/Analytics/Analytics/Private/Analytics.cpp
+++ b/Engine/Source/Runtime/Analytics/Analytics/Private/Analytics.cpp
@@ -1,5 +1,6 @@
 // Copyright 1998-2016 Epic Games, Inc. All Rights Reserved.
 
+#include "Core.h"
 #include "Interfaces/IAnalyticsProviderModule.h"
 #include "Interfaces/IAnalyticsProvider.h"
 #include "Analytics.h"
diff --git a/Engine/Source/Runtime/Analytics/AnalyticsET/Private/AnalyticsETPrivatePCH.h b/Engine/Source/Runtime/Analytics/AnalyticsET/Private/AnalyticsETPrivatePCH.h
--- a/Engine/Source/Runtime/Analytics/AnalyticsET/Private/AnalyticsETPrivatePCH.h
+++ b/Engine/Source/Runtime/Analytics/AnalyticsET/Private/AnalyticsETPrivatePCH.h
@@ -1,4 +1,5 @@
 // Copyright 1998-2016 Epic Games, Inc. All Rights Reserved.
 
-// You should place include statements to your module's private header files here.  You only need to
-// add includes for headers that are used in most of your module's source files though.
+#pragma once
+
+#include "Core.h"
diff --git a/Engine/Source/Runtime/AssetRegistry/Private/AssetRegistryPCH.h b/Engine/Source/Runtime/AssetRegistry/Private/AssetRegistryPCH.h
--- a/Engine/Source/Runtime/AssetRegistry/Private/AssetRegistryPCH.h
+++ b/Engine/Source/Runtime/AssetRegistry/Private/AssetRegistryPCH.h
@@ -1,7 +1,6 @@
 // Copyright 1998-2016 Epic Games, Inc. All Rights Reserved.
 
-#ifndef _ASSET_REGISTRY_PRIVATE_
-#define _ASSET_REGISTRY_PRIVATE_
+#pragma once
 
 #include "Core.h"
 #include "CoreUObject.h"
@@ -18,4 +17,3 @@ DECLARE_LOG_CATEGORY_EXTERN(LogAssetRegistry, Log, All);
 #include "AssetDataGatherer.h"
 #include "AssetRegistry.h"
 
-#endif // _ASSET_REGISTRY_PRIVATE_
diff --git a/Engine/Source/Runtime/Core/Private/CorePrivatePCH.h b/Engine/Source/Runtime/Core/Private/CorePrivatePCH.h
--- a/Engine/Source/Runtime/Core/Private/CorePrivatePCH.h
+++ b/Engine/Source/Runtime/Core/Private/CorePrivatePCH.h
@@ -67,6 +67,7 @@ Includes.
 #include "ThreadingBase.h"				// Non-platform specific multi-threaded support.
 #include "Internationalization/Internationalization.h"
 #include "Internationalization/Culture.h"
+#include "Internationalization/CulturePointer.h"
 #include "Guid.h"						// FGuid class
 #include "UnrealMath.h"					// Vector math functions.
 #include "OutputDevices.h"				// Output devices
diff --git a/Engine/Source/Runtime/Core/Private/Internationalization/Cultures/InvariantCulture.h b/Engine/Source/Runtime/Core/Private/Internationalization/Cultures/InvariantCulture.h
--- a/Engine/Source/Runtime/Core/Private/Internationalization/Cultures/InvariantCulture.h
+++ b/Engine/Source/Runtime/Core/Private/Internationalization/Cultures/InvariantCulture.h
@@ -4,6 +4,7 @@
 #if !UE_ENABLE_ICU
 
 #include "Culture.h"
+#include "Internationalization/Internationalization.h"
 #include "FastDecimalFormat.h"
 
 #define LOCTEXT_NAMESPACE "Internationalization"
diff --git a/Engine/Source/Runtime/Core/Private/Logging/TokenizedMessage.cpp b/Engine/Source/Runtime/Core/Private/Logging/TokenizedMessage.cpp
--- a/Engine/Source/Runtime/Core/Private/Logging/TokenizedMessage.cpp
+++ b/Engine/Source/Runtime/Core/Private/Logging/TokenizedMessage.cpp
@@ -180,6 +180,11 @@ void FAssetNameToken::FindAsset(const TSharedRef<IMessageToken>& Token, FString
 	}
 }
 
+TSharedRef<FAssetNameToken> FAssetNameToken::Create(const FString& InAssetName, const FText& InMessage)
+{
+	return MakeShareable(new FAssetNameToken(InAssetName, InMessage));
+}
+
 FAssetNameToken::FAssetNameToken(const FString& InAssetName, const FText& InMessage)
 	: AssetName(InAssetName)
 {
@@ -201,4 +206,9 @@ FDocumentationToken::FDocumentationToken( const FString& InDocumentationLink, co
 	, PreviewExcerptName(InPreviewExcerptName)
 { }
 
+TSharedRef<FDocumentationToken> FDocumentationToken::Create(const FString& InDocumentationLink, const FString& InPreviewExcerptLink, const FString& InPreviewExcerptName)
+{
+	return MakeShareable(new FDocumentationToken(InDocumentationLink, InPreviewExcerptLink, InPreviewExcerptName));
+}
+
 #undef LOCTEXT_NAMESPACE
\ No newline at end of file
diff --git a/Engine/Source/Runtime/Core/Private/Misc/MapErrors.cpp b/Engine/Source/Runtime/Core/Private/Misc/MapErrors.cpp
--- a/Engine/Source/Runtime/Core/Private/Misc/MapErrors.cpp
+++ b/Engine/Source/Runtime/Core/Private/Misc/MapErrors.cpp
@@ -65,4 +65,9 @@ static const FString MapErrorsPath = TEXT("Shared/Editor/MapErrors");
 FMapErrorToken::FMapErrorToken(const FName& InErrorName)
 	: FDocumentationToken(MapErrorsPath, MapErrorsPath, InErrorName.ToString())
 {
+}
+
+TSharedRef<FMapErrorToken> FMapErrorToken::Create(const FName& InErrorName)
+{
+	return MakeShareable(new FMapErrorToken(InErrorName));
 }
\ No newline at end of file
diff --git a/Engine/Source/Runtime/Core/Public/Android/AndroidPlatform.h b/Engine/Source/Runtime/Core/Public/Android/AndroidPlatform.h
--- a/Engine/Source/Runtime/Core/Public/Android/AndroidPlatform.h
+++ b/Engine/Source/Runtime/Core/Public/Android/AndroidPlatform.h
@@ -60,9 +60,6 @@ typedef FAndroidTypes FPlatformTypes;
 #define FUNCTION_CHECK_RETURN_END __attribute__ ((warn_unused_result))	/* Warn that callers should not ignore the return value. */
 #define FUNCTION_NO_RETURN_END __attribute__ ((noreturn))				/* Indicate that the function never returns. */
 
-#define TEXT_HELPER(a,b)	a ## b
-#define TEXT(s)				TEXT_HELPER(L, s)
-
 #define ABSTRACT abstract
 
 // Alignment.
diff --git a/Engine/Source/Runtime/Core/Public/Containers/ResourceArray.h b/Engine/Source/Runtime/Core/Public/Containers/ResourceArray.h
--- a/Engine/Source/Runtime/Core/Public/Containers/ResourceArray.h
+++ b/Engine/Source/Runtime/Core/Public/Containers/ResourceArray.h
@@ -111,6 +111,3 @@ public:
 	virtual ~FTexture2DResourceMem() {}
 };
 
-
-// Default to the dynamically bound RHI.
-#include "DynamicRHIResourceArray.h"
diff --git a/Engine/Source/Runtime/Core/Public/Core.h b/Engine/Source/Runtime/Core/Public/Core.h
--- a/Engine/Source/Runtime/Core/Public/Core.h
+++ b/Engine/Source/Runtime/Core/Public/Core.h
@@ -71,6 +71,7 @@
 #include "MapBuilder.h"					// Builder template for maps.
 #include "List.h"						// Dynamic list definitions.
 #include "ResourceArray.h"				// Resource array definitions.
+#include "DynamicRHIResourceArray.h"	// Default to the dynamically bound RHI.
 #include "RefCounting.h"				// Reference counting definitions.
 #include "ScriptDelegates.h"
 #include "Delegate.h"					// C++ delegate system
diff --git a/Engine/Source/Runtime/Core/Public/HAL/Platform.h b/Engine/Source/Runtime/Core/Public/HAL/Platform.h
--- a/Engine/Source/Runtime/Core/Public/HAL/Platform.h
+++ b/Engine/Source/Runtime/Core/Public/HAL/Platform.h
@@ -732,6 +732,13 @@ namespace TypeTests
 	#error Unknown Compiler
 #endif
 
+// If we don't have a platform-specific define for the TEXT macro, define it now.
+#if !defined(TEXT) && !UE_BUILD_DOCS
+	#define TEXT_PASTE(x) L ## x
+	#define TEXT(x) TEXT_PASTE(x)
+#endif
+
 // Include defaults for defines that aren't explicitly set by the platform
 #include "UMemoryDefines.h"
 #include "../Misc/CoreMiscDefines.h"
+#include "../Misc/CoreDefines.h"
diff --git a/Engine/Source/Runtime/Core/Public/HAL/PlatformString.h b/Engine/Source/Runtime/Core/Public/HAL/PlatformString.h
--- a/Engine/Source/Runtime/Core/Public/HAL/PlatformString.h
+++ b/Engine/Source/Runtime/Core/Public/HAL/PlatformString.h
@@ -22,7 +22,3 @@
 #include "Linux/LinuxPlatformString.h"
 #endif
 
-//TEXT macro, may or may not have come from OS includes
-#if !defined(TEXT) && !UE_BUILD_DOCS
-#define TEXT(s) L##s
-#endif
diff --git a/Engine/Source/Runtime/Core/Public/HTML5/HTML5SystemIncludes.h b/Engine/Source/Runtime/Core/Public/HTML5/HTML5SystemIncludes.h
--- a/Engine/Source/Runtime/Core/Public/HTML5/HTML5SystemIncludes.h
+++ b/Engine/Source/Runtime/Core/Public/HTML5/HTML5SystemIncludes.h
@@ -11,11 +11,6 @@
 // Set up compiler pragmas, etc
 #include "HTML5/HTML5PlatformCompilerSetup.h"
 
-// Define TEXT early
-#define TEXT_PASTE(x) L ## x
-#define TEXT_HELPER(a,b)	a ## b
-#define TEXT(s)				TEXT_HELPER(L, s)
-
 // map the Windows functions (that UE4 unfortunately uses be default) to normal functions
 #if PLATFORM_HTML5_WIN32
 
diff --git a/Engine/Source/Runtime/Core/Public/IOS/IOSPlatform.h b/Engine/Source/Runtime/Core/Public/IOS/IOSPlatform.h
--- a/Engine/Source/Runtime/Core/Public/IOS/IOSPlatform.h
+++ b/Engine/Source/Runtime/Core/Public/IOS/IOSPlatform.h
@@ -66,9 +66,6 @@ typedef FIOSPlatformTypes FPlatformTypes;
 #define FUNCTION_CHECK_RETURN_END __attribute__ ((warn_unused_result))	/* Warn that callers should not ignore the return value. */
 #define FUNCTION_NO_RETURN_END __attribute__ ((noreturn))				/* Indicate that the function never returns. */
 
-#define TEXT_HELPER(a,b)	a ## b
-#define TEXT(s)				TEXT_HELPER(L, s)
-
 #define ABSTRACT abstract
 
 // Strings.
diff --git a/Engine/Source/Runtime/Core/Public/Internationalization/Culture.h b/Engine/Source/Runtime/Core/Public/Internationalization/Culture.h
--- a/Engine/Source/Runtime/Core/Public/Internationalization/Culture.h
+++ b/Engine/Source/Runtime/Core/Public/Internationalization/Culture.h
@@ -134,5 +134,3 @@ private:
 	FString CachedScript;
 	FString CachedVariant;
 };
-
-#include "CulturePointer.h"
diff --git a/Engine/Source/Runtime/Core/Public/Logging/MessageLog.h b/Engine/Source/Runtime/Core/Public/Logging/MessageLog.h
--- a/Engine/Source/Runtime/Core/Public/Logging/MessageLog.h
+++ b/Engine/Source/Runtime/Core/Public/Logging/MessageLog.h
@@ -118,9 +118,9 @@ private:
 	bool bSuppressLoggingToOutputLog;
 
 	/** Delegate for retrieving the message log */
-	static FGetLog GetLog;
+	CORE_API static FGetLog GetLog;
 
 	/** Delegate for message selection */
-	static FMessageSelectionChanged MessageSelectionChanged;
+	CORE_API static FMessageSelectionChanged MessageSelectionChanged;
 };
 
diff --git a/Engine/Source/Runtime/Core/Public/Logging/TokenizedMessage.h b/Engine/Source/Runtime/Core/Public/Logging/TokenizedMessage.h
--- a/Engine/Source/Runtime/Core/Public/Logging/TokenizedMessage.h
+++ b/Engine/Source/Runtime/Core/Public/Logging/TokenizedMessage.h
@@ -360,7 +360,7 @@ private:
 	FString URL;
 
 	/** The delegate we will use to generate our URL */
-	static FGenerateURL GenerateURL;
+	CORE_API static FGenerateURL GenerateURL;
 };
 
 /** 
@@ -371,10 +371,7 @@ class FAssetNameToken : public IMessageToken
 {
 public:
 	/** Factory method, tokens can only be constructed as shared refs */
-	CORE_API static TSharedRef<FAssetNameToken> Create( const FString& InAssetName, const FText& InMessage = FText() )
-	{
-		return MakeShareable(new FAssetNameToken(InAssetName, InMessage));
-	}
+	CORE_API static TSharedRef<FAssetNameToken> Create(const FString& InAssetName, const FText& InMessage = FText());
 
 	/** Begin IMessageToken interface */
 	virtual EMessageToken::Type GetType() const override
@@ -410,7 +407,7 @@ private:
 	FString AssetName;
 
 	/** The delegate we will use to go to our file */
-	static FOnGotoAsset GotoAsset;
+	CORE_API static FOnGotoAsset GotoAsset;
 };
 
 /** 
@@ -420,10 +417,7 @@ class FDocumentationToken : public IMessageToken
 {
 public:
 	/** Factory method, tokens can only be constructed as shared refs */
-	CORE_API static TSharedRef<FDocumentationToken> Create( const FString& InDocumentationLink, const FString& InPreviewExcerptLink = FString(), const FString& InPreviewExcerptName = FString() )
-	{
-		return MakeShareable(new FDocumentationToken(InDocumentationLink, InPreviewExcerptLink, InPreviewExcerptName));
-	}
+	CORE_API static TSharedRef<FDocumentationToken> Create(const FString& InDocumentationLink, const FString& InPreviewExcerptLink = FString(), const FString& InPreviewExcerptName = FString());
 
 	/** Begin IMessageToken interface */
 	virtual EMessageToken::Type GetType() const override
diff --git a/Engine/Source/Runtime/Core/Public/Math/Matrix.h b/Engine/Source/Runtime/Core/Public/Math/Matrix.h
--- a/Engine/Source/Runtime/Core/Public/Math/Matrix.h
+++ b/Engine/Source/Runtime/Core/Public/Math/Matrix.h
@@ -489,3 +489,5 @@ static inline void Inverse4x4( double* dst, const float* src )
 		dst[i] = inv[i] * det;
 	}
 }
+
+#include "Matrix.inl"
diff --git a/Engine/Source/Runtime/Core/Public/Math/Matrix.inl b/Engine/Source/Runtime/Core/Public/Math/Matrix.inl
new file mode 100644
--- /dev/null
+++ b/Engine/Source/Runtime/Core/Public/Math/Matrix.inl
@@ -0,0 +1,761 @@
+// Copyright 1998-2016 Epic Games, Inc. All Rights Reserved.
+
+/*=============================================================================
+	NOTE: This file should ONLY be included by UnrealMath.h!
+=============================================================================*/
+
+#pragma once
+
+#include "OutputDevices.h"
+#include "AssertionMacros.h"
+
+
+/**
+ * FMatrix inline functions.
+ */
+
+// Constructors.
+
+FORCEINLINE FMatrix::FMatrix()
+{
+}
+
+FORCEINLINE FMatrix::FMatrix(const FPlane& InX,const FPlane& InY,const FPlane& InZ,const FPlane& InW)
+{
+	M[0][0] = InX.X; M[0][1] = InX.Y;  M[0][2] = InX.Z;  M[0][3] = InX.W;
+	M[1][0] = InY.X; M[1][1] = InY.Y;  M[1][2] = InY.Z;  M[1][3] = InY.W;
+	M[2][0] = InZ.X; M[2][1] = InZ.Y;  M[2][2] = InZ.Z;  M[2][3] = InZ.W;
+	M[3][0] = InW.X; M[3][1] = InW.Y;  M[3][2] = InW.Z;  M[3][3] = InW.W;
+}
+
+FORCEINLINE FMatrix::FMatrix(const FVector& InX,const FVector& InY,const FVector& InZ,const FVector& InW)
+{
+	M[0][0] = InX.X; M[0][1] = InX.Y;  M[0][2] = InX.Z;  M[0][3] = 0.0f;
+	M[1][0] = InY.X; M[1][1] = InY.Y;  M[1][2] = InY.Z;  M[1][3] = 0.0f;
+	M[2][0] = InZ.X; M[2][1] = InZ.Y;  M[2][2] = InZ.Z;  M[2][3] = 0.0f;
+	M[3][0] = InW.X; M[3][1] = InW.Y;  M[3][2] = InW.Z;  M[3][3] = 1.0f;
+}
+
+
+inline void FMatrix::SetIdentity()
+{
+	M[0][0] = 1; M[0][1] = 0;  M[0][2] = 0;  M[0][3] = 0;
+	M[1][0] = 0; M[1][1] = 1;  M[1][2] = 0;  M[1][3] = 0;
+	M[2][0] = 0; M[2][1] = 0;  M[2][2] = 1;  M[2][3] = 0;
+	M[3][0] = 0; M[3][1] = 0;  M[3][2] = 0;  M[3][3] = 1;
+}
+
+
+FORCEINLINE void FMatrix::operator*=(const FMatrix& Other)
+{
+	VectorMatrixMultiply( this, this, &Other );
+}
+
+
+FORCEINLINE FMatrix FMatrix::operator*(const FMatrix& Other) const
+{
+	FMatrix Result;
+	VectorMatrixMultiply( &Result, this, &Other );
+	return Result;
+}
+
+
+FORCEINLINE FMatrix	FMatrix::operator+(const FMatrix& Other) const
+{
+	FMatrix ResultMat;
+
+	for(int32 X = 0;X < 4;X++)
+	{
+		for(int32 Y = 0;Y < 4;Y++)
+		{
+			ResultMat.M[X][Y] = M[X][Y]+Other.M[X][Y];
+		}
+	}
+
+	return ResultMat;
+}
+
+FORCEINLINE void FMatrix::operator+=(const FMatrix& Other)
+{
+	*this = *this + Other;
+}
+
+FORCEINLINE FMatrix	FMatrix::operator*(float Other) const
+{
+	FMatrix ResultMat;
+
+	for(int32 X = 0;X < 4;X++)
+	{
+		for(int32 Y = 0;Y < 4;Y++)
+		{
+			ResultMat.M[X][Y] = M[X][Y]*Other;
+		}
+	}
+
+	return ResultMat;
+}
+
+FORCEINLINE void FMatrix::operator*=(float Other)
+{
+	*this = *this*Other;
+}
+
+// Comparison operators.
+
+inline bool FMatrix::operator==(const FMatrix& Other) const
+{
+	for(int32 X = 0;X < 4;X++)
+	{
+		for(int32 Y = 0;Y < 4;Y++)
+		{
+			if(M[X][Y] != Other.M[X][Y])
+			{
+				return false;
+			}
+		}
+	}
+
+	return true;
+}
+
+// Error-tolerant comparison.
+inline bool FMatrix::Equals(const FMatrix& Other, float Tolerance/*=KINDA_SMALL_NUMBER*/) const
+{
+	for(int32 X = 0;X < 4;X++)
+	{
+		for(int32 Y = 0;Y < 4;Y++)
+		{
+			if( FMath::Abs(M[X][Y] - Other.M[X][Y]) > Tolerance )
+			{
+				return false;
+			}
+		}
+	}
+
+	return true;
+}
+
+inline bool FMatrix::operator!=(const FMatrix& Other) const
+{
+	return !(*this == Other);
+}
+
+
+// Homogeneous transform.
+
+FORCEINLINE FVector4 FMatrix::TransformFVector4(const FVector4 &P) const
+{
+	FVector4 Result;
+	VectorRegister VecP = VectorLoadAligned(&P);
+	VectorRegister VecR = VectorTransformVector(VecP, this);
+	VectorStoreAligned(VecR, &Result);
+	return Result;
+}
+
+
+// Transform position
+
+/** Transform a location - will take into account translation part of the FMatrix. */
+FORCEINLINE FVector4 FMatrix::TransformPosition(const FVector &V) const
+{
+	return TransformFVector4(FVector4(V.X,V.Y,V.Z,1.0f));
+}
+
+/** Inverts the matrix and then transforms V - correctly handles scaling in this matrix. */
+FORCEINLINE FVector FMatrix::InverseTransformPosition(const FVector &V) const
+{
+	FMatrix InvSelf = this->InverseFast();
+	return InvSelf.TransformPosition(V);
+}
+
+// Transform vector
+
+/** 
+ *	Transform a direction vector - will not take into account translation part of the FMatrix. 
+ *	If you want to transform a surface normal (or plane) and correctly account for non-uniform scaling you should use TransformByUsingAdjointT.
+ */
+FORCEINLINE FVector4 FMatrix::TransformVector(const FVector& V) const
+{
+	return TransformFVector4(FVector4(V.X,V.Y,V.Z,0.0f));
+}
+
+/** Faster version of InverseTransformVector that assumes no scaling. WARNING: Will NOT work correctly if there is scaling in the matrix. */
+FORCEINLINE FVector FMatrix::InverseTransformVector(const FVector &V) const
+{
+	FMatrix InvSelf = this->InverseFast();
+	return InvSelf.TransformVector(V);
+}
+
+
+// Transpose.
+
+FORCEINLINE FMatrix FMatrix::GetTransposed() const
+{
+	FMatrix	Result;
+
+	Result.M[0][0] = M[0][0];
+	Result.M[0][1] = M[1][0];
+	Result.M[0][2] = M[2][0];
+	Result.M[0][3] = M[3][0];
+
+	Result.M[1][0] = M[0][1];
+	Result.M[1][1] = M[1][1];
+	Result.M[1][2] = M[2][1];
+	Result.M[1][3] = M[3][1];
+
+	Result.M[2][0] = M[0][2];
+	Result.M[2][1] = M[1][2];
+	Result.M[2][2] = M[2][2];
+	Result.M[2][3] = M[3][2];
+
+	Result.M[3][0] = M[0][3];
+	Result.M[3][1] = M[1][3];
+	Result.M[3][2] = M[2][3];
+	Result.M[3][3] = M[3][3];
+
+	return Result;
+}
+
+// Determinant.
+
+inline float FMatrix::Determinant() const
+{
+	return	M[0][0] * (
+				M[1][1] * (M[2][2] * M[3][3] - M[2][3] * M[3][2]) -
+				M[2][1] * (M[1][2] * M[3][3] - M[1][3] * M[3][2]) +
+				M[3][1] * (M[1][2] * M[2][3] - M[1][3] * M[2][2])
+				) -
+			M[1][0] * (
+				M[0][1] * (M[2][2] * M[3][3] - M[2][3] * M[3][2]) -
+				M[2][1] * (M[0][2] * M[3][3] - M[0][3] * M[3][2]) +
+				M[3][1] * (M[0][2] * M[2][3] - M[0][3] * M[2][2])
+				) +
+			M[2][0] * (
+				M[0][1] * (M[1][2] * M[3][3] - M[1][3] * M[3][2]) -
+				M[1][1] * (M[0][2] * M[3][3] - M[0][3] * M[3][2]) +
+				M[3][1] * (M[0][2] * M[1][3] - M[0][3] * M[1][2])
+				) -
+			M[3][0] * (
+				M[0][1] * (M[1][2] * M[2][3] - M[1][3] * M[2][2]) -
+				M[1][1] * (M[0][2] * M[2][3] - M[0][3] * M[2][2]) +
+				M[2][1] * (M[0][2] * M[1][3] - M[0][3] * M[1][2])
+				);
+}
+
+/** Calculate determinant of rotation 3x3 matrix */
+inline float FMatrix::RotDeterminant() const
+{
+	return	
+		M[0][0] * (M[1][1] * M[2][2] - M[1][2] * M[2][1]) -
+		M[1][0] * (M[0][1] * M[2][2] - M[0][2] * M[2][1]) +
+		M[2][0] * (M[0][1] * M[1][2] - M[0][2] * M[1][1]);
+}
+
+// Inverse.
+/** Fast path, doesn't check for nil matrices in final release builds */
+inline FMatrix FMatrix::InverseFast() const
+{
+	// If we're in non final release, then make sure we're not creating NaNs
+#if !(UE_BUILD_SHIPPING || UE_BUILD_TEST)
+	// Check for zero scale matrix to invert
+	if(	GetScaledAxis( EAxis::X ).IsNearlyZero(SMALL_NUMBER) && 
+		GetScaledAxis( EAxis::Y ).IsNearlyZero(SMALL_NUMBER) && 
+		GetScaledAxis( EAxis::Z ).IsNearlyZero(SMALL_NUMBER) ) 
+	{
+		UE_LOG(LogUnrealMath, Error, TEXT("FMatrix::InverseFast(), trying to invert a NIL matrix, this results in NaNs! Use Inverse() instead."));
+		ensureMsgf(false, TEXT("FMatrix::InverseFast(), trying to invert a NIL matrix, this results in NaNs! Use Inverse() instead."));
+	}
+#endif
+	FMatrix Result;
+	VectorMatrixInverse( &Result, this );
+	return Result;
+}
+
+// Inverse.
+inline FMatrix FMatrix::Inverse() const
+{
+	FMatrix Result;
+
+	// Check for zero scale matrix to invert
+	if(	GetScaledAxis( EAxis::X ).IsNearlyZero(SMALL_NUMBER) && 
+		GetScaledAxis( EAxis::Y ).IsNearlyZero(SMALL_NUMBER) && 
+		GetScaledAxis( EAxis::Z ).IsNearlyZero(SMALL_NUMBER) ) 
+	{
+		// just set to zero - avoids unsafe inverse of zero and duplicates what QNANs were resulting in before (scaling away all children)
+		Result = FMatrix::Identity;
+	}
+	else
+	{
+		const float	Det = Determinant();
+
+		if(Det == 0.0f)
+		{
+			Result = FMatrix::Identity;
+		}
+		else
+		{
+			VectorMatrixInverse( &Result, this );
+		}
+	}
+
+	return Result;
+}
+
+inline FMatrix FMatrix::TransposeAdjoint() const
+{
+	FMatrix TA;
+
+	TA.M[0][0] = this->M[1][1] * this->M[2][2] - this->M[1][2] * this->M[2][1];
+	TA.M[0][1] = this->M[1][2] * this->M[2][0] - this->M[1][0] * this->M[2][2];
+	TA.M[0][2] = this->M[1][0] * this->M[2][1] - this->M[1][1] * this->M[2][0];
+	TA.M[0][3] = 0.f;
+
+	TA.M[1][0] = this->M[2][1] * this->M[0][2] - this->M[2][2] * this->M[0][1];
+	TA.M[1][1] = this->M[2][2] * this->M[0][0] - this->M[2][0] * this->M[0][2];
+	TA.M[1][2] = this->M[2][0] * this->M[0][1] - this->M[2][1] * this->M[0][0];
+	TA.M[1][3] = 0.f;
+
+	TA.M[2][0] = this->M[0][1] * this->M[1][2] - this->M[0][2] * this->M[1][1];
+	TA.M[2][1] = this->M[0][2] * this->M[1][0] - this->M[0][0] * this->M[1][2];
+	TA.M[2][2] = this->M[0][0] * this->M[1][1] - this->M[0][1] * this->M[1][0];
+	TA.M[2][3] = 0.f;
+
+	TA.M[3][0] = 0.f;
+	TA.M[3][1] = 0.f;
+	TA.M[3][2] = 0.f;
+	TA.M[3][3] = 1.f;
+
+	return TA;
+}
+
+// NOTE: There is some compiler optimization issues with WIN64 that cause FORCEINLINE to cause a crash
+// Remove any scaling from this matrix (ie magnitude of each row is 1)
+inline void FMatrix::RemoveScaling(float Tolerance/*=SMALL_NUMBER*/)
+{
+	// For each row, find magnitude, and if its non-zero re-scale so its unit length.
+	const float SquareSum0 = (M[0][0] * M[0][0]) + (M[0][1] * M[0][1]) + (M[0][2] * M[0][2]);
+	const float SquareSum1 = (M[1][0] * M[1][0]) + (M[1][1] * M[1][1]) + (M[1][2] * M[1][2]);
+	const float SquareSum2 = (M[2][0] * M[2][0]) + (M[2][1] * M[2][1]) + (M[2][2] * M[2][2]);
+	const float Scale0 = FMath::FloatSelect( SquareSum0 - Tolerance, FMath::InvSqrt(SquareSum0), 1.0f );
+	const float Scale1 = FMath::FloatSelect( SquareSum1 - Tolerance, FMath::InvSqrt(SquareSum1), 1.0f );
+	const float Scale2 = FMath::FloatSelect( SquareSum2 - Tolerance, FMath::InvSqrt(SquareSum2), 1.0f );
+	M[0][0] *= Scale0; 
+	M[0][1] *= Scale0; 
+	M[0][2] *= Scale0; 
+	M[1][0] *= Scale1; 
+	M[1][1] *= Scale1; 
+	M[1][2] *= Scale1; 
+	M[2][0] *= Scale2; 
+	M[2][1] *= Scale2; 
+	M[2][2] *= Scale2;
+}
+
+// Returns matrix without scale information
+inline FMatrix FMatrix::GetMatrixWithoutScale(float Tolerance/*=SMALL_NUMBER*/) const
+{
+	FMatrix Result = *this;
+	Result.RemoveScaling(Tolerance);
+	return Result;
+}
+
+/** Remove any scaling from this matrix (ie magnitude of each row is 1) and return the 3D scale vector that was initially present. */
+inline FVector FMatrix::ExtractScaling(float Tolerance/*=SMALL_NUMBER*/)
+{
+	FVector Scale3D(0,0,0);
+
+	// For each row, find magnitude, and if its non-zero re-scale so its unit length.
+	const float SquareSum0 = (M[0][0] * M[0][0]) + (M[0][1] * M[0][1]) + (M[0][2] * M[0][2]);
+	const float SquareSum1 = (M[1][0] * M[1][0]) + (M[1][1] * M[1][1]) + (M[1][2] * M[1][2]);
+	const float SquareSum2 = (M[2][0] * M[2][0]) + (M[2][1] * M[2][1]) + (M[2][2] * M[2][2]);
+
+	if( SquareSum0 > Tolerance )
+	{
+		float Scale0 = FMath::Sqrt(SquareSum0);
+		Scale3D[0] =  Scale0;
+		float InvScale0 = 1.f / Scale0;
+		M[0][0] *= InvScale0; 
+		M[0][1] *= InvScale0; 
+		M[0][2] *= InvScale0; 
+	}
+	else
+	{
+		Scale3D[0] =  0;
+	}
+
+	if( SquareSum1 > Tolerance )
+	{
+		float Scale1 = FMath::Sqrt(SquareSum1);
+		Scale3D[1] =  Scale1;
+		float InvScale1 = 1.f / Scale1;
+		M[1][0] *= InvScale1; 
+		M[1][1] *= InvScale1; 
+		M[1][2] *= InvScale1; 
+	}
+	else
+	{
+		Scale3D[1] =  0;
+	}
+
+	if( SquareSum2 > Tolerance )
+	{
+		float Scale2 = FMath::Sqrt(SquareSum2);
+		Scale3D[2] =  Scale2;
+		float InvScale2 = 1.f / Scale2;
+		M[2][0] *= InvScale2; 
+		M[2][1] *= InvScale2; 
+		M[2][2] *= InvScale2; 
+	}
+	else
+	{
+		Scale3D[2] =  0;
+	}
+
+	return Scale3D;
+}
+
+/** return a 3D scale vector calculated from this matrix (where each component is the magnitude of a row vector). */
+inline FVector FMatrix::GetScaleVector(float Tolerance/*=SMALL_NUMBER*/) const
+{
+	FVector Scale3D(1,1,1);
+
+	// For each row, find magnitude, and if its non-zero re-scale so its unit length.
+	for(int32 i=0; i<3; i++)
+	{
+		const float SquareSum = (M[i][0] * M[i][0]) + (M[i][1] * M[i][1]) + (M[i][2] * M[i][2]);
+		if(SquareSum > Tolerance)
+		{
+			Scale3D[i] = FMath::Sqrt(SquareSum);
+		}
+		else
+		{
+			Scale3D[i] = 0.f;
+		}
+	}
+
+	return Scale3D;
+}
+// Remove any translation from this matrix
+inline FMatrix FMatrix::RemoveTranslation() const
+{
+	FMatrix Result = *this;
+	Result.M[3][0] = 0.0f;
+	Result.M[3][1] = 0.0f;
+	Result.M[3][2] = 0.0f;
+	return Result;
+}
+
+FORCEINLINE FMatrix FMatrix::ConcatTranslation(const FVector& Translation) const
+{
+	FMatrix Result;
+
+	float* RESTRICT Dest = &Result.M[0][0];
+	const float* RESTRICT Src = &M[0][0];
+	const float* RESTRICT Trans = &Translation.X;
+
+	Dest[0] = Src[0];
+	Dest[1] = Src[1];
+	Dest[2] = Src[2];
+	Dest[3] = Src[3];
+	Dest[4] = Src[4];
+	Dest[5] = Src[5];
+	Dest[6] = Src[6];
+	Dest[7] = Src[7];
+	Dest[8] = Src[8];
+	Dest[9] = Src[9];
+	Dest[10] = Src[10];
+	Dest[11] = Src[11];
+	Dest[12] = Src[12] + Trans[0];
+	Dest[13] = Src[13] + Trans[1];
+	Dest[14] = Src[14] + Trans[2];
+	Dest[15] = Src[15];
+
+	return Result;
+}
+
+/** Returns true if any element of this matrix is not finite */
+inline bool FMatrix::ContainsNaN() const
+{
+	for(int32 i=0; i<4; i++)
+	{
+		for(int32 j=0; j<4; j++)
+		{
+			if(!FMath::IsFinite(M[i][j]))
+			{
+				return true;
+			}
+		}
+	}
+
+	return false;
+}
+
+/** @return the maximum magnitude of any row of the matrix. */
+inline float FMatrix::GetMaximumAxisScale() const
+{
+	const float MaxRowScaleSquared = FMath::Max(
+		GetScaledAxis( EAxis::X ).SizeSquared(),
+		FMath::Max(
+			GetScaledAxis( EAxis::Y ).SizeSquared(),
+			GetScaledAxis( EAxis::Z ).SizeSquared()
+			)
+		);
+	return FMath::Sqrt(MaxRowScaleSquared);
+}
+
+inline void FMatrix::ScaleTranslation(const FVector& InScale3D)
+{
+	M[3][0] *= InScale3D.X;
+	M[3][1] *= InScale3D.Y;
+	M[3][2] *= InScale3D.Z;
+}
+
+// GetOrigin
+
+inline FVector FMatrix::GetOrigin() const
+{
+	return FVector(M[3][0],M[3][1],M[3][2]);
+}
+
+inline FVector FMatrix::GetScaledAxis( EAxis::Type InAxis ) const
+{
+	switch ( InAxis )
+	{
+	case EAxis::X:
+		return FVector(M[0][0], M[0][1], M[0][2]);
+
+	case EAxis::Y:
+		return FVector(M[1][0], M[1][1], M[1][2]);
+
+	case EAxis::Z:
+		return FVector(M[2][0], M[2][1], M[2][2]);
+
+	default:
+		ensure(0);
+		return FVector::ZeroVector;
+	}
+}
+
+inline void FMatrix::GetScaledAxes(FVector &X, FVector &Y, FVector &Z) const
+{
+	X.X = M[0][0]; X.Y = M[0][1]; X.Z = M[0][2];
+	Y.X = M[1][0]; Y.Y = M[1][1]; Y.Z = M[1][2];
+	Z.X = M[2][0]; Z.Y = M[2][1]; Z.Z = M[2][2];
+}
+
+inline FVector FMatrix::GetUnitAxis( EAxis::Type InAxis ) const
+{
+	return GetScaledAxis( InAxis ).GetSafeNormal();
+}
+
+inline void FMatrix::GetUnitAxes(FVector &X, FVector &Y, FVector &Z) const
+{
+	GetScaledAxes(X,Y,Z);
+	X.Normalize();
+	Y.Normalize();
+	Z.Normalize();
+}
+
+inline void FMatrix::SetAxis( int32 i, const FVector& Axis )
+{
+	checkSlow(i >= 0 && i <= 2);
+	M[i][0] = Axis.X;
+	M[i][1] = Axis.Y;
+	M[i][2] = Axis.Z;
+}
+
+inline void FMatrix::SetOrigin( const FVector& NewOrigin )
+{
+	M[3][0] = NewOrigin.X;
+	M[3][1] = NewOrigin.Y;
+	M[3][2] = NewOrigin.Z;
+}
+
+inline void FMatrix::SetAxes(FVector* Axis0 /*= NULL*/, FVector* Axis1 /*= NULL*/, FVector* Axis2 /*= NULL*/, FVector* Origin /*= NULL*/)
+{
+	if (Axis0 != NULL)
+	{
+		M[0][0] = Axis0->X;
+		M[0][1] = Axis0->Y;
+		M[0][2] = Axis0->Z;
+	}
+	if (Axis1 != NULL)
+	{
+		M[1][0] = Axis1->X;
+		M[1][1] = Axis1->Y;
+		M[1][2] = Axis1->Z;
+	}
+	if (Axis2 != NULL)
+	{
+		M[2][0] = Axis2->X;
+		M[2][1] = Axis2->Y;
+		M[2][2] = Axis2->Z;
+	}
+	if (Origin != NULL)
+	{
+		M[3][0] = Origin->X;
+		M[3][1] = Origin->Y;
+		M[3][2] = Origin->Z;
+	}
+}
+
+inline FVector FMatrix::GetColumn(int32 i) const
+{
+	checkSlow(i >= 0 && i <= 3);
+	return FVector(M[0][i], M[1][i], M[2][i]);
+}
+
+FORCEINLINE bool MakeFrustumPlane(float A,float B,float C,float D,FPlane& OutPlane)
+{
+	const float	LengthSquared = A * A + B * B + C * C;
+	if(LengthSquared > DELTA*DELTA)
+	{
+		const float	InvLength = FMath::InvSqrt(LengthSquared);
+		OutPlane = FPlane(-A * InvLength,-B * InvLength,-C * InvLength,D * InvLength);
+		return 1;
+	}
+	else
+		return 0;
+}
+
+// Frustum plane extraction.
+FORCEINLINE bool FMatrix::GetFrustumNearPlane(FPlane& OutPlane) const
+{
+	return MakeFrustumPlane(
+		M[0][2],
+		M[1][2],
+		M[2][2],
+		M[3][2],
+		OutPlane
+		);
+}
+
+FORCEINLINE bool FMatrix::GetFrustumFarPlane(FPlane& OutPlane) const
+{
+	return MakeFrustumPlane(
+		M[0][3] - M[0][2],
+		M[1][3] - M[1][2],
+		M[2][3] - M[2][2],
+		M[3][3] - M[3][2],
+		OutPlane
+		);
+}
+
+FORCEINLINE bool FMatrix::GetFrustumLeftPlane(FPlane& OutPlane) const
+{
+	return MakeFrustumPlane(
+		M[0][3] + M[0][0],
+		M[1][3] + M[1][0],
+		M[2][3] + M[2][0],
+		M[3][3] + M[3][0],
+		OutPlane
+		);
+}
+
+FORCEINLINE bool FMatrix::GetFrustumRightPlane(FPlane& OutPlane) const
+{
+	return MakeFrustumPlane(
+		M[0][3] - M[0][0],
+		M[1][3] - M[1][0],
+		M[2][3] - M[2][0],
+		M[3][3] - M[3][0],
+		OutPlane
+		);
+}
+
+FORCEINLINE bool FMatrix::GetFrustumTopPlane(FPlane& OutPlane) const
+{
+	return MakeFrustumPlane(
+		M[0][3] - M[0][1],
+		M[1][3] - M[1][1],
+		M[2][3] - M[2][1],
+		M[3][3] - M[3][1],
+		OutPlane
+		);
+}
+
+FORCEINLINE bool FMatrix::GetFrustumBottomPlane(FPlane& OutPlane) const
+{
+	return MakeFrustumPlane(
+		M[0][3] + M[0][1],
+		M[1][3] + M[1][1],
+		M[2][3] + M[2][1],
+		M[3][3] + M[3][1],
+		OutPlane
+		);
+}
+
+/**
+ * Utility for mirroring this transform across a certain plane,
+ * and flipping one of the axis as well.
+ */
+inline void FMatrix::Mirror(EAxis::Type MirrorAxis, EAxis::Type FlipAxis)
+{
+	if(MirrorAxis == EAxis::X)
+	{
+		M[0][0] *= -1.f;
+		M[1][0] *= -1.f;
+		M[2][0] *= -1.f;
+
+		M[3][0] *= -1.f;
+	}
+	else if(MirrorAxis == EAxis::Y)
+	{
+		M[0][1] *= -1.f;
+		M[1][1] *= -1.f;
+		M[2][1] *= -1.f;
+
+		M[3][1] *= -1.f;
+	}
+	else if(MirrorAxis == EAxis::Z)
+	{
+		M[0][2] *= -1.f;
+		M[1][2] *= -1.f;
+		M[2][2] *= -1.f;
+
+		M[3][2] *= -1.f;
+	}
+
+	if(FlipAxis == EAxis::X)
+	{
+		M[0][0] *= -1.f;
+		M[0][1] *= -1.f;
+		M[0][2] *= -1.f;
+	}
+	else if(FlipAxis == EAxis::Y)
+	{
+		M[1][0] *= -1.f;
+		M[1][1] *= -1.f;
+		M[1][2] *= -1.f;
+	}
+	else if(FlipAxis == EAxis::Z)
+	{
+		M[2][0] *= -1.f;
+		M[2][1] *= -1.f;
+		M[2][2] *= -1.f;
+	}
+}
+
+/** 
+ * Apply Scale to this matrix
+ */
+inline FMatrix FMatrix::ApplyScale(float Scale)
+{
+	FMatrix ScaleMatrix(
+		FPlane(Scale, 0.0f, 0.0f, 0.0f),
+		FPlane(0.0f, Scale, 0.0f, 0.0f),
+		FPlane(0.0f, 0.0f, Scale, 0.0f),
+		FPlane(0.0f, 0.0f, 0.0f, 1.0f)
+	);
+	return ScaleMatrix*(*this);
+}
+
+// Serializer.
+inline FArchive& operator<<(FArchive& Ar,FMatrix& M)
+{
+	Ar << M.M[0][0] << M.M[0][1] << M.M[0][2] << M.M[0][3];
+	Ar << M.M[1][0] << M.M[1][1] << M.M[1][2] << M.M[1][3];
+	Ar << M.M[2][0] << M.M[2][1] << M.M[2][2] << M.M[2][3];
+	Ar << M.M[3][0] << M.M[3][1] << M.M[3][2] << M.M[3][3];
+	return Ar;
+}
+
diff --git a/Engine/Source/Runtime/Core/Public/Math/ScalarRegister.h b/Engine/Source/Runtime/Core/Public/Math/ScalarRegister.h
--- a/Engine/Source/Runtime/Core/Public/Math/ScalarRegister.h
+++ b/Engine/Source/Runtime/Core/Public/Math/ScalarRegister.h
@@ -238,16 +238,3 @@ FORCEINLINE ScalarRegister ScalarReciprocal(const ScalarRegister& A)
 #define NonOneAnimWeight(A) ((A) < 1.0f - ZERO_ANIMWEIGHT_THRESH)
 
 #endif
-
-/**
-* Include the current implementation of a FTransform, depending on the vector processing mode
-*/
-
-#if ENABLE_VECTORIZED_TRANSFORM
-#include "TransformVectorized.h"
-#else
-#include "Transform.h"
-#endif
-
-#include "UnrealMatrix.h"
-
