Copying //UE4/Dev-Build to //UE4/Dev-Main (Source: //UE4/Dev-Build @ 3092544)

diff --git a/Engine/Source/Runtime/Core/Public/Math/Transform.h b/Engine/Source/Runtime/Core/Public/Math/Transform.h
--- a/Engine/Source/Runtime/Core/Public/Math/Transform.h
+++ b/Engine/Source/Runtime/Core/Public/Math/Transform.h
@@ -2,1555 +2,14 @@
 
 #pragma once
 
+/**
+* Include the current implementation of a FTransform, depending on the vector processing mode
+*/
 
-/** 
- * Transform composed of Scale, Rotation (as a quaternion), and Translation.
- *
- * Transforms can be used to convert from one space to another, for example by transforming
- * positions and directions from local space to world space.
- *
- * Transformation of position vectors is applied in the order:  Scale -> Rotate -> Translate.
- * Transformation of direction vectors is applied in the order: Scale -> Rotate.
- *
- * Order matters when composing transforms: C = A * B will yield a transform C that logically
- * first applies A then B to any subsequent transformation. Note that this is the opposite order of quaternion (FQuat) multiplication.
- *
- * Example: LocalToWorld = (DeltaRotation * LocalToWorld) will change rotation in local space by DeltaRotation.
- * Example: LocalToWorld = (LocalToWorld * DeltaRotation) will change rotation in world space by DeltaRotation.
- */
-struct FTransform
-{
-#if !defined(COREUOBJECT_API)
-	#define MAYBE_COREUOBJECT_API
-#else
-	#define MAYBE_COREUOBJECT_API COREUOBJECT_API
-#endif
-	friend MAYBE_COREUOBJECT_API class UScriptStruct* Z_Construct_UScriptStruct_FTransform();
-
-protected:
-	/** Rotation of this transformation, as a quaternion. */
-	FQuat	Rotation;
-	/** Translation of this transformation, as a vector. */
-	FVector	Translation;
-	/** 3D scale (always applied in local space) as a vector. */
-	FVector	Scale3D;
-public:
-	/**
-	 * The identity transformation (Rotation = FQuat::Identity, Translation = FVector::ZeroVector, Scale3D = (1,1,1)).
-	 */
-	static CORE_API const FTransform Identity;
-
-#if ENABLE_NAN_DIAGNOSTIC
-	FORCEINLINE void DiagnosticCheckNaN_Scale3D() const
-	{
-		if (Scale3D.ContainsNaN())
-		{
-			logOrEnsureNanError(TEXT("FTransform Scale3D contains NaN: %s"), *Scale3D.ToString());
-			const_cast<FTransform*>(this)->Scale3D = FVector(1.f);
-		}
-	}
+#include "ScalarRegister.h"
 
-	FORCEINLINE void DiagnosticCheckNaN_Translate() const
-	{
-		if (Translation.ContainsNaN())
-		{
-			logOrEnsureNanError(TEXT("FTransform Translation contains NaN: %s"), *Translation.ToString());
-			const_cast<FTransform*>(this)->Translation = FVector::ZeroVector;
-		}
-	}
-
-	FORCEINLINE void DiagnosticCheckNaN_Rotate() const
-	{
-		if (Rotation.ContainsNaN())
-		{
-			logOrEnsureNanError(TEXT("FTransform Rotation contains NaN: %s"), *Rotation.ToString());
-			const_cast<FTransform*>(this)->Rotation = FQuat::Identity;			
-		}
-	}
-
-	FORCEINLINE void DiagnosticCheckNaN_All() const
-	{
-		DiagnosticCheckNaN_Scale3D();
-		DiagnosticCheckNaN_Rotate();
-		DiagnosticCheckNaN_Translate();
-	}
-
-	FORCEINLINE void DiagnosticCheck_IsValid() const
-	{
-		DiagnosticCheckNaN_All();
-		if (!IsValid())
-		{
-			logOrEnsureNanError(TEXT("FTransform transform is not valid: %s"), *ToHumanReadableString());
-		}
-		
-	}
+#if ENABLE_VECTORIZED_TRANSFORM
+#include "TransformVectorized.h"
 #else
-	FORCEINLINE void DiagnosticCheckNaN_Translate() const {}
-	FORCEINLINE void DiagnosticCheckNaN_Rotate() const {}
-	FORCEINLINE void DiagnosticCheckNaN_Scale3D() const {}
-	FORCEINLINE void DiagnosticCheckNaN_All() const {}
-	FORCEINLINE void DiagnosticCheck_IsValid() const {}
-#endif
-
-	/** Default constructor. */
-	FORCEINLINE FTransform()
-		: Rotation(0.f,0.f,0.f,1.f)
-		, Translation(0.f)
-		, Scale3D(1.f)
-	{
-	}
-
-	/**
-	 * Constructor with an initial translation
-	 *
-	 * @param InTranslation The value to use for the translation component
-	 */
-	FORCEINLINE explicit FTransform(const FVector& InTranslation) 
-		:	Rotation(FQuat::Identity), 
-			Translation(InTranslation), 
-			Scale3D(FVector(1.f))
-	{
-		DiagnosticCheckNaN_All();
-	}
-
-	/**
-	 * Constructor with leaving uninitialized memory
-	 */
-	FORCEINLINE explicit FTransform(ENoInit) 
-	{
-		// Note: This can be used to track down initialization issues with bone transform arrays; but it will
-		// cause issues with transient fields such as RootMotionDelta that get initialized to 0 by default
-#if ENABLE_NAN_DIAGNOSTIC
-		float qnan = FMath::Log2(-5.3f);
-		check(FMath::IsNaN(qnan));
-		Translation = FVector(qnan, qnan, qnan);
-		Rotation = FQuat(qnan, qnan, qnan, qnan);
-		Scale3D = FVector(qnan, qnan, qnan);
+#include "TransformNonVectorized.h"
 #endif
-	}
-
-	/**
-	 * Constructor with an initial rotation
-	 *
-	 * @param InRotation The value to use for rotation component
-	 */
-	FORCEINLINE explicit FTransform(const FQuat& InRotation) 
-		:	Rotation(InRotation), 
-			Translation(FVector::ZeroVector), 
-			Scale3D(FVector(1.f))
-	{
-		DiagnosticCheckNaN_All();
-	}
-
-	/**
-	 * Constructor with an initial rotation
-	 *
-	 * @param InRotation The value to use for rotation component  (after being converted to a quaternion)
-	 */
-	FORCEINLINE explicit FTransform(const FRotator& InRotation) 
-		:	Rotation(InRotation), 
-			Translation(FVector::ZeroVector), 
-			Scale3D(FVector(1.f))
-	{
-		DiagnosticCheckNaN_All();
-	}
-
-	/**
-	 * Constructor with all components initialized
-	 *
-	 * @param InRotation The value to use for rotation component
-	 * @param InTranslation The value to use for the translation component
-	 * @param InScale3D The value to use for the scale component
-	 */
-	FORCEINLINE FTransform(const FQuat& InRotation, const FVector& InTranslation, const FVector& InScale3D = FVector(1.f)) 
-		:	Rotation(InRotation), 
-			Translation(InTranslation), 
-			Scale3D(InScale3D)
-	{
-		DiagnosticCheckNaN_All();
-	}
-
-	/**
-	 * Constructor with all components initialized, taking a FRotator as the rotation component
-	 *
-	 * @param InRotation The value to use for rotation component (after being converted to a quaternion)
-	 * @param InTranslation The value to use for the translation component
-	 * @param InScale3D The value to use for the scale component
-	 */
-	FORCEINLINE FTransform(const FRotator& InRotation, const FVector& InTranslation, const FVector& InScale3D = FVector(1.f))
-		:	Rotation(InRotation), 
-			Translation(InTranslation), 
-			Scale3D(InScale3D)
-	{
-		DiagnosticCheckNaN_All();
-	}
-
-	/**
-	 * Copy-constructor
-	 *
-	 * @param InTransform The source transform from which all components will be copied
-	 */
-	FORCEINLINE FTransform(const FTransform& InTransform) : 
-		Rotation(InTransform.Rotation), 
-		Translation(InTransform.Translation), 
-		Scale3D(InTransform.Scale3D)
-	{
-		DiagnosticCheckNaN_All();
-	}
-
-	/**
-	 * Constructor for converting a Matrix (including scale) into a FTransform.
-	 */
-	FORCEINLINE explicit FTransform(const FMatrix& InMatrix)
-	{
-		SetFromMatrix(InMatrix);
-		DiagnosticCheckNaN_All();
-	}
-
-	/** Constructor that takes basis axes and translation */
-	FORCEINLINE FTransform(const FVector& InX, const FVector& InY, const FVector& InZ, const FVector& InTranslation)
-	{
-		SetFromMatrix(FMatrix(InX, InY, InZ, InTranslation));
-		DiagnosticCheckNaN_All();
-	}
-
-	/**
-	 * Does a debugf of the contents of this Transform.
-	 */
-	CORE_API void DebugPrint() const;
-
-	/** Debug purpose only **/
-	bool DebugEqualMatrix(const FMatrix& Matrix) const;
-
-	/** Convert FTransform contents to a string */
-	CORE_API FString ToHumanReadableString() const;
-
-	CORE_API FString ToString() const;
-
-	/** Acceptable form: "%f,%f,%f|%f,%f,%f|%f,%f,%f" */
-	CORE_API bool InitFromString(const FString& InSourceString);
-
-#ifdef IMPLEMENT_ASSIGNMENT_OPERATOR_MANUALLY
-	/**
-	 * Copy another Transform into this one
-	 */
-	FORCEINLINE FTransform& operator=(const FTransform& Other)
-	{
-		this->Rotation = Other.Rotation;
-		this->Translation = Other.Translation;
-		this->Scale3D = Other.Scale3D;
-
-		return *this;
-	}
-#endif
-
-	/**
-	 * Convert this Transform to a transformation matrix with scaling.
-	 */
-	FORCEINLINE FMatrix ToMatrixWithScale() const
-	{
-		FMatrix OutMatrix;
-
-#if !(UE_BUILD_SHIPPING || UE_BUILD_TEST) && WITH_EDITORONLY_DATA
-		// Make sure Rotation is normalized when we turn it into a matrix.
-		check(IsRotationNormalized());
-#endif
-		OutMatrix.M[3][0] = Translation.X;
-		OutMatrix.M[3][1] = Translation.Y;
-		OutMatrix.M[3][2] = Translation.Z;
-
-		const float x2 = Rotation.X + Rotation.X;	
-		const float y2 = Rotation.Y + Rotation.Y;  
-		const float z2 = Rotation.Z + Rotation.Z;
-		{
-			const float xx2 = Rotation.X * x2;
-			const float yy2 = Rotation.Y * y2;			
-			const float zz2 = Rotation.Z * z2;
-
-			OutMatrix.M[0][0] = (1.0f - (yy2 + zz2)) * Scale3D.X;	
-			OutMatrix.M[1][1] = (1.0f - (xx2 + zz2)) * Scale3D.Y;
-			OutMatrix.M[2][2] = (1.0f - (xx2 + yy2)) * Scale3D.Z;
-		}
-		{
-			const float yz2 = Rotation.Y * z2;   
-			const float wx2 = Rotation.W * x2;	
-
-			OutMatrix.M[2][1] = (yz2 - wx2) * Scale3D.Z;
-			OutMatrix.M[1][2] = (yz2 + wx2) * Scale3D.Y;
-		}
-		{
-			const float xy2 = Rotation.X * y2;
-			const float wz2 = Rotation.W * z2;
-
-			OutMatrix.M[1][0] = (xy2 - wz2) * Scale3D.Y;
-			OutMatrix.M[0][1] = (xy2 + wz2) * Scale3D.X;
-		}
-		{
-			const float xz2 = Rotation.X * z2;
-			const float wy2 = Rotation.W * y2;   
-
-			OutMatrix.M[2][0] = (xz2 + wy2) * Scale3D.Z;
-			OutMatrix.M[0][2] = (xz2 - wy2) * Scale3D.X;
-		}
-
-		OutMatrix.M[0][3] = 0.0f;
-		OutMatrix.M[1][3] = 0.0f;
-		OutMatrix.M[2][3] = 0.0f;
-		OutMatrix.M[3][3] = 1.0f;
-
-		return OutMatrix;
-	}
-
-	/**
-	 * Convert this Transform to matrix with scaling and compute the inverse of that.
-	 */
-	FORCEINLINE FMatrix ToInverseMatrixWithScale() const
-	{
-		// todo: optimize
-		return ToMatrixWithScale().Inverse();
-	}
-
-	/**
-	 * Convert this Transform to inverse.
-	 */
-	FORCEINLINE FTransform Inverse() const
-	{
-		FQuat   InvRotation    = Rotation.Inverse();
-		// this used to cause NaN if Scale contained 0 
-		FVector InvScale3D     = GetSafeScaleReciprocal(Scale3D);
-		FVector InvTranslation = InvRotation * (InvScale3D * -Translation);
-
-		return FTransform(InvRotation, InvTranslation, InvScale3D);
-	}
-
-	/**
-	 * Convert this Transform to a transformation matrix, ignoring its scaling
-	 */
-	FORCEINLINE FMatrix ToMatrixNoScale() const
-	{
-		FMatrix OutMatrix;
-
-#if !(UE_BUILD_SHIPPING || UE_BUILD_TEST) && WITH_EDITORONLY_DATA
-		// Make sure Rotation is normalized when we turn it into a matrix.
-		check(IsRotationNormalized());
-#endif
-		OutMatrix.M[3][0] = Translation.X;
-		OutMatrix.M[3][1] = Translation.Y;
-		OutMatrix.M[3][2] = Translation.Z;
-
-		const float x2 = Rotation.X + Rotation.X;	
-		const float y2 = Rotation.Y + Rotation.Y;  
-		const float z2 = Rotation.Z + Rotation.Z;
-		{
-			const float xx2 = Rotation.X * x2;
-			const float yy2 = Rotation.Y * y2;			
-			const float zz2 = Rotation.Z * z2;
-
-			OutMatrix.M[0][0] = (1.0f - (yy2 + zz2));	
-			OutMatrix.M[1][1] = (1.0f - (xx2 + zz2));
-			OutMatrix.M[2][2] = (1.0f - (xx2 + yy2));
-		}
-		{
-			const float yz2 = Rotation.Y * z2;   
-			const float wx2 = Rotation.W * x2;	
-
-			OutMatrix.M[2][1] = (yz2 - wx2);
-			OutMatrix.M[1][2] = (yz2 + wx2);
-		}
-		{
-			const float xy2 = Rotation.X * y2;
-			const float wz2 = Rotation.W * z2;
-
-			OutMatrix.M[1][0] = (xy2 - wz2);
-			OutMatrix.M[0][1] = (xy2 + wz2);
-		}
-		{
-			const float xz2 = Rotation.X * z2;
-			const float wy2 = Rotation.W * y2;   
-
-			OutMatrix.M[2][0] = (xz2 + wy2);
-			OutMatrix.M[0][2] = (xz2 - wy2);
-		}
-
-		OutMatrix.M[0][3] = 0.0f;
-		OutMatrix.M[1][3] = 0.0f;
-		OutMatrix.M[2][3] = 0.0f;
-		OutMatrix.M[3][3] = 1.0f;
-
-		return OutMatrix;
-	}
-
-	/** Set this transform to the weighted blend of the supplied two transforms. */
-	FORCEINLINE void Blend(const FTransform& Atom1, const FTransform& Atom2, float Alpha)
-	{
-#if !(UE_BUILD_SHIPPING || UE_BUILD_TEST) && WITH_EDITORONLY_DATA
-		// Check that all bone atoms coming from animation are normalized
-		check(Atom1.IsRotationNormalized());
-		check(Atom2.IsRotationNormalized());
-#endif
-		if(Alpha <= ZERO_ANIMWEIGHT_THRESH)
-		{
-			// if blend is all the way for child1, then just copy its bone atoms
-			(*this) = Atom1;
-		}
-		else if(Alpha >= 1.f - ZERO_ANIMWEIGHT_THRESH)
-		{
-			// if blend is all the way for child2, then just copy its bone atoms
-			(*this) = Atom2;
-		}
-		else
-		{
-			// Simple linear interpolation for translation and scale.
-			Translation = FMath::Lerp(Atom1.Translation, Atom2.Translation, Alpha);
-			Scale3D		= FMath::Lerp(Atom1.Scale3D, Atom2.Scale3D, Alpha);
-			Rotation	= FQuat::FastLerp(Atom1.Rotation, Atom2.Rotation, Alpha);
-
-			// ..and renormalize
-			Rotation.Normalize();
-		}
-	}
-
-	/** Set this Transform to the weighted blend of it and the supplied Transform. */
-	FORCEINLINE void BlendWith(const FTransform& OtherAtom, float Alpha)
-	{
-#if !(UE_BUILD_SHIPPING || UE_BUILD_TEST) && WITH_EDITORONLY_DATA
-		// Check that all bone atoms coming from animation are normalized
-		check(IsRotationNormalized());
-		check(OtherAtom.IsRotationNormalized());
-#endif
-		if(Alpha > ZERO_ANIMWEIGHT_THRESH)
-		{
-			if(Alpha >= 1.f - ZERO_ANIMWEIGHT_THRESH)
-			{
-				// if blend is all the way for child2, then just copy its bone atoms
-				(*this) = OtherAtom;
-			}
-			else 
-			{
-				// Simple linear interpolation for translation and scale.
-				Translation = FMath::Lerp(Translation, OtherAtom.Translation, Alpha);
-				Scale3D		= FMath::Lerp(Scale3D, OtherAtom.Scale3D, Alpha);
-				Rotation	= FQuat::FastLerp(Rotation, OtherAtom.Rotation, Alpha);
-
-				// ..and renormalize
-				Rotation.Normalize();
-			}
-		}
-	}
-
-
-	/**
-	 * Quaternion addition is wrong here. This is just a special case for linear interpolation.
-	 * Use only within blends!!
-	 * Rotation part is NOT normalized!!
-	 */
-	FORCEINLINE FTransform operator+(const FTransform& Atom) const
-	{
-		return FTransform(Rotation + Atom.Rotation, Translation + Atom.Translation, Scale3D + Atom.Scale3D);
-	}
-
-	FORCEINLINE FTransform& operator+=(const FTransform& Atom)
-	{
-		Translation += Atom.Translation;
-
-		Rotation.X += Atom.Rotation.X;
-		Rotation.Y += Atom.Rotation.Y;
-		Rotation.Z += Atom.Rotation.Z;
-		Rotation.W += Atom.Rotation.W;
-
-		Scale3D += Atom.Scale3D;
-
-		DiagnosticCheckNaN_All();
-		return *this;
-	}
-
-	FORCEINLINE FTransform operator*(ScalarRegister Mult) const
-	{
-		return FTransform(Rotation * Mult, Translation * Mult, Scale3D * Mult);
-	}
-
-	FORCEINLINE FTransform& operator*=(ScalarRegister Mult)
-	{
-		Translation *= Mult;
-		Rotation.X	*= Mult;
-		Rotation.Y	*= Mult;
-		Rotation.Z	*= Mult;
-		Rotation.W	*= Mult;
-		Scale3D		*= Mult;
-		DiagnosticCheckNaN_All();
-
-		return *this;
-	}
-
-	/**
-	 * Return a transform that is the result of this multiplied by another transform.
-	 * Order matters when composing transforms : C = A * B will yield a transform C that logically first applies A then B to any subsequent transformation.
-	 * 
-	 * @param  Other other transform by which to multiply.
-	 * @return new transform: this * Other
-	 */
-	FORCEINLINE FTransform operator*(const FTransform& Other) const;
-
-
-	/**
-	 * Sets this transform to the result of this multiplied by another transform.
-	 * Order matters when composing transforms : C = A * B will yield a transform C that logically first applies A then B to any subsequent transformation.
-	 *
-	 * @param  Other other transform by which to multiply.
-	 */
-	FORCEINLINE void operator*=(const FTransform& Other);
-	
-	/**
-	 * Return a transform that is the result of this multiplied by another transform (made only from a rotation).
-	 * Order matters when composing transforms : C = A * B will yield a transform C that logically first applies A then B to any subsequent transformation.
-	 *
-	 * @param  Other other quaternion rotation by which to multiply.
-	 * @return new transform: this * FTransform(Other)
-	 */
-	FORCEINLINE FTransform operator*(const FQuat& Other) const;
-	
-	/**
-	 * Sets this transform to the result of this multiplied by another transform (made only from a rotation).
-	 * Order matters when composing transforms : C = A * B will yield a transform C that logically first applies A then B to any subsequent transformation.
-	 *
-	 * @param  Other other quaternion rotation by which to multiply.
-	 */
-	FORCEINLINE void operator*=(const FQuat& Other);
-
-	FORCEINLINE void ScaleTranslation(const FVector& InScale3D);
-	FORCEINLINE void ScaleTranslation(const float& Scale);
-	FORCEINLINE void RemoveScaling(float Tolerance=SMALL_NUMBER);
-	FORCEINLINE float GetMaximumAxisScale() const;
-	FORCEINLINE float GetMinimumAxisScale() const;
-
-	// Inverse does not work well with VQS format(in particular non-uniform), so removing it, but made two below functions to be used instead. 
-
-	/*******************************************************************************************
-	 * The below 2 functions are the ones to get delta transform and return FTransform format that can be concatenated
-	 * Inverse itself can't concatenate with VQS format(since VQS always transform from S->Q->T, where inverse happens from T(-1)->Q(-1)->S(-1))
-	 * So these 2 provides ways to fix this
-	 * GetRelativeTransform returns this*Other(-1) and parameter is Other(not Other(-1))
-	 * GetRelativeTransformReverse returns this(-1)*Other, and parameter is Other. 
-	 *******************************************************************************************/
-	CORE_API FTransform GetRelativeTransform(const FTransform& Other) const;
-	CORE_API FTransform GetRelativeTransformReverse(const FTransform& Other) const;
-	/**
-	 * Set current transform and the relative to ParentTransform.
-	 * Equates to This = This->GetRelativeTransform(Parent), but saves the intermediate FTransform storage and copy.
-	 */
-	CORE_API void SetToRelativeTransform(const FTransform& ParentTransform);
-
-	FORCEINLINE FVector4 TransformFVector4(const FVector4& V) const;
-	FORCEINLINE FVector4 TransformFVector4NoScale(const FVector4& V) const;
-	FORCEINLINE FVector TransformPosition(const FVector& V) const;
-	FORCEINLINE FVector TransformPositionNoScale(const FVector& V) const;
-
-	/** Inverts the matrix and then transforms V - correctly handles scaling in this matrix. */
-	FORCEINLINE FVector InverseTransformPosition(const FVector &V) const;
-	FORCEINLINE FVector InverseTransformPositionNoScale(const FVector &V) const;
-	FORCEINLINE FVector TransformVector(const FVector& V) const;
-	FORCEINLINE FVector TransformVectorNoScale(const FVector& V) const;
-
-	/** 
-	 *	Transform a direction vector by the inverse of this matrix - will not take into account translation part.
-	 *	If you want to transform a surface normal (or plane) and correctly account for non-uniform scaling you should use TransformByUsingAdjointT with adjoint of matrix inverse.
-	 */
-	FORCEINLINE FVector InverseTransformVector(const FVector &V) const;
-	FORCEINLINE FVector InverseTransformVectorNoScale(const FVector &V) const;
-
-	FORCEINLINE FTransform GetScaled(float Scale) const;
-	FORCEINLINE FTransform GetScaled(FVector Scale) const;
-	FORCEINLINE FVector GetScaledAxis(EAxis::Type InAxis) const;
-	FORCEINLINE FVector GetUnitAxis(EAxis::Type InAxis) const;
-	FORCEINLINE void Mirror(EAxis::Type MirrorAxis, EAxis::Type FlipAxis);
-	FORCEINLINE FVector GetSafeScaleReciprocal(const FVector& InScale, float Tolerance=SMALL_NUMBER) const;
-
-	// temp function for easy conversion
-	FORCEINLINE FVector GetLocation() const
-	{
-		return GetTranslation();
-	}
-
-	FORCEINLINE FRotator Rotator() const
-	{
-		return Rotation.Rotator();
-	}
-
-	/** Calculate the  */
-	FORCEINLINE float GetDeterminant() const
-	{
-		return Scale3D.X * Scale3D.Y * Scale3D.Z;
-	}
-
-	/** Set the translation of this transformation */
-	FORCEINLINE void SetLocation(const FVector& Origin)
-	{
-		Translation = Origin;
-		DiagnosticCheckNaN_Translate();
-	}
-
-	/**
-	 * Checks the components for non-finite values (NaN or Inf).
-	 * @return Returns true if any component (rotation, translation, or scale) is not finite.
-	 */
-	bool ContainsNaN() const
-	{
-		return (Translation.ContainsNaN() || Rotation.ContainsNaN() || Scale3D.ContainsNaN());
-	}
-
-	inline bool IsValid() const
-	{
-		if (ContainsNaN())
-		{
-			return false;
-		}
-
-		if (!Rotation.IsNormalized())
-		{
-			return false;
-		}
-
-		return true;
-	}
-
-	// Serializer.
-	inline friend FArchive& operator<<(FArchive& Ar,FTransform& M)
-	{
-		Ar << M.Rotation;
-		Ar << M.Translation;
-		Ar << M.Scale3D;
-		return Ar;
-	}
-
-	// Binary comparison operators.
-	/*
-	bool operator==(const FTransform& Other) const
-	{
-		return Rotation==Other.Rotation && Translation==Other.Translation && Scale3D==Other.Scale3D;
-	}
-	bool operator!=(const FTransform& Other) const
-	{
-		return Rotation!=Other.Rotation || Translation!=Other.Translation || Scale3D!=Other.Scale3D;
-	}
-	*/
-
-private:
-
-	FORCEINLINE bool Private_RotationEquals(const FQuat& InRotation, const float Tolerance = KINDA_SMALL_NUMBER) const
-	{
-		return Rotation.Equals(InRotation, Tolerance);
-	}
-
-	FORCEINLINE bool Private_TranslationEquals(const FVector& InTranslation, const float Tolerance = KINDA_SMALL_NUMBER) const
-	{
-		return Translation.Equals(InTranslation, Tolerance);
-	}
-
-	FORCEINLINE bool Private_Scale3DEquals(const FVector& InScale3D, const float Tolerance = KINDA_SMALL_NUMBER) const
-	{
-		return Scale3D.Equals(InScale3D, Tolerance);
-	}
-
-public:
-
-	// Test if A's rotation equals B's rotation, within a tolerance. Preferred over "A.GetRotation().Equals(B.GetRotation())" because it is faster on some platforms.
-	FORCEINLINE static bool AreRotationsEqual(const FTransform& A, const FTransform& B, float Tolerance=KINDA_SMALL_NUMBER)
-	{
-		return A.Private_RotationEquals(B.Rotation, Tolerance);
-	}
-
-	// Test if A's translation equals B's translation, within a tolerance. Preferred over "A.GetTranslation().Equals(B.GetTranslation())" because it is faster on some platforms.
-	FORCEINLINE static bool AreTranslationsEqual(const FTransform& A, const FTransform& B, float Tolerance=KINDA_SMALL_NUMBER)
-	{
-		return A.Private_TranslationEquals(B.Translation, Tolerance);
-	}
-
-	// Test if A's scale equals B's scale, within a tolerance. Preferred over "A.GetScale3D().Equals(B.GetScale3D())" because it is faster on some platforms.
-	FORCEINLINE static bool AreScale3DsEqual(const FTransform& A, const FTransform& B, float Tolerance=KINDA_SMALL_NUMBER)
-	{
-		return A.Private_Scale3DEquals(B.Scale3D, Tolerance);
-	}
-
-
-
-	// Test if this Transform's rotation equals another's rotation, within a tolerance. Preferred over "GetRotation().Equals(Other.GetRotation())" because it is faster on some platforms.
-	FORCEINLINE bool RotationEquals(const FTransform& Other, float Tolerance=KINDA_SMALL_NUMBER) const
-	{
-		return AreRotationsEqual(*this, Other, Tolerance);
-	}
-
-	// Test if this Transform's translation equals another's translation, within a tolerance. Preferred over "GetTranslation().Equals(Other.GetTranslation())" because it is faster on some platforms.
-	FORCEINLINE bool TranslationEquals(const FTransform& Other, float Tolerance=KINDA_SMALL_NUMBER) const
-	{
-		return AreTranslationsEqual(*this, Other, Tolerance);
-	}
-
-	// Test if this Transform's scale equals another's scale, within a tolerance. Preferred over "GetScale3D().Equals(Other.GetScale3D())" because it is faster on some platforms.
-	FORCEINLINE bool Scale3DEquals(const FTransform& Other, float Tolerance=KINDA_SMALL_NUMBER) const
-	{
-		return AreScale3DsEqual(*this, Other, Tolerance);
-	}
-
-
-	// Test if all components of the transforms are equal, within a tolerance.
-	inline bool Equals(const FTransform& Other, float Tolerance=KINDA_SMALL_NUMBER) const
-	{
-		return Private_TranslationEquals(Other.Translation, Tolerance) && Private_RotationEquals(Other.Rotation, Tolerance) && Private_Scale3DEquals(Other.Scale3D, Tolerance);
-	}
-
-	// Test if rotation and translation components of the transforms are equal, within a tolerance.
-	inline bool EqualsNoScale(const FTransform& Other, float Tolerance=KINDA_SMALL_NUMBER) const
-	{
-		return Private_TranslationEquals(Other.Translation, Tolerance) && Private_RotationEquals(Other.Rotation, Tolerance);
-	}
-
-	/**
-	 * Create a new transform: OutTransform = A * B.
-	 *
-	 * Order matters when composing transforms : A * B will yield a transform that logically first applies A then B to any subsequent transformation.
-	 * 
-	 * @param  OutTransform pointer to transform that will store the result of A * B.
-	 * @param  A Transform A.
-	 * @param  B Transform B.
-	 */
-	FORCEINLINE static void Multiply(FTransform* OutTransform, const FTransform* A, const FTransform* B);
-
-	/**
-	* Create a new transform: OutTransform = A * B using the matrix while keeping the scale that's given by A and B
-	* Please note that this operation is a lot more expensive than normal Multiply
-	*
-	* Order matters when composing transforms : A * B will yield a transform that logically first applies A then B to any subsequent transformation.
-	*
-	* @param  OutTransform pointer to transform that will store the result of A * B.
-	* @param  A Transform A.
-	* @param  B Transform B.
-	*/
-	FORCEINLINE static void MultiplyUsingMatrixWithScale(FTransform* OutTransform, const FTransform* A, const FTransform* B);
-
-	/**
-	 * Sets the components
-	 * @param InRotation The new value for the Rotation component
-	 * @param InTranslation The new value for the Translation component
-	 * @param InScale3D The new value for the Scale3D component
-	 */
-	FORCEINLINE void SetComponents(const FQuat& InRotation, const FVector& InTranslation, const FVector& InScale3D) 
-	{
-		Rotation = InRotation;
-		Translation = InTranslation;
-		Scale3D = InScale3D;
-
-		DiagnosticCheckNaN_All();
-	}
-
-	/**
-	 * Sets the components to the identity transform:
-	 *   Rotation = (0,0,0,1)
-	 *   Translation = (0,0,0)
-	 *   Scale3D = (1,1,1)
-	 */
-	FORCEINLINE void SetIdentity()
-	{
-		Rotation = FQuat::Identity;
-		Translation	= FVector::ZeroVector;
-		Scale3D = FVector(1,1,1);
-	}
-
-	/**
-	 * Scales the Scale3D component by a new factor
-	 * @param Scale3DMultiplier The value to multiply Scale3D with
-	 */
-	FORCEINLINE void MultiplyScale3D(const FVector& Scale3DMultiplier)
-	{
-		Scale3D *= Scale3DMultiplier;
-		DiagnosticCheckNaN_Scale3D();
-	}
-
-	/**
-	 * Sets the translation component
-	 * @param NewTranslation The new value for the translation component
-	 */
-	FORCEINLINE void SetTranslation(const FVector& NewTranslation)
-	{
-		Translation = NewTranslation;
-		DiagnosticCheckNaN_Translate();
-	}
-
-	/** Copy translation from another FTransform. */
-	FORCEINLINE void CopyTranslation(const FTransform& Other)
-	{
-		Translation = Other.Translation;
-	}
-
-	/**
-	 * Concatenates another rotation to this transformation 
-	 * @param DeltaRotation The rotation to concatenate in the following fashion: Rotation = Rotation * DeltaRotation
-	 */
-	FORCEINLINE void ConcatenateRotation(const FQuat& DeltaRotation)
-	{
-		Rotation = Rotation * DeltaRotation;
-		DiagnosticCheckNaN_Rotate();
-	}
-
-	/**
-	 * Adjusts the translation component of this transformation 
-	 * @param DeltaTranslation The translation to add in the following fashion: Translation += DeltaTranslation
-	 */
-	FORCEINLINE void AddToTranslation(const FVector& DeltaTranslation)
-	{
-		Translation += DeltaTranslation;
-		DiagnosticCheckNaN_Translate();
-	}
-
-	/**
-	 * Add the translations from two FTransforms and return the result.
-	 * @return A.Translation + B.Translation
-	 */
-	FORCEINLINE static FVector AddTranslations(const FTransform& A, const FTransform& B)
-	{
-		return A.Translation + B.Translation;
-	}
-
-	/**
-	 * Subtract translations from two FTransforms and return the difference.
-	 * @return A.Translation - B.Translation.
-	 */
-	FORCEINLINE static FVector SubtractTranslations(const FTransform& A, const FTransform& B)
-	{
-		return A.Translation - B.Translation;
-	}
-
-	/**
-	 * Sets the rotation component
-	 * @param NewRotation The new value for the rotation component
-	 */
-	FORCEINLINE void SetRotation(const FQuat& NewRotation)
-	{
-		Rotation = NewRotation;
-		DiagnosticCheckNaN_Rotate();
-	}
-
-	/** Copy rotation from another FTransform. */
-	FORCEINLINE void CopyRotation(const FTransform& Other)
-	{
-		Rotation = Other.Rotation;
-	}
-
-	/**
-	 * Sets the Scale3D component
-	 * @param NewScale3D The new value for the Scale3D component
-	 */
-	FORCEINLINE void SetScale3D(const FVector& NewScale3D)
-	{
-		Scale3D = NewScale3D;
-		DiagnosticCheckNaN_Scale3D();
-	}
-
-	/** Copy scale from another FTransform. */
-	FORCEINLINE void CopyScale3D(const FTransform& Other)
-	{
-		Scale3D = Other.Scale3D;
-	}
-
-	/**
-	 * Sets both the translation and Scale3D components at the same time
-	 * @param NewTranslation The new value for the translation component
-	 * @param NewScale3D The new value for the Scale3D component
-	 */
-	FORCEINLINE void SetTranslationAndScale3D(const FVector& NewTranslation, const FVector& NewScale3D)
-	{
-		Translation = NewTranslation;
-		Scale3D = NewScale3D;
-
-		DiagnosticCheckNaN_Translate();
-		DiagnosticCheckNaN_Scale3D();
-	}
-
-	/** @note: Added template type function for Accumulate
-	  * The template type isn't much useful yet, but it is with the plan to move forward
-	  * to unify blending features with just type of additive or full pose
-	  * Eventually it would be nice to just call blend and it all works depending on full pose
-	  * or additive, but right now that is a lot more refactoring
-	  * For now this types only defines the different functionality of accumulate
-	  */
-	  
-	/**
-	 * Accumulates another transform with this one
-	 *
-	 * Rotation is accumulated multiplicatively (Rotation = SourceAtom.Rotation * Rotation)
-	 * Translation is accumulated additively (Translation += SourceAtom.Translation)
-	 * Scale3D is accumulated multiplicatively (Scale3D *= SourceAtom.Scale3D)
-	 *
-	 * @param SourceAtom The other transform to accumulate into this one
-	 */
-	FORCEINLINE void Accumulate(const FTransform& SourceAtom)
-	{
-		// Add ref pose relative animation to base animation, only if rotation is significant.
-		if(FMath::Square(SourceAtom.Rotation.W) < 1.f - DELTA * DELTA)
-		{
-			Rotation = SourceAtom.Rotation * Rotation;
-		}
-
-		Translation += SourceAtom.Translation;
-		Scale3D *= SourceAtom.Scale3D;
-
-		DiagnosticCheckNaN_All();
-
-		checkSlow(IsRotationNormalized());
-	}
-	
-	/** Accumulates another transform with this one, with a blending weight
-	*
-	* Let SourceAtom = Atom * BlendWeight
-	* Rotation is accumulated multiplicatively (Rotation = SourceAtom.Rotation * Rotation).
-	* Translation is accumulated additively (Translation += SourceAtom.Translation)
-	* Scale3D is accumulated multiplicatively (Scale3D *= SourceAtom.Scale3D)
-	*
-	* Note: Rotation will not be normalized! Will have to be done manually.
-	*
-	* @param Atom The other transform to accumulate into this one
-	* @param BlendWeight The weight to multiply Atom by before it is accumulated.
-	*/
-	FORCEINLINE void Accumulate(const FTransform& Atom, float BlendWeight/* default param doesn't work since vectorized version takes ref param */)
-	{
-		FTransform SourceAtom(Atom * BlendWeight);
-
-		// Add ref pose relative animation to base animation, only if rotation is significant.
-		if (FMath::Square(SourceAtom.Rotation.W) < 1.f - DELTA * DELTA)
-		{
-			Rotation = SourceAtom.Rotation * Rotation;
-		}
-
-		Translation += SourceAtom.Translation;
-		Scale3D *= SourceAtom.Scale3D;
-
-		DiagnosticCheckNaN_All();
-	}
-
-	/**
-	 * Accumulates another transform with this one, with an optional blending weight
-	 *
-	 * Rotation is accumulated additively, in the shortest direction (Rotation = Rotation +/- DeltaAtom.Rotation * Weight)
-	 * Translation is accumulated additively (Translation += DeltaAtom.Translation * Weight)
-	 * Scale3D is accumulated additively (Scale3D += DeltaAtom.Scale3D * Weight)
-	 *
-	 * @param DeltaAtom The other transform to accumulate into this one
-	 * @param Weight The weight to multiply DeltaAtom by before it is accumulated.
-	 */
-	FORCEINLINE void AccumulateWithShortestRotation(const FTransform& DeltaAtom, float BlendWeight/* default param doesn't work since vectorized version takes ref param */)
-	{
-		FTransform Atom(DeltaAtom * BlendWeight);
-
-		// To ensure the 'shortest route', we make sure the dot product between the accumulator and the incoming child atom is positive.
-		if((Atom.Rotation | Rotation) < 0.f)
-		{
-			Rotation.X -= Atom.Rotation.X;
-			Rotation.Y -= Atom.Rotation.Y;
-			Rotation.Z -= Atom.Rotation.Z;
-			Rotation.W -= Atom.Rotation.W;
-		}
-		else
-		{
-			Rotation.X += Atom.Rotation.X;
-			Rotation.Y += Atom.Rotation.Y;
-			Rotation.Z += Atom.Rotation.Z;
-			Rotation.W += Atom.Rotation.W;
-		}
-
-		Translation += Atom.Translation;
-		Scale3D += Atom.Scale3D;
-
-		DiagnosticCheckNaN_All();
-	}
-
-	/** Accumulates another transform with this one, with a blending weight
-	*
-	* Let SourceAtom = Atom * BlendWeight
-	* Rotation is accumulated multiplicatively (Rotation = SourceAtom.Rotation * Rotation).
-	* Translation is accumulated additively (Translation += SourceAtom.Translation)
-	* Scale3D is accumulated assuming incoming scale is additive scale (Scale3D *= (1 + SourceAtom.Scale3D))
-	* 
-	* When we create additive, we create additive scale based on [TargetScale/SourceScale -1]
-	* because that way when you apply weight of 0.3, you don't shrink. We only saves the % of grow/shrink
-	* when we apply that back to it, we add back the 1, so that it goes back to it. 
-	* This solves issue where you blend two additives with 0.3, you don't come back to 0.6 scale, but 1 scale at the end 
-	* because [1 + [1-1]*0.3 + [1-1]*0.3] becomes 1, so you don't shrink by applying additive scale
-	*
-	* Note: Rotation will not be normalized! Will have to be done manually.
-	*
-	* @param Atom The other transform to accumulate into this one
-	* @param BlendWeight The weight to multiply Atom by before it is accumulated.
-	*/
-	FORCEINLINE void AccumulateWithAdditiveScale(const FTransform& Atom, float BlendWeight/* default param doesn't work since vectorized version takes ref param */)
-	{
-		const FVector DefaultScale(1.f);
-
-		FTransform SourceAtom(Atom * BlendWeight);
-		
-		// Add ref pose relative animation to base animation, only if rotation is significant.
-		if (FMath::Square(SourceAtom.Rotation.W) < 1.f - DELTA * DELTA)
-		{
-			Rotation = SourceAtom.Rotation * Rotation;
-		}
-
-		Translation += SourceAtom.Translation;
-		Scale3D *= (DefaultScale + SourceAtom.Scale3D);
-
-		DiagnosticCheckNaN_All();
-	}
-	/**
-	 * Set the translation and Scale3D components of this transform to a linearly interpolated combination of two other transforms
-	 *
-	 * Translation = FMath::Lerp(SourceAtom1.Translation, SourceAtom2.Translation, Alpha)
-	 * Scale3D = FMath::Lerp(SourceAtom1.Scale3D, SourceAtom2.Scale3D, Alpha)
-	 *
-	 * @param SourceAtom1 The starting point source atom (used 100% if Alpha is 0)
-	 * @param SourceAtom2 The ending point source atom (used 100% if Alpha is 1)
-	 * @param Alpha The blending weight between SourceAtom1 and SourceAtom2
-	 */
-	FORCEINLINE void LerpTranslationScale3D(const FTransform& SourceAtom1, const FTransform& SourceAtom2, ScalarRegister Alpha)
-	{
-		Translation	= FMath::Lerp(SourceAtom1.Translation, SourceAtom2.Translation, Alpha);
-		Scale3D = FMath::Lerp(SourceAtom1.Scale3D, SourceAtom2.Scale3D, Alpha);
-
-		DiagnosticCheckNaN_Translate();
-		DiagnosticCheckNaN_Scale3D();
-	}
-
-	/**
-	 * Normalize the rotation component of this transformation
-	 */
-	FORCEINLINE void NormalizeRotation()
-	{
-		Rotation.Normalize();
-		DiagnosticCheckNaN_Rotate();
-	}
-
-	/**
-	 * Checks whether the rotation component is normalized or not
-	 *
-	 * @return true if the rotation component is normalized, and false otherwise.
-	 */
-	FORCEINLINE bool IsRotationNormalized() const
-	{
-		return Rotation.IsNormalized();
-	}
-
-	/**
-	 * Blends the Identity transform with a weighted source transform and accumulates that into a destination transform
-	 *
-	 * SourceAtom = Blend(Identity, SourceAtom, BlendWeight)
-	 * FinalAtom.Rotation = SourceAtom.Rotation * FinalAtom.Rotation
-	 * FinalAtom.Translation += SourceAtom.Translation
-	 * FinalAtom.Scale3D *= SourceAtom.Scale3D
-	 *
-	 * @param FinalAtom [in/out] The atom to accumulate the blended source atom into
-	 * @param SourceAtom The target transformation (used when BlendWeight = 1); this is modified during the process
-	 * @param BlendWeight The blend weight between Identity and SourceAtom
-	 */
-	FORCEINLINE static void BlendFromIdentityAndAccumulate(FTransform& FinalAtom, FTransform& SourceAtom, float BlendWeight)
-	{
-		const  FTransform AdditiveIdentity(FQuat::Identity, FVector::ZeroVector, FVector::ZeroVector);
-		const FVector DefaultScale(1.f);
-
-		// Scale delta by weight
-		if(BlendWeight < (1.f - ZERO_ANIMWEIGHT_THRESH))
-		{
-			SourceAtom.Blend(AdditiveIdentity, SourceAtom, BlendWeight);
-		}
-
-		// Add ref pose relative animation to base animation, only if rotation is significant.
-		if(FMath::Square(SourceAtom.Rotation.W) < 1.f - DELTA * DELTA)
-		{
-			FinalAtom.Rotation = SourceAtom.Rotation * FinalAtom.Rotation;
-		}
-
-		FinalAtom.Translation += SourceAtom.Translation;
-		FinalAtom.Scale3D *= (DefaultScale + SourceAtom.Scale3D);
-
-		FinalAtom.DiagnosticCheckNaN_All();
-
-		checkSlow(FinalAtom.IsRotationNormalized());
-	}
-
-	/**
-	 * Returns the rotation component
-	 *
-	 * @return The rotation component
-	 */
-	FORCEINLINE FQuat GetRotation() const
-	{
-		DiagnosticCheckNaN_Rotate();
-		return Rotation;
-	}
-
-	/**
-	 * Returns the translation component
-	 *
-	 * @return The translation component
-	 */
-	FORCEINLINE FVector GetTranslation() const
-	{
-		DiagnosticCheckNaN_Translate();
-		return Translation;
-	}
-
-	/**
-	 * Returns the Scale3D component
-	 *
-	 * @return The Scale3D component
-	 */
-	FORCEINLINE FVector GetScale3D() const
-	{
-		DiagnosticCheckNaN_Scale3D();
-		return Scale3D;
-	}
-
-	/**
-	 * Returns an opaque copy of the rotation component
-	 * This method should be used when passing rotation from one FTransform to another
-	 *
-	 * @return The rotation component
-	 */
-	DEPRECATED(4.5, "FTransform::GetRotationV() is deprecated, use FTransform::GetRotation() instead.")
-	FORCEINLINE FQuat GetRotationV() const
-	{
-		DiagnosticCheckNaN_Rotate();
-		return Rotation;
-	}
-
-	/**
-	 * Returns an opaque copy of the translation component
-	 * This method should be used when passing translation from one FTransform to another
-	 *
-	 * @return The translation component
-	 */
-	DEPRECATED(4.5, "FTransform::GetTranslationV() is deprecated, use FTransform::GetTranslation() instead.")
-	FORCEINLINE FVector GetTranslationV() const
-	{
-		DiagnosticCheckNaN_Translate();
-		return Translation;
-	}
-
-	/**
-	 * Returns an opaque copy of the Scale3D component
-	 * This method should be used when passing Scale3D from one FTransform to another
-	 *
-	 * @return The Scale3D component
-	 */
-	DEPRECATED(4.5, "FTransform::GetScale3DV() is deprecated, use FTransform::GetScale3D() instead.")
-	FORCEINLINE FVector GetScale3DV() const
-	{
-		DiagnosticCheckNaN_Scale3D();
-		return Scale3D;
-	}
-
-	/**
-	 * Sets the Rotation and Scale3D of this transformation from another transform
-	 *
-	 * @param SrcBA The transform to copy rotation and Scale3D from
-	 */
-	FORCEINLINE void CopyRotationPart(const FTransform& SrcBA)
-	{
-		Rotation = SrcBA.Rotation;
-		Scale3D = SrcBA.Scale3D;
-
-		DiagnosticCheckNaN_Rotate();
-		DiagnosticCheckNaN_Scale3D();
-	}
-
-	/**
-	 * Sets the Translation and Scale3D of this transformation from another transform
-	 *
-	 * @param SrcBA The transform to copy translation and Scale3D from
-	 */
-	FORCEINLINE void CopyTranslationAndScale3D(const FTransform& SrcBA)
-	{
-		Translation = SrcBA.Translation;
-		Scale3D = SrcBA.Scale3D;
-
-		DiagnosticCheckNaN_Translate();
-		DiagnosticCheckNaN_Scale3D();
-	}
-
-	void SetFromMatrix(const FMatrix& InMatrix)
-	{
-		FMatrix M = InMatrix;
-
-		// Get the 3D scale from the matrix
-		Scale3D = M.ExtractScaling();
-
-		// If there is negative scaling going on, we handle that here
-		if(InMatrix.Determinant() < 0.f)
-		{
-			// Assume it is along X and modify transform accordingly. 
-			// It doesn't actually matter which axis we choose, the 'appearance' will be the same
-			Scale3D.X *= -1.f;
-			M.SetAxis(0, -M.GetScaledAxis(EAxis::X));
-		}
-
-		Rotation = FQuat(M);
-		Translation = InMatrix.GetOrigin();
-
-		// Normalize rotation
-		Rotation.Normalize();
-	}
-};
-
-
-template <> struct TIsPODType<FTransform> { enum { Value = true }; };
-
-
-/** Scale the translation part of the Transform by the supplied vector. */
-FORCEINLINE void FTransform::ScaleTranslation(const FVector& InScale3D)
-{
-	Translation *= InScale3D;
-
-	DiagnosticCheckNaN_Translate();
-}
-
-
-FORCEINLINE void FTransform::ScaleTranslation(const float& Scale)
-{
-	Translation *= Scale;
-
-	DiagnosticCheckNaN_Translate();
-}
-
-
-// this function is from matrix, and all it does is to normalize rotation portion
-FORCEINLINE void FTransform::RemoveScaling(float Tolerance/*=SMALL_NUMBER*/)
-{
-	Scale3D = FVector(1,1,1);
-	Rotation.Normalize();
-
-	DiagnosticCheckNaN_Rotate();
-	DiagnosticCheckNaN_Scale3D();
-}
-
-FORCEINLINE void FTransform::MultiplyUsingMatrixWithScale(FTransform* OutTransform, const FTransform* A, const FTransform* B)
-{
-	// the goal of using M is to get the correct orientation
-	// but for translation, we still need scale
-	FMatrix M = A->ToMatrixWithScale() * B->ToMatrixWithScale();
-	M.RemoveScaling();
-
-	// get combined scale
-	FVector Scale3D = A->Scale3D*B->Scale3D;
-
-	// apply negative scale back to axes
-	FVector SignedScale = Scale3D.GetSignVector();
-
-	M.SetAxis(0, SignedScale.X * M.GetScaledAxis(EAxis::X));
-	M.SetAxis(1, SignedScale.Y * M.GetScaledAxis(EAxis::Y));
-	M.SetAxis(2, SignedScale.Z * M.GetScaledAxis(EAxis::Z));
-
-	// @note: if you have negative with 0 scale, this will return rotation that is identity
-	// since matrix loses that axes
-	FQuat Rotation = FQuat(M);
-	Rotation.Normalize();
-
-	// set values back to output
-	OutTransform->Scale3D = Scale3D;
-	OutTransform->Rotation = Rotation;
-	
-	// technically I could calculate this using FTransform but then it does more quat multiplication 
-	// instead of using Scale in matrix multiplication
-	// it's a question of between RemoveScaling vs using FTransform to move translation
-	OutTransform->Translation = M.GetOrigin();
-}
-
-
-/** Returns Multiplied Transform of 2 FTransforms **/
-FORCEINLINE void FTransform::Multiply(FTransform* OutTransform, const FTransform* A, const FTransform* B)
-{
-	A->DiagnosticCheckNaN_All();
-	B->DiagnosticCheckNaN_All();
-
-	checkSlow(A->IsRotationNormalized());
-	checkSlow(B->IsRotationNormalized());
-
-	//	When Q = quaternion, S = single scalar scale, and T = translation
-	//	QST(A) = Q(A), S(A), T(A), and QST(B) = Q(B), S(B), T(B)
-
-	//	QST (AxB) 
-
-	// QST(A) = Q(A)*S(A)*P*-Q(A) + T(A)
-	// QST(AxB) = Q(B)*S(B)*QST(A)*-Q(B) + T(B)
-	// QST(AxB) = Q(B)*S(B)*[Q(A)*S(A)*P*-Q(A) + T(A)]*-Q(B) + T(B)
-	// QST(AxB) = Q(B)*S(B)*Q(A)*S(A)*P*-Q(A)*-Q(B) + Q(B)*S(B)*T(A)*-Q(B) + T(B)
-	// QST(AxB) = [Q(B)*Q(A)]*[S(B)*S(A)]*P*-[Q(B)*Q(A)] + Q(B)*S(B)*T(A)*-Q(B) + T(B)
-
-	//	Q(AxB) = Q(B)*Q(A)
-	//	S(AxB) = S(A)*S(B)
-	//	T(AxB) = Q(B)*S(B)*T(A)*-Q(B) + T(B)
-
-	const bool bHaveNegativeScale = A->Scale3D.GetMin() < 0 || B->Scale3D.GetMin() < 0;
-	if (bHaveNegativeScale)
-	{
-		// @note, if you have 0 scale with negative, you're going to lose rotation as it can't convert back to quat
-		MultiplyUsingMatrixWithScale(OutTransform, A, B);
-	}
-	else
-	{
-		OutTransform->Rotation = B->Rotation*A->Rotation;
-		OutTransform->Scale3D = A->Scale3D*B->Scale3D;
-		OutTransform->Translation = B->Rotation*(B->Scale3D*A->Translation) + B->Translation;
-	}
-
-	// we do not support matrix transform when non-uniform
-	// that was removed at rev 21 with UE4
-	OutTransform->DiagnosticCheckNaN_All();
-}
-/** 
- * Apply Scale to this transform
- */
-FORCEINLINE FTransform FTransform::GetScaled(float InScale) const
-{
-	FTransform A(*this);
-	A.Scale3D *= InScale;
-
-	A.DiagnosticCheckNaN_Scale3D();
-
-	return A;
-}
-
-
-/** 
- * Apply Scale to this transform
- */
-FORCEINLINE FTransform FTransform::GetScaled(FVector InScale) const
-{
-	FTransform A(*this);
-	A.Scale3D *= InScale;
-
-	A.DiagnosticCheckNaN_Scale3D();
-
-	return A;
-}
-
-
-/** Transform homogenous FVector4, ignoring the scaling part of this transform **/
-FORCEINLINE FVector4 FTransform::TransformFVector4NoScale(const FVector4& V) const
-{
-	DiagnosticCheckNaN_All();
-
-	// if not, this won't work
-	checkSlow (V.W == 0.f || V.W == 1.f);
-
-	//Transform using QST is following
-	//QST(P) = Q*S*P*-Q + T where Q = quaternion, S = scale, T = translation
-	FVector4 Transform = FVector4(Rotation.RotateVector(FVector(V)), 0.f);
-	if (V.W == 1.f)
-	{
-		Transform += FVector4(Translation, 1.f);
-	}
-
-	return Transform;
-}
-
-
-/** Transform FVector4 **/
-FORCEINLINE FVector4 FTransform::TransformFVector4(const FVector4& V) const
-{
-	DiagnosticCheckNaN_All();
-
-	// if not, this won't work
-	checkSlow (V.W == 0.f || V.W == 1.f);
-
-	//Transform using QST is following
-	//QST(P) = Q*S*P*-Q + T where Q = quaternion, S = scale, T = translation
-
-	FVector4 Transform = FVector4(Rotation.RotateVector(Scale3D*FVector(V)), 0.f);
-	if (V.W == 1.f)
-	{
-		Transform += FVector4(Translation, 1.f);
-	}
-
-	return Transform;
-}
-
-
-FORCEINLINE FVector FTransform::TransformPosition(const FVector& V) const
-{
-	DiagnosticCheckNaN_All();
-	return Rotation.RotateVector(Scale3D*V) + Translation;
-}
-
-
-FORCEINLINE FVector FTransform::TransformPositionNoScale(const FVector& V) const
-{
-	DiagnosticCheckNaN_All();
-	return Rotation.RotateVector(V) + Translation;
-}
-
-
-FORCEINLINE FVector FTransform::TransformVector(const FVector& V) const
-{
-	DiagnosticCheckNaN_All();
-	return Rotation.RotateVector(Scale3D*V);
-}
-
-
-FORCEINLINE FVector FTransform::TransformVectorNoScale(const FVector& V) const
-{
-	DiagnosticCheckNaN_All();
-	return Rotation.RotateVector(V);
-}
-
-
-// do backward operation when inverse, translation -> rotation -> scale
-FORCEINLINE FVector FTransform::InverseTransformPosition(const FVector &V) const
-{
-	DiagnosticCheckNaN_All();
-	return (Rotation.UnrotateVector(V-Translation)) * GetSafeScaleReciprocal(Scale3D);
-}
-
-
-// do backward operation when inverse, translation -> rotation
-FORCEINLINE FVector FTransform::InverseTransformPositionNoScale(const FVector &V) const
-{
-	DiagnosticCheckNaN_All();
-	return (Rotation.UnrotateVector(V-Translation));
-}
-
-
-// do backward operation when inverse, translation -> rotation -> scale
-FORCEINLINE FVector FTransform::InverseTransformVector(const FVector &V) const
-{
-	DiagnosticCheckNaN_All();
-	return (Rotation.UnrotateVector(V)) * GetSafeScaleReciprocal(Scale3D);
-}
-
-
-// do backward operation when inverse, translation -> rotation
-FORCEINLINE FVector FTransform::InverseTransformVectorNoScale(const FVector &V) const
-{
-	DiagnosticCheckNaN_All();
-	return (Rotation.UnrotateVector(V));
-}
-
-
-FORCEINLINE FTransform FTransform::operator*(const FTransform& Other) const
-{
-	FTransform Output;
-	Multiply(&Output, this, &Other);
-	return Output;
-}
-
-
-FORCEINLINE void FTransform::operator*=(const FTransform& Other)
-{
-	Multiply(this, this, &Other);
-}
-
-
-FORCEINLINE FTransform FTransform::operator*(const FQuat& Other) const
-{
-	FTransform Output, OtherTransform(Other, FVector::ZeroVector, FVector(1.f));
-	Multiply(&Output, this, &OtherTransform);
-	return Output;
-}
-
-
-FORCEINLINE void FTransform::operator*=(const FQuat& Other)
-{
-	FTransform OtherTransform(Other, FVector::ZeroVector, FVector(1.f));
-	Multiply(this, this, &OtherTransform);
-}
-
-
-// x = 0, y = 1, z = 2
-FORCEINLINE FVector FTransform::GetScaledAxis(EAxis::Type InAxis) const
-{
-	if (InAxis == EAxis::X)
-	{
-		return TransformVector(FVector(1.f, 0.f, 0.f));
-	}
-	else if (InAxis == EAxis::Y)
-	{
-		return TransformVector(FVector(0.f, 1.f, 0.f));
-	}
-
-	return TransformVector(FVector(0.f, 0.f, 1.f));
-}
-
-
-// x = 0, y = 1, z = 2
-FORCEINLINE FVector FTransform::GetUnitAxis(EAxis::Type InAxis) const
-{
-	if (InAxis == EAxis::X)
-	{
-		return TransformVectorNoScale(FVector(1.f, 0.f, 0.f));
-	}
-	else if (InAxis == EAxis::Y)
-	{
-		return TransformVectorNoScale(FVector(0.f, 1.f, 0.f));
-	}
-
-	return TransformVectorNoScale(FVector(0.f, 0.f, 1.f));
-}
-
-
-FORCEINLINE void FTransform::Mirror(EAxis::Type MirrorAxis, EAxis::Type FlipAxis)
-{
-	// We do convert to Matrix for mirroring. 
-	FMatrix M = ToMatrixWithScale();
-	M.Mirror(MirrorAxis, FlipAxis);
-	SetFromMatrix(M);
-}
-
-
-/** same version of FMatrix::GetMaximumAxisScale function **/
-/** @return the maximum magnitude of all components of the 3D scale. */
-inline float FTransform::GetMaximumAxisScale() const
-{
-	DiagnosticCheckNaN_Scale3D();
-	return Scale3D.GetAbsMax();
-}
-
-
-/** @return the minimum magnitude of all components of the 3D scale. */
-inline float FTransform::GetMinimumAxisScale() const
-{
-	DiagnosticCheckNaN_Scale3D();
-	return Scale3D.GetAbsMin();
-}
-
-
-// mathematically if you have 0 scale, it should be infinite, 
-// however, in practice if you have 0 scale, and relative transform doesn't make much sense 
-// anymore because you should be instead of showing gigantic infinite mesh
-// also returning BIG_NUMBER causes sequential NaN issues by multiplying 
-// so we hardcode as 0
-FORCEINLINE FVector FTransform::GetSafeScaleReciprocal(const FVector& InScale, float Tolerance) const
-{
-	FVector SafeReciprocalScale;
-	if (FMath::Abs(InScale.X) <= Tolerance)
-	{
-		SafeReciprocalScale.X = 0.f;
-	}
-	else
-	{
-		SafeReciprocalScale.X = 1 / InScale.X;
-	}
-
-	if (FMath::Abs(InScale.Y) <= Tolerance)
-	{
-		SafeReciprocalScale.Y = 0.f;
-	}
-	else
-	{
-		SafeReciprocalScale.Y = 1 / InScale.Y;
-	}
-
-	if (FMath::Abs(InScale.Z) <= Tolerance)
-	{
-		SafeReciprocalScale.Z = 0.f;
-	}
-	else
-	{
-		SafeReciprocalScale.Z = 1 / InScale.Z;
-	}
-
-	return SafeReciprocalScale;
-}
