Copying //UE4/Dev-Build to //UE4/Dev-Main (Source: //UE4/Dev-Build @ 3104200)

==========================
MAJOR FEATURES + CHANGES
==========================

Change 3096076 on 2016/08/22 by [REDACTED_PERSON]

	Fixes for building with -nosharedPCH
	Added nodes for building with -nosharedPCH and -noPCH but not using them by default as they take so much longer than NonUnity and there are still errors with -noPCH

Change 3097064 on 2016/08/22 by [REDACTED_PERSON]

	Core: Refactoring OutputDevice.h, OutputDevices.h, FeedbackDevice.h, CallbackDevice.h into one-class per file, and files named after their classes.

Change 3097102 on 2016/08/22 by [REDACTED_PERSON]

	Core: Move log macros and related functionality into the Logging folder.

Change 3097240 on 2016/08/22 by [REDACTED_PERSON]

	Core: Move FCommandLine into its own header.

Change 3097258 on 2016/08/22 by [REDACTED_PERSON]

	Core: Move FFileHelper into its own header.

Change 3097279 on 2016/08/22 by [REDACTED_PERSON]

	Core: Split FArchive derived classes into separate headers per class.

Change 3097742 on 2016/08/23 by [REDACTED_PERSON]

	UnrealGameSync: Add the -r5 parameter to P4 sync command lines, to retry syncs on poor connections before failing.

Change 3099283 on 2016/08/24 by [REDACTED_PERSON]

	Core: Move math functions which are dependent on multiple class definitions into individual class headers. This forces an include order, but allows us to include individual math classes separately. Would ideally split out the FMath class into a Math namespace and declare/implement things like intersection tests and interpolation routines in separate headers.

Change 3099417 on 2016/08/24 by [REDACTED_PERSON]

	CoreUObject: Separating out UObject archive classes into separate headers.

Change 3099544 on 2016/08/24 by [REDACTED_PERSON]

	CoreUObject: Rename UTextProperty.h/cpp to TextProperty.h/cpp

Change 3099678 on 2016/08/24 by [REDACTED_PERSON]

	CoreUObject: Rename ObjectBase.h to ObjectMacros.h to avoid ambiguiity with UObjectBase.h, and move Interface.h/cpp into the UObject folder.

Change 3099951 on 2016/08/24 by [REDACTED_PERSON]

	Core: Split out ThreadingBase into class-per-file headers.

Change 3102385 on 2016/08/26 by [REDACTED_PERSON]

	First round of fixes from attempting to compile Editor in Monolithic mode
	Changed multiple DetermineOptimalSpacing free functions to be members of their respective classes
	Changed duplicate FScopedTempFile classes to be named for their respective source control types (Git & SVN)
	Export Declaration of LogUnrealAudioDevice extern instead of defining it in multiple modules
	Changed duplicate FDiffResultItem struct to FTreeDiffResultItem
	Added FORCEINLINE to FFeaturedClasses functions in .inl file

Change 3102728 on 2016/08/26 by [REDACTED_PERSON]

	Core: Move threading classes from ThreadingBase.h into separate headers, and globals into CoreGlobals.h.

Change 3102850 on 2016/08/26 by [REDACTED_PERSON]

	Core: Move FCriticalSection into its own file (CriticalSection.h and platform-specific versions).

Change 3102877 on 2016/08/26 by [REDACTED_PERSON]

	Core: Move definitions of GLog, GError, GWarn into CoreGlobals.h

Change 3102898 on 2016/08/26 by [REDACTED_PERSON]

	Core: Move supporting assert functionality (the FDebug class) into AssertionMacros.h to remove circular header dependencies.

Change 3102956 on 2016/08/26 by [REDACTED_PERSON]

	Core: Move FExec into its own header (Misc/Exec.h)

Change 3102961 on 2016/08/26 by [REDACTED_PERSON]

	Core: Fix matrix functions referencing the forceinline'd FMatrix constructor before it's declared.

Change 3102982 on 2016/08/26 by [REDACTED_PERSON]

	Core: Move FMessageDialog into its own header (Misc/MessageDialog.h)

Change 3103008 on 2016/08/26 by [REDACTED_PERSON]

	Core: Move definition of LowLevelFatalError into AssertionMacros.h.

Change 3103040 on 2016/08/26 by [REDACTED_PERSON]

	Core: Move FError::Throwf (only enabled for UHT) into its own header in CoreUObject.

Change 3103057 on 2016/08/26 by [REDACTED_PERSON]

	Core: Move FMsg function definitions into the LogMacros.h header, where it's primarily used from.

Change 3103102 on 2016/08/26 by [REDACTED_PERSON]

	Core: Make the LogMacros.h header standalone. Untangles a few header dependencies necessary to include log macros early on. Global Core log channels are now defined in CoreGlobals.h

Change 3103160 on 2016/08/26 by [REDACTED_PERSON]

	CoreUObject: Separate Linker.h into LinkerLoad.h, LinkerSave.h, PackageFileSummary.h, ObjectResource.h (matching existing cpp files).

Change 3103265 on 2016/08/26 by [REDACTED_PERSON]

	Core: Split FScriptArray, TMRUArray, TTransArray, TIndirectArray into separate headers.

Change 3103642 on 2016/08/26 by [REDACTED_PERSON]

	Core: Remove reference to FMath function from allocator, to avoid unnecessarily bringing in math library.

diff --git a/Engine/Plugins/Developer/GitSourceControl/Source/GitSourceControl/Private/GitSourceControlOperations.cpp b/Engine/Plugins/Developer/GitSourceControl/Source/GitSourceControl/Private/GitSourceControlOperations.cpp
--- a/Engine/Plugins/Developer/GitSourceControl/Source/GitSourceControl/Private/GitSourceControlOperations.cpp
+++ b/Engine/Plugins/Developer/GitSourceControl/Source/GitSourceControl/Private/GitSourceControlOperations.cpp
@@ -55,7 +55,7 @@ bool FGitCheckInWorker::Execute(FGitSourceControlCommand& InCommand)
 	TSharedRef<FCheckIn, ESPMode::ThreadSafe> Operation = StaticCastSharedRef<FCheckIn>(InCommand.Operation);
 
 	// make a temp file to place our commit message in
-	FScopedTempFile CommitMsgFile(Operation->GetDescription());
+	FGitScopedTempFile CommitMsgFile(Operation->GetDescription());
 	if(CommitMsgFile.GetFilename().Len() > 0)
 	{
 		TArray<FString> Parameters;
diff --git a/Engine/Plugins/Developer/GitSourceControl/Source/GitSourceControl/Private/GitSourceControlUtils.cpp b/Engine/Plugins/Developer/GitSourceControl/Source/GitSourceControl/Private/GitSourceControlUtils.cpp
--- a/Engine/Plugins/Developer/GitSourceControl/Source/GitSourceControl/Private/GitSourceControlUtils.cpp
+++ b/Engine/Plugins/Developer/GitSourceControl/Source/GitSourceControl/Private/GitSourceControlUtils.cpp
@@ -17,7 +17,7 @@ namespace GitSourceControlConstants
 	const int32 MaxFilesPerBatch = 50;
 }
 
-FScopedTempFile::FScopedTempFile(const FText& InText)
+FGitScopedTempFile::FGitScopedTempFile(const FText& InText)
 {
 	Filename = FPaths::CreateTempFilename(*FPaths::GameLogDir(), TEXT("Git-Temp"), TEXT(".txt"));
 	if(!FFileHelper::SaveStringToFile(InText.ToString(), *Filename, FFileHelper::EEncodingOptions::ForceUTF8WithoutBOM))
@@ -26,7 +26,7 @@ FScopedTempFile::FScopedTempFile(const FText& InText)
 	}
 }
 
-FScopedTempFile::~FScopedTempFile()
+FGitScopedTempFile::~FGitScopedTempFile()
 {
 	if(FPaths::FileExists(Filename))
 	{
@@ -37,7 +37,7 @@ FScopedTempFile::~FScopedTempFile()
 	}
 }
 
-const FString& FScopedTempFile::GetFilename() const
+const FString& FGitScopedTempFile::GetFilename() const
 {
 	return Filename;
 }
diff --git a/Engine/Plugins/Developer/GitSourceControl/Source/GitSourceControl/Private/GitSourceControlUtils.h b/Engine/Plugins/Developer/GitSourceControl/Source/GitSourceControl/Private/GitSourceControlUtils.h
--- a/Engine/Plugins/Developer/GitSourceControl/Source/GitSourceControl/Private/GitSourceControlUtils.h
+++ b/Engine/Plugins/Developer/GitSourceControl/Source/GitSourceControl/Private/GitSourceControlUtils.h
@@ -10,15 +10,15 @@ class FGitSourceControlCommand;
 /**
  * Helper struct for maintaining temporary files for passing to commands
  */
-class FScopedTempFile
+class FGitScopedTempFile
 {
 public:
 
 	/** Constructor - open & write string to temp file */
-	FScopedTempFile(const FText& InText);
+	FGitScopedTempFile(const FText& InText);
 
 	/** Destructor - delete temp file */
-	~FScopedTempFile();
+	~FGitScopedTempFile();
 
 	/** Get the filename of this temp file - empty if it failed to be created */
 	const FString& GetFilename() const;
diff --git a/Engine/Plugins/Experimental/AlembicImporter/Source/AlembicImporter/Private/AlembicImporterPrivatePCH.h b/Engine/Plugins/Experimental/AlembicImporter/Source/AlembicImporter/Private/AlembicImporterPrivatePCH.h
--- a/Engine/Plugins/Experimental/AlembicImporter/Source/AlembicImporter/Private/AlembicImporterPrivatePCH.h
+++ b/Engine/Plugins/Experimental/AlembicImporter/Source/AlembicImporter/Private/AlembicImporterPrivatePCH.h
@@ -1 +1,3 @@
 // Copyright 1998-2016 Epic Games, Inc. All Rights Reserved.
+
+#include "UnrealEd.h"
diff --git a/Engine/Plugins/Experimental/AlembicImporter/Source/AlembicLibrary/Public/AbcImportData.h b/Engine/Plugins/Experimental/AlembicImporter/Source/AlembicLibrary/Public/AbcImportData.h
--- a/Engine/Plugins/Experimental/AlembicImporter/Source/AlembicLibrary/Public/AbcImportData.h
+++ b/Engine/Plugins/Experimental/AlembicImporter/Source/AlembicLibrary/Public/AbcImportData.h
@@ -14,6 +14,7 @@
 
 class FAbcImporter;
 class UAbcImportSettings;
+class UMaterial;
 
 /** TODO
 - Refactor MatrixSamples, TimeSamples into one structure
diff --git a/Engine/Plugins/Experimental/AlembicImporter/Source/AlembicLibrary/Public/AbcImporter.h b/Engine/Plugins/Experimental/AlembicImporter/Source/AlembicLibrary/Public/AbcImporter.h
--- a/Engine/Plugins/Experimental/AlembicImporter/Source/AlembicLibrary/Public/AbcImporter.h
+++ b/Engine/Plugins/Experimental/AlembicImporter/Source/AlembicLibrary/Public/AbcImporter.h
@@ -5,6 +5,7 @@
 #include <Alembic/AbcGeom/All.h>
 
 class UStaticMesh;
+class USkeletalMesh;
 class UGeometryCache;
 class UGeometryCacheTrack_FlipbookAnimation;
 class UGeometryCacheTrack_TransformAnimation;
@@ -18,6 +19,7 @@ struct FAbcTransformObject;
 struct FAbcMeshSample;
 struct FAbcCompressionSettings;
 struct FCompressedAbcData;
+struct FRawMesh;
 
 enum EAbcImportError : uint32
 {
diff --git a/Engine/Plugins/Experimental/AlembicImporter/Source/AlembicLibrary/Public/AlembicLibraryPublicPCH.h b/Engine/Plugins/Experimental/AlembicImporter/Source/AlembicLibrary/Public/AlembicLibraryPublicPCH.h
--- a/Engine/Plugins/Experimental/AlembicImporter/Source/AlembicLibrary/Public/AlembicLibraryPublicPCH.h
+++ b/Engine/Plugins/Experimental/AlembicImporter/Source/AlembicLibrary/Public/AlembicLibraryPublicPCH.h
@@ -1,3 +1,7 @@
 // Copyright 1998-2016 Epic Games, Inc. All Rights Reserved.
 
 #include "Core.h"
+#include "CoreUObject.h"
+#include "SecureHash.h"
+#include "UnrealEd.h"
+#include "SlateBasics.h"
diff --git a/Engine/Plugins/Runtime/OculusAudio/Source/OculusAudio/Private/OculusAudio.h b/Engine/Plugins/Runtime/OculusAudio/Source/OculusAudio/Private/OculusAudio.h
--- a/Engine/Plugins/Runtime/OculusAudio/Source/OculusAudio/Private/OculusAudio.h
+++ b/Engine/Plugins/Runtime/OculusAudio/Source/OculusAudio/Private/OculusAudio.h
@@ -1,4 +1,7 @@
 // Copyright 1998-2016 Epic Games, Inc. All Rights Reserved.
 
 #pragma once
+
+#include "CoreUObject.h"
+#include "Engine.h"
 #include "IOculusAudioPlugin.h"
\ No newline at end of file
diff --git a/Engine/Plugins/Runtime/PacketHandlers/CompressionComponents/Oodle/Source/OodleHandlerComponent/Public/OodleArchives.h b/Engine/Plugins/Runtime/PacketHandlers/CompressionComponents/Oodle/Source/OodleHandlerComponent/Public/OodleArchives.h
--- a/Engine/Plugins/Runtime/PacketHandlers/CompressionComponents/Oodle/Source/OodleHandlerComponent/Public/OodleArchives.h
+++ b/Engine/Plugins/Runtime/PacketHandlers/CompressionComponents/Oodle/Source/OodleHandlerComponent/Public/OodleArchives.h
@@ -3,7 +3,7 @@
 #pragma once
 
 #if HAS_OODLE_SDK
-#include "ArchiveBase.h"
+#include "ArchiveProxy.h"
 
 #define CAPTURE_HEADER_MAGIC		0x41091CC4
 #define CAPTURE_FILE_VERSION		0x00000002
diff --git a/Engine/Source/Developer/LogVisualizer/Private/VisualLoggerDatabase.h b/Engine/Source/Developer/LogVisualizer/Private/VisualLoggerDatabase.h
--- a/Engine/Source/Developer/LogVisualizer/Private/VisualLoggerDatabase.h
+++ b/Engine/Source/Developer/LogVisualizer/Private/VisualLoggerDatabase.h
@@ -1,6 +1,6 @@
 // Copyright 1998-2016 Epic Games, Inc. All Rights Reserved.
 #pragma once
-#include "ObjectBase.h"
+#include "ObjectMacros.h"
 #include "SDockTab.h"
 #include "VisualLogger/VisualLogger.h"
 
diff --git a/Engine/Source/Developer/LogVisualizer/Private/VisualLoggerTimeSliderController.cpp b/Engine/Source/Developer/LogVisualizer/Private/VisualLoggerTimeSliderController.cpp
--- a/Engine/Source/Developer/LogVisualizer/Private/VisualLoggerTimeSliderController.cpp
+++ b/Engine/Source/Developer/LogVisualizer/Private/VisualLoggerTimeSliderController.cpp
@@ -68,16 +68,18 @@ static float GetNextSpacing( uint32 CurrentStep )
 	}
 }
 
-/**
- * Determines the optimal spacing between tick marks in the slider for a given pixel density
- * Increments until a minimum amount of slate units specified by MinTick is reached
- * 
- * @param InPixelsPerInput	The density of pixels between each input
- * @param MinTick			The minimum slate units per tick allowed
- * @param MinTickSpacing	The minimum tick spacing in time units allowed
- * @return the optimal spacing in time units
- */
-float DetermineOptimalSpacing( float InPixelsPerInput, uint32 MinTick, float MinTickSpacing )
+FVisualLoggerTimeSliderController::FVisualLoggerTimeSliderController(const FVisualLoggerTimeSliderArgs& InArgs)
+	: TimeSliderArgs( InArgs )
+	, DistanceDragged( 0.0f )
+	, bDraggingScrubber( false )
+	, bPanning( false )
+{
+	ScrubHandleUp = FEditorStyle::GetBrush( TEXT( "Sequencer.Timeline.ScrubHandleUp" ) ); 
+	ScrubHandleDown = FEditorStyle::GetBrush( TEXT( "Sequencer.Timeline.ScrubHandleDown" ) );
+	CursorBackground = FEditorStyle::GetBrush("Sequencer.SectionArea.Background");
+}
+
+float FVisualLoggerTimeSliderController::DetermineOptimalSpacing(float InPixelsPerInput, uint32 MinTick, float MinTickSpacing) const
 {
 	uint32 CurStep = 0;
 
@@ -93,18 +95,6 @@ float DetermineOptimalSpacing( float InPixelsPerInput, uint32 MinTick, float Min
 	return Spacing;
 }
 
-
-FVisualLoggerTimeSliderController::FVisualLoggerTimeSliderController(const FVisualLoggerTimeSliderArgs& InArgs)
-	: TimeSliderArgs( InArgs )
-	, DistanceDragged( 0.0f )
-	, bDraggingScrubber( false )
-	, bPanning( false )
-{
-	ScrubHandleUp = FEditorStyle::GetBrush( TEXT( "Sequencer.Timeline.ScrubHandleUp" ) ); 
-	ScrubHandleDown = FEditorStyle::GetBrush( TEXT( "Sequencer.Timeline.ScrubHandleDown" ) );
-	CursorBackground = FEditorStyle::GetBrush("Sequencer.SectionArea.Background");
-}
-
 void FVisualLoggerTimeSliderController::SetTimesliderArgs(const FVisualLoggerTimeSliderArgs& InArgs)
 {
 	TimeSliderArgs = InArgs;
diff --git a/Engine/Source/Developer/LogVisualizer/Private/VisualLoggerTimeSliderController.h b/Engine/Source/Developer/LogVisualizer/Private/VisualLoggerTimeSliderController.h
--- a/Engine/Source/Developer/LogVisualizer/Private/VisualLoggerTimeSliderController.h
+++ b/Engine/Source/Developer/LogVisualizer/Private/VisualLoggerTimeSliderController.h
@@ -24,6 +24,17 @@ class FVisualLoggerTimeSliderController : public ITimeSliderController
 {
 public:
 	FVisualLoggerTimeSliderController(const FVisualLoggerTimeSliderArgs& InArgs);
+
+	/**
+	* Determines the optimal spacing between tick marks in the slider for a given pixel density
+	* Increments until a minimum amount of slate units specified by MinTick is reached
+	*
+	* @param InPixelsPerInput	The density of pixels between each input
+	* @param MinTick			The minimum slate units per tick allowed
+	* @param MinTickSpacing	The minimum tick spacing in time units allowed
+	* @return the optimal spacing in time units
+	*/
+	float DetermineOptimalSpacing(float InPixelsPerInput, uint32 MinTick, float MinTickSpacing) const;
 	void SetTimesliderArgs(const FVisualLoggerTimeSliderArgs& InArgs);
 
 	/** ITimeSliderController Interface */
diff --git a/Engine/Source/Developer/MaterialUtilities/Public/MaterialUtilities.h b/Engine/Source/Developer/MaterialUtilities/Public/MaterialUtilities.h
--- a/Engine/Source/Developer/MaterialUtilities/Public/MaterialUtilities.h
+++ b/Engine/Source/Developer/MaterialUtilities/Public/MaterialUtilities.h
@@ -3,7 +3,7 @@
 #pragma once
 
 #include "ModuleInterface.h"
-#include "Runtime/CoreUObject/Public/UObject/ObjectBase.h"
+#include "Runtime/CoreUObject/Public/UObject/ObjectMacros.h"
 #include "Runtime/Engine/Classes/Engine/EngineTypes.h"
 #include "Runtime/Engine/Classes/Engine/Texture.h"
 #include "Runtime/Engine/Public/SceneTypes.h"
diff --git a/Engine/Source/Editor/BehaviorTreeEditor/Private/SBehaviorTreeDiff.cpp b/Engine/Source/Editor/BehaviorTreeEditor/Private/SBehaviorTreeDiff.cpp
--- a/Engine/Source/Editor/BehaviorTreeEditor/Private/SBehaviorTreeDiff.cpp
+++ b/Engine/Source/Editor/BehaviorTreeEditor/Private/SBehaviorTreeDiff.cpp
@@ -13,15 +13,15 @@
 #define LOCTEXT_NAMESPACE "SBehaviorTreeDiff"
 
 //////////////////////////////////////////////////////////////////////////
-// FDiffResultItem
+// FTreeDiffResultItem
 
-struct FDiffResultItem: public TSharedFromThis<FDiffResultItem>
+struct FTreeDiffResultItem : public TSharedFromThis<FTreeDiffResultItem>
 {
 	/**
 	 * Constructor
 	 * @param InResult A difference result 
 	 */
-	FDiffResultItem(const FDiffSingleResult& InResult): Result(InResult){}
+	FTreeDiffResultItem(const FDiffSingleResult& InResult): Result(InResult){}
 
 	/**
 	 * GenerateWidget for the diff item
@@ -257,7 +257,7 @@ void SBehaviorTreeDiff::BuildDiffSourceArray()
 	DiffListSource.Empty();
 	for (auto DiffIt(FoundDiffs.CreateConstIterator()); DiffIt; ++DiffIt)
 	{
-		DiffListSource.Add(FSharedDiffOnGraph(new FDiffResultItem(*DiffIt)));
+		DiffListSource.Add(FSharedDiffOnGraph(new FTreeDiffResultItem(*DiffIt)));
 	}
 }
 
diff --git a/Engine/Source/Editor/BehaviorTreeEditor/Private/SBehaviorTreeDiff.h b/Engine/Source/Editor/BehaviorTreeEditor/Private/SBehaviorTreeDiff.h
--- a/Engine/Source/Editor/BehaviorTreeEditor/Private/SBehaviorTreeDiff.h
+++ b/Engine/Source/Editor/BehaviorTreeEditor/Private/SBehaviorTreeDiff.h
@@ -93,7 +93,7 @@ private:
 	};
 
 	// Type def for our Table item
-	typedef TSharedPtr<struct FDiffResultItem> FSharedDiffOnGraph;
+	typedef TSharedPtr<struct FTreeDiffResultItem> FSharedDiffOnGraph;
 
 	// Type def of our Table Type
 	typedef SListView<FSharedDiffOnGraph > SListViewType;
diff --git a/Engine/Source/Editor/GameProjectGeneration/Public/FeaturedClasses.inl b/Engine/Source/Editor/GameProjectGeneration/Public/FeaturedClasses.inl
--- a/Engine/Source/Editor/GameProjectGeneration/Public/FeaturedClasses.inl
+++ b/Engine/Source/Editor/GameProjectGeneration/Public/FeaturedClasses.inl
@@ -22,7 +22,7 @@
 
 
 /** Get a list of all featured native class types */
-TArray<FNewClassInfo> FFeaturedClasses::AllNativeClasses()
+FORCEINLINE TArray<FNewClassInfo> FFeaturedClasses::AllNativeClasses()
 {
 	TArray<FNewClassInfo> Array;
 	Array.Add(FNewClassInfo(FNewClassInfo::EClassType::EmptyCpp));
@@ -42,7 +42,7 @@ TArray<FNewClassInfo> FFeaturedClasses::AllNativeClasses()
 }
 
 /** Get a list of all featured Actor class types */
-TArray<FNewClassInfo> FFeaturedClasses::ActorClasses()
+FORCEINLINE TArray<FNewClassInfo> FFeaturedClasses::ActorClasses()
 {
 	TArray<FNewClassInfo> Array;
 	AddCommonActorClasses(Array);
@@ -51,7 +51,7 @@ TArray<FNewClassInfo> FFeaturedClasses::ActorClasses()
 }
 
 /** Get a list of all featured Component class types */
-TArray<FNewClassInfo> FFeaturedClasses::ComponentClasses()
+FORCEINLINE TArray<FNewClassInfo> FFeaturedClasses::ComponentClasses()
 {
 	TArray<FNewClassInfo> Array;
 	AddCommonComponentClasses(Array);
@@ -60,7 +60,7 @@ TArray<FNewClassInfo> FFeaturedClasses::ComponentClasses()
 }
 
 /** Append the featured Actor class types that are commonly used */
-void FFeaturedClasses::AddCommonActorClasses(TArray<FNewClassInfo>& Array)
+FORCEINLINE void FFeaturedClasses::AddCommonActorClasses(TArray<FNewClassInfo>& Array)
 {
 	Array.Add(FNewClassInfo(ACharacter::StaticClass()));
 	Array.Add(FNewClassInfo(APawn::StaticClass()));
@@ -68,7 +68,7 @@ void FFeaturedClasses::AddCommonActorClasses(TArray<FNewClassInfo>& Array)
 }
 
 /** Append the featured Actor class types that are less commonly used */
-void FFeaturedClasses::AddExtraActorClasses(TArray<FNewClassInfo>& Array)
+FORCEINLINE void FFeaturedClasses::AddExtraActorClasses(TArray<FNewClassInfo>& Array)
 {
 	Array.Add(FNewClassInfo(APlayerCameraManager::StaticClass()));
 	Array.Add(FNewClassInfo(APlayerController::StaticClass()));
@@ -80,14 +80,14 @@ void FFeaturedClasses::AddExtraActorClasses(TArray<FNewClassInfo>& Array)
 }
 
 /** Append the featured Component class types that are commonly used */
-void FFeaturedClasses::AddCommonComponentClasses(TArray<FNewClassInfo>& Array)
+FORCEINLINE void FFeaturedClasses::AddCommonComponentClasses(TArray<FNewClassInfo>& Array)
 {
 	Array.Add(FNewClassInfo(UActorComponent::StaticClass()));
 	Array.Add(FNewClassInfo(USceneComponent::StaticClass()));
 }
 
 /** Append the featured Component class types that are less commonly used */
-void FFeaturedClasses::AddExtraComponentClasses(TArray<FNewClassInfo>& Array)
+FORCEINLINE void FFeaturedClasses::AddExtraComponentClasses(TArray<FNewClassInfo>& Array)
 {
 	Array.Add(FNewClassInfo(UStaticMeshComponent::StaticClass()));
 	Array.Add(FNewClassInfo(USkeletalMeshComponent::StaticClass()));
diff --git a/Engine/Source/Editor/Sequencer/Private/SequencerTimeSliderController.cpp b/Engine/Source/Editor/Sequencer/Private/SequencerTimeSliderController.cpp
--- a/Engine/Source/Editor/Sequencer/Private/SequencerTimeSliderController.cpp
+++ b/Engine/Source/Editor/Sequencer/Private/SequencerTimeSliderController.cpp
@@ -69,16 +69,19 @@ static float GetNextSpacing( uint32 CurrentStep )
 	}
 }
 
-/**
- * Determines the optimal spacing between tick marks in the slider for a given pixel density
- * Increments until a minimum amount of slate units specified by MinTick is reached
- * 
- * @param InPixelsPerInput	The density of pixels between each input
- * @param MinTick			The minimum slate units per tick allowed
- * @param MinTickSpacing	The minimum tick spacing in time units allowed
- * @return the optimal spacing in time units
- */
-float DetermineOptimalSpacing( float InPixelsPerInput, uint32 MinTick, float MinTickSpacing )
+FSequencerTimeSliderController::FSequencerTimeSliderController( const FTimeSliderArgs& InArgs )
+	: TimeSliderArgs( InArgs )
+	, DistanceDragged( 0.0f )
+	, MouseDragType( DRAG_NONE )
+	, bPanning( false )
+{
+	ScrubHandleUp = FEditorStyle::GetBrush( TEXT( "Sequencer.Timeline.ScrubHandleUp" ) ); 
+	ScrubHandleDown = FEditorStyle::GetBrush( TEXT( "Sequencer.Timeline.ScrubHandleDown" ) );
+	ScrubHandleSize = 13.f;
+	ContextMenuSuppression = 0;
+}
+
+float FSequencerTimeSliderController::DetermineOptimalSpacing(float InPixelsPerInput, uint32 MinTick, float MinTickSpacing) const
 {
 	uint32 CurStep = 0;
 
@@ -97,19 +100,6 @@ float DetermineOptimalSpacing( float InPixelsPerInput, uint32 MinTick, float Min
 	return Spacing;
 }
 
-
-FSequencerTimeSliderController::FSequencerTimeSliderController( const FTimeSliderArgs& InArgs )
-	: TimeSliderArgs( InArgs )
-	, DistanceDragged( 0.0f )
-	, MouseDragType( DRAG_NONE )
-	, bPanning( false )
-{
-	ScrubHandleUp = FEditorStyle::GetBrush( TEXT( "Sequencer.Timeline.ScrubHandleUp" ) ); 
-	ScrubHandleDown = FEditorStyle::GetBrush( TEXT( "Sequencer.Timeline.ScrubHandleDown" ) );
-	ScrubHandleSize = 13.f;
-	ContextMenuSuppression = 0;
-}
-
 struct FDrawTickArgs
 {
 	/** Geometry of the area */
diff --git a/Engine/Source/Editor/Sequencer/Private/SequencerTimeSliderController.h b/Engine/Source/Editor/Sequencer/Private/SequencerTimeSliderController.h
--- a/Engine/Source/Editor/Sequencer/Private/SequencerTimeSliderController.h
+++ b/Engine/Source/Editor/Sequencer/Private/SequencerTimeSliderController.h
@@ -46,6 +46,17 @@ class FSequencerTimeSliderController : public ITimeSliderController
 public:
 	FSequencerTimeSliderController( const FTimeSliderArgs& InArgs );
 
+	/**
+	* Determines the optimal spacing between tick marks in the slider for a given pixel density
+	* Increments until a minimum amount of slate units specified by MinTick is reached
+	*
+	* @param InPixelsPerInput	The density of pixels between each input
+	* @param MinTick			The minimum slate units per tick allowed
+	* @param MinTickSpacing	The minimum tick spacing in time units allowed
+	* @return the optimal spacing in time units
+	*/
+	float DetermineOptimalSpacing(float InPixelsPerInput, uint32 MinTick, float MinTickSpacing) const;
+
 	/** ITimeSliderController Interface */
 	virtual int32 OnPaintTimeSlider( bool bMirrorLabels, const FGeometry& AllottedGeometry, const FSlateRect& MyClippingRect, FSlateWindowElementList& OutDrawElements, int32 LayerId, const FWidgetStyle& InWidgetStyle, bool bParentEnabled ) const override;
 	virtual FReply OnMouseButtonDown( SWidget& WidgetOwner, const FGeometry& MyGeometry, const FPointerEvent& MouseEvent ) override;
diff --git a/Engine/Source/Editor/UMGEditor/Private/Designer/SRuler.cpp b/Engine/Source/Editor/UMGEditor/Private/Designer/SRuler.cpp
--- a/Engine/Source/Editor/UMGEditor/Private/Designer/SRuler.cpp
+++ b/Engine/Source/Editor/UMGEditor/Private/Designer/SRuler.cpp
@@ -92,16 +92,19 @@ static float GetNextSpacing( uint32 CurrentStep )
 	}
 }
 
-/**
- * Determines the optimal spacing between tick marks in the slider for a given pixel density
- * Increments until a minimum amount of slate units specified by MinTick is reached
- * 
- * @param InPixelsPerInput	The density of pixels between each input
- * @param MinTick			The minimum slate units per tick allowed
- * @param MinTickSpacing	The minimum tick spacing in time units allowed
- * @return the optimal spacing in time units
- */
-float DetermineOptimalSpacing( float InPixelsPerInput, uint32 MinTick, float MinTickSpacing )
+/////////////////////////////////////////////////////
+// SRuler
+
+void SRuler::Construct(const SRuler::FArguments& InArgs)
+{
+	Orientation = InArgs._Orientation;
+	AbsoluteOrigin = FVector2D(0, 0);
+	SlateToUnitScale = 1;
+
+	MouseButtonDownHandler = InArgs._OnMouseButtonDown;
+}
+
+float SRuler::DetermineOptimalSpacing(float InPixelsPerInput, uint32 MinTick, float MinTickSpacing) const
 {
 	if (InPixelsPerInput == 0.0f)
 		return MinTickSpacing;
@@ -120,18 +123,6 @@ float DetermineOptimalSpacing( float InPixelsPerInput, uint32 MinTick, float Min
 	return Spacing;
 }
 
-/////////////////////////////////////////////////////
-// SRuler
-
-void SRuler::Construct(const SRuler::FArguments& InArgs)
-{
-	Orientation = InArgs._Orientation;
-	AbsoluteOrigin = FVector2D(0, 0);
-	SlateToUnitScale = 1;
-
-	MouseButtonDownHandler = InArgs._OnMouseButtonDown;
-}
-
 void SRuler::SetRuling(FVector2D InAbsoluteOrigin, float InSlateToUnitScale)
 {
 	AbsoluteOrigin = InAbsoluteOrigin;
diff --git a/Engine/Source/Editor/UMGEditor/Private/Designer/SRuler.h b/Engine/Source/Editor/UMGEditor/Private/Designer/SRuler.h
--- a/Engine/Source/Editor/UMGEditor/Private/Designer/SRuler.h
+++ b/Engine/Source/Editor/UMGEditor/Private/Designer/SRuler.h
@@ -25,6 +25,17 @@ public:
 	 */
 	void Construct( const FArguments& InArgs );
 
+	/**
+	* Determines the optimal spacing between tick marks in the slider for a given pixel density
+	* Increments until a minimum amount of slate units specified by MinTick is reached
+	*
+	* @param InPixelsPerInput	The density of pixels between each input
+	* @param MinTick			The minimum slate units per tick allowed
+	* @param MinTickSpacing	The minimum tick spacing in time units allowed
+	* @return the optimal spacing in time units
+	*/
+	float DetermineOptimalSpacing(float InPixelsPerInput, uint32 MinTick, float MinTickSpacing) const;
+
 	/**
 	 * Sets the rules to use when rendering the ruler.
 	 *
diff --git a/Engine/Source/Editor/UnrealEd/Private/EditorAnimUtils.cpp b/Engine/Source/Editor/UnrealEd/Private/EditorAnimUtils.cpp
--- a/Engine/Source/Editor/UnrealEd/Private/EditorAnimUtils.cpp
+++ b/Engine/Source/Editor/UnrealEd/Private/EditorAnimUtils.cpp
@@ -13,7 +13,7 @@
 #include "ObjectEditorUtils.h"
 #include "SNotificationList.h"
 #include "Editor/ContentBrowser/Public/ContentBrowserModule.h"
-#include "Serialization/ArchiveUObjectBase.h"
+#include "Serialization/ArchiveUObject.h"
 
 #define LOCTEXT_NAMESPACE "EditorAnimUtils"
 
diff --git a/Engine/Source/Programs/UnrealGameSync/UnrealGameSync/Perforce.cs b/Engine/Source/Programs/UnrealGameSync/UnrealGameSync/Perforce.cs
--- a/Engine/Source/Programs/UnrealGameSync/UnrealGameSync/Perforce.cs
+++ b/Engine/Source/Programs/UnrealGameSync/UnrealGameSync/Perforce.cs
@@ -800,7 +800,7 @@ namespace UnrealGameSync
 			bool bResult;
 			using(PerforceTagRecordParser Parser = new PerforceTagRecordParser(x => SyncOutput(new PerforceFileRecord(x))))
 			{
-				bResult = RunCommand(String.Format("-x \"{0}\" -z tag sync", TempFileName), null, Line => FilterSyncOutput(Line, Parser, TamperedFiles, Log), CommandOptions.NoFailOnErrors | CommandOptions.IgnoreFilesUpToDateError | CommandOptions.IgnoreExitCode, Log);
+				bResult = RunCommand(String.Format("-r5 -x \"{0}\" -z tag sync", TempFileName), null, Line => FilterSyncOutput(Line, Parser, TamperedFiles, Log), CommandOptions.NoFailOnErrors | CommandOptions.IgnoreFilesUpToDateError | CommandOptions.IgnoreExitCode, Log);
 			}
 			return bResult;
 		}
diff --git a/Engine/Source/Programs/UnrealGameSync/UnrealGameSync/Properties/AssemblyInfo.cs b/Engine/Source/Programs/UnrealGameSync/UnrealGameSync/Properties/AssemblyInfo.cs
--- a/Engine/Source/Programs/UnrealGameSync/UnrealGameSync/Properties/AssemblyInfo.cs
+++ b/Engine/Source/Programs/UnrealGameSync/UnrealGameSync/Properties/AssemblyInfo.cs
@@ -34,5 +34,5 @@ using System.Runtime.InteropServices;
 // You can specify all the values or you can default the Build and Revision Numbers 
 // by using the '*' as shown below:
 // [assembly: AssemblyVersion("1.0.*")]
-[assembly: AssemblyVersion("1.78.*")]
-[assembly: AssemblyFileVersion("1.78.0.0")]
+[assembly: AssemblyVersion("1.79.*")]
+[assembly: AssemblyFileVersion("1.79.0.0")]
diff --git a/Engine/Source/Runtime/AIModule/Public/SimpleCellGrid.h b/Engine/Source/Runtime/AIModule/Public/SimpleCellGrid.h
--- a/Engine/Source/Runtime/AIModule/Public/SimpleCellGrid.h
+++ b/Engine/Source/Runtime/AIModule/Public/SimpleCellGrid.h
@@ -2,7 +2,7 @@
 
 #pragma once
 
-#include "Serialization/ArchiveBase.h"
+#include "Serialization/Archive.h"
 #include "HAL/UnrealMemory.h"
 
 struct FGridSize2D
diff --git a/Engine/Source/Runtime/Core/Private/CorePrivatePCH.h b/Engine/Source/Runtime/Core/Private/CorePrivatePCH.h
--- a/Engine/Source/Runtime/Core/Private/CorePrivatePCH.h
+++ b/Engine/Source/Runtime/Core/Private/CorePrivatePCH.h
@@ -18,6 +18,8 @@ Low level includes.
 #include "AssertionMacros.h"					// Various assertion macros
 #include "UObject/UnrealNames.h"				// EName definition.
 #include "OutputDevice.h"						// Output devices, logf, debugf, etc
+#include "Misc/MessageDialog.h"
+#include "Misc/Exec.h"
 #include "NumericLimits.h"						// Numeric limits
 #include "UnrealMathUtility.h"					// FMath
 #include "UnrealTypeTraits.h"					// Type Traits
@@ -49,16 +51,25 @@ Includes.
 #include "ObjectVersion.h"				// Object version info.
 #include "TypeHash.h"					// Type hashing functions
 #include "EnumAsByte.h"					// Common template definitions.
-#include "ArchiveBase.h"				// Archive class.
+#include "Archive.h"					// Archive class.
+#include "ArchiveProxy.h"
+#include "NameAsStringProxyArchive.h"
 #include "ScopedPointer.h"				// Scoped pointer definitions.
 #include "Sorting.h"					// Sorting definitions.
 #include "Array.h"						// Dynamic array definitions.
+#include "ScriptArray.h"
+#include "MRUArray.h"
+#include "TransArray.h"
+#include "IndirectArray.h"
 #include "BitArray.h"					// Bit array definition.
-#include "Bitstreams.h"					// Bit stream archiver.
+#include "BitReader.h"					// Bit stream archiver.
+#include "BitWriter.h"					// Bit stream archiver.
 #include "SparseArray.h"				// Sparse array definitions.
 #include "UnrealString.h"				// Dynamic string definitions.
 #include "NameTypes.h"					// Global name subsystem.
 #include "CoreMisc.h"					// Low level utility code.
+#include "FileHelper.h"
+#include "CommandLine.h"
 #include "Paths.h"						// Path helper functions
 #include "Set.h"						// Set definitions.
 #include "Map.h"						// Dynamic map definitions.
@@ -71,11 +82,31 @@ Includes.
 #include "Guid.h"						// FGuid class
 #include "UnrealMath.h"					// Vector math functions.
 #include "OutputDevices.h"				// Output devices
+#include "OutputDeviceRedirector.h"
+#include "OutputDeviceNull.h"
+#include "OutputDeviceMemory.h"
+#include "OutputDeviceFile.h"
+#include "OutputDeviceDebug.h"
+#include "OutputDeviceArchiveWrapper.h"
+#include "OutputDeviceAnsiError.h"
+#include "BufferedOutputDevice.h"
+#include "LogScopedVerbosityOverride.h"
 #include "CoreStats.h"
 #include "TimeGuard.h"
 #include "MemStack.h"					// Stack based memory management.
 #include "AsyncWork.h"					// Async threaded work
-#include "Archive.h"					// Utility archive classes
+#include "MemoryArchive.h"
+#include "MemoryWriter.h"
+#include "LargeMemoryWriter.h"
+#include "LargeMemoryReader.h"
+#include "BufferArchive.h"
+#include "MemoryReader.h"
+#include "ArrayReader.h"
+#include "ArrayWriter.h"
+#include "BufferReader.h"
+#include "BufferWriter.h"
+#include "ArchiveSaveCompressedProxy.h"
+#include "ArchiveLoadCompressedProxy.h"
 #include "Variant.h"
 #include "CircularBuffer.h"
 #include "Queue.h"
@@ -84,6 +115,8 @@ Includes.
 #include "ConfigCacheIni.h"				// The configuration cache declarations
 #include "IConsoleManager.h"
 #include "FeedbackContext.h"
+#include "SlowTask.h"
+#include "ScopedSlowTask.h"
 #include "AutomationTest.h"				// Automation testing functionality for game, editor, etc.
 #include "CallbackDevice.h"				// Base class for callback devices.
 #include "GenericWindowDefinition.h"
diff --git a/Engine/Source/Runtime/Core/Private/GenericPlatform/GenericPlatformOutputDevices.cpp b/Engine/Source/Runtime/Core/Private/GenericPlatform/GenericPlatformOutputDevices.cpp
--- a/Engine/Source/Runtime/Core/Private/GenericPlatform/GenericPlatformOutputDevices.cpp
+++ b/Engine/Source/Runtime/Core/Private/GenericPlatform/GenericPlatformOutputDevices.cpp
@@ -4,7 +4,7 @@
 #include "HAL/FeedbackContextAnsi.h"
 #include "Misc/App.h"
 #include "Misc/OutputDeviceConsole.h"
-#include "HAL/OutputDeviceMemory.h"
+#include "Misc/OutputDeviceMemory.h"
 
 void FGenericPlatformOutputDevices::SetupOutputDevices()
 {
diff --git a/Engine/Source/Runtime/Core/Private/Linux/LinuxCriticalSection.cpp b/Engine/Source/Runtime/Core/Private/Linux/LinuxCriticalSection.cpp
new file mode 100644
--- /dev/null
+++ b/Engine/Source/Runtime/Core/Private/Linux/LinuxCriticalSection.cpp
@@ -0,0 +1,65 @@
+// Copyright 1998-2016 Epic Games, Inc. All Rights Reserved.
+
+#include "CorePrivatePCH.h"
+#include "Linux/LinuxCriticalSection.h"
+#include "LinuxPlatformRunnableThread.h"
+#include "EngineVersion.h"
+#include <spawn.h>
+#include <sys/wait.h>
+#include <sys/resource.h>
+#include <sys/ioctl.h> // ioctl
+#include <sys/file.h>
+#include <asm/ioctls.h> // FIONREAD
+#include <sys/file.h> // flock
+#include "LinuxApplication.h" // FLinuxApplication::IsForeground()
+
+FLinuxSystemWideCriticalSection::FLinuxSystemWideCriticalSection(const FString& InName, FTimespan InTimeout)
+{
+	check(InName.Len() > 0)
+	check(InTimeout >= FTimespan::Zero())
+	check(InTimeout.GetTotalSeconds() < (double)FLT_MAX)
+
+	const FString LockPath = FString(FLinuxPlatformProcess::ApplicationSettingsDir()) / InName;
+	FString NormalizedFilepath(LockPath);
+	NormalizedFilepath.ReplaceInline(TEXT("\\"), TEXT("/"));
+
+	// Attempt to open a file and then lock with flock (NOTE: not an atomic operation, but best we can do)
+	FileHandle = open(TCHAR_TO_UTF8(*NormalizedFilepath), O_CREAT | O_WRONLY | O_NONBLOCK, S_IRUSR | S_IWUSR | S_IRGRP | S_IWGRP | S_IROTH | S_IWOTH);
+
+	if (FileHandle == -1 && InTimeout != FTimespan::Zero())
+	{
+		FDateTime ExpireTime = FDateTime::UtcNow() + InTimeout;
+		const float RetrySeconds = FMath::Min((float)InTimeout.GetTotalSeconds(), 0.25f);
+
+		do
+		{
+			// retry until timeout
+			FLinuxPlatformProcess::Sleep(RetrySeconds);
+			FileHandle = open(TCHAR_TO_UTF8(*NormalizedFilepath), O_CREAT | O_WRONLY | O_NONBLOCK, S_IRUSR | S_IWUSR | S_IRGRP | S_IWGRP | S_IROTH | S_IWOTH);
+		} while (FileHandle == -1 && FDateTime::UtcNow() < ExpireTime);
+	}
+	if (FileHandle != -1)
+	{
+		flock(FileHandle, LOCK_EX);
+	}
+}
+
+FLinuxSystemWideCriticalSection::~FLinuxSystemWideCriticalSection()
+{
+	Release();
+}
+
+bool FLinuxSystemWideCriticalSection::IsValid() const
+{
+	return FileHandle != -1;
+}
+
+void FLinuxSystemWideCriticalSection::Release()
+{
+	if (IsValid())
+	{
+		flock(FileHandle, LOCK_UN);
+		close(FileHandle);
+		FileHandle = -1;
+	}
+}
diff --git a/Engine/Source/Runtime/Core/Private/Linux/LinuxPlatformProcess.cpp b/Engine/Source/Runtime/Core/Private/Linux/LinuxPlatformProcess.cpp
--- a/Engine/Source/Runtime/Core/Private/Linux/LinuxPlatformProcess.cpp
+++ b/Engine/Source/Runtime/Core/Private/Linux/LinuxPlatformProcess.cpp
@@ -1463,54 +1463,3 @@ void FLinuxPlatformProcess::CeaseBeingFirstInstance()
 	}
 #endif
 }
-
-FLinuxSystemWideCriticalSection::FLinuxSystemWideCriticalSection(const FString& InName, FTimespan InTimeout)
-{
-	check(InName.Len() > 0)
-	check(InTimeout >= FTimespan::Zero())
-	check(InTimeout.GetTotalSeconds() < (double)FLT_MAX)
-
-	const FString LockPath = FString(FLinuxPlatformProcess::ApplicationSettingsDir()) / InName;
-	FString NormalizedFilepath(LockPath);
-	NormalizedFilepath.ReplaceInline(TEXT("\\"), TEXT("/"));
-
-	// Attempt to open a file and then lock with flock (NOTE: not an atomic operation, but best we can do)
-	FileHandle = open(TCHAR_TO_UTF8(*NormalizedFilepath), O_CREAT | O_WRONLY | O_NONBLOCK, S_IRUSR | S_IWUSR | S_IRGRP | S_IWGRP | S_IROTH | S_IWOTH);
-
-	if (FileHandle == -1 && InTimeout != FTimespan::Zero())
-	{
-		FDateTime ExpireTime = FDateTime::UtcNow() + InTimeout;
-		const float RetrySeconds = FMath::Min((float)InTimeout.GetTotalSeconds(), 0.25f);
-
-		do
-		{
-			// retry until timeout
-			FLinuxPlatformProcess::Sleep(RetrySeconds);
-			FileHandle = open(TCHAR_TO_UTF8(*NormalizedFilepath), O_CREAT | O_WRONLY | O_NONBLOCK, S_IRUSR | S_IWUSR | S_IRGRP | S_IWGRP | S_IROTH | S_IWOTH);
-		} while (FileHandle == -1 && FDateTime::UtcNow() < ExpireTime);
-	}
-	if (FileHandle != -1)
-	{
-		flock(FileHandle, LOCK_EX);
-	}
-}
-
-FLinuxSystemWideCriticalSection::~FLinuxSystemWideCriticalSection()
-{
-	Release();
-}
-
-bool FLinuxSystemWideCriticalSection::IsValid() const
-{
-	return FileHandle != -1;
-}
-
-void FLinuxSystemWideCriticalSection::Release()
-{
-	if (IsValid())
-	{
-		flock(FileHandle, LOCK_UN);
-		close(FileHandle);
-		FileHandle = -1;
-	}
-}
diff --git a/Engine/Source/Runtime/Core/Private/Logging/LogMacros.cpp b/Engine/Source/Runtime/Core/Private/Logging/LogMacros.cpp
new file mode 100644
--- /dev/null
+++ b/Engine/Source/Runtime/Core/Private/Logging/LogMacros.cpp
@@ -0,0 +1,120 @@
+// Copyright 1998-2016 Epic Games, Inc. All Rights Reserved.
+
+#include "CorePrivatePCH.h"
+#include "Logging/LogMacros.h"
+#include "VarargsHelper.h"
+
+void StaticFailDebug( const TCHAR* Error, const ANSICHAR* File, int32 Line, const TCHAR* Description, bool bIsEnsure );
+
+/** Statics to prevent FMsg::Logf from allocating too much stack memory. */
+static FCriticalSection					MsgLogfStaticBufferGuard;
+/** Increased from 4096 to fix crashes in the renderthread without autoreporter. */
+static TCHAR							MsgLogfStaticBuffer[8192];
+
+VARARG_BODY(void, FMsg::Logf, const TCHAR*, VARARG_EXTRA(const ANSICHAR* File) VARARG_EXTRA(int32 Line) VARARG_EXTRA(const class FName& Category) VARARG_EXTRA(ELogVerbosity::Type Verbosity))
+{
+#if !NO_LOGGING
+	if (Verbosity != ELogVerbosity::Fatal)
+	{
+		// SetColour is routed to GWarn just like the other verbosities and handled in the 
+		// device that does the actual printing.
+		FOutputDevice* LogDevice = NULL;
+		switch (Verbosity)
+		{
+		case ELogVerbosity::Error:
+		case ELogVerbosity::Warning:
+		case ELogVerbosity::Display:
+		case ELogVerbosity::SetColor:
+			if (GWarn)
+			{
+				LogDevice = GWarn;
+				break;
+			}
+		default:
+		{
+			LogDevice = GLog;
+		}
+		break;
+		}
+		GROWABLE_LOGF(LogDevice->Log(Category, Verbosity, Buffer))
+	}
+	else
+	{
+		// Keep Message buffer small, in some cases, this code is executed with 16KB stack.
+		TCHAR Message[4096];
+		{
+			// Simulate Sprintf_s
+			// @todo: implement platform independent sprintf_S
+			// We're using one big shared static buffer here, so guard against re-entry
+			FScopeLock MsgLock(&MsgLogfStaticBufferGuard);
+			// Print to a large static buffer so we can keep the stack allocation below 16K
+			GET_VARARGS(MsgLogfStaticBuffer, ARRAY_COUNT(MsgLogfStaticBuffer), ARRAY_COUNT(MsgLogfStaticBuffer) - 1, Fmt, Fmt);
+			// Copy the message to the stack-allocated buffer)
+			FCString::Strncpy(Message, MsgLogfStaticBuffer, ARRAY_COUNT(Message) - 1);
+			Message[ARRAY_COUNT(Message) - 1] = '\0';
+		}
+
+		StaticFailDebug(TEXT("Fatal error:"), File, Line, Message, false);
+		FDebug::AssertFailed("", File, Line, Message);
+	}
+#endif
+}
+
+VARARG_BODY(void, FMsg::Logf_Internal, const TCHAR*, VARARG_EXTRA(const ANSICHAR* File) VARARG_EXTRA(int32 Line) VARARG_EXTRA(const class FName& Category) VARARG_EXTRA(ELogVerbosity::Type Verbosity))
+{
+#if !NO_LOGGING
+	if (Verbosity != ELogVerbosity::Fatal)
+	{
+		// SetColour is routed to GWarn just like the other verbosities and handled in the 
+		// device that does the actual printing.
+		FOutputDevice* LogDevice = NULL;
+		switch (Verbosity)
+		{
+		case ELogVerbosity::Error:
+		case ELogVerbosity::Warning:
+		case ELogVerbosity::Display:
+		case ELogVerbosity::SetColor:
+			if (GWarn)
+			{
+				LogDevice = GWarn;
+				break;
+			}
+		default:
+		{
+			LogDevice = GLog;
+		}
+		break;
+		}
+		GROWABLE_LOGF(LogDevice->Log(Category, Verbosity, Buffer))
+	}
+	else
+	{
+		// Keep Message buffer small, in some cases, this code is executed with 16KB stack.
+		TCHAR Message[4096];
+		{
+			// Simulate Sprintf_s
+			// @todo: implement platform independent sprintf_S
+			// We're using one big shared static buffer here, so guard against re-entry
+			FScopeLock MsgLock(&MsgLogfStaticBufferGuard);
+			// Print to a large static buffer so we can keep the stack allocation below 16K
+			GET_VARARGS(MsgLogfStaticBuffer, ARRAY_COUNT(MsgLogfStaticBuffer), ARRAY_COUNT(MsgLogfStaticBuffer) - 1, Fmt, Fmt);
+			// Copy the message to the stack-allocated buffer)
+			FCString::Strncpy(Message, MsgLogfStaticBuffer, ARRAY_COUNT(Message) - 1);
+			Message[ARRAY_COUNT(Message) - 1] = '\0';
+		}
+
+		StaticFailDebug(TEXT("Fatal error:"), File, Line, Message, false);
+	}
+#endif
+}
+
+/** Sends a formatted message to a remote tool. */
+void VARARGS FMsg::SendNotificationStringf( const TCHAR *Fmt, ... )
+{
+	GROWABLE_LOGF(SendNotificationString(Buffer));
+}
+
+void FMsg::SendNotificationString( const TCHAR* Message )
+{
+	FPlatformMisc::LowLevelOutputDebugString(Message);
+}
diff --git a/Engine/Source/Runtime/Core/Private/Mac/MacCriticalSection.cpp b/Engine/Source/Runtime/Core/Private/Mac/MacCriticalSection.cpp
new file mode 100644
--- /dev/null
+++ b/Engine/Source/Runtime/Core/Private/Mac/MacCriticalSection.cpp
@@ -0,0 +1,63 @@
+// Copyright 1998-2016 Epic Games, Inc. All Rights Reserved.
+
+/*=============================================================================
+	MacPlatformProcess.mm: Mac implementations of Process functions
+=============================================================================*/
+
+#include "CorePrivatePCH.h"
+#include "Mac/MacCriticalSection.h"
+#include "Misc/App.h"
+#include "MacApplication.h"
+#include <mach-o/dyld.h>
+#include <libproc.h>
+
+FMacSystemWideCriticalSection::FMacSystemWideCriticalSection(const FString& InName, FTimespan InTimeout)
+{
+	check(InName.Len() > 0)
+	check(InTimeout >= FTimespan::Zero())
+	check(InTimeout.GetTotalSeconds() < (double)FLT_MAX)	// we'll need this to fit in a single-precision float later so check it here
+
+	FDateTime ExpireTime = FDateTime::UtcNow() + InTimeout;
+
+	// This lock implementation is using files so correct any backslashes in the name
+	const FString LockPath = FString(FMacPlatformProcess::ApplicationSettingsDir()) / InName;
+	FString NormalizedFilepath(LockPath);
+	NormalizedFilepath.ReplaceInline(TEXT("\\"), TEXT("/"));
+
+	// Attempt to open a file with O_EXLOCK (equivalent of atomic open() + flock())
+	FileHandle = open(TCHAR_TO_UTF8(*NormalizedFilepath), O_CREAT | O_WRONLY | O_EXLOCK | O_NONBLOCK, S_IRUSR | S_IWUSR | S_IRGRP | S_IWGRP | S_IROTH | S_IWOTH);
+
+	if (FileHandle == -1 && InTimeout != FTimespan::Zero())
+	{
+		// Failed to get a valid file handle so the file is probably open and locked by another owner - retry until we timeout
+
+		for (FTimespan TimeoutRemaining = ExpireTime - FDateTime::UtcNow();
+			 FileHandle == -1 && TimeoutRemaining > FTimespan::Zero();
+			 TimeoutRemaining = ExpireTime - FDateTime::UtcNow())
+		{
+			// retry until timeout
+			float RetrySeconds = FMath::Min((float)TimeoutRemaining.GetTotalSeconds(), 0.25f);
+			FMacPlatformProcess::Sleep(RetrySeconds);
+			FileHandle = open(TCHAR_TO_UTF8(*NormalizedFilepath), O_CREAT | O_WRONLY | O_EXLOCK | O_NONBLOCK, S_IRUSR | S_IWUSR | S_IRGRP | S_IWGRP | S_IROTH | S_IWOTH);
+		}
+	}
+}
+
+FMacSystemWideCriticalSection::~FMacSystemWideCriticalSection()
+{
+	Release();
+}
+
+bool FMacSystemWideCriticalSection::IsValid() const
+{
+	return FileHandle != -1;
+}
+
+void FMacSystemWideCriticalSection::Release()
+{
+	if (IsValid())
+	{
+		close(FileHandle);
+		FileHandle = -1;
+	}
+}
diff --git a/Engine/Source/Runtime/Core/Private/Mac/MacPlatformProcess.cpp b/Engine/Source/Runtime/Core/Private/Mac/MacPlatformProcess.cpp
--- a/Engine/Source/Runtime/Core/Private/Mac/MacPlatformProcess.cpp
+++ b/Engine/Source/Runtime/Core/Private/Mac/MacPlatformProcess.cpp
@@ -1204,55 +1204,3 @@ FRunnableThread* FMacPlatformProcess::CreateRunnableThread()
 {
 	return new FRunnableThreadMac();
 }
-
-
-FMacSystemWideCriticalSection::FMacSystemWideCriticalSection(const FString& InName, FTimespan InTimeout)
-{
-	check(InName.Len() > 0)
-	check(InTimeout >= FTimespan::Zero())
-	check(InTimeout.GetTotalSeconds() < (double)FLT_MAX)	// we'll need this to fit in a single-precision float later so check it here
-
-	FDateTime ExpireTime = FDateTime::UtcNow() + InTimeout;
-
-	// This lock implementation is using files so correct any backslashes in the name
-	const FString LockPath = FString(FMacPlatformProcess::ApplicationSettingsDir()) / InName;
-	FString NormalizedFilepath(LockPath);
-	NormalizedFilepath.ReplaceInline(TEXT("\\"), TEXT("/"));
-
-	// Attempt to open a file with O_EXLOCK (equivalent of atomic open() + flock())
-	FileHandle = open(TCHAR_TO_UTF8(*NormalizedFilepath), O_CREAT | O_WRONLY | O_EXLOCK | O_NONBLOCK, S_IRUSR | S_IWUSR | S_IRGRP | S_IWGRP | S_IROTH | S_IWOTH);
-
-	if (FileHandle == -1 && InTimeout != FTimespan::Zero())
-	{
-		// Failed to get a valid file handle so the file is probably open and locked by another owner - retry until we timeout
-
-		for (FTimespan TimeoutRemaining = ExpireTime - FDateTime::UtcNow();
-			 FileHandle == -1 && TimeoutRemaining > FTimespan::Zero();
-			 TimeoutRemaining = ExpireTime - FDateTime::UtcNow())
-		{
-			// retry until timeout
-			float RetrySeconds = FMath::Min((float)TimeoutRemaining.GetTotalSeconds(), 0.25f);
-			FMacPlatformProcess::Sleep(RetrySeconds);
-			FileHandle = open(TCHAR_TO_UTF8(*NormalizedFilepath), O_CREAT | O_WRONLY | O_EXLOCK | O_NONBLOCK, S_IRUSR | S_IWUSR | S_IRGRP | S_IWGRP | S_IROTH | S_IWOTH);
-		}
-	}
-}
-
-FMacSystemWideCriticalSection::~FMacSystemWideCriticalSection()
-{
-	Release();
-}
-
-bool FMacSystemWideCriticalSection::IsValid() const
-{
-	return FileHandle != -1;
-}
-
-void FMacSystemWideCriticalSection::Release()
-{
-	if (IsValid())
-	{
-		close(FileHandle);
-		FileHandle = -1;
-	}
-}
diff --git a/Engine/Source/Runtime/Core/Private/Math/Matrix.cpp b/Engine/Source/Runtime/Core/Private/Math/Matrix.cpp
new file mode 100644
--- /dev/null
+++ b/Engine/Source/Runtime/Core/Private/Math/Matrix.cpp
@@ -0,0 +1,10 @@
+// Copyright 1998-2016 Epic Games, Inc. All Rights Reserved.
+
+#include "CorePrivatePCH.h"
+#include "Matrix.h"
+
+void FMatrix::ErrorEnsure(const TCHAR* Message)
+{
+	UE_LOG(LogUnrealMath, Error, TEXT("FMatrix::InverseFast(), trying to invert a NIL matrix, this results in NaNs! Use Inverse() instead."));
+	ensureMsgf(false, TEXT("FMatrix::InverseFast(), trying to invert a NIL matrix, this results in NaNs! Use Inverse() instead."));
+}
diff --git a/Engine/Source/Runtime/Core/Private/Math/UnrealMath.cpp b/Engine/Source/Runtime/Core/Private/Math/UnrealMath.cpp
--- a/Engine/Source/Runtime/Core/Private/Math/UnrealMath.cpp
+++ b/Engine/Source/Runtime/Core/Private/Math/UnrealMath.cpp
@@ -188,6 +188,11 @@ FRotator FVector::ToOrientationRotator() const
 	return R;
 }
 
+FRotator FVector::Rotation() const
+{
+	return ToOrientationRotator();
+}
+
 FRotator FVector4::ToOrientationRotator() const
 {
 	FRotator R;
@@ -212,6 +217,10 @@ FRotator FVector4::ToOrientationRotator() const
 	return R;
 }
 
+FRotator FVector4::Rotation() const
+{
+	return ToOrientationRotator();
+}
 
 FQuat FVector::ToOrientationQuat() const
 {
diff --git a/Engine/Source/Runtime/Core/Private/Misc/AssertionMacros.cpp b/Engine/Source/Runtime/Core/Private/Misc/AssertionMacros.cpp
new file mode 100644
--- /dev/null
+++ b/Engine/Source/Runtime/Core/Private/Misc/AssertionMacros.cpp
@@ -0,0 +1,356 @@
+// Copyright 1998-2016 Epic Games, Inc. All Rights Reserved.
+
+#include "CorePrivatePCH.h"
+#include "ExceptionHandling.h"
+#include "VarargsHelper.h"
+#include "OutputDeviceHelper.h"
+#include "HAL/ThreadHeartBeat.h"
+
+bool FDebug::bHasAsserted = false;
+
+#define FILE_LINE_DESC TEXT(" [File:%s] [Line: %i] ")
+
+/** Lock used to synchronize the fail debug calls. */
+static FCriticalSection	FailDebugCriticalSection;
+
+/** Number of top function calls to hide when dumping the callstack as text. */
+#if PLATFORM_LINUX
+
+	// Rationale: check() and ensure() handlers have different depth - worse, ensure() can optionally end up calling the same path as check().
+	// It is better to show the full callstack as is than accidentaly ignore a part of the problem
+	#define CALLSTACK_IGNOREDEPTH 0
+
+#else
+	#define CALLSTACK_IGNOREDEPTH 2
+#endif // PLATFORM_LINUX
+
+void PrintScriptCallstack(bool bEmptyWhenDone)
+{
+#if DO_BLUEPRINT_GUARD
+	// Walk the script stack, if any
+	FBlueprintExceptionTracker& BlueprintExceptionTracker = FBlueprintExceptionTracker::Get();
+	if( BlueprintExceptionTracker.ScriptStack.Num() > 0 )
+	{
+		FString ScriptStack = TEXT( "\n\nScript Stack:\n" );
+		for (int32 FrameIdx = BlueprintExceptionTracker.ScriptStack.Num() - 1; FrameIdx >= 0; --FrameIdx)
+		{
+			ScriptStack += BlueprintExceptionTracker.ScriptStack[FrameIdx].GetStackDescription() + TEXT( "\n" );
+		}
+
+		UE_LOG( LogOutputDevice, Warning, TEXT( "%s" ), *ScriptStack );
+
+		if (bEmptyWhenDone)
+		{
+			BlueprintExceptionTracker.ScriptStack.Empty();
+		}
+	}
+#endif
+}
+
+/**
+ *	Prints error to the debug output, 
+ *	prompts for the remote debugging if there is not debugger, breaks into the debugger 
+ *	and copies the error into the global error message.
+ */
+void StaticFailDebug( const TCHAR* Error, const ANSICHAR* File, int32 Line, const TCHAR* Description, bool bIsEnsure )
+{
+	// For ensure log should be flushed in the engine loop.
+	if( !bIsEnsure )
+	{
+		GLog->PanicFlushThreadedLogs();
+	}
+
+	FScopeLock Lock( &FailDebugCriticalSection );
+	FPlatformMisc::LowLevelOutputDebugStringf(TEXT("%s") FILE_LINE_DESC TEXT("\n%s\n"), Error, ANSI_TO_TCHAR(File), Line, Description);
+
+	// Copy the detailed error into the error message.
+	FCString::Snprintf( GErrorMessage, ARRAY_COUNT( GErrorMessage ), TEXT( "%s" ) FILE_LINE_DESC TEXT( "\n%s\n" ), Error, ANSI_TO_TCHAR( File ), Line, Description );
+
+	// Copy the error message to the error history.
+	FCString::Strncpy( GErrorHist, GErrorMessage, ARRAY_COUNT( GErrorHist ) );
+	FCString::Strncat( GErrorHist, TEXT( "\r\n\r\n" ), ARRAY_COUNT( GErrorHist ) );
+}
+
+void OutputMultiLineCallstack(const ANSICHAR* File, int32 Line, const FName& LogName, const TCHAR* Heading, TCHAR* Message, ELogVerbosity::Type Verbosity)
+{
+	const bool bWriteUATMarkers = FParse::Param(FCommandLine::Get(), TEXT("CrashForUAT")) && FParse::Param(FCommandLine::Get(), TEXT("stdout"));
+
+	if (bWriteUATMarkers)
+	{
+		FMsg::Logf(File, Line, LogName, Verbosity, TEXT("begin: stack for UAT"));
+	}
+
+	FMsg::Logf(File, Line, LogName, Verbosity, TEXT("%s"), Heading);
+	FMsg::Logf(File, Line, LogName, Verbosity, TEXT(""));
+
+	for (TCHAR* LineStart = Message;; )
+	{
+		// Find the end of the current line
+		TCHAR* LineEnd = LineStart;
+		while (*LineEnd != 0 && *LineEnd != '\r' && *LineEnd != '\n')
+		{
+			LineEnd++;
+		}
+
+		// Output it
+		TCHAR LineEndCharacter = *LineEnd;
+		*LineEnd = 0;
+		FMsg::Logf(File, Line, LogName, Verbosity, TEXT("%s"), LineStart);
+		*LineEnd = LineEndCharacter;
+
+		// Quit if this was the last line
+		if (*LineEnd == 0)
+		{
+			break;
+		}
+
+		// Move to the next line
+		LineStart = (LineEnd[0] == '\r' && LineEnd[1] == '\n') ? (LineEnd + 2) : (LineEnd + 1);
+	}
+
+	if (bWriteUATMarkers)
+	{
+		FMsg::Logf(File, Line, LogName, Verbosity, TEXT("end: stack for UAT"));
+	}
+}
+
+#if DO_CHECK || DO_GUARD_SLOW
+//
+// Failed assertion handler.
+//warning: May be called at library startup time.
+//
+void VARARGS FDebug::LogAssertFailedMessage(const ANSICHAR* Expr, const ANSICHAR* File, int32 Line, const TCHAR* Format/*=TEXT("")*/, ...)
+{
+	// Print out the blueprint callstack
+	PrintScriptCallstack(true);
+
+	// Ignore this assert if we're already forcibly shutting down because of a critical error.
+	if( !GIsCriticalError )
+	{
+		TCHAR DescriptionString[4096];
+		GET_VARARGS( DescriptionString, ARRAY_COUNT( DescriptionString ), ARRAY_COUNT( DescriptionString ) - 1, Format, Format );
+
+		TCHAR ErrorString[MAX_SPRINTF];
+		FCString::Sprintf( ErrorString, TEXT( "Assertion failed: %s" ), ANSI_TO_TCHAR( Expr ) );
+
+		if( FPlatformProperties::AllowsCallStackDumpDuringAssert() )
+		{
+			ANSICHAR StackTrace[4096];
+			if( StackTrace != NULL )
+			{
+				StackTrace[0] = 0;
+				FPlatformStackWalk::StackWalkAndDump( StackTrace, ARRAY_COUNT(StackTrace), CALLSTACK_IGNOREDEPTH );
+
+				FCString::Strncat( DescriptionString, ANSI_TO_TCHAR( StackTrace ), ARRAY_COUNT( DescriptionString ) - 1 );
+			}
+		}
+
+		StaticFailDebug( ErrorString, File, Line, DescriptionString, false );
+	}
+}
+
+/**
+ * Called when an 'ensure' assertion fails; gathers stack data and generates and error report.
+ *
+ * @param	Expr	Code expression ANSI string (#code)
+ * @param	File	File name ANSI string (__FILE__)
+ * @param	Line	Line number (__LINE__)
+ * @param	Msg		Informative error message text
+ */
+void FDebug::EnsureFailed(const ANSICHAR* Expr, const ANSICHAR* File, int32 Line, const TCHAR* Msg)
+{
+
+#if STATS
+	FString EnsureFailedPerfMessage = FString::Printf(TEXT("FDebug::EnsureFailed"));
+	SCOPE_LOG_TIME_IN_SECONDS(*EnsureFailedPerfMessage, nullptr)
+#endif
+
+	// You can set bShouldCrash to true to cause a regular assertion to trigger (stopping program execution) when an ensure() error occurs
+	const bool bShouldCrash = false;		// By default, don't crash on ensure()
+	if( bShouldCrash )
+	{
+		// Just trigger a regular assertion which will crash via GError->Logf()
+		FDebug::LogAssertFailedMessage( Expr, File, Line, Msg );
+		return;
+	}
+
+	// Print out the blueprint callstack
+	PrintScriptCallstack(false);
+
+	// Print initial debug message for this error
+	TCHAR ErrorString[MAX_SPRINTF];
+	FCString::Sprintf(ErrorString,TEXT("Ensure condition failed: %s"),ANSI_TO_TCHAR(Expr));
+
+	StaticFailDebug( ErrorString, File, Line, Msg, true );
+
+	// Is there a debugger attached?  If not we'll submit an error report.
+	if (FPlatformMisc::IsDebuggerPresent())
+	{
+#if !NO_LOGGING
+		UE_LOG(LogOutputDevice, Error, TEXT("%s [File:%s] [Line: %i]"), ErrorString, ANSI_TO_TCHAR(File), Line);
+#endif
+		return;
+	}
+
+	// If we determine that we have not sent a report for this ensure yet, send the report below.
+	bool bShouldSendNewReport = false;
+
+	// No debugger attached, so generate a call stack and submit a crash report
+	// Walk the stack and dump it to the allocated memory.
+	const SIZE_T StackTraceSize = 65535;
+	ANSICHAR* StackTrace = (ANSICHAR*) FMemory::SystemMalloc( StackTraceSize );
+	if( StackTrace != NULL )
+	{
+		// Stop checking heartbeat for this thread. Ensure can take a lot of time (when stackwalking)
+		// Thread heartbeat will be resumed the next time this thread calls FThreadHeartBeat::Get().HeartBeat();
+		// The reason why we don't call HeartBeat() at the end of this function is that maybe this thread
+		// Never had a heartbeat checked and may not be sending heartbeats at all which would later lead to a false positives when detecting hangs.
+		FThreadHeartBeat::Get().KillHeartBeat();
+
+		{
+#if STATS
+			FString StackWalkPerfMessage = FString::Printf(TEXT("FPlatformStackWalk::StackWalkAndDump"));
+			SCOPE_LOG_TIME_IN_SECONDS(*StackWalkPerfMessage, nullptr)
+#endif
+			StackTrace[0] = 0;
+			FPlatformStackWalk::StackWalkAndDumpEx( StackTrace, StackTraceSize, CALLSTACK_IGNOREDEPTH, FGenericPlatformStackWalk::EStackWalkFlags::FlagsUsedWhenHandlingEnsure );
+		}
+
+		// Create a final string that we'll output to the log (and error history buffer)
+		TCHAR ErrorMsg[16384];
+		FCString::Snprintf( ErrorMsg, ARRAY_COUNT( ErrorMsg ), TEXT( "Ensure condition failed: %s [File:%s] [Line: %i]" ) LINE_TERMINATOR TEXT( "%s" ) LINE_TERMINATOR TEXT( "Stack: " ) LINE_TERMINATOR, ANSI_TO_TCHAR( Expr ), ANSI_TO_TCHAR( File ), Line, Msg );
+
+		// Also append the stack trace
+		FCString::Strncat( ErrorMsg, ANSI_TO_TCHAR(StackTrace), ARRAY_COUNT(ErrorMsg) - 1 );
+		FMemory::SystemFree( StackTrace );
+
+		// Dump the error and flush the log.
+#if !NO_LOGGING
+		OutputMultiLineCallstack(__FILE__, __LINE__, LogOutputDevice.GetCategoryName(), TEXT("=== Handled ensure: ==="), ErrorMsg, ELogVerbosity::Error);
+#endif
+		GLog->Flush();
+
+		// Submit the error report to the server! (and display a balloon in the system tray)
+		{
+			// How many unique previous errors we should keep track of
+			const uint32 MaxPreviousErrorsToTrack = 4;
+			static uint32 StaticPreviousErrorCount = 0;
+			if( StaticPreviousErrorCount < MaxPreviousErrorsToTrack )
+			{
+				// Check to see if we've already reported this error.  No point in blasting the server with
+				// the same error over and over again in a single application session.
+				bool bHasErrorAlreadyBeenReported = false;
+
+				// Static: Array of previous unique error message CRCs
+				static uint32 StaticPreviousErrorCRCs[ MaxPreviousErrorsToTrack ];
+
+				// Compute CRC of error string.  Note that along with the call stack, this includes the message
+				// string passed to the macro, so only truly redundant errors will go unreported.  Though it also
+				// means you shouldn't pass loop counters to ensureMsgf(), otherwise failures may spam the server!
+				const uint32 ErrorStrCRC = FCrc::StrCrc_DEPRECATED( ErrorMsg );
+
+				for( uint32 CurErrorIndex = 0; CurErrorIndex < StaticPreviousErrorCount; ++CurErrorIndex )
+				{
+					if( StaticPreviousErrorCRCs[ CurErrorIndex ] == ErrorStrCRC )
+					{
+						// Found it!  This is a redundant error message.
+						bHasErrorAlreadyBeenReported = true;
+						break;
+					}
+				}
+
+				// Add the element to the list and bump the count
+				StaticPreviousErrorCRCs[ StaticPreviousErrorCount++ ] = ErrorStrCRC;
+
+				if( !bHasErrorAlreadyBeenReported )
+				{
+#if STATS
+					FString SubmitErrorReporterfMessage = FString::Printf(TEXT("SubmitErrorReport"));
+					SCOPE_LOG_TIME_IN_SECONDS(*SubmitErrorReporterfMessage, nullptr)
+#endif
+
+					FCoreDelegates::OnHandleSystemEnsure.Broadcast();
+
+					FPlatformMisc::SubmitErrorReport( ErrorMsg, EErrorReportMode::Balloon );
+
+					bShouldSendNewReport = true;
+				}
+			}
+		}
+	}
+	else
+	{
+		// If we fail to generate the string to identify the crash we don't know if we should skip sending the report,
+		// so we will just send the report anyway.
+		bShouldSendNewReport = true;
+
+		// Add message to log even without stacktrace. It is useful for testing fail on ensure.
+#if !NO_LOGGING
+		UE_LOG(LogOutputDevice, Error, TEXT("%s [File:%s] [Line: %i]"), ErrorString, ANSI_TO_TCHAR(File), Line);
+#endif
+	}
+
+	if ( bShouldSendNewReport )
+	{
+#if STATS
+		FString SendNewReportMessage = FString::Printf(TEXT("SendNewReport"));
+		SCOPE_LOG_TIME_IN_SECONDS(*SendNewReportMessage, nullptr)
+#endif
+
+#if PLATFORM_DESKTOP
+		FScopeLock Lock( &FailDebugCriticalSection );
+
+		NewReportEnsure( GErrorMessage );
+
+		GErrorHist[0] = 0;
+		GErrorMessage[0] = 0;
+		GErrorExceptionDescription[0] = 0;
+#endif
+	}
+}
+
+#endif // DO_CHECK || DO_GUARD_SLOW
+
+void VARARGS FDebug::AssertFailed(const ANSICHAR* Expr, const ANSICHAR* File, int32 Line, const TCHAR* Format/* = TEXT("")*/, ...)
+{
+	if (GIsCriticalError)
+	{
+		return;
+	}
+
+	// This is not perfect because another thread might crash and be handled before this assert
+	// but this static varible will report the crash as an assert. Given complexity of a thread
+	// aware solution, this should be good enough. If crash reports are obviously wrong we can
+	// look into fixing this.
+	bHasAsserted = true;
+
+	TCHAR DescriptionString[4096];
+	GET_VARARGS(DescriptionString, ARRAY_COUNT(DescriptionString), ARRAY_COUNT(DescriptionString) - 1, Format, Format);
+
+	TCHAR ErrorString[MAX_SPRINTF];
+	FCString::Sprintf(ErrorString, TEXT("Assertion failed: %s"), ANSI_TO_TCHAR(Expr));
+	GError->Logf(TEXT("Assertion failed: %s") FILE_LINE_DESC TEXT("\n%s\n"), ErrorString, ANSI_TO_TCHAR(File), Line, DescriptionString);
+}
+
+#if DO_CHECK || DO_GUARD_SLOW
+bool VARARGS FDebug::OptionallyLogFormattedEnsureMessageReturningFalse( bool bLog, const ANSICHAR* Expr, const ANSICHAR* File, int32 Line, const TCHAR* FormattedMsg, ... )
+{
+	if (bLog)
+	{
+		const int32 TempStrSize = 4096;
+		TCHAR TempStr[ TempStrSize ];
+		GET_VARARGS( TempStr, TempStrSize, TempStrSize - 1, FormattedMsg, FormattedMsg );
+		EnsureFailed( Expr, File, Line, TempStr );
+	}
+	
+	return false;
+}
+#endif
+
+void VARARGS LowLevelFatalErrorHandler(const ANSICHAR* File, int32 Line, const TCHAR* Format, ...)
+{
+	TCHAR DescriptionString[4096];
+	GET_VARARGS( DescriptionString, ARRAY_COUNT(DescriptionString), ARRAY_COUNT(DescriptionString)-1, Format, Format );
+
+	StaticFailDebug(TEXT("LowLevelFatalError"), File, Line, DescriptionString, false);
+}
diff --git a/Engine/Source/Runtime/Core/Private/Misc/BufferedOutputDevice.cpp b/Engine/Source/Runtime/Core/Private/Misc/BufferedOutputDevice.cpp
new file mode 100644
--- /dev/null
+++ b/Engine/Source/Runtime/Core/Private/Misc/BufferedOutputDevice.cpp
@@ -0,0 +1,24 @@
+// Copyright 1998-2016 Epic Games, Inc. All Rights Reserved.
+
+#include "CorePrivatePCH.h"
+#include "BufferedOutputDevice.h"
+
+void FBufferedOutputDevice::Serialize(const TCHAR* InData, ELogVerbosity::Type Verbosity, const class FName& Category)
+{
+	if (Verbosity > FilterLevel)
+		return;
+
+	FScopeLock ScopeLock(&SynchronizationObject);
+	new(BufferedLines)FBufferedLine(InData, Category, Verbosity);
+}
+
+void FBufferedOutputDevice::GetContents(TArray<FBufferedLine>& DestBuffer, bool ClearDevice/*=true*/)
+{
+	FScopeLock ScopeLock(&SynchronizationObject);
+	DestBuffer = BufferedLines;
+
+	if (ClearDevice)
+	{
+		BufferedLines.Empty();
+	}
+}
diff --git a/Engine/Source/Runtime/Core/Private/Misc/CommandLine.cpp b/Engine/Source/Runtime/Core/Private/Misc/CommandLine.cpp
new file mode 100644
--- /dev/null
+++ b/Engine/Source/Runtime/Core/Private/Misc/CommandLine.cpp
@@ -0,0 +1,281 @@
+// Copyright 1998-2016 Epic Games, Inc. All Rights Reserved.
+
+#include "CorePrivatePCH.h"
+#include "CommandLine.h"
+
+/*-----------------------------------------------------------------------------
+	FCommandLine
+-----------------------------------------------------------------------------*/
+
+bool FCommandLine::bIsInitialized = false;
+TCHAR FCommandLine::CmdLine[FCommandLine::MaxCommandLineSize] = TEXT("");
+TCHAR FCommandLine::OriginalCmdLine[FCommandLine::MaxCommandLineSize] = TEXT("");
+TCHAR FCommandLine::LoggingCmdLine[FCommandLine::MaxCommandLineSize] = TEXT("");
+TCHAR FCommandLine::LoggingOriginalCmdLine[FCommandLine::MaxCommandLineSize] = TEXT("");
+FString FCommandLine::SubprocessCommandLine(TEXT(" -Multiprocess"));
+
+bool FCommandLine::IsInitialized()
+{
+	return bIsInitialized;
+}
+
+const TCHAR* FCommandLine::Get()
+{
+	UE_CLOG(!bIsInitialized, LogInit, Fatal, TEXT("Attempting to get the command line but it hasn't been initialized yet."));
+	return CmdLine;
+}
+
+const TCHAR* FCommandLine::GetForLogging()
+{
+	UE_CLOG(!bIsInitialized, LogInit, Fatal, TEXT("Attempting to get the command line but it hasn't been initialized yet."));
+	return LoggingCmdLine;
+}
+
+const TCHAR* FCommandLine::GetOriginal()
+{
+	UE_CLOG(!bIsInitialized, LogInit, Fatal, TEXT("Attempting to get the command line but it hasn't been initialized yet."));
+	return OriginalCmdLine;
+}
+
+const TCHAR* FCommandLine::GetOriginalForLogging()
+{
+	UE_CLOG(!bIsInitialized, LogInit, Fatal, TEXT("Attempting to get the command line but it hasn't been initialized yet."));
+	return LoggingOriginalCmdLine;
+}
+
+bool FCommandLine::Set(const TCHAR* NewCommandLine)
+{
+	if (!bIsInitialized)
+	{
+		FCString::Strncpy(OriginalCmdLine, NewCommandLine, ARRAY_COUNT(OriginalCmdLine));
+		FCString::Strncpy(LoggingOriginalCmdLine, NewCommandLine, ARRAY_COUNT(LoggingOriginalCmdLine));
+	}
+
+	FCString::Strncpy( CmdLine, NewCommandLine, ARRAY_COUNT(CmdLine) );
+	FCString::Strncpy(LoggingCmdLine, NewCommandLine, ARRAY_COUNT(LoggingCmdLine));
+	// If configured as part of the build, strip out any unapproved args
+	WhitelistCommandLines();
+
+	bIsInitialized = true;
+
+	// Check for the '-' that normal ones get converted to in Outlook. It's important to do it AFTER the command line is initialized
+	if (StringHasBadDashes(NewCommandLine))
+	{
+		FText ErrorMessage = FText::Format(NSLOCTEXT("Engine", "ComdLineHasInvalidChar", "Error: Command-line contains an invalid '-' character, likely pasted from an email.\nCmdline = {0}"), FText::FromString(NewCommandLine));
+#if !UE_BUILD_SHIPPING
+		FMessageDialog::Open(EAppMsgType::Ok, ErrorMessage);
+		return false;
+#else
+		UE_LOG(LogInit, Fatal, TEXT("%s"), *ErrorMessage.ToString());
+#endif
+	}
+
+	return true;
+}
+
+void FCommandLine::Append(const TCHAR* AppendString)
+{
+	FCString::Strncat( CmdLine, AppendString, ARRAY_COUNT(CmdLine) );
+	// If configured as part of the build, strip out any unapproved args
+	WhitelistCommandLines();
+}
+
+#if WANTS_COMMANDLINE_WHITELIST
+TArray<FString> FCommandLine::ApprovedArgs;
+TArray<FString> FCommandLine::FilterArgsForLogging;
+
+#ifdef OVERRIDE_COMMANDLINE_WHITELIST
+/**
+ * When overriding this setting make sure that your define looks like the following in your .cs file:
+ *
+ *		OutCPPEnvironmentConfiguration.Definitions.Add("OVERRIDE_COMMANDLINE_WHITELIST=\"-arg1 -arg2 -arg3 -arg4\"");
+ *
+ * The important part is the \" as they quotes get stripped off by the compiler without them
+ */
+const TCHAR* OverrideList = TEXT(OVERRIDE_COMMANDLINE_WHITELIST);
+#else
+// Default list most conservative restrictions
+const TCHAR* OverrideList = TEXT("-fullscreen /windowed");
+#endif
+
+#ifdef FILTER_COMMANDLINE_LOGGING
+/**
+ * When overriding this setting make sure that your define looks like the following in your .cs file:
+ *
+ *		OutCPPEnvironmentConfiguration.Definitions.Add("FILTER_COMMANDLINE_LOGGING=\"-arg1 -arg2 -arg3 -arg4\"");
+ *
+ * The important part is the \" as they quotes get stripped off by the compiler without them
+ */
+const TCHAR* FilterForLoggingList = TEXT(FILTER_COMMANDLINE_LOGGING);
+#else
+const TCHAR* FilterForLoggingList = TEXT("");
+#endif
+
+void FCommandLine::WhitelistCommandLines()
+{
+	if (ApprovedArgs.Num() == 0)
+	{
+		TArray<FString> Ignored;
+		FCommandLine::Parse(OverrideList, ApprovedArgs, Ignored);
+	}
+	if (FilterArgsForLogging.Num() == 0)
+	{
+		TArray<FString> Ignored;
+		FCommandLine::Parse(FilterForLoggingList, FilterArgsForLogging, Ignored);
+	}
+	// Process the original command line
+	TArray<FString> OriginalList = FilterCommandLine(OriginalCmdLine);
+	BuildWhitelistCommandLine(OriginalCmdLine, ARRAY_COUNT(OriginalCmdLine), OriginalList);
+	// Process the current command line
+	TArray<FString> CmdList = FilterCommandLine(CmdLine);
+	BuildWhitelistCommandLine(CmdLine, ARRAY_COUNT(CmdLine), CmdList);
+	// Process the command line for logging purposes
+	TArray<FString> LoggingCmdList = FilterCommandLineForLogging(LoggingCmdLine);
+	BuildWhitelistCommandLine(LoggingCmdLine, ARRAY_COUNT(LoggingCmdLine), LoggingCmdList);
+	// Process the original command line for logging purposes
+	TArray<FString> LoggingOriginalCmdList = FilterCommandLineForLogging(LoggingOriginalCmdLine);
+	BuildWhitelistCommandLine(LoggingOriginalCmdLine, ARRAY_COUNT(LoggingOriginalCmdLine), LoggingOriginalCmdList);
+}
+
+TArray<FString> FCommandLine::FilterCommandLine(TCHAR* CommandLine)
+{
+	TArray<FString> Ignored;
+	TArray<FString> ParsedList;
+	// Parse the command line list
+	FCommandLine::Parse(CommandLine, ParsedList, Ignored);
+	// Remove any that are not in our approved list
+	for (int32 Index = 0; Index < ParsedList.Num(); Index++)
+	{
+		bool bFound = false;
+		for (auto ApprovedArg : ApprovedArgs)
+		{
+			if (ParsedList[Index].StartsWith(ApprovedArg))
+			{
+				bFound = true;
+				break;
+			}
+		}
+		if (!bFound)
+		{
+			ParsedList.RemoveAt(Index);
+			Index--;
+		}
+	}
+	return ParsedList;
+}
+
+TArray<FString> FCommandLine::FilterCommandLineForLogging(TCHAR* CommandLine)
+{
+	TArray<FString> Ignored;
+	TArray<FString> ParsedList;
+	// Parse the command line list
+	FCommandLine::Parse(CommandLine, ParsedList, Ignored);
+	// Remove any that are not in our approved list
+	for (int32 Index = 0; Index < ParsedList.Num(); Index++)
+	{
+		for (auto Filter : FilterArgsForLogging)
+		{
+			if (ParsedList[Index].StartsWith(Filter))
+			{
+				ParsedList.RemoveAt(Index);
+				Index--;
+				break;
+			}
+		}
+	}
+	return ParsedList;
+}
+
+void FCommandLine::BuildWhitelistCommandLine(TCHAR* CommandLine, uint32 ArrayCount, const TArray<FString>& FilteredArgs)
+{
+	check(ArrayCount > 0);
+	// Zero the whole string
+	FMemory::Memzero(CommandLine, sizeof(TCHAR) * ArrayCount);
+
+	uint32 StartIndex = 0;
+	for (auto Arg : FilteredArgs)
+	{
+		if ((StartIndex + Arg.Len() + 2) < ArrayCount)
+		{
+			if (StartIndex != 0)
+			{
+				CommandLine[StartIndex++] = TEXT(' ');
+			}
+			CommandLine[StartIndex++] = TEXT('-');
+			FCString::Strncpy(&CommandLine[StartIndex], *Arg, ArrayCount - StartIndex);
+			StartIndex += Arg.Len();
+		}
+	}
+}
+#endif
+
+void FCommandLine::AddToSubprocessCommandline( const TCHAR* Param )
+{
+	check( Param != NULL );
+	if (Param[0] != ' ')
+	{
+		SubprocessCommandLine += TEXT(" ");
+	}
+	SubprocessCommandLine += Param;
+}
+
+const FString& FCommandLine::GetSubprocessCommandline()
+{
+	return SubprocessCommandLine;
+}
+
+/** 
+ * Removes the executable name from a commandline, denoted by parentheses.
+ */
+const TCHAR* FCommandLine::RemoveExeName(const TCHAR* InCmdLine)
+{
+	// Skip over executable that is in the commandline
+	if( *InCmdLine=='\"' )
+	{
+		InCmdLine++;
+		while( *InCmdLine && *InCmdLine!='\"' )
+		{
+			InCmdLine++;
+		}
+		if( *InCmdLine )
+		{
+			InCmdLine++;
+		}
+	}
+	while( *InCmdLine && *InCmdLine!=' ' )
+	{
+		InCmdLine++;
+	}
+	// skip over any spaces at the start, which Vista likes to toss in multiple
+	while (*InCmdLine == ' ')
+	{
+		InCmdLine++;
+	}
+	return InCmdLine;
+}
+
+
+/**
+ * Parses a string into tokens, separating switches (beginning with - or /) from
+ * other parameters
+ *
+ * @param	CmdLine		the string to parse
+ * @param	Tokens		[out] filled with all parameters found in the string
+ * @param	Switches	[out] filled with all switches found in the string
+ */
+void FCommandLine::Parse(const TCHAR* InCmdLine, TArray<FString>& Tokens, TArray<FString>& Switches)
+{
+	FString NextToken;
+	while (FParse::Token(InCmdLine, NextToken, false))
+	{
+		if ((**NextToken == TCHAR('-')) || (**NextToken == TCHAR('/')))
+		{
+			new(Switches) FString(NextToken.Mid(1));
+			new(Tokens) FString(NextToken.Right(NextToken.Len() - 1));
+		}
+		else
+		{
+			new(Tokens) FString(NextToken);
+		}
+	}
+}
diff --git a/Engine/Source/Runtime/Core/Private/Misc/CoreDelegates.cpp b/Engine/Source/Runtime/Core/Private/Misc/CoreDelegates.cpp
new file mode 100644
--- /dev/null
+++ b/Engine/Source/Runtime/Core/Private/Misc/CoreDelegates.cpp
@@ -0,0 +1,96 @@
+// Copyright 1998-2016 Epic Games, Inc. All Rights Reserved.
+
+// Core includes.
+#include "CorePrivatePCH.h"
+
+const FVector DefaultRefVector(0);
+
+
+//////////////////////////////////////////////////////////////////////////
+// FCoreDelegates
+
+TArray<FCoreDelegates::FHotFixDelegate> FCoreDelegates::HotFixDelegates;
+TArray<FCoreDelegates::FResolvePackageNameDelegate> FCoreDelegates::PackageNameResolvers;
+
+FCoreDelegates::FHotFixDelegate& FCoreDelegates::GetHotfixDelegate(EHotfixDelegates::Type HotFix)
+{
+	if (HotFix >= HotFixDelegates.Num())
+	{
+		HotFixDelegates.SetNum(HotFix + 1);
+	}
+	return HotFixDelegates[HotFix];
+}
+
+FCoreDelegates::FObjectArrayForDebugVisualizersDelegate& FCoreDelegates::GetObjectArrayForDebugVisualizersDelegate()
+{
+	static FObjectArrayForDebugVisualizersDelegate StaticDelegate;
+	return StaticDelegate;
+}
+
+FCoreDelegates::FOnPreMainInit& FCoreDelegates::GetPreMainInitDelegate()
+{
+	static FOnPreMainInit StaticDelegate;
+	return StaticDelegate;
+}
+
+FCoreDelegates::FOnMountPak FCoreDelegates::OnMountPak;
+FCoreDelegates::FOnUnmountPak FCoreDelegates::OnUnmountPak;
+FCoreDelegates::FOnUserLoginChangedEvent FCoreDelegates::OnUserLoginChangedEvent; 
+FCoreDelegates::FOnUserControllerConnectionChange FCoreDelegates::OnControllerConnectionChange;
+FCoreDelegates::FOnSafeFrameChangedEvent FCoreDelegates::OnSafeFrameChangedEvent;
+FCoreDelegates::FOnHandleSystemEnsure FCoreDelegates::OnHandleSystemEnsure;
+FCoreDelegates::FOnHandleSystemError FCoreDelegates::OnHandleSystemError;
+FCoreDelegates::FOnActorLabelChanged FCoreDelegates::OnActorLabelChanged;
+
+#if WITH_EDITOR
+	FSimpleMulticastDelegate FCoreDelegates::PreModal;
+	FSimpleMulticastDelegate FCoreDelegates::PostModal;
+#endif	//WITH_EDITOR
+FSimpleMulticastDelegate FCoreDelegates::OnShutdownAfterError;
+FSimpleMulticastDelegate FCoreDelegates::OnInit;
+FSimpleMulticastDelegate FCoreDelegates::OnExit;
+FSimpleMulticastDelegate FCoreDelegates::OnPreExit;
+FSimpleMulticastDelegate FCoreDelegates::ColorPickerChanged;
+FCoreDelegates::FOnModalMessageBox FCoreDelegates::ModalErrorMessage;
+FCoreDelegates::FOnInviteAccepted FCoreDelegates::OnInviteAccepted;
+FCoreDelegates::FWorldOriginOffset FCoreDelegates::PreWorldOriginOffset;
+FCoreDelegates::FWorldOriginOffset FCoreDelegates::PostWorldOriginOffset;
+FCoreDelegates::FStarvedGameLoop FCoreDelegates::StarvedGameLoop;
+
+FCoreDelegates::FApplicationLifetimeDelegate FCoreDelegates::ApplicationWillDeactivateDelegate;
+FCoreDelegates::FApplicationLifetimeDelegate FCoreDelegates::ApplicationHasReactivatedDelegate;
+FCoreDelegates::FApplicationLifetimeDelegate FCoreDelegates::ApplicationWillEnterBackgroundDelegate;
+FCoreDelegates::FApplicationLifetimeDelegate FCoreDelegates::ApplicationHasEnteredForegroundDelegate;
+FCoreDelegates::FApplicationLifetimeDelegate FCoreDelegates::ApplicationWillTerminateDelegate; 
+
+FCoreDelegates::FApplicationRegisteredForRemoteNotificationsDelegate FCoreDelegates::ApplicationRegisteredForRemoteNotificationsDelegate;
+FCoreDelegates::FApplicationRegisteredForUserNotificationsDelegate FCoreDelegates::ApplicationRegisteredForUserNotificationsDelegate;
+FCoreDelegates::FApplicationFailedToRegisterForRemoteNotificationsDelegate FCoreDelegates::ApplicationFailedToRegisterForRemoteNotificationsDelegate;
+FCoreDelegates::FApplicationReceivedRemoteNotificationDelegate FCoreDelegates::ApplicationReceivedRemoteNotificationDelegate;
+FCoreDelegates::FApplicationReceivedLocalNotificationDelegate FCoreDelegates::ApplicationReceivedLocalNotificationDelegate;
+
+FCoreDelegates::FStatCheckEnabled FCoreDelegates::StatCheckEnabled;
+FCoreDelegates::FStatEnabled FCoreDelegates::StatEnabled;
+FCoreDelegates::FStatDisabled FCoreDelegates::StatDisabled;
+FCoreDelegates::FStatDisableAll FCoreDelegates::StatDisableAll;
+
+FCoreDelegates::FApplicationLicenseChange FCoreDelegates::ApplicationLicenseChange;
+FCoreDelegates::FPlatformChangedLaptopMode FCoreDelegates::PlatformChangedLaptopMode;
+
+FCoreDelegates::FLoadStringAssetReferenceInCook FCoreDelegates::LoadStringAssetReferenceInCook;
+
+FCoreDelegates::FVRHeadsetRecenter FCoreDelegates::VRHeadsetRecenter;
+FCoreDelegates::FVRHeadsetLost FCoreDelegates::VRHeadsetLost;
+FCoreDelegates::FVRHeadsetReconnected FCoreDelegates::VRHeadsetReconnected;
+
+FCoreDelegates::FOnUserActivityStringChanged FCoreDelegates::UserActivityStringChanged;
+FCoreDelegates::FOnGameSessionIDChange FCoreDelegates::GameSessionIDChanged;
+FCoreDelegates::FOnCrashOverrideParamsChanged FCoreDelegates::CrashOverrideParamsChanged;
+
+FCoreDelegates::FOnAsyncLoadingFlush FCoreDelegates::OnAsyncLoadingFlush;
+FCoreDelegates::FRenderingThreadChanged FCoreDelegates::PostRenderingThreadCreated;
+FCoreDelegates::FRenderingThreadChanged FCoreDelegates::PreRenderingThreadDestroyed;
+FSimpleMulticastDelegate FCoreDelegates::OnFEngineLoopInitComplete;
+FCoreDelegates::FApplicationReceivedOnScreenOrientationChangedNotificationDelegate FCoreDelegates::ApplicationReceivedScreenOrientationChangedNotificationDelegate;
+
+FCoreDelegates::FConfigReadyForUse FCoreDelegates::ConfigReadyForUse;
