Copying //UE4/Dev-Build to //UE4/Dev-Main (Source: //UE4/Dev-Build @ 3104200)

diff --git a/Engine/Source/Runtime/Core/Private/Misc/CoreMisc.cpp b/Engine/Source/Runtime/Core/Private/Misc/CoreMisc.cpp
--- a/Engine/Source/Runtime/Core/Private/Misc/CoreMisc.cpp
+++ b/Engine/Source/Runtime/Core/Private/Misc/CoreMisc.cpp
@@ -76,834 +76,6 @@ bool FStaticSelfRegisteringExec::Exec( UWorld* InWorld, const TCHAR* Cmd, FOutpu
 }
 
 
-
-/*-----------------------------------------------------------------------------
-	FFileHelper
------------------------------------------------------------------------------*/
-
-/**
- * Load a binary file to a dynamic array.
- */
-bool FFileHelper::LoadFileToArray( TArray<uint8>& Result, const TCHAR* Filename, uint32 Flags )
-{
-	FArchive* Reader = IFileManager::Get().CreateFileReader( Filename, Flags );
-	if( !Reader )
-	{
-		if (!(Flags & FILEREAD_Silent))
-		{
-			UE_LOG(LogStreaming,Warning,TEXT("Failed to read file '%s' error."),Filename);
-		}
-		return 0;
-	}
-	int64 TotalSize = Reader->TotalSize();
-	Result.Reset( TotalSize );
-	Result.AddUninitialized( TotalSize );
-	Reader->Serialize(Result.GetData(), Result.Num());
-	bool Success = Reader->Close();
-	delete Reader;
-	return Success;
-}
-
-/**
- * Converts an arbitrary text buffer to an FString.
- * Supports all combination of ANSI/Unicode files and platforms.
- */
-void FFileHelper::BufferToString( FString& Result, const uint8* Buffer, int32 Size )
-{
-	TArray<TCHAR>& ResultArray = Result.GetCharArray();
-	ResultArray.Empty();
-
-	if( Size >= 2 && !( Size & 1 ) && Buffer[0] == 0xff && Buffer[1] == 0xfe )
-	{
-		// Unicode Intel byte order. Less 1 for the FFFE header, additional 1 for null terminator.
-		ResultArray.AddUninitialized( Size / 2 );
-		for( int32 i = 0; i < ( Size / 2 ) - 1; i++ )
-		{
-			ResultArray[ i ] = CharCast<TCHAR>( (UCS2CHAR)(( uint16 )Buffer[i * 2 + 2] + ( uint16 )Buffer[i * 2 + 3] * 256) );
-		}
-	}
-	else if( Size >= 2 && !( Size & 1 ) && Buffer[0] == 0xfe && Buffer[1] == 0xff )
-	{
-		// Unicode non-Intel byte order. Less 1 for the FFFE header, additional 1 for null terminator.
-		ResultArray.AddUninitialized( Size / 2 );
-		for( int32 i = 0; i < ( Size / 2 ) - 1; i++ )
-		{
-			ResultArray[ i ] = CharCast<TCHAR>( (UCS2CHAR)(( uint16 )Buffer[i * 2 + 3] + ( uint16 )Buffer[i * 2 + 2] * 256) );
-		}
-	}
-	else
-	{
-		if ( Size >= 3 && Buffer[0] == 0xef && Buffer[1] == 0xbb && Buffer[2] == 0xbf )
-		{
-			// Skip over UTF-8 BOM if there is one
-			Buffer += 3;
-			Size   -= 3;
-		}
-
-		FUTF8ToTCHAR Conv((const ANSICHAR*)Buffer, Size);
-		int32 Length = Conv.Length();
-		ResultArray.AddUninitialized(Length + 1); // For the null terminator
-		CopyAssignItems(ResultArray.GetData(), Conv.Get(), Length);
-	}
-
-	if (ResultArray.Num() == 1)
-	{
-		// If it's only a zero terminator then make the result actually empty
-		ResultArray.Empty();
-	}
-	else
-	{
-		// Else ensure null terminator is present
-		ResultArray.Last() = 0;
-	}
-}
-
-/**
- * Load a text file to an FString.
- * Supports all combination of ANSI/Unicode files and platforms.
- * @param Result string representation of the loaded file
- * @param Filename name of the file to load
- * @param VerifyFlags flags controlling the hash verification behavior ( see EHashOptions )
- */
-bool FFileHelper::LoadFileToString( FString& Result, const TCHAR* Filename, uint32 VerifyFlags )
-{
-	TUniquePtr<FArchive> Reader( IFileManager::Get().CreateFileReader( Filename ) );
-	if( !Reader )
-	{
-		return 0;
-	}
-	
-	int32 Size = Reader->TotalSize();
-	if( !Size )
-	{
-		Result.Empty();
-		return true;
-	}
-
-	uint8* Ch = (uint8*)FMemory::Malloc(Size);
-	Reader->Serialize( Ch, Size );
-	bool Success = Reader->Close();
-	Reader = nullptr;
-	BufferToString( Result, Ch, Size );
-
-	// handle SHA verify of the file
-	if( (VerifyFlags & EHashOptions::EnableVerify) && ( (VerifyFlags & EHashOptions::ErrorMissingHash) || FSHA1::GetFileSHAHash(Filename, NULL) ) )
-	{
-		// kick off SHA verify task. this frees the buffer on close
-		FBufferReaderWithSHA Ar( Ch, Size, true, Filename, false, true );
-	}
-	else
-	{
-		// free manually since not running SHA task
-		FMemory::Free(Ch);
-	}
-
-	return Success;
-}
-
-/**
- * Save a binary array to a file.
- */
-bool FFileHelper::SaveArrayToFile(TArrayView<const uint8> Array, const TCHAR* Filename, IFileManager* FileManager /*= &IFileManager::Get()*/, uint32 WriteFlags)
-{
-	FArchive* Ar = FileManager->CreateFileWriter( Filename, WriteFlags );
-	if( !Ar )
-	{
-		return 0;
-	}
-	Ar->Serialize(const_cast<uint8*>(Array.GetData()), Array.Num());
-	delete Ar;
-	return true;
-}
-
-/**
- * Write the FString to a file.
- * Supports all combination of ANSI/Unicode files and platforms.
- */
-bool FFileHelper::SaveStringToFile( const FString& String, const TCHAR* Filename,  EEncodingOptions::Type EncodingOptions, IFileManager* FileManager /*= &IFileManager::Get()*/, uint32 WriteFlags )
-{
-	// max size of the string is a UCS2CHAR for each character and some UNICODE magic 
-	auto Ar = TUniquePtr<FArchive>( FileManager->CreateFileWriter( Filename, WriteFlags ) );
-	if( !Ar )
-		return false;
-
-	if( String.IsEmpty() )
-		return true;
-
-	const TCHAR* StrPtr = *String;
-
-	bool SaveAsUnicode = EncodingOptions == EEncodingOptions::ForceUnicode || ( EncodingOptions == EEncodingOptions::AutoDetect && !FCString::IsPureAnsi(StrPtr) );
-	if( EncodingOptions == EEncodingOptions::ForceUTF8 )
-	{
-		UTF8CHAR UTF8BOM[] = { 0xEF, 0xBB, 0xBF };
-		Ar->Serialize( &UTF8BOM, ARRAY_COUNT(UTF8BOM) * sizeof(UTF8CHAR) );
-
-		FTCHARToUTF8 UTF8String(StrPtr);
-		Ar->Serialize( (UTF8CHAR*)UTF8String.Get(), UTF8String.Length() * sizeof(UTF8CHAR) );
-	}
-	else if ( EncodingOptions == EEncodingOptions::ForceUTF8WithoutBOM )
-	{
-		FTCHARToUTF8 UTF8String(StrPtr);
-		Ar->Serialize((UTF8CHAR*)UTF8String.Get(), UTF8String.Length() * sizeof(UTF8CHAR));
-	}
-	else if (SaveAsUnicode)
-	{
-		UCS2CHAR BOM = UNICODE_BOM;
-		Ar->Serialize( &BOM, sizeof(UCS2CHAR) );
-
-		auto Src = StringCast<UCS2CHAR>(StrPtr, String.Len());
-		Ar->Serialize( (UCS2CHAR*)Src.Get(), Src.Length() * sizeof(UCS2CHAR) );
-	}
-	else
-	{
-		auto Src = StringCast<ANSICHAR>(StrPtr, String.Len());
-		Ar->Serialize( (ANSICHAR*)Src.Get(), Src.Length() * sizeof(ANSICHAR) );
-	}
-
-	return true;
-}
-
-/**
- * Generates the next unique bitmap filename with a specified extension
- * 
- * @param Pattern		Filename with path, but without extension.
- * @oaran Extension		File extension to be appended
- * @param OutFilename	Reference to an FString where the newly generated filename will be placed
- * @param FileManager	Reference to a IFileManager (or the global instance by default)
- *
- * @return true if success
- */
-bool FFileHelper::GenerateNextBitmapFilename( const FString& Pattern, const FString& Extension, FString& OutFilename, IFileManager* FileManager /*= &IFileManager::Get()*/ )
-{
-	FString File;
-	OutFilename = "";
-	bool bSuccess = false;
-
-	for( int32 TestBitmapIndex = GScreenshotBitmapIndex + 1; TestBitmapIndex < 100000; ++TestBitmapIndex )
-	{
-		File = FString::Printf(TEXT("%s%05i.%s"), *Pattern, TestBitmapIndex, *Extension);
-		if( FileManager->FileSize(*File) < 0 )
-		{
-			GScreenshotBitmapIndex = TestBitmapIndex;
-			OutFilename = File;
-			bSuccess = true;
-			break;
-		}
-	}
-
-	return bSuccess;
-}
-
-/**
- * Saves a 24Bit BMP file to disk
- * 
- * @param Pattern filename with path, must not be 0, if with "bmp" extension (e.g. "out.bmp") the filename stays like this, if without (e.g. "out") automatic index numbers are addended (e.g. "out00002.bmp")
- * @param DataWidth - Width of the bitmap supplied in Data >0
- * @param DataHeight - Height of the bitmap supplied in Data >0
- * @param Data must not be 0
- * @param SubRectangle optional, specifies a sub-rectangle of the source image to save out. If NULL, the whole bitmap is saved
- * @param FileManager must not be 0
- * @param OutFilename optional, if specified filename will be output
- *
- * @return true if success
- */
-bool FFileHelper::CreateBitmap( const TCHAR* Pattern, int32 SourceWidth, int32 SourceHeight, const FColor* Data, struct FIntRect* SubRectangle, IFileManager* FileManager /*= &IFileManager::Get()*/, FString* OutFilename /*= NULL*/, bool bInWriteAlpha /*= false*/ )
-{
-#if ALLOW_DEBUG_FILES
-	FIntRect Src(0, 0, SourceWidth, SourceHeight);
-	if (SubRectangle == NULL || SubRectangle->Area() == 0)
-	{
-		SubRectangle = &Src;
-	}
-
-	FString File;
-	// if the Pattern already has a .bmp extension, then use that the file to write to
-	if (FPaths::GetExtension(Pattern) == TEXT("bmp"))
-	{
-		File = Pattern;
-	}
-	else
-	{
-		if (GenerateNextBitmapFilename(Pattern, TEXT("bmp"), File, FileManager))
-		{
-			if ( OutFilename )
-			{
-				*OutFilename = File;
-			}
-		}
-		else
-		{
-			return false;
-		}
-	}
-
-	FArchive* Ar = FileManager->CreateDebugFileWriter( *File );
-	if( Ar )
-	{
-		// Types.
-		#if PLATFORM_SUPPORTS_PRAGMA_PACK
-			#pragma pack (push,1)
-		#endif
-		struct BITMAPFILEHEADER
-		{
-			uint16 bfType GCC_PACK(1);
-			uint32 bfSize GCC_PACK(1);
-			uint16 bfReserved1 GCC_PACK(1); 
-			uint16 bfReserved2 GCC_PACK(1);
-			uint32 bfOffBits GCC_PACK(1);
-		} FH; 
-		struct BITMAPINFOHEADER
-		{
-			uint32 biSize GCC_PACK(1); 
-			int32  biWidth GCC_PACK(1);
-			int32  biHeight GCC_PACK(1);
-			uint16 biPlanes GCC_PACK(1);
-			uint16 biBitCount GCC_PACK(1);
-			uint32 biCompression GCC_PACK(1);
-			uint32 biSizeImage GCC_PACK(1);
-			int32  biXPelsPerMeter GCC_PACK(1); 
-			int32  biYPelsPerMeter GCC_PACK(1);
-			uint32 biClrUsed GCC_PACK(1);
-			uint32 biClrImportant GCC_PACK(1); 
-		} IH;
-		struct BITMAPV4HEADER
-		{
-			uint32 bV4RedMask GCC_PACK(1);
-			uint32 bV4GreenMask GCC_PACK(1);
-			uint32 bV4BlueMask GCC_PACK(1);
-			uint32 bV4AlphaMask GCC_PACK(1);
-			uint32 bV4CSType GCC_PACK(1);
-			uint32 bV4EndpointR[3] GCC_PACK(1);
-			uint32 bV4EndpointG[3] GCC_PACK(1);
-			uint32 bV4EndpointB[3] GCC_PACK(1);
-			uint32 bV4GammaRed GCC_PACK(1);
-			uint32 bV4GammaGreen GCC_PACK(1);
-			uint32 bV4GammaBlue GCC_PACK(1);
-		} IHV4;
-		#if PLATFORM_SUPPORTS_PRAGMA_PACK
-			#pragma pack (pop)
-		#endif
-
-		int32 Width = SubRectangle->Width();
-		int32 Height = SubRectangle->Height();
-		uint32 BytesPerPixel = bInWriteAlpha ? 4 : 3;
-		uint32 BytesPerLine = Align(Width * BytesPerPixel, 4);
-
-		uint32 InfoHeaderSize = sizeof(BITMAPINFOHEADER) + (bInWriteAlpha ? sizeof(BITMAPV4HEADER) : 0);
-
-		// File header.
-		FH.bfType       		= INTEL_ORDER16((uint16) ('B' + 256*'M'));
-		FH.bfSize       		= INTEL_ORDER32((uint32) (sizeof(BITMAPFILEHEADER) + InfoHeaderSize + BytesPerLine * Height));
-		FH.bfReserved1  		= INTEL_ORDER16((uint16) 0);
-		FH.bfReserved2  		= INTEL_ORDER16((uint16) 0);
-		FH.bfOffBits    		= INTEL_ORDER32((uint32) (sizeof(BITMAPFILEHEADER) + InfoHeaderSize));
-		Ar->Serialize( &FH, sizeof(FH) );
-
-		// Info header.
-		IH.biSize               = INTEL_ORDER32((uint32) InfoHeaderSize);
-		IH.biWidth              = INTEL_ORDER32((uint32) Width);
-		IH.biHeight             = INTEL_ORDER32((uint32) Height);
-		IH.biPlanes             = INTEL_ORDER16((uint16) 1);
-		IH.biBitCount           = INTEL_ORDER16((uint16) BytesPerPixel * 8);
-		if(bInWriteAlpha)
-		{
-			IH.biCompression    = INTEL_ORDER32((uint32) 3); //BI_BITFIELDS
-		}
-		else
-		{
-			IH.biCompression    = INTEL_ORDER32((uint32) 0); //BI_RGB
-		}
-		IH.biSizeImage          = INTEL_ORDER32((uint32) BytesPerLine * Height);
-		IH.biXPelsPerMeter      = INTEL_ORDER32((uint32) 0);
-		IH.biYPelsPerMeter      = INTEL_ORDER32((uint32) 0);
-		IH.biClrUsed            = INTEL_ORDER32((uint32) 0);
-		IH.biClrImportant       = INTEL_ORDER32((uint32) 0);
-		Ar->Serialize( &IH, sizeof(IH) );
-
-		// If we're writing alpha, we need to write the extra portion of the V4 header
-		if (bInWriteAlpha)
-		{
-			IHV4.bV4RedMask     = INTEL_ORDER32((uint32) 0x00ff0000);
-			IHV4.bV4GreenMask   = INTEL_ORDER32((uint32) 0x0000ff00);
-			IHV4.bV4BlueMask    = INTEL_ORDER32((uint32) 0x000000ff);
-			IHV4.bV4AlphaMask   = INTEL_ORDER32((uint32) 0xff000000);
-			IHV4.bV4CSType      = INTEL_ORDER32((uint32) 'Win ');
-			IHV4.bV4GammaRed    = INTEL_ORDER32((uint32) 0);
-			IHV4.bV4GammaGreen  = INTEL_ORDER32((uint32) 0);
-			IHV4.bV4GammaBlue   = INTEL_ORDER32((uint32) 0);
-			Ar->Serialize( &IHV4, sizeof(IHV4) );
-		}
-
-		// Colors.
-		for( int32 i = SubRectangle->Max.Y - 1; i >= SubRectangle->Min.Y; i-- )
-		{
-			for( int32 j = SubRectangle->Min.X; j < SubRectangle->Max.X; j++ )
-			{
-				Ar->Serialize( (void *)&Data[i*SourceWidth+j].B, 1 );
-				Ar->Serialize( (void *)&Data[i*SourceWidth+j].G, 1 );
-				Ar->Serialize( (void *)&Data[i*SourceWidth+j].R, 1 );
-
-				if (bInWriteAlpha)
-				{
-					Ar->Serialize( (void *)&Data[i * SourceWidth + j].A, 1 );
-				}
-			}
-
-			// Pad each row's length to be a multiple of 4 bytes.
-
-			for(uint32 PadIndex = Width * BytesPerPixel; PadIndex < BytesPerLine; PadIndex++)
-			{
-				uint8 B = 0;
-				Ar->Serialize(&B, 1);
-			}
-		}
-
-		// Success.
-		delete Ar;
-		if (!GIsEditor)
-		{
-			SendDataToPCViaUnrealConsole( TEXT("UE_PROFILER!BUGIT:"), File );
-		}
-	}
-	else 
-	{
-		return 0;
-	}
-#endif
-	// Success.
-	return 1;
-}
-
-/**
- *	Load the given ANSI text file to an array of strings - one FString per line of the file.
- *	Intended for use in simple text parsing actions
- *
- *	@param	InFilename			The text file to read, full path
- *	@param	InFileManager		The filemanager to use - NULL will use &IFileManager::Get()
- *	@param	OutStrings			The array of FStrings to fill in
- *
- *	@return	bool				true if successful, false if not
- */
-bool FFileHelper::LoadANSITextFileToStrings(const TCHAR* InFilename, IFileManager* InFileManager, TArray<FString>& OutStrings)
-{
-	IFileManager* FileManager = (InFileManager != NULL) ? InFileManager : &IFileManager::Get();
-	// Read and parse the file, adding the pawns and their sounds to the list
-	FArchive* TextFile = FileManager->CreateFileReader(InFilename, 0);
-	if (TextFile != NULL)
-	{
-		// get the size of the file
-		int32 Size = TextFile->TotalSize();
-		// read the file
-		TArray<uint8> Buffer;
-		Buffer.Empty(Size);
-		Buffer.AddUninitialized(Size);
-		TextFile->Serialize(Buffer.GetData(), Size);
-		// zero terminate it
-		Buffer.Add(0);
-		// Release the file
-		delete TextFile;
-
-		// Now read it
-		// init traveling pointer
-		ANSICHAR* Ptr = (ANSICHAR*)Buffer.GetData();
-
-		// iterate over the lines until complete
-		bool bIsDone = false;
-		while (!bIsDone)
-		{
-			// Store the location of the first character of this line
-			ANSICHAR* Start = Ptr;
-
-			// Advance the char pointer until we hit a newline character
-			while (*Ptr && *Ptr!='\r' && *Ptr!='\n')
-			{
-				Ptr++;
-			}
-
-			// If this is the end of the file, we're done
-			if (*Ptr == 0)
-			{
-				bIsDone = 1;
-			}
-			// Handle different line endings. If \r\n then NULL and advance 2, otherwise NULL and advance 1
-			// This handles \r, \n, or \r\n
-			else if ( *Ptr=='\r' && *(Ptr+1)=='\n' )
-			{
-				// This was \r\n. Terminate the current line, and advance the pointer forward 2 characters in the stream
-				*Ptr++ = 0;
-				*Ptr++ = 0;
-			}
-			else
-			{
-				// Terminate the current line, and advance the pointer to the next character in the stream
-				*Ptr++ = 0;
-			}
-
-			FString CurrLine = ANSI_TO_TCHAR(Start);
-			OutStrings.Add(CurrLine);
-		}
-
-		return true;
-	}
-	else
-	{
-		UE_LOG(LogStreaming, Warning, TEXT("Failed to open ANSI TEXT file %s"), InFilename);
-		return false;
-	}
-}
-
-/*-----------------------------------------------------------------------------
-	FCommandLine
------------------------------------------------------------------------------*/
-
-bool FCommandLine::bIsInitialized = false;
-TCHAR FCommandLine::CmdLine[FCommandLine::MaxCommandLineSize] = TEXT("");
-TCHAR FCommandLine::OriginalCmdLine[FCommandLine::MaxCommandLineSize] = TEXT("");
-TCHAR FCommandLine::LoggingCmdLine[FCommandLine::MaxCommandLineSize] = TEXT("");
-TCHAR FCommandLine::LoggingOriginalCmdLine[FCommandLine::MaxCommandLineSize] = TEXT("");
-FString FCommandLine::SubprocessCommandLine(TEXT(" -Multiprocess"));
-
-bool FCommandLine::IsInitialized()
-{
-	return bIsInitialized;
-}
-
-const TCHAR* FCommandLine::Get()
-{
-	UE_CLOG(!bIsInitialized, LogInit, Fatal, TEXT("Attempting to get the command line but it hasn't been initialized yet."));
-	return CmdLine;
-}
-
-const TCHAR* FCommandLine::GetForLogging()
-{
-	UE_CLOG(!bIsInitialized, LogInit, Fatal, TEXT("Attempting to get the command line but it hasn't been initialized yet."));
-	return LoggingCmdLine;
-}
-
-const TCHAR* FCommandLine::GetOriginal()
-{
-	UE_CLOG(!bIsInitialized, LogInit, Fatal, TEXT("Attempting to get the command line but it hasn't been initialized yet."));
-	return OriginalCmdLine;
-}
-
-const TCHAR* FCommandLine::GetOriginalForLogging()
-{
-	UE_CLOG(!bIsInitialized, LogInit, Fatal, TEXT("Attempting to get the command line but it hasn't been initialized yet."));
-	return LoggingOriginalCmdLine;
-}
-
-bool FCommandLine::Set(const TCHAR* NewCommandLine)
-{
-	if (!bIsInitialized)
-	{
-		FCString::Strncpy(OriginalCmdLine, NewCommandLine, ARRAY_COUNT(OriginalCmdLine));
-		FCString::Strncpy(LoggingOriginalCmdLine, NewCommandLine, ARRAY_COUNT(LoggingOriginalCmdLine));
-	}
-
-	FCString::Strncpy( CmdLine, NewCommandLine, ARRAY_COUNT(CmdLine) );
-	FCString::Strncpy(LoggingCmdLine, NewCommandLine, ARRAY_COUNT(LoggingCmdLine));
-	// If configured as part of the build, strip out any unapproved args
-	WhitelistCommandLines();
-
-	bIsInitialized = true;
-
-	// Check for the '-' that normal ones get converted to in Outlook. It's important to do it AFTER the command line is initialized
-	if (StringHasBadDashes(NewCommandLine))
-	{
-		FText ErrorMessage = FText::Format(NSLOCTEXT("Engine", "ComdLineHasInvalidChar", "Error: Command-line contains an invalid '-' character, likely pasted from an email.\nCmdline = {0}"), FText::FromString(NewCommandLine));
-#if !UE_BUILD_SHIPPING
-		FMessageDialog::Open(EAppMsgType::Ok, ErrorMessage);
-		return false;
-#else
-		UE_LOG(LogInit, Fatal, TEXT("%s"), *ErrorMessage.ToString());
-#endif
-	}
-
-	return true;
-}
-
-void FCommandLine::Append(const TCHAR* AppendString)
-{
-	FCString::Strncat( CmdLine, AppendString, ARRAY_COUNT(CmdLine) );
-	// If configured as part of the build, strip out any unapproved args
-	WhitelistCommandLines();
-}
-
-#if WANTS_COMMANDLINE_WHITELIST
-TArray<FString> FCommandLine::ApprovedArgs;
-TArray<FString> FCommandLine::FilterArgsForLogging;
-
-#ifdef OVERRIDE_COMMANDLINE_WHITELIST
-/**
- * When overriding this setting make sure that your define looks like the following in your .cs file:
- *
- *		OutCPPEnvironmentConfiguration.Definitions.Add("OVERRIDE_COMMANDLINE_WHITELIST=\"-arg1 -arg2 -arg3 -arg4\"");
- *
- * The important part is the \" as they quotes get stripped off by the compiler without them
- */
-const TCHAR* OverrideList = TEXT(OVERRIDE_COMMANDLINE_WHITELIST);
-#else
-// Default list most conservative restrictions
-const TCHAR* OverrideList = TEXT("-fullscreen /windowed");
-#endif
-
-#ifdef FILTER_COMMANDLINE_LOGGING
-/**
- * When overriding this setting make sure that your define looks like the following in your .cs file:
- *
- *		OutCPPEnvironmentConfiguration.Definitions.Add("FILTER_COMMANDLINE_LOGGING=\"-arg1 -arg2 -arg3 -arg4\"");
- *
- * The important part is the \" as they quotes get stripped off by the compiler without them
- */
-const TCHAR* FilterForLoggingList = TEXT(FILTER_COMMANDLINE_LOGGING);
-#else
-const TCHAR* FilterForLoggingList = TEXT("");
-#endif
-
-void FCommandLine::WhitelistCommandLines()
-{
-	if (ApprovedArgs.Num() == 0)
-	{
-		TArray<FString> Ignored;
-		FCommandLine::Parse(OverrideList, ApprovedArgs, Ignored);
-	}
-	if (FilterArgsForLogging.Num() == 0)
-	{
-		TArray<FString> Ignored;
-		FCommandLine::Parse(FilterForLoggingList, FilterArgsForLogging, Ignored);
-	}
-	// Process the original command line
-	TArray<FString> OriginalList = FilterCommandLine(OriginalCmdLine);
-	BuildWhitelistCommandLine(OriginalCmdLine, ARRAY_COUNT(OriginalCmdLine), OriginalList);
-	// Process the current command line
-	TArray<FString> CmdList = FilterCommandLine(CmdLine);
-	BuildWhitelistCommandLine(CmdLine, ARRAY_COUNT(CmdLine), CmdList);
-	// Process the command line for logging purposes
-	TArray<FString> LoggingCmdList = FilterCommandLineForLogging(LoggingCmdLine);
-	BuildWhitelistCommandLine(LoggingCmdLine, ARRAY_COUNT(LoggingCmdLine), LoggingCmdList);
-	// Process the original command line for logging purposes
-	TArray<FString> LoggingOriginalCmdList = FilterCommandLineForLogging(LoggingOriginalCmdLine);
-	BuildWhitelistCommandLine(LoggingOriginalCmdLine, ARRAY_COUNT(LoggingOriginalCmdLine), LoggingOriginalCmdList);
-}
-
-TArray<FString> FCommandLine::FilterCommandLine(TCHAR* CommandLine)
-{
-	TArray<FString> Ignored;
-	TArray<FString> ParsedList;
-	// Parse the command line list
-	FCommandLine::Parse(CommandLine, ParsedList, Ignored);
-	// Remove any that are not in our approved list
-	for (int32 Index = 0; Index < ParsedList.Num(); Index++)
-	{
-		bool bFound = false;
-		for (auto ApprovedArg : ApprovedArgs)
-		{
-			if (ParsedList[Index].StartsWith(ApprovedArg))
-			{
-				bFound = true;
-				break;
-			}
-		}
-		if (!bFound)
-		{
-			ParsedList.RemoveAt(Index);
-			Index--;
-		}
-	}
-	return ParsedList;
-}
-
-TArray<FString> FCommandLine::FilterCommandLineForLogging(TCHAR* CommandLine)
-{
-	TArray<FString> Ignored;
-	TArray<FString> ParsedList;
-	// Parse the command line list
-	FCommandLine::Parse(CommandLine, ParsedList, Ignored);
-	// Remove any that are not in our approved list
-	for (int32 Index = 0; Index < ParsedList.Num(); Index++)
-	{
-		for (auto Filter : FilterArgsForLogging)
-		{
-			if (ParsedList[Index].StartsWith(Filter))
-			{
-				ParsedList.RemoveAt(Index);
-				Index--;
-				break;
-			}
-		}
-	}
-	return ParsedList;
-}
-
-void FCommandLine::BuildWhitelistCommandLine(TCHAR* CommandLine, uint32 ArrayCount, const TArray<FString>& FilteredArgs)
-{
-	check(ArrayCount > 0);
-	// Zero the whole string
-	FMemory::Memzero(CommandLine, sizeof(TCHAR) * ArrayCount);
-
-	uint32 StartIndex = 0;
-	for (auto Arg : FilteredArgs)
-	{
-		if ((StartIndex + Arg.Len() + 2) < ArrayCount)
-		{
-			if (StartIndex != 0)
-			{
-				CommandLine[StartIndex++] = TEXT(' ');
-			}
-			CommandLine[StartIndex++] = TEXT('-');
-			FCString::Strncpy(&CommandLine[StartIndex], *Arg, ArrayCount - StartIndex);
-			StartIndex += Arg.Len();
-		}
-	}
-}
-#endif
-
-void FCommandLine::AddToSubprocessCommandline( const TCHAR* Param )
-{
-	check( Param != NULL );
-	if (Param[0] != ' ')
-	{
-		SubprocessCommandLine += TEXT(" ");
-	}
-	SubprocessCommandLine += Param;
-}
-
-const FString& FCommandLine::GetSubprocessCommandline()
-{
-	return SubprocessCommandLine;
-}
-
-/** 
- * Removes the executable name from a commandline, denoted by parentheses.
- */
-const TCHAR* FCommandLine::RemoveExeName(const TCHAR* InCmdLine)
-{
-	// Skip over executable that is in the commandline
-	if( *InCmdLine=='\"' )
-	{
-		InCmdLine++;
-		while( *InCmdLine && *InCmdLine!='\"' )
-		{
-			InCmdLine++;
-		}
-		if( *InCmdLine )
-		{
-			InCmdLine++;
-		}
-	}
-	while( *InCmdLine && *InCmdLine!=' ' )
-	{
-		InCmdLine++;
-	}
-	// skip over any spaces at the start, which Vista likes to toss in multiple
-	while (*InCmdLine == ' ')
-	{
-		InCmdLine++;
-	}
-	return InCmdLine;
-}
-
-
-/**
- * Parses a string into tokens, separating switches (beginning with - or /) from
- * other parameters
- *
- * @param	CmdLine		the string to parse
- * @param	Tokens		[out] filled with all parameters found in the string
- * @param	Switches	[out] filled with all switches found in the string
- */
-void FCommandLine::Parse(const TCHAR* InCmdLine, TArray<FString>& Tokens, TArray<FString>& Switches)
-{
-	FString NextToken;
-	while (FParse::Token(InCmdLine, NextToken, false))
-	{
-		if ((**NextToken == TCHAR('-')) || (**NextToken == TCHAR('/')))
-		{
-			new(Switches) FString(NextToken.Mid(1));
-			new(Tokens) FString(NextToken.Right(NextToken.Len() - 1));
-		}
-		else
-		{
-			new(Tokens) FString(NextToken);
-		}
-	}
-}
-
-
-/*-----------------------------------------------------------------------------
-	FMaintenance
------------------------------------------------------------------------------*/
-void FMaintenance::DeleteOldLogs()
-{
-	int32 PurgeLogsDays = -1; // -1 means don't delete old files
-	int32 MaxLogFilesOnDisk = -1; // -1 means keep all files
-	GConfig->GetInt(TEXT("LogFiles"), TEXT("PurgeLogsDays"), PurgeLogsDays, GEngineIni);
-	GConfig->GetInt(TEXT("LogFiles"), TEXT("MaxLogFilesOnDisk"), MaxLogFilesOnDisk, GEngineIni);
-	if (PurgeLogsDays >= 0 || MaxLogFilesOnDisk >= 0)
-	{
-		// get a list of files in the log dir
-		TArray<FString> Files;
-		IFileManager::Get().FindFiles(Files, *FString::Printf(TEXT("%s*.*"), *FPaths::GameLogDir()), true, false);
-		for (FString& Filename : Files)
-		{
-			Filename = FPaths::GameLogDir() / Filename;
-		}
-
-		struct FSortByDateNewestFirst
-		{
-			bool operator()(const FString& A, const FString& B) const
-			{
-				const FDateTime TimestampA = IFileManager::Get().GetTimeStamp(*A);
-				const FDateTime TimestampB = IFileManager::Get().GetTimeStamp(*B);
-				return TimestampB < TimestampA;
-			}
-		};
-		Files.Sort(FSortByDateNewestFirst());
-
-		// delete all those with the backup text in their name and that are older than the specified number of days
-		double MaxFileAgeSeconds = 60.0 * 60.0 * 24.0 * double(PurgeLogsDays);
-		for (int32 FileIndex = Files.Num() - 1; FileIndex >= 0; --FileIndex)
-		{
-			const FString& Filename = Files[FileIndex];
-			if (FOutputDeviceFile::IsBackupCopy(*Filename) && IFileManager::Get().GetFileAgeSeconds(*Filename) > MaxFileAgeSeconds)
-			{
-				UE_LOG(LogStreaming, Log, TEXT("Deleting old log file %s"), *Filename);
-				IFileManager::Get().Delete(*Filename);
-				Files.RemoveAt(FileIndex);
-			}
-		}
-
-		// If required, trim the number of files on disk
-		if (MaxLogFilesOnDisk >= 0 && Files.Num() > MaxLogFilesOnDisk)
-		{
-			for (int32 FileIndex = Files.Num() - 1; FileIndex >= 0 && Files.Num() > MaxLogFilesOnDisk; --FileIndex)
-			{
-				if (FOutputDeviceFile::IsBackupCopy(*Files[FileIndex]))
-				{
-					IFileManager::Get().Delete(*Files[FileIndex]);
-					Files.RemoveAt(FileIndex);
-				}
-			}
-		}
-
-		// Remove old UE4 crash contexts
-		TArray<FString> Directories;
-		IFileManager::Get().FindFiles( Directories, *FString::Printf( TEXT( "%s/UE4CC*" ), *FPaths::GameLogDir() ), false, true );
-
-		for (const FString& Dir : Directories)
-		{
-			const FString CrashContextDirectory = FPaths::GameLogDir() / Dir;
-			const FDateTime DirectoryAccessTime = IFileManager::Get().GetTimeStamp( *CrashContextDirectory );
-			if (FDateTime::Now() - DirectoryAccessTime > FTimespan::FromDays( PurgeLogsDays ))
-			{
-				UE_LOG( LogStreaming, Log, TEXT( "Deleting old crash context %s" ), *Dir );
-				IFileManager::Get().DeleteDirectory( *CrashContextDirectory, false, true );
-			}
-		}
-	}
-}
-
-
 /*-----------------------------------------------------------------------------
 	Module singletons.
 -----------------------------------------------------------------------------*/
diff --git a/Engine/Source/Runtime/Core/Private/Misc/Exec.cpp b/Engine/Source/Runtime/Core/Private/Misc/Exec.cpp
new file mode 100644
--- /dev/null
+++ b/Engine/Source/Runtime/Core/Private/Misc/Exec.cpp
@@ -0,0 +1,8 @@
+// Copyright 1998-2016 Epic Games, Inc. All Rights Reserved.
+
+#include "CorePrivatePCH.h"
+#include "Misc/Exec.h"
+
+FExec::~FExec()
+{
+}
diff --git a/Engine/Source/Runtime/Core/Private/Misc/FeedbackContext.cpp b/Engine/Source/Runtime/Core/Private/Misc/FeedbackContext.cpp
--- a/Engine/Source/Runtime/Core/Private/Misc/FeedbackContext.cpp
+++ b/Engine/Source/Runtime/Core/Private/Misc/FeedbackContext.cpp
@@ -3,6 +3,12 @@
 #include "CorePrivatePCH.h"
 #include "FeedbackContext.h"
 
+FFeedbackContext::FFeedbackContext()
+	: TreatWarningsAsErrors(0)
+	, ScopeStack(MakeShareable(new FSlowTaskStack))
+{
+}
+
 FFeedbackContext::~FFeedbackContext()
 {
 	ensureMsgf(LegacyAPIScopes.Num() == 0, TEXT("EndSlowTask has not been called for %d outstanding tasks"), LegacyAPIScopes.Num());
@@ -102,36 +108,3 @@ bool FFeedbackContext::IsPlayingInEditor() const
 {
 	return GIsPlayInEditorWorld;
 }
-
-void FSlowTask::MakeDialog(bool bShowCancelButton, bool bAllowInPIE)
-{
-	const bool bIsDisabledByPIE = Context.IsPlayingInEditor() && !bAllowInPIE;
-	const bool bIsDialogAllowed = bEnabled && !GIsSilent && !bIsDisabledByPIE && !IsRunningCommandlet() && IsInGameThread();
-	if (!GIsSlowTask && bIsDialogAllowed)
-	{
-		Context.StartSlowTask(GetCurrentMessage(), bShowCancelButton);
-		if (GIsSlowTask)
-		{
-			bCreatedDialog = true;
-		}
-	}
-}
-
-float FSlowTaskStack::GetProgressFraction(int32 Index) const
-{
-	const int32 StartIndex = Num() - 1;
-	const int32 EndIndex = Index;
-
-	float Progress = 0.f;
-	for (int32 CurrentIndex = StartIndex; CurrentIndex >= EndIndex; --CurrentIndex)
-	{
-		const FSlowTask* Scope = (*this)[CurrentIndex];
-		
-		const float ThisScopeCompleted = float(Scope->CompletedWork) / Scope->TotalAmountOfWork;
-		const float ThisScopeCurrentFrame = float(Scope->CurrentFrameScope) / Scope->TotalAmountOfWork;
-
-		Progress = ThisScopeCompleted + ThisScopeCurrentFrame*Progress;
-	}
-
-	return Progress;
-}
\ No newline at end of file
diff --git a/Engine/Source/Runtime/Core/Private/Misc/FileHelper.cpp b/Engine/Source/Runtime/Core/Private/Misc/FileHelper.cpp
new file mode 100644
--- /dev/null
+++ b/Engine/Source/Runtime/Core/Private/Misc/FileHelper.cpp
@@ -0,0 +1,554 @@
+// Copyright 1998-2016 Epic Games, Inc. All Rights Reserved.
+
+#include "CorePrivatePCH.h"
+#include "FileHelper.h"
+#include "SecureHash.h"
+
+/*-----------------------------------------------------------------------------
+	FFileHelper
+-----------------------------------------------------------------------------*/
+
+/**
+ * Load a binary file to a dynamic array.
+ */
+bool FFileHelper::LoadFileToArray( TArray<uint8>& Result, const TCHAR* Filename, uint32 Flags )
+{
+	FArchive* Reader = IFileManager::Get().CreateFileReader( Filename, Flags );
+	if( !Reader )
+	{
+		if (!(Flags & FILEREAD_Silent))
+		{
+			UE_LOG(LogStreaming,Warning,TEXT("Failed to read file '%s' error."),Filename);
+		}
+		return 0;
+	}
+	int64 TotalSize = Reader->TotalSize();
+	Result.Reset( TotalSize );
+	Result.AddUninitialized( TotalSize );
+	Reader->Serialize(Result.GetData(), Result.Num());
+	bool Success = Reader->Close();
+	delete Reader;
+	return Success;
+}
+
+/**
+ * Converts an arbitrary text buffer to an FString.
+ * Supports all combination of ANSI/Unicode files and platforms.
+ */
+void FFileHelper::BufferToString( FString& Result, const uint8* Buffer, int32 Size )
+{
+	TArray<TCHAR>& ResultArray = Result.GetCharArray();
+	ResultArray.Empty();
+
+	if( Size >= 2 && !( Size & 1 ) && Buffer[0] == 0xff && Buffer[1] == 0xfe )
+	{
+		// Unicode Intel byte order. Less 1 for the FFFE header, additional 1 for null terminator.
+		ResultArray.AddUninitialized( Size / 2 );
+		for( int32 i = 0; i < ( Size / 2 ) - 1; i++ )
+		{
+			ResultArray[ i ] = CharCast<TCHAR>( (UCS2CHAR)(( uint16 )Buffer[i * 2 + 2] + ( uint16 )Buffer[i * 2 + 3] * 256) );
+		}
+	}
+	else if( Size >= 2 && !( Size & 1 ) && Buffer[0] == 0xfe && Buffer[1] == 0xff )
+	{
+		// Unicode non-Intel byte order. Less 1 for the FFFE header, additional 1 for null terminator.
+		ResultArray.AddUninitialized( Size / 2 );
+		for( int32 i = 0; i < ( Size / 2 ) - 1; i++ )
+		{
+			ResultArray[ i ] = CharCast<TCHAR>( (UCS2CHAR)(( uint16 )Buffer[i * 2 + 3] + ( uint16 )Buffer[i * 2 + 2] * 256) );
+		}
+	}
+	else
+	{
+		if ( Size >= 3 && Buffer[0] == 0xef && Buffer[1] == 0xbb && Buffer[2] == 0xbf )
+		{
+			// Skip over UTF-8 BOM if there is one
+			Buffer += 3;
+			Size   -= 3;
+		}
+
+		FUTF8ToTCHAR Conv((const ANSICHAR*)Buffer, Size);
+		int32 Length = Conv.Length();
+		ResultArray.AddUninitialized(Length + 1); // For the null terminator
+		CopyAssignItems(ResultArray.GetData(), Conv.Get(), Length);
+	}
+
+	if (ResultArray.Num() == 1)
+	{
+		// If it's only a zero terminator then make the result actually empty
+		ResultArray.Empty();
+	}
+	else
+	{
+		// Else ensure null terminator is present
+		ResultArray.Last() = 0;
+	}
+}
+
+/**
+ * Load a text file to an FString.
+ * Supports all combination of ANSI/Unicode files and platforms.
+ * @param Result string representation of the loaded file
+ * @param Filename name of the file to load
+ * @param VerifyFlags flags controlling the hash verification behavior ( see EHashOptions )
+ */
+bool FFileHelper::LoadFileToString( FString& Result, const TCHAR* Filename, uint32 VerifyFlags )
+{
+	TUniquePtr<FArchive> Reader( IFileManager::Get().CreateFileReader( Filename ) );
+	if( !Reader )
+	{
+		return 0;
+	}
+	
+	int32 Size = Reader->TotalSize();
+	if( !Size )
+	{
+		Result.Empty();
+		return true;
+	}
+
+	uint8* Ch = (uint8*)FMemory::Malloc(Size);
+	Reader->Serialize( Ch, Size );
+	bool Success = Reader->Close();
+	Reader = nullptr;
+	BufferToString( Result, Ch, Size );
+
+	// handle SHA verify of the file
+	if( (VerifyFlags & EHashOptions::EnableVerify) && ( (VerifyFlags & EHashOptions::ErrorMissingHash) || FSHA1::GetFileSHAHash(Filename, NULL) ) )
+	{
+		// kick off SHA verify task. this frees the buffer on close
+		FBufferReaderWithSHA Ar( Ch, Size, true, Filename, false, true );
+	}
+	else
+	{
+		// free manually since not running SHA task
+		FMemory::Free(Ch);
+	}
+
+	return Success;
+}
+
+/**
+ * Save a binary array to a file.
+ */
+bool FFileHelper::SaveArrayToFile(TArrayView<const uint8> Array, const TCHAR* Filename, IFileManager* FileManager /*= &IFileManager::Get()*/, uint32 WriteFlags)
+{
+	FArchive* Ar = FileManager->CreateFileWriter( Filename, WriteFlags );
+	if( !Ar )
+	{
+		return 0;
+	}
+	Ar->Serialize(const_cast<uint8*>(Array.GetData()), Array.Num());
+	delete Ar;
+	return true;
+}
+
+/**
+ * Write the FString to a file.
+ * Supports all combination of ANSI/Unicode files and platforms.
+ */
+bool FFileHelper::SaveStringToFile( const FString& String, const TCHAR* Filename,  EEncodingOptions::Type EncodingOptions, IFileManager* FileManager /*= &IFileManager::Get()*/, uint32 WriteFlags )
+{
+	// max size of the string is a UCS2CHAR for each character and some UNICODE magic 
+	auto Ar = TUniquePtr<FArchive>( FileManager->CreateFileWriter( Filename, WriteFlags ) );
+	if( !Ar )
+		return false;
+
+	if( String.IsEmpty() )
+		return true;
+
+	const TCHAR* StrPtr = *String;
+
+	bool SaveAsUnicode = EncodingOptions == EEncodingOptions::ForceUnicode || ( EncodingOptions == EEncodingOptions::AutoDetect && !FCString::IsPureAnsi(StrPtr) );
+	if( EncodingOptions == EEncodingOptions::ForceUTF8 )
+	{
+		UTF8CHAR UTF8BOM[] = { 0xEF, 0xBB, 0xBF };
+		Ar->Serialize( &UTF8BOM, ARRAY_COUNT(UTF8BOM) * sizeof(UTF8CHAR) );
+
+		FTCHARToUTF8 UTF8String(StrPtr);
+		Ar->Serialize( (UTF8CHAR*)UTF8String.Get(), UTF8String.Length() * sizeof(UTF8CHAR) );
+	}
+	else if ( EncodingOptions == EEncodingOptions::ForceUTF8WithoutBOM )
+	{
+		FTCHARToUTF8 UTF8String(StrPtr);
+		Ar->Serialize((UTF8CHAR*)UTF8String.Get(), UTF8String.Length() * sizeof(UTF8CHAR));
+	}
+	else if (SaveAsUnicode)
+	{
+		UCS2CHAR BOM = UNICODE_BOM;
+		Ar->Serialize( &BOM, sizeof(UCS2CHAR) );
+
+		auto Src = StringCast<UCS2CHAR>(StrPtr, String.Len());
+		Ar->Serialize( (UCS2CHAR*)Src.Get(), Src.Length() * sizeof(UCS2CHAR) );
+	}
+	else
+	{
+		auto Src = StringCast<ANSICHAR>(StrPtr, String.Len());
+		Ar->Serialize( (ANSICHAR*)Src.Get(), Src.Length() * sizeof(ANSICHAR) );
+	}
+
+	return true;
+}
+
+/**
+ * Generates the next unique bitmap filename with a specified extension
+ * 
+ * @param Pattern		Filename with path, but without extension.
+ * @oaran Extension		File extension to be appended
+ * @param OutFilename	Reference to an FString where the newly generated filename will be placed
+ * @param FileManager	Reference to a IFileManager (or the global instance by default)
+ *
+ * @return true if success
+ */
+bool FFileHelper::GenerateNextBitmapFilename( const FString& Pattern, const FString& Extension, FString& OutFilename, IFileManager* FileManager /*= &IFileManager::Get()*/ )
+{
+	FString File;
+	OutFilename = "";
+	bool bSuccess = false;
+
+	for( int32 TestBitmapIndex = GScreenshotBitmapIndex + 1; TestBitmapIndex < 100000; ++TestBitmapIndex )
+	{
+		File = FString::Printf(TEXT("%s%05i.%s"), *Pattern, TestBitmapIndex, *Extension);
+		if( FileManager->FileSize(*File) < 0 )
+		{
+			GScreenshotBitmapIndex = TestBitmapIndex;
+			OutFilename = File;
+			bSuccess = true;
+			break;
+		}
+	}
+
+	return bSuccess;
+}
+
+/**
+ * Saves a 24Bit BMP file to disk
+ * 
+ * @param Pattern filename with path, must not be 0, if with "bmp" extension (e.g. "out.bmp") the filename stays like this, if without (e.g. "out") automatic index numbers are addended (e.g. "out00002.bmp")
+ * @param DataWidth - Width of the bitmap supplied in Data >0
+ * @param DataHeight - Height of the bitmap supplied in Data >0
+ * @param Data must not be 0
+ * @param SubRectangle optional, specifies a sub-rectangle of the source image to save out. If NULL, the whole bitmap is saved
+ * @param FileManager must not be 0
+ * @param OutFilename optional, if specified filename will be output
+ *
+ * @return true if success
+ */
+bool FFileHelper::CreateBitmap( const TCHAR* Pattern, int32 SourceWidth, int32 SourceHeight, const FColor* Data, struct FIntRect* SubRectangle, IFileManager* FileManager /*= &IFileManager::Get()*/, FString* OutFilename /*= NULL*/, bool bInWriteAlpha /*= false*/ )
+{
+#if ALLOW_DEBUG_FILES
+	FIntRect Src(0, 0, SourceWidth, SourceHeight);
+	if (SubRectangle == NULL || SubRectangle->Area() == 0)
+	{
+		SubRectangle = &Src;
+	}
+
+	FString File;
+	// if the Pattern already has a .bmp extension, then use that the file to write to
+	if (FPaths::GetExtension(Pattern) == TEXT("bmp"))
+	{
+		File = Pattern;
+	}
+	else
+	{
+		if (GenerateNextBitmapFilename(Pattern, TEXT("bmp"), File, FileManager))
+		{
+			if ( OutFilename )
+			{
+				*OutFilename = File;
+			}
+		}
+		else
+		{
+			return false;
+		}
+	}
+
+	FArchive* Ar = FileManager->CreateDebugFileWriter( *File );
+	if( Ar )
+	{
+		// Types.
+		#if PLATFORM_SUPPORTS_PRAGMA_PACK
+			#pragma pack (push,1)
+		#endif
+		struct BITMAPFILEHEADER
+		{
+			uint16 bfType GCC_PACK(1);
+			uint32 bfSize GCC_PACK(1);
+			uint16 bfReserved1 GCC_PACK(1); 
+			uint16 bfReserved2 GCC_PACK(1);
+			uint32 bfOffBits GCC_PACK(1);
+		} FH; 
+		struct BITMAPINFOHEADER
+		{
+			uint32 biSize GCC_PACK(1); 
+			int32  biWidth GCC_PACK(1);
+			int32  biHeight GCC_PACK(1);
+			uint16 biPlanes GCC_PACK(1);
+			uint16 biBitCount GCC_PACK(1);
+			uint32 biCompression GCC_PACK(1);
+			uint32 biSizeImage GCC_PACK(1);
+			int32  biXPelsPerMeter GCC_PACK(1); 
+			int32  biYPelsPerMeter GCC_PACK(1);
+			uint32 biClrUsed GCC_PACK(1);
+			uint32 biClrImportant GCC_PACK(1); 
+		} IH;
+		struct BITMAPV4HEADER
+		{
+			uint32 bV4RedMask GCC_PACK(1);
+			uint32 bV4GreenMask GCC_PACK(1);
+			uint32 bV4BlueMask GCC_PACK(1);
+			uint32 bV4AlphaMask GCC_PACK(1);
+			uint32 bV4CSType GCC_PACK(1);
+			uint32 bV4EndpointR[3] GCC_PACK(1);
+			uint32 bV4EndpointG[3] GCC_PACK(1);
+			uint32 bV4EndpointB[3] GCC_PACK(1);
+			uint32 bV4GammaRed GCC_PACK(1);
+			uint32 bV4GammaGreen GCC_PACK(1);
+			uint32 bV4GammaBlue GCC_PACK(1);
+		} IHV4;
+		#if PLATFORM_SUPPORTS_PRAGMA_PACK
+			#pragma pack (pop)
+		#endif
+
+		int32 Width = SubRectangle->Width();
+		int32 Height = SubRectangle->Height();
+		uint32 BytesPerPixel = bInWriteAlpha ? 4 : 3;
+		uint32 BytesPerLine = Align(Width * BytesPerPixel, 4);
+
+		uint32 InfoHeaderSize = sizeof(BITMAPINFOHEADER) + (bInWriteAlpha ? sizeof(BITMAPV4HEADER) : 0);
+
+		// File header.
+		FH.bfType       		= INTEL_ORDER16((uint16) ('B' + 256*'M'));
+		FH.bfSize       		= INTEL_ORDER32((uint32) (sizeof(BITMAPFILEHEADER) + InfoHeaderSize + BytesPerLine * Height));
+		FH.bfReserved1  		= INTEL_ORDER16((uint16) 0);
+		FH.bfReserved2  		= INTEL_ORDER16((uint16) 0);
+		FH.bfOffBits    		= INTEL_ORDER32((uint32) (sizeof(BITMAPFILEHEADER) + InfoHeaderSize));
+		Ar->Serialize( &FH, sizeof(FH) );
+
+		// Info header.
+		IH.biSize               = INTEL_ORDER32((uint32) InfoHeaderSize);
+		IH.biWidth              = INTEL_ORDER32((uint32) Width);
+		IH.biHeight             = INTEL_ORDER32((uint32) Height);
+		IH.biPlanes             = INTEL_ORDER16((uint16) 1);
+		IH.biBitCount           = INTEL_ORDER16((uint16) BytesPerPixel * 8);
+		if(bInWriteAlpha)
+		{
+			IH.biCompression    = INTEL_ORDER32((uint32) 3); //BI_BITFIELDS
+		}
+		else
+		{
+			IH.biCompression    = INTEL_ORDER32((uint32) 0); //BI_RGB
+		}
+		IH.biSizeImage          = INTEL_ORDER32((uint32) BytesPerLine * Height);
+		IH.biXPelsPerMeter      = INTEL_ORDER32((uint32) 0);
+		IH.biYPelsPerMeter      = INTEL_ORDER32((uint32) 0);
+		IH.biClrUsed            = INTEL_ORDER32((uint32) 0);
+		IH.biClrImportant       = INTEL_ORDER32((uint32) 0);
+		Ar->Serialize( &IH, sizeof(IH) );
+
+		// If we're writing alpha, we need to write the extra portion of the V4 header
+		if (bInWriteAlpha)
+		{
+			IHV4.bV4RedMask     = INTEL_ORDER32((uint32) 0x00ff0000);
+			IHV4.bV4GreenMask   = INTEL_ORDER32((uint32) 0x0000ff00);
+			IHV4.bV4BlueMask    = INTEL_ORDER32((uint32) 0x000000ff);
+			IHV4.bV4AlphaMask   = INTEL_ORDER32((uint32) 0xff000000);
+			IHV4.bV4CSType      = INTEL_ORDER32((uint32) 'Win ');
+			IHV4.bV4GammaRed    = INTEL_ORDER32((uint32) 0);
+			IHV4.bV4GammaGreen  = INTEL_ORDER32((uint32) 0);
+			IHV4.bV4GammaBlue   = INTEL_ORDER32((uint32) 0);
+			Ar->Serialize( &IHV4, sizeof(IHV4) );
+		}
+
+		// Colors.
+		for( int32 i = SubRectangle->Max.Y - 1; i >= SubRectangle->Min.Y; i-- )
+		{
+			for( int32 j = SubRectangle->Min.X; j < SubRectangle->Max.X; j++ )
+			{
+				Ar->Serialize( (void *)&Data[i*SourceWidth+j].B, 1 );
+				Ar->Serialize( (void *)&Data[i*SourceWidth+j].G, 1 );
+				Ar->Serialize( (void *)&Data[i*SourceWidth+j].R, 1 );
+
+				if (bInWriteAlpha)
+				{
+					Ar->Serialize( (void *)&Data[i * SourceWidth + j].A, 1 );
+				}
+			}
+
+			// Pad each row's length to be a multiple of 4 bytes.
+
+			for(uint32 PadIndex = Width * BytesPerPixel; PadIndex < BytesPerLine; PadIndex++)
+			{
+				uint8 B = 0;
+				Ar->Serialize(&B, 1);
+			}
+		}
+
+		// Success.
+		delete Ar;
+		if (!GIsEditor)
+		{
+			SendDataToPCViaUnrealConsole( TEXT("UE_PROFILER!BUGIT:"), File );
+		}
+	}
+	else 
+	{
+		return 0;
+	}
+#endif
+	// Success.
+	return 1;
+}
+
+/**
+ *	Load the given ANSI text file to an array of strings - one FString per line of the file.
+ *	Intended for use in simple text parsing actions
+ *
+ *	@param	InFilename			The text file to read, full path
+ *	@param	InFileManager		The filemanager to use - NULL will use &IFileManager::Get()
+ *	@param	OutStrings			The array of FStrings to fill in
+ *
+ *	@return	bool				true if successful, false if not
+ */
+bool FFileHelper::LoadANSITextFileToStrings(const TCHAR* InFilename, IFileManager* InFileManager, TArray<FString>& OutStrings)
+{
+	IFileManager* FileManager = (InFileManager != NULL) ? InFileManager : &IFileManager::Get();
+	// Read and parse the file, adding the pawns and their sounds to the list
+	FArchive* TextFile = FileManager->CreateFileReader(InFilename, 0);
+	if (TextFile != NULL)
+	{
+		// get the size of the file
+		int32 Size = TextFile->TotalSize();
+		// read the file
+		TArray<uint8> Buffer;
+		Buffer.Empty(Size);
+		Buffer.AddUninitialized(Size);
+		TextFile->Serialize(Buffer.GetData(), Size);
+		// zero terminate it
+		Buffer.Add(0);
+		// Release the file
+		delete TextFile;
+
+		// Now read it
+		// init traveling pointer
+		ANSICHAR* Ptr = (ANSICHAR*)Buffer.GetData();
+
+		// iterate over the lines until complete
+		bool bIsDone = false;
+		while (!bIsDone)
+		{
+			// Store the location of the first character of this line
+			ANSICHAR* Start = Ptr;
+
+			// Advance the char pointer until we hit a newline character
+			while (*Ptr && *Ptr!='\r' && *Ptr!='\n')
+			{
+				Ptr++;
+			}
+
+			// If this is the end of the file, we're done
+			if (*Ptr == 0)
+			{
+				bIsDone = 1;
+			}
+			// Handle different line endings. If \r\n then NULL and advance 2, otherwise NULL and advance 1
+			// This handles \r, \n, or \r\n
+			else if ( *Ptr=='\r' && *(Ptr+1)=='\n' )
+			{
+				// This was \r\n. Terminate the current line, and advance the pointer forward 2 characters in the stream
+				*Ptr++ = 0;
+				*Ptr++ = 0;
+			}
+			else
+			{
+				// Terminate the current line, and advance the pointer to the next character in the stream
+				*Ptr++ = 0;
+			}
+
+			FString CurrLine = ANSI_TO_TCHAR(Start);
+			OutStrings.Add(CurrLine);
+		}
+
+		return true;
+	}
+	else
+	{
+		UE_LOG(LogStreaming, Warning, TEXT("Failed to open ANSI TEXT file %s"), InFilename);
+		return false;
+	}
+}
+
+/*-----------------------------------------------------------------------------
+	FMaintenance
+-----------------------------------------------------------------------------*/
+void FMaintenance::DeleteOldLogs()
+{
+	int32 PurgeLogsDays = -1; // -1 means don't delete old files
+	int32 MaxLogFilesOnDisk = -1; // -1 means keep all files
+	GConfig->GetInt(TEXT("LogFiles"), TEXT("PurgeLogsDays"), PurgeLogsDays, GEngineIni);
+	GConfig->GetInt(TEXT("LogFiles"), TEXT("MaxLogFilesOnDisk"), MaxLogFilesOnDisk, GEngineIni);
+	if (PurgeLogsDays >= 0 || MaxLogFilesOnDisk >= 0)
+	{
+		// get a list of files in the log dir
+		TArray<FString> Files;
+		IFileManager::Get().FindFiles(Files, *FString::Printf(TEXT("%s*.*"), *FPaths::GameLogDir()), true, false);
+		for (FString& Filename : Files)
+		{
+			Filename = FPaths::GameLogDir() / Filename;
+		}
+
+		struct FSortByDateNewestFirst
+		{
+			bool operator()(const FString& A, const FString& B) const
+			{
+				const FDateTime TimestampA = IFileManager::Get().GetTimeStamp(*A);
+				const FDateTime TimestampB = IFileManager::Get().GetTimeStamp(*B);
+				return TimestampB < TimestampA;
+			}
+		};
+		Files.Sort(FSortByDateNewestFirst());
+
+		// delete all those with the backup text in their name and that are older than the specified number of days
+		double MaxFileAgeSeconds = 60.0 * 60.0 * 24.0 * double(PurgeLogsDays);
+		for (int32 FileIndex = Files.Num() - 1; FileIndex >= 0; --FileIndex)
+		{
+			const FString& Filename = Files[FileIndex];
+			if (FOutputDeviceFile::IsBackupCopy(*Filename) && IFileManager::Get().GetFileAgeSeconds(*Filename) > MaxFileAgeSeconds)
+			{
+				UE_LOG(LogStreaming, Log, TEXT("Deleting old log file %s"), *Filename);
+				IFileManager::Get().Delete(*Filename);
+				Files.RemoveAt(FileIndex);
+			}
+		}
+
+		// If required, trim the number of files on disk
+		if (MaxLogFilesOnDisk >= 0 && Files.Num() > MaxLogFilesOnDisk)
+		{
+			for (int32 FileIndex = Files.Num() - 1; FileIndex >= 0 && Files.Num() > MaxLogFilesOnDisk; --FileIndex)
+			{
+				if (FOutputDeviceFile::IsBackupCopy(*Files[FileIndex]))
+				{
+					IFileManager::Get().Delete(*Files[FileIndex]);
+					Files.RemoveAt(FileIndex);
+				}
+			}
+		}
+
+		// Remove old UE4 crash contexts
+		TArray<FString> Directories;
+		IFileManager::Get().FindFiles( Directories, *FString::Printf( TEXT( "%s/UE4CC*" ), *FPaths::GameLogDir() ), false, true );
+
+		for (const FString& Dir : Directories)
+		{
+			const FString CrashContextDirectory = FPaths::GameLogDir() / Dir;
+			const FDateTime DirectoryAccessTime = IFileManager::Get().GetTimeStamp( *CrashContextDirectory );
+			if (FDateTime::Now() - DirectoryAccessTime > FTimespan::FromDays( PurgeLogsDays ))
+			{
+				UE_LOG( LogStreaming, Log, TEXT( "Deleting old crash context %s" ), *Dir );
+				IFileManager::Get().DeleteDirectory( *CrashContextDirectory, false, true );
+			}
+		}
+	}
+}
+
diff --git a/Engine/Source/Runtime/Core/Private/Misc/MessageDialog.cpp b/Engine/Source/Runtime/Core/Private/Misc/MessageDialog.cpp
new file mode 100644
--- /dev/null
+++ b/Engine/Source/Runtime/Core/Private/Misc/MessageDialog.cpp
@@ -0,0 +1,97 @@
+﻿// Copyright 1998-2016 Epic Games, Inc. All Rights Reserved.
+
+#include "CorePrivatePCH.h"
+#include "Misc/MessageDialog.h"
+
+namespace
+{
+	/**
+	 * Singleton to only create error title text if needed (and after localization system is in place)
+	 */
+	const FText& GetDefaultMessageTitle()
+	{
+		// Will be initialised on first call
+		static FText DefaultMessageTitle(NSLOCTEXT("MessageDialog", "DefaultMessageTitle", "Message"));
+		return DefaultMessageTitle;
+	}
+}
+
+void FMessageDialog::Debugf( const FText& Message, const FText* OptTitle )
+{
+	if( FApp::IsUnattended() == true )
+	{
+		GLog->Logf( *Message.ToString() );
+	}
+	else
+	{
+		FText Title = OptTitle ? *OptTitle : GetDefaultMessageTitle();
+		if ( GIsEditor && FCoreDelegates::ModalErrorMessage.IsBound() )
+		{
+			FCoreDelegates::ModalErrorMessage.Execute(EAppMsgType::Ok, Message, Title);
+		}
+		else
+		{
+			FPlatformMisc::MessageBoxExt( EAppMsgType::Ok, *Message.ToString(), *NSLOCTEXT("MessageDialog", "DefaultDebugMessageTitle", "ShowDebugMessagef").ToString() );
+		}
+	}
+}
+
+void FMessageDialog::ShowLastError()
+{
+	uint32 LastError = FPlatformMisc::GetLastError();
+
+	TCHAR TempStr[MAX_SPRINTF]=TEXT("");
+	TCHAR ErrorBuffer[1024];
+	FCString::Sprintf( TempStr, TEXT("GetLastError : %d\n\n%s"), LastError, FPlatformMisc::GetSystemErrorMessage(ErrorBuffer, 1024, 0) );
+	if( FApp::IsUnattended() == true )
+	{
+		UE_LOG(LogOutputDevice, Fatal, TempStr);
+	}
+	else
+	{
+		FPlatformMisc::MessageBoxExt( EAppMsgType::Ok, TempStr, *NSLOCTEXT("MessageDialog", "DefaultSystemErrorTitle", "System Error").ToString() );
+	}
+}
+
+EAppReturnType::Type FMessageDialog::Open( EAppMsgType::Type MessageType, const FText& Message, const FText* OptTitle )
+{
+	if( FApp::IsUnattended() == true )
+	{
+		if (GWarn)
+		{
+			GWarn->Logf( *Message.ToString() );
+		}
+
+		switch(MessageType)
+		{
+		case EAppMsgType::Ok:
+			return EAppReturnType::Ok;
+		case EAppMsgType::YesNo:
+			return EAppReturnType::No;
+		case EAppMsgType::OkCancel:
+			return EAppReturnType::Cancel;
+		case EAppMsgType::YesNoCancel:
+			return EAppReturnType::Cancel;
+		case EAppMsgType::CancelRetryContinue:
+			return EAppReturnType::Cancel;
+		case EAppMsgType::YesNoYesAllNoAll:
+			return EAppReturnType::No;
+		case EAppMsgType::YesNoYesAllNoAllCancel:
+		default:
+			return EAppReturnType::Yes;
+			break;
+		}
+	}
+	else
+	{
+		FText Title = OptTitle ? *OptTitle : GetDefaultMessageTitle();
+		if ( GIsEditor && !IsRunningCommandlet() && FCoreDelegates::ModalErrorMessage.IsBound() )
+		{
+			return FCoreDelegates::ModalErrorMessage.Execute( MessageType, Message, Title );
+		}
+		else
+		{
+			return FPlatformMisc::MessageBoxExt( MessageType, *Message.ToString(), *Title.ToString() );
+		}
+	}
+}
diff --git a/Engine/Source/Runtime/Core/Private/Misc/OutputDevice.cpp b/Engine/Source/Runtime/Core/Private/Misc/OutputDevice.cpp
--- a/Engine/Source/Runtime/Core/Private/Misc/OutputDevice.cpp
+++ b/Engine/Source/Runtime/Core/Private/Misc/OutputDevice.cpp
@@ -6,21 +6,7 @@
 #include "OutputDeviceHelper.h"
 #include "HAL/ThreadHeartBeat.h"
 
-DEFINE_LOG_CATEGORY_STATIC(LogOutputDevice, Log, All);
-
-
-namespace
-{
-	/**
-	 * Singleton to only create error title text if needed (and after localization system is in place)
-	 */
-	const FText& GetDefaultMessageTitle()
-	{
-		// Will be initialised on first call
-		static FText DefaultMessageTitle(NSLOCTEXT("MessageDialog", "DefaultMessageTitle", "Message"));
-		return DefaultMessageTitle;
-	}
-}
+DEFINE_LOG_CATEGORY(LogOutputDevice);
 
 const TCHAR* FOutputDevice::VerbosityToString(ELogVerbosity::Type Verbosity)
 {
@@ -67,18 +53,6 @@ void FOutputDevice::Log( const FText& T )
 	Formatted printing and messages.
 -----------------------------------------------------------------------------*/
 
-
-/** Number of top function calls to hide when dumping the callstack as text. */
-#if PLATFORM_LINUX
-
-	// Rationale: check() and ensure() handlers have different depth - worse, ensure() can optionally end up calling the same path as check().
-	// It is better to show the full callstack as is than accidentaly ignore a part of the problem
-	#define CALLSTACK_IGNOREDEPTH 0
-
-#else
-	#define CALLSTACK_IGNOREDEPTH 2
-#endif // PLATFORM_LINUX
-
 VARARG_BODY( void, FOutputDevice::CategorizedLogf, const TCHAR*, VARARG_EXTRA(const class FName& Category) VARARG_EXTRA(ELogVerbosity::Type Verbosity)  )
 {
 	GROWABLE_LOGF(Serialize(Buffer, Verbosity, Category))
@@ -94,557 +68,8 @@ VARARG_BODY( void, FOutputDevice::Logf, const TCHAR*, VARARG_NONE )
 	GROWABLE_LOGF(Serialize(Buffer, TLS->Verbosity, TLS->Category))
 }
 
-#define FILE_LINE_DESC TEXT(" [File:%s] [Line: %i] ")
 
 
 /** Critical errors. */
 CORE_API FOutputDeviceError* GError = NULL;
 
-/** Lock used to synchronize the fail debug calls. */
-static FCriticalSection	FailDebugCriticalSection;
-
-void PrintScriptCallstack(bool bEmptyWhenDone)
-{
-#if DO_BLUEPRINT_GUARD
-	// Walk the script stack, if any
-	FBlueprintExceptionTracker& BlueprintExceptionTracker = FBlueprintExceptionTracker::Get();
-	if( BlueprintExceptionTracker.ScriptStack.Num() > 0 )
-	{
-		FString ScriptStack = TEXT( "\n\nScript Stack:\n" );
-		for (int32 FrameIdx = BlueprintExceptionTracker.ScriptStack.Num() - 1; FrameIdx >= 0; --FrameIdx)
-		{
-			ScriptStack += BlueprintExceptionTracker.ScriptStack[FrameIdx].GetStackDescription() + TEXT( "\n" );
-		}
-
-		UE_LOG( LogOutputDevice, Warning, TEXT( "%s" ), *ScriptStack );
-
-		if (bEmptyWhenDone)
-		{
-			BlueprintExceptionTracker.ScriptStack.Empty();
-		}
-	}
-#endif
-}
-
-/**
- *	Prints error to the debug output, 
- *	prompts for the remote debugging if there is not debugger, breaks into the debugger 
- *	and copies the error into the global error message.
- */
-void StaticFailDebug( const TCHAR* Error, const ANSICHAR* File, int32 Line, const TCHAR* Description, bool bIsEnsure = false )
-{
-	// For ensure log should be flushed in the engine loop.
-	if( !bIsEnsure )
-	{
-		GLog->PanicFlushThreadedLogs();
-	}
-
-	FScopeLock Lock( &FailDebugCriticalSection );
-	FPlatformMisc::LowLevelOutputDebugStringf(TEXT("%s") FILE_LINE_DESC TEXT("\n%s\n"), Error, ANSI_TO_TCHAR(File), Line, Description);
-
-	// Copy the detailed error into the error message.
-	FCString::Snprintf( GErrorMessage, ARRAY_COUNT( GErrorMessage ), TEXT( "%s" ) FILE_LINE_DESC TEXT( "\n%s\n" ), Error, ANSI_TO_TCHAR( File ), Line, Description );
-
-	// Copy the error message to the error history.
-	FCString::Strncpy( GErrorHist, GErrorMessage, ARRAY_COUNT( GErrorHist ) );
-	FCString::Strncat( GErrorHist, TEXT( "\r\n\r\n" ), ARRAY_COUNT( GErrorHist ) );
-}
-
-bool FDebug::bHasAsserted = false;
-
-void FDebug::OutputMultiLineCallstack(const ANSICHAR* File, int32 Line, const FName& LogName, const TCHAR* Heading, TCHAR* Message, ELogVerbosity::Type Verbosity)
-{
-	const bool bWriteUATMarkers = FParse::Param(FCommandLine::Get(), TEXT("CrashForUAT")) && FParse::Param(FCommandLine::Get(), TEXT("stdout"));
-
-	if (bWriteUATMarkers)
-	{
-		FMsg::Logf(File, Line, LogName, Verbosity, TEXT("begin: stack for UAT"));
-	}
-
-	FMsg::Logf(File, Line, LogName, Verbosity, TEXT("%s"), Heading);
-	FMsg::Logf(File, Line, LogName, Verbosity, TEXT(""));
-
-	for (TCHAR* LineStart = Message;; )
-	{
-		// Find the end of the current line
-		TCHAR* LineEnd = LineStart;
-		while (*LineEnd != 0 && *LineEnd != '\r' && *LineEnd != '\n')
-		{
-			LineEnd++;
-		}
-
-		// Output it
-		TCHAR LineEndCharacter = *LineEnd;
-		*LineEnd = 0;
-		FMsg::Logf(File, Line, LogName, Verbosity, TEXT("%s"), LineStart);
-		*LineEnd = LineEndCharacter;
-
-		// Quit if this was the last line
-		if (*LineEnd == 0)
-		{
-			break;
-		}
-
-		// Move to the next line
-		LineStart = (LineEnd[0] == '\r' && LineEnd[1] == '\n') ? (LineEnd + 2) : (LineEnd + 1);
-	}
-
-	if (bWriteUATMarkers)
-	{
-		FMsg::Logf(File, Line, LogName, Verbosity, TEXT("end: stack for UAT"));
-	}
-}
-
-#if DO_CHECK || DO_GUARD_SLOW
-//
-// Failed assertion handler.
-//warning: May be called at library startup time.
-//
-void VARARGS FDebug::LogAssertFailedMessage(const ANSICHAR* Expr, const ANSICHAR* File, int32 Line, const TCHAR* Format/*=TEXT("")*/, ...)
-{
-	// Print out the blueprint callstack
-	PrintScriptCallstack(true);
-
-	// Ignore this assert if we're already forcibly shutting down because of a critical error.
-	if( !GIsCriticalError )
-	{
-		TCHAR DescriptionString[4096];
-		GET_VARARGS( DescriptionString, ARRAY_COUNT( DescriptionString ), ARRAY_COUNT( DescriptionString ) - 1, Format, Format );
-
-		TCHAR ErrorString[MAX_SPRINTF];
-		FCString::Sprintf( ErrorString, TEXT( "Assertion failed: %s" ), ANSI_TO_TCHAR( Expr ) );
-
-		if( FPlatformProperties::AllowsCallStackDumpDuringAssert() )
-		{
-			ANSICHAR StackTrace[4096];
-			if( StackTrace != NULL )
-			{
-				StackTrace[0] = 0;
-				FPlatformStackWalk::StackWalkAndDump( StackTrace, ARRAY_COUNT(StackTrace), CALLSTACK_IGNOREDEPTH );
-
-				FCString::Strncat( DescriptionString, ANSI_TO_TCHAR( StackTrace ), ARRAY_COUNT( DescriptionString ) - 1 );
-			}
-		}
-
-		StaticFailDebug( ErrorString, File, Line, DescriptionString );
-	}
-}
-
-/**
- * Called when an 'ensure' assertion fails; gathers stack data and generates and error report.
- *
- * @param	Expr	Code expression ANSI string (#code)
- * @param	File	File name ANSI string (__FILE__)
- * @param	Line	Line number (__LINE__)
- * @param	Msg		Informative error message text
- */
-void FDebug::EnsureFailed(const ANSICHAR* Expr, const ANSICHAR* File, int32 Line, const TCHAR* Msg)
-{
-
-#if STATS
-	FString EnsureFailedPerfMessage = FString::Printf(TEXT("FDebug::EnsureFailed"));
-	SCOPE_LOG_TIME_IN_SECONDS(*EnsureFailedPerfMessage, nullptr)
-#endif
-
-	// You can set bShouldCrash to true to cause a regular assertion to trigger (stopping program execution) when an ensure() error occurs
-	const bool bShouldCrash = false;		// By default, don't crash on ensure()
-	if( bShouldCrash )
-	{
-		// Just trigger a regular assertion which will crash via GError->Logf()
-		FDebug::LogAssertFailedMessage( Expr, File, Line, Msg );
-		return;
-	}
-
-	// Print out the blueprint callstack
-	PrintScriptCallstack(false);
-
-	// Print initial debug message for this error
-	TCHAR ErrorString[MAX_SPRINTF];
-	FCString::Sprintf(ErrorString,TEXT("Ensure condition failed: %s"),ANSI_TO_TCHAR(Expr));
-
-	StaticFailDebug( ErrorString, File, Line, Msg, true );
-
-	// Is there a debugger attached?  If not we'll submit an error report.
-	if (FPlatformMisc::IsDebuggerPresent())
-	{
-#if !NO_LOGGING
-		UE_LOG(LogOutputDevice, Error, TEXT("%s [File:%s] [Line: %i]"), ErrorString, ANSI_TO_TCHAR(File), Line);
-#endif
-		return;
-	}
-
-	// If we determine that we have not sent a report for this ensure yet, send the report below.
-	bool bShouldSendNewReport = false;
-
-	// No debugger attached, so generate a call stack and submit a crash report
-	// Walk the stack and dump it to the allocated memory.
-	const SIZE_T StackTraceSize = 65535;
-	ANSICHAR* StackTrace = (ANSICHAR*) FMemory::SystemMalloc( StackTraceSize );
-	if( StackTrace != NULL )
-	{
-		// Stop checking heartbeat for this thread. Ensure can take a lot of time (when stackwalking)
-		// Thread heartbeat will be resumed the next time this thread calls FThreadHeartBeat::Get().HeartBeat();
-		// The reason why we don't call HeartBeat() at the end of this function is that maybe this thread
-		// Never had a heartbeat checked and may not be sending heartbeats at all which would later lead to a false positives when detecting hangs.
-		FThreadHeartBeat::Get().KillHeartBeat();
-
-		{
-#if STATS
-			FString StackWalkPerfMessage = FString::Printf(TEXT("FPlatformStackWalk::StackWalkAndDump"));
-			SCOPE_LOG_TIME_IN_SECONDS(*StackWalkPerfMessage, nullptr)
-#endif
-			StackTrace[0] = 0;
-			FPlatformStackWalk::StackWalkAndDumpEx( StackTrace, StackTraceSize, CALLSTACK_IGNOREDEPTH, FGenericPlatformStackWalk::EStackWalkFlags::FlagsUsedWhenHandlingEnsure );
-		}
-
-		// Create a final string that we'll output to the log (and error history buffer)
-		TCHAR ErrorMsg[16384];
-		FCString::Snprintf( ErrorMsg, ARRAY_COUNT( ErrorMsg ), TEXT( "Ensure condition failed: %s [File:%s] [Line: %i]" ) LINE_TERMINATOR TEXT( "%s" ) LINE_TERMINATOR TEXT( "Stack: " ) LINE_TERMINATOR, ANSI_TO_TCHAR( Expr ), ANSI_TO_TCHAR( File ), Line, Msg );
-
-		// Also append the stack trace
-		FCString::Strncat( ErrorMsg, ANSI_TO_TCHAR(StackTrace), ARRAY_COUNT(ErrorMsg) - 1 );
-		FMemory::SystemFree( StackTrace );
-
-		// Dump the error and flush the log.
-#if !NO_LOGGING
-		FDebug::OutputMultiLineCallstack(__FILE__, __LINE__, LogOutputDevice.GetCategoryName(), TEXT("=== Handled ensure: ==="), ErrorMsg, ELogVerbosity::Error);
-#endif
-		GLog->Flush();
-
-		// Submit the error report to the server! (and display a balloon in the system tray)
-		{
-			// How many unique previous errors we should keep track of
-			const uint32 MaxPreviousErrorsToTrack = 4;
-			static uint32 StaticPreviousErrorCount = 0;
-			if( StaticPreviousErrorCount < MaxPreviousErrorsToTrack )
-			{
-				// Check to see if we've already reported this error.  No point in blasting the server with
-				// the same error over and over again in a single application session.
-				bool bHasErrorAlreadyBeenReported = false;
-
-				// Static: Array of previous unique error message CRCs
-				static uint32 StaticPreviousErrorCRCs[ MaxPreviousErrorsToTrack ];
-
-				// Compute CRC of error string.  Note that along with the call stack, this includes the message
-				// string passed to the macro, so only truly redundant errors will go unreported.  Though it also
-				// means you shouldn't pass loop counters to ensureMsgf(), otherwise failures may spam the server!
-				const uint32 ErrorStrCRC = FCrc::StrCrc_DEPRECATED( ErrorMsg );
-
-				for( uint32 CurErrorIndex = 0; CurErrorIndex < StaticPreviousErrorCount; ++CurErrorIndex )
-				{
-					if( StaticPreviousErrorCRCs[ CurErrorIndex ] == ErrorStrCRC )
-					{
-						// Found it!  This is a redundant error message.
-						bHasErrorAlreadyBeenReported = true;
-						break;
-					}
-				}
-
-				// Add the element to the list and bump the count
-				StaticPreviousErrorCRCs[ StaticPreviousErrorCount++ ] = ErrorStrCRC;
-
-				if( !bHasErrorAlreadyBeenReported )
-				{
-#if STATS
-					FString SubmitErrorReporterfMessage = FString::Printf(TEXT("SubmitErrorReport"));
-					SCOPE_LOG_TIME_IN_SECONDS(*SubmitErrorReporterfMessage, nullptr)
-#endif
-
-					FCoreDelegates::OnHandleSystemEnsure.Broadcast();
-
-					FPlatformMisc::SubmitErrorReport( ErrorMsg, EErrorReportMode::Balloon );
-
-					bShouldSendNewReport = true;
-				}
-			}
-		}
-	}
-	else
-	{
-		// If we fail to generate the string to identify the crash we don't know if we should skip sending the report,
-		// so we will just send the report anyway.
-		bShouldSendNewReport = true;
-
-		// Add message to log even without stacktrace. It is useful for testing fail on ensure.
-#if !NO_LOGGING
-		UE_LOG(LogOutputDevice, Error, TEXT("%s [File:%s] [Line: %i]"), ErrorString, ANSI_TO_TCHAR(File), Line);
-#endif
-	}
-
-	if ( bShouldSendNewReport )
-	{
-#if STATS
-		FString SendNewReportMessage = FString::Printf(TEXT("SendNewReport"));
-		SCOPE_LOG_TIME_IN_SECONDS(*SendNewReportMessage, nullptr)
-#endif
-
-#if PLATFORM_DESKTOP
-		FScopeLock Lock( &FailDebugCriticalSection );
-
-		NewReportEnsure( GErrorMessage );
-
-		GErrorHist[0] = 0;
-		GErrorMessage[0] = 0;
-		GErrorExceptionDescription[0] = 0;
-#endif
-	}
-}
-
-#endif // DO_CHECK || DO_GUARD_SLOW
-
-void VARARGS FError::LowLevelFatal(const ANSICHAR* File, int32 Line, const TCHAR* Format, ...)
-{
-	TCHAR DescriptionString[4096];
-	GET_VARARGS( DescriptionString, ARRAY_COUNT(DescriptionString), ARRAY_COUNT(DescriptionString)-1, Format, Format );
-
-	StaticFailDebug(TEXT("LowLevelFatalError"),File,Line,DescriptionString);
-}
-
-void VARARGS FDebug::AssertFailed(const ANSICHAR* Expr, const ANSICHAR* File, int32 Line, const TCHAR* Format/* = TEXT("")*/, ...)
-{
-	if (GIsCriticalError)
-	{
-		return;
-	}
-
-	// This is not perfect because another thread might crash and be handled before this assert
-	// but this static varible will report the crash as an assert. Given complexity of a thread
-	// aware solution, this should be good enough. If crash reports are obviously wrong we can
-	// look into fixing this.
-	bHasAsserted = true;
-
-	TCHAR DescriptionString[4096];
-	GET_VARARGS(DescriptionString, ARRAY_COUNT(DescriptionString), ARRAY_COUNT(DescriptionString) - 1, Format, Format);
-
-	TCHAR ErrorString[MAX_SPRINTF];
-	FCString::Sprintf(ErrorString, TEXT("Assertion failed: %s"), ANSI_TO_TCHAR(Expr));
-	GError->Logf(TEXT("Assertion failed: %s") FILE_LINE_DESC TEXT("\n%s\n"), ErrorString, ANSI_TO_TCHAR(File), Line, DescriptionString);
-}
-
-#if DO_CHECK || DO_GUARD_SLOW
-bool VARARGS FDebug::OptionallyLogFormattedEnsureMessageReturningFalse( bool bLog, const ANSICHAR* Expr, const ANSICHAR* File, int32 Line, const TCHAR* FormattedMsg, ... )
-{
-	if (bLog)
-	{
-		const int32 TempStrSize = 4096;
-		TCHAR TempStr[ TempStrSize ];
-		GET_VARARGS( TempStr, TempStrSize, TempStrSize - 1, FormattedMsg, FormattedMsg );
-		EnsureFailed( Expr, File, Line, TempStr );
-	}
-	
-	return false;
-}
-#endif
-
-void FMessageDialog::Debugf( const FText& Message, const FText* OptTitle )
-{
-	if( FApp::IsUnattended() == true )
-	{
-		GLog->Logf( *Message.ToString() );
-	}
-	else
-	{
-		FText Title = OptTitle ? *OptTitle : GetDefaultMessageTitle();
-		if ( GIsEditor && FCoreDelegates::ModalErrorMessage.IsBound() )
-		{
-			FCoreDelegates::ModalErrorMessage.Execute(EAppMsgType::Ok, Message, Title);
-		}
-		else
-		{
-			FPlatformMisc::MessageBoxExt( EAppMsgType::Ok, *Message.ToString(), *NSLOCTEXT("MessageDialog", "DefaultDebugMessageTitle", "ShowDebugMessagef").ToString() );
-		}
-	}
-}
-
-void FMessageDialog::ShowLastError()
-{
-	uint32 LastError = FPlatformMisc::GetLastError();
-
-	TCHAR TempStr[MAX_SPRINTF]=TEXT("");
-	TCHAR ErrorBuffer[1024];
-	FCString::Sprintf( TempStr, TEXT("GetLastError : %d\n\n%s"), LastError, FPlatformMisc::GetSystemErrorMessage(ErrorBuffer, 1024, 0) );
-	if( FApp::IsUnattended() == true )
-	{
-		UE_LOG(LogOutputDevice, Fatal, TempStr);
-	}
-	else
-	{
-		FPlatformMisc::MessageBoxExt( EAppMsgType::Ok, TempStr, *NSLOCTEXT("MessageDialog", "DefaultSystemErrorTitle", "System Error").ToString() );
-	}
-}
-
-EAppReturnType::Type FMessageDialog::Open( EAppMsgType::Type MessageType, const FText& Message, const FText* OptTitle )
-{
-	if( FApp::IsUnattended() == true )
-	{
-		if (GWarn)
-		{
-			GWarn->Logf( *Message.ToString() );
-		}
-
-		switch(MessageType)
-		{
-		case EAppMsgType::Ok:
-			return EAppReturnType::Ok;
-		case EAppMsgType::YesNo:
-			return EAppReturnType::No;
-		case EAppMsgType::OkCancel:
-			return EAppReturnType::Cancel;
-		case EAppMsgType::YesNoCancel:
-			return EAppReturnType::Cancel;
-		case EAppMsgType::CancelRetryContinue:
-			return EAppReturnType::Cancel;
-		case EAppMsgType::YesNoYesAllNoAll:
-			return EAppReturnType::No;
-		case EAppMsgType::YesNoYesAllNoAllCancel:
-		default:
-			return EAppReturnType::Yes;
-			break;
-		}
-	}
-	else
-	{
-		FText Title = OptTitle ? *OptTitle : GetDefaultMessageTitle();
-		if ( GIsEditor && !IsRunningCommandlet() && FCoreDelegates::ModalErrorMessage.IsBound() )
-		{
-			return FCoreDelegates::ModalErrorMessage.Execute( MessageType, Message, Title );
-		}
-		else
-		{
-			return FPlatformMisc::MessageBoxExt( MessageType, *Message.ToString(), *Title.ToString() );
-		}
-	}
-}
-
-/*-----------------------------------------------------------------------------
-	Exceptions.
------------------------------------------------------------------------------*/
-
-//
-// Throw a string exception with a message.
-//
-#if HACK_HEADER_GENERATOR
-void VARARGS FError::ThrowfImpl(const TCHAR* Fmt, ...)
-{
-	static TCHAR TempStr[4096];
-	GET_VARARGS( TempStr, ARRAY_COUNT(TempStr), ARRAY_COUNT(TempStr)-1, Fmt, Fmt );
-#if !PLATFORM_EXCEPTIONS_DISABLED
-	throw( TempStr );
-#else
-	UE_LOG(LogOutputDevice, Fatal, TEXT("THROW: %s"), TempStr);
-#endif
-}					
-#endif
-
-/** Statics to prevent FMsg::Logf from allocating too much stack memory. */
-static FCriticalSection					MsgLogfStaticBufferGuard;
-/** Increased from 4096 to fix crashes in the renderthread without autoreporter. */
-static TCHAR							MsgLogfStaticBuffer[8192];
-
-VARARG_BODY(void, FMsg::Logf, const TCHAR*, VARARG_EXTRA(const ANSICHAR* File) VARARG_EXTRA(int32 Line) VARARG_EXTRA(const class FName& Category) VARARG_EXTRA(ELogVerbosity::Type Verbosity))
-{
-#if !NO_LOGGING
-	if (Verbosity != ELogVerbosity::Fatal)
-	{
-		// SetColour is routed to GWarn just like the other verbosities and handled in the 
-		// device that does the actual printing.
-		FOutputDevice* LogDevice = NULL;
-		switch (Verbosity)
-		{
-		case ELogVerbosity::Error:
-		case ELogVerbosity::Warning:
-		case ELogVerbosity::Display:
-		case ELogVerbosity::SetColor:
-			if (GWarn)
-			{
-				LogDevice = GWarn;
-				break;
-			}
-		default:
-		{
-			LogDevice = GLog;
-		}
-		break;
-		}
-		GROWABLE_LOGF(LogDevice->Log(Category, Verbosity, Buffer))
-	}
-	else
-	{
-		// Keep Message buffer small, in some cases, this code is executed with 16KB stack.
-		TCHAR Message[4096];
-		{
-			// Simulate Sprintf_s
-			// @todo: implement platform independent sprintf_S
-			// We're using one big shared static buffer here, so guard against re-entry
-			FScopeLock MsgLock(&MsgLogfStaticBufferGuard);
-			// Print to a large static buffer so we can keep the stack allocation below 16K
-			GET_VARARGS(MsgLogfStaticBuffer, ARRAY_COUNT(MsgLogfStaticBuffer), ARRAY_COUNT(MsgLogfStaticBuffer) - 1, Fmt, Fmt);
-			// Copy the message to the stack-allocated buffer)
-			FCString::Strncpy(Message, MsgLogfStaticBuffer, ARRAY_COUNT(Message) - 1);
-			Message[ARRAY_COUNT(Message) - 1] = '\0';
-		}
-
-		StaticFailDebug(TEXT("Fatal error:"), File, Line, Message);
-		FDebug::AssertFailed("", File, Line, Message);
-	}
-#endif
-}
-
-VARARG_BODY(void, FMsg::Logf_Internal, const TCHAR*, VARARG_EXTRA(const ANSICHAR* File) VARARG_EXTRA(int32 Line) VARARG_EXTRA(const class FName& Category) VARARG_EXTRA(ELogVerbosity::Type Verbosity))
-{
-#if !NO_LOGGING
-	if (Verbosity != ELogVerbosity::Fatal)
-	{
-		// SetColour is routed to GWarn just like the other verbosities and handled in the 
-		// device that does the actual printing.
-		FOutputDevice* LogDevice = NULL;
-		switch (Verbosity)
-		{
-		case ELogVerbosity::Error:
-		case ELogVerbosity::Warning:
-		case ELogVerbosity::Display:
-		case ELogVerbosity::SetColor:
-			if (GWarn)
-			{
-				LogDevice = GWarn;
-				break;
-			}
-		default:
-		{
-			LogDevice = GLog;
-		}
-		break;
-		}
-		GROWABLE_LOGF(LogDevice->Log(Category, Verbosity, Buffer))
-	}
-	else
-	{
-		// Keep Message buffer small, in some cases, this code is executed with 16KB stack.
-		TCHAR Message[4096];
-		{
-			// Simulate Sprintf_s
-			// @todo: implement platform independent sprintf_S
-			// We're using one big shared static buffer here, so guard against re-entry
-			FScopeLock MsgLock(&MsgLogfStaticBufferGuard);
-			// Print to a large static buffer so we can keep the stack allocation below 16K
-			GET_VARARGS(MsgLogfStaticBuffer, ARRAY_COUNT(MsgLogfStaticBuffer), ARRAY_COUNT(MsgLogfStaticBuffer) - 1, Fmt, Fmt);
-			// Copy the message to the stack-allocated buffer)
-			FCString::Strncpy(Message, MsgLogfStaticBuffer, ARRAY_COUNT(Message) - 1);
-			Message[ARRAY_COUNT(Message) - 1] = '\0';
-		}
-
-		StaticFailDebug(TEXT("Fatal error:"), File, Line, Message);
-	}
-#endif
-}
-
-/** Sends a formatted message to a remote tool. */
-void VARARGS FMsg::SendNotificationStringf( const TCHAR *Fmt, ... )
-{
-	GROWABLE_LOGF(SendNotificationString(Buffer));
-}
-
-void FMsg::SendNotificationString( const TCHAR* Message )
-{
-	FPlatformMisc::LowLevelOutputDebugString(Message);
-}
diff --git a/Engine/Source/Runtime/Core/Private/Misc/OutputDeviceAnsiError.cpp b/Engine/Source/Runtime/Core/Private/Misc/OutputDeviceAnsiError.cpp
new file mode 100644
--- /dev/null
+++ b/Engine/Source/Runtime/Core/Private/Misc/OutputDeviceAnsiError.cpp
@@ -0,0 +1,92 @@
+// Copyright 1998-2016 Epic Games, Inc. All Rights Reserved.
+
+#include "CorePrivatePCH.h"
+#include "OutputDeviceAnsiError.h"
+
+/** Constructor, initializing member variables */
+FOutputDeviceAnsiError::FOutputDeviceAnsiError()
+:	ErrorPos(0)
+{}
+
+/**
+ * Serializes the passed in data unless the current event is suppressed.
+ *
+ * @param	Data	Text to log
+ * @param	Event	Event name used for suppression purposes
+ */
+void FOutputDeviceAnsiError::Serialize( const TCHAR* Msg, ELogVerbosity::Type Verbosity, const class FName& Category )
+{
+	// Display the error and exit.
+	FPlatformMisc::LocalPrint( TEXT("\nappError called: \n") );
+	FPlatformMisc::LocalPrint( Msg );
+	FPlatformMisc::LocalPrint( TEXT("\n") );
+
+	if( !GIsCriticalError )
+	{
+		// First appError.
+		GIsCriticalError = 1;
+		UE_LOG(LogHAL, Error, TEXT("appError called: %s"), Msg );
+		FCString::Strncpy( GErrorHist, Msg, ARRAY_COUNT(GErrorHist) );
+		FCString::Strncat( GErrorHist, TEXT("\r\n\r\n"), ARRAY_COUNT(GErrorHist) );
+		ErrorPos = FCString::Strlen(GErrorHist);
+	}
+	else
+	{
+		UE_LOG(LogHAL, Error, TEXT("Error reentered: %s"), Msg );
+	}
+
+	FPlatformMisc::DebugBreak();
+
+	if( GIsGuarded )
+	{
+		// Propagate error so structured exception handler can perform necessary work.
+#if PLATFORM_EXCEPTIONS_DISABLED
+		FPlatformMisc::DebugBreak();
+#endif
+		FPlatformMisc::RaiseException( 1 );
+	}
+	else
+	{
+		// We crashed outside the guarded code (e.g. appExit).
+		HandleError();
+		// pop up a crash window if we are not in unattended mode
+		if( FApp::IsUnattended() == false )
+		{
+			FPlatformMisc::RequestExit( true );
+		}
+		else
+		{
+			UE_LOG(LogHAL, Error, TEXT("%s"), Msg );
+		}		
+	}
+}
+
+/**
+ * Error handling function that is being called from within the system wide global
+ * error handler, e.g. using structured exception handling on the PC.
+ */
+void FOutputDeviceAnsiError::HandleError()
+{
+	GIsGuarded			= 0;
+	GIsRunning			= 0;
+	GIsCriticalError	= 1;
+	GLogConsole			= NULL;
+	GErrorHist[ARRAY_COUNT(GErrorHist)-1]=0;
+
+	if (GLog)
+	{
+		// print to log and flush it
+		UE_LOG( LogHAL, Log, TEXT( "=== Critical error: ===" ) LINE_TERMINATOR LINE_TERMINATOR TEXT( "%s" ) LINE_TERMINATOR, GErrorExceptionDescription );
+		UE_LOG(LogHAL, Log, GErrorHist);
+
+		GLog->Flush();
+	}
+	else
+	{
+		FPlatformMisc::LocalPrint( GErrorHist );
+	}
+
+	FPlatformMisc::LocalPrint( TEXT("\n\nExiting due to error\n") );
+
+	FCoreDelegates::OnShutdownAfterError.Broadcast();
+}
diff --git a/Engine/Source/Runtime/Core/Private/Misc/OutputDeviceDebug.cpp b/Engine/Source/Runtime/Core/Private/Misc/OutputDeviceDebug.cpp
new file mode 100644
--- /dev/null
+++ b/Engine/Source/Runtime/Core/Private/Misc/OutputDeviceDebug.cpp
@@ -0,0 +1,33 @@
+// Copyright 1998-2016 Epic Games, Inc. All Rights Reserved.
+
+#include "CorePrivatePCH.h"
+#include "OutputDeviceDebug.h"
+
+/**
+ * Serializes the passed in data unless the current event is suppressed.
+ *
+ * @param	Data	Text to log
+ * @param	Event	Event name used for suppression purposes
+ */
+void FOutputDeviceDebug::Serialize( const TCHAR* Data, ELogVerbosity::Type Verbosity, const class FName& Category, const double Time )
+{
+	static bool Entry=false;
+	if( !GIsCriticalError || Entry )
+	{
+		if (Verbosity != ELogVerbosity::SetColor)
+		{
+			FPlatformMisc::LowLevelOutputDebugStringf(TEXT("%s%s"),*FOutputDeviceHelper::FormatLogLine(Verbosity, Category, Data, GPrintLogTimes, Time),LINE_TERMINATOR);
+		}
+	}
+	else
+	{
+		Entry=true;
+		Serialize( Data, Verbosity, Category );
+		Entry=false;
+	}
+}
+
+void FOutputDeviceDebug::Serialize( const TCHAR* Data, ELogVerbosity::Type Verbosity, const class FName& Category )
+{
+	Serialize( Data, Verbosity, Category, -1.0 );
+}
diff --git a/Engine/Source/Runtime/Core/Private/Misc/OutputDeviceFile.cpp b/Engine/Source/Runtime/Core/Private/Misc/OutputDeviceFile.cpp
new file mode 100644
--- /dev/null
+++ b/Engine/Source/Runtime/Core/Private/Misc/OutputDeviceFile.cpp
@@ -0,0 +1,418 @@
+// Copyright 1998-2016 Epic Games, Inc. All Rights Reserved.
+
+#include "CorePrivatePCH.h"
+#include "OutputDeviceFile.h"
+
+/** Used by tools which include only core to disable log file creation. */
+#ifndef ALLOW_LOG_FILE
+	#define ALLOW_LOG_FILE 1
+#endif
+
+typedef uint8 UTF8BOMType[3];
+static UTF8BOMType UTF8BOM = { 0xEF, 0xBB, 0xBF };
+
+/**
+* Thread heartbeat check class.
+* Used by crash handling code to check for hangs.
+*/
+class CORE_API FAsyncWriter : public FRunnable, public FArchive
+{
+	enum EConstants
+	{
+		InitialBufferSize = 128 * 1024
+	};
+
+	/** Thread to run the worker FRunnable on. Serializes the ring buffer to disk. */
+	FRunnableThread* Thread;
+	/** Stops this thread */
+	FThreadSafeCounter StopTaskCounter;
+
+	/** Writer archive */
+	FArchive& Ar;
+	/** Data ring buffer */
+	TArray<uint8> Buffer;
+	/** Position where the unserialized data starts in the buffer */
+	int32 BufferStartPos;
+	/** Position where the unserialized data ends in the buffer (such as if (BufferEndPos > BufferStartPos) Length = BufferEndPos - BufferStartPos; */
+	int32 BufferEndPos;
+	/** Sync object for the buffer pos */
+	FCriticalSection BufferPosCritical;
+	/** Event to let the worker thread know there is data to save to disk. */
+	FEvent* BufferDirtyEvent;
+	/** Event flagged when the worker thread finished processing data */
+	FEvent* BufferEmptyEvent;
+
+	/** Serialize the contents of the ring buffer to disk */
+	void SerializeBufferToArchive()
+	{
+		// Grab a local copy of the end pos. It's ok if it changes on the client thread later on.
+		// We won't be modifying it anyway and will later serialize new data in the next iteration.
+		// Here we only serialize what we know exists at the beginning of this function.
+		int32 ThisThreadEndPos = BufferEndPos;
+		if (ThisThreadEndPos >= BufferStartPos)
+		{
+			Ar.Serialize(Buffer.GetData() + BufferStartPos, ThisThreadEndPos - BufferStartPos);
+		}
+		else
+		{
+			// Data is wrapped around the ring buffer
+			Ar.Serialize(Buffer.GetData() + BufferStartPos, Buffer.Num() - BufferStartPos);
+			Ar.Serialize(Buffer.GetData(), BufferEndPos);
+		}
+		// Modify the start pos. Only the worker thread modifies this value so it's ok to not guard it with a critical section.
+		BufferStartPos = ThisThreadEndPos;
+		// Let the client threads know we're done (this is used by FlushBuffer).
+		BufferEmptyEvent->Trigger();
+	}
+
+	/** Flush the memory buffer (doesn't force the archive to flush) */
+	void FlushBuffer()
+	{
+		BufferDirtyEvent->Trigger();
+		if (!Thread)
+		{
+			SerializeBufferToArchive();
+		}
+		BufferEmptyEvent->Wait();
+	}
+
+public:
+
+	FAsyncWriter(FArchive& InAr)
+		: Thread(nullptr)
+		, Ar(InAr)
+		, BufferStartPos(0)
+		, BufferEndPos(0)
+	{
+		Buffer.AddUninitialized(InitialBufferSize);
+		BufferDirtyEvent = FPlatformProcess::GetSynchEventFromPool();
+		BufferEmptyEvent = FPlatformProcess::GetSynchEventFromPool(true);
+		BufferEmptyEvent->Trigger();
+
+		if (FPlatformProcess::SupportsMultithreading())
+		{
+			FString WriterName = FString::Printf(TEXT("FAsyncWriter_%s"), *FPaths::GetBaseFilename(Ar.GetArchiveName()));
+			Thread = FRunnableThread::Create(this, *WriterName, 0, TPri_BelowNormal);
+		}
+	}
+
+	virtual ~FAsyncWriter()
+	{
+		Flush();
+		delete Thread;
+		Thread = nullptr;
+		FPlatformProcess::ReturnSynchEventToPool(BufferDirtyEvent);
+		BufferDirtyEvent = nullptr;
+		FPlatformProcess::ReturnSynchEventToPool(BufferEmptyEvent);
+		BufferEmptyEvent = nullptr;
+	}
+
+	/** Serialize data to buffer that will later be saved to disk by the async thread */
+	virtual void Serialize(void* InData, int64 Length) override
+	{
+		if (!InData || Length <= 0)
+		{
+			return;
+		}
+
+		const uint8* Data = (uint8*)InData;
+
+		BufferEmptyEvent->Reset();
+		FScopeLock WriteLock(&BufferPosCritical);
+
+		// Store the local copy of the current buffer start pos. It may get moved by the worker thread but we don't
+		// care about it too much because we only modify BufferEndPos. Copy should be atomic enough. We only use it
+		// for checking the remaining space in the buffer so underestimating is ok.
+		{
+			const int32 ThisThreadStartPos = BufferStartPos;
+			// Calculate the remaining size in the ring buffer
+			const int32 BufferFreeSize = ThisThreadStartPos <= BufferEndPos ? (Buffer.Num() - BufferEndPos + ThisThreadStartPos) : (ThisThreadStartPos - BufferEndPos);
+			// Make sure the buffer is BIGGER than we require otherwise we may calculate the wrong (0) buffer EndPos for StartPos = 0 and Length = Buffer.Num()
+			if (BufferFreeSize <= Length)
+			{
+				// Force the async thread to call SerializeBufferToArchive even if it's currently empty
+				BufferDirtyEvent->Trigger();
+				FlushBuffer();
+
+				// Resize the buffer if needed
+				if (Length >= Buffer.Num())
+				{
+					// Keep the buffer bigger than we require so that % Buffer.Num() does not return 0 for Lengths = Buffer.Num()
+					Buffer.SetNumUninitialized(Length + 1);
+				}
+			}
+		}
+
+		// We now know there's enough space in the buffer to copy data
+		const int32 WritePos = BufferEndPos;
+		if ((WritePos + Length) <= Buffer.Num())
+		{
+			// Copy straight into the ring buffer
+			FMemory::Memcpy(Buffer.GetData() + WritePos, Data, Length);
+		}
+		else
+		{
+			// Wrap around the ring buffer
+			int32 BufferSizeToEnd = Buffer.Num() - WritePos;
+			FMemory::Memcpy(Buffer.GetData() + WritePos, Data, BufferSizeToEnd);
+			FMemory::Memcpy(Buffer.GetData(), Data + BufferSizeToEnd, Length - BufferSizeToEnd);
+		}
+
+		// Update the end position and let the async thread know we need to write to disk
+		BufferEndPos = (BufferEndPos + Length) % Buffer.Num();
+		BufferDirtyEvent->Trigger();
+
+		// No async thread? Serialize now.
+		if (!Thread)
+		{
+			SerializeBufferToArchive();
+		}
+	}
+
+	/** Flush all buffers to disk */
+	void Flush()
+	{
+		FScopeLock WriteLock(&BufferPosCritical);
+		FlushBuffer();
+		Ar.Flush();
+	}
+
+	//~ Begin FRunnable Interface.
+	virtual bool Init() 
+	{
+		return true;
+	}
+	virtual uint32 Run()
+	{
+		while (StopTaskCounter.GetValue() == 0)
+		{
+			if (BufferDirtyEvent->Wait(10))
+			{
+				SerializeBufferToArchive();
+			}
+		}
+		return 0;
+	}
+	virtual void Stop()
+	{
+		StopTaskCounter.Increment();
+	}
+	//~ End FRunnable Interface
+};
+
+/** 
+ * Constructor, initializing member variables.
+ *
+ * @param InFilename		Filename to use, can be NULL
+ * @param bInDisableBackup	If true, existing files will not be backed up
+ */
+FOutputDeviceFile::FOutputDeviceFile( const TCHAR* InFilename, bool bInDisableBackup  )
+: AsyncWriter(nullptr)
+, WriterArchive(nullptr)
+, Opened(false)
+, Dead(false)
+, bDisableBackup(bInDisableBackup)
+{
+	if( InFilename )
+	{
+		FCString::Strncpy( Filename, InFilename, ARRAY_COUNT(Filename) );
+	}
+	else
+	{
+		Filename[0]	= 0;
+	}
+}
+
+void FOutputDeviceFile::SetFilename(const TCHAR* InFilename)
+{
+	// Close any existing file.
+	TearDown();
+
+	FCString::Strncpy( Filename, InFilename, ARRAY_COUNT(Filename) );
+}
+
+/**
+ * Closes output device and cleans up. This can't happen in the destructor
+ * as we have to call "delete" which cannot be done for static/ global
+ * objects.
+ */
+void FOutputDeviceFile::TearDown()
+{
+	if (AsyncWriter)
+	{
+		if (!bSuppressEventTag)
+		{			
+			Logf(TEXT("Log file closed, %s"), FPlatformTime::StrTimestamp());
+		}
+		delete AsyncWriter;
+		AsyncWriter = nullptr;
+	}
+	delete WriterArchive;
+	WriterArchive = nullptr;
+}
+
+/**
+ * Flush the write cache so the file isn't truncated in case we crash right
+ * after calling this function.
+ */
+void FOutputDeviceFile::Flush()
+{
+	if (AsyncWriter)
+	{
+		AsyncWriter->Flush();
+	}
+}
+
+/** if the passed in file exists, makes a timestamped backup copy
+ * @param Filename the name of the file to check
+ */
+void FOutputDeviceFile::CreateBackupCopy(const TCHAR* Filename)
+{
+	if (IFileManager::Get().FileSize(Filename) > 0)
+	{
+		FString SystemTime = FDateTime::Now().ToString();
+		FString Name, Extension;
+		FString(Filename).Split(TEXT("."), &Name, &Extension, ESearchCase::CaseSensitive, ESearchDir::FromEnd);
+		FString BackupFilename = FString::Printf(TEXT("%s%s%s.%s"), *Name, BACKUP_LOG_FILENAME_POSTFIX, *SystemTime, *Extension);
+		IFileManager::Get().Copy(*BackupFilename, Filename, false);
+	}
+}
+
+bool FOutputDeviceFile::IsBackupCopy(const TCHAR* Filename)
+{
+	return Filename != nullptr && FCString::Stristr(const_cast<TCHAR*>(Filename), BACKUP_LOG_FILENAME_POSTFIX) != nullptr;
+}
+
+void FOutputDeviceFile::WriteByteOrderMarkToArchive(EByteOrderMark ByteOrderMark)
+{
+	switch (ByteOrderMark)
+	{
+	case EByteOrderMark::UTF8:
+		AsyncWriter->Serialize(UTF8BOM, sizeof(UTF8BOM));
+		break;
+
+	case EByteOrderMark::Unspecified:
+	default:
+		check(false);
+		break;
+	}
+}
+
+bool FOutputDeviceFile::CreateWriter(uint32 MaxAttempts)
+{
+	uint32 WriteFlags = FILEWRITE_AllowRead | (Opened ? FILEWRITE_Append : 0);
+
+	// Open log file.
+	FArchive* Ar = IFileManager::Get().CreateFileWriter(Filename, WriteFlags);
+
+	// If that failed, append an _2 and try again (unless we don't want extra copies). This 
+	// happens in the case of running a server and client on same computer for example.
+	if (!bDisableBackup && !Ar)
+	{
+		FString FilenamePart = FPaths::GetBaseFilename(Filename, false) + "_";
+		FString ExtensionPart = FPaths::GetExtension(Filename, true);
+		FString FinalFilename;
+		uint32 FileIndex = 2;
+		do
+		{
+			// Continue to increment indices until a valid filename is found
+			FinalFilename = FilenamePart + FString::FromInt(FileIndex++) + ExtensionPart;
+			if (!Opened)
+			{
+				CreateBackupCopy(*FinalFilename);
+			}
+			Ar = IFileManager::Get().CreateFileWriter(*FinalFilename, WriteFlags);
+		} while (!Ar && FileIndex < MaxAttempts);
+	}
+
+	FAsyncWriter* Result = nullptr;
+	if (Ar)
+	{
+		WriterArchive = Ar;
+		AsyncWriter = new FAsyncWriter(*WriterArchive);
+	}
+
+	return !!AsyncWriter;
+}
+
+/**
+ * Serializes the passed in data unless the current event is suppressed.
+ *
+ * @param	Data	Text to log
+ * @param	Event	Event name used for suppression purposes
+ */
+void FOutputDeviceFile::Serialize( const TCHAR* Data, ELogVerbosity::Type Verbosity, const class FName& Category, const double Time )
+{
+#if ALLOW_LOG_FILE && !NO_LOGGING
+	static bool Entry = false;
+	if( !GIsCriticalError || Entry )
+	{
+		if (!AsyncWriter && !Dead)
+		{
+			// Make log filename.
+			if( !Filename[0] )
+			{
+				FCString::Strcpy(Filename, *FPlatformOutputDevices::GetAbsoluteLogFilename());
+			}
+
+			// if the file already exists, create a backup as we are going to overwrite it
+			if (!bDisableBackup && !Opened)
+			{
+				CreateBackupCopy(Filename);
+			}
+
+			// Open log file and create the worker thread.
+			if (CreateWriter())
+			{
+				Opened = true;
+
+				WriteByteOrderMarkToArchive(EByteOrderMark::UTF8);
+
+				if (!bSuppressEventTag)
+				{
+					Logf( TEXT("Log file open, %s"), FPlatformTime::StrTimestamp() );
+				}
+			}
+			else 
+			{
+				Dead = true;
+			}
+		}
+
+		if (AsyncWriter && Verbosity != ELogVerbosity::SetColor)
+		{
+			FOutputDeviceHelper::FormatCastAndSerializeLine(*AsyncWriter, Data, Verbosity, Category, Time, bSuppressEventTag, bAutoEmitLineTerminator);
+
+			static bool GForceLogFlush = false;
+			static bool GTestedCmdLine = false;
+			if (!GTestedCmdLine)
+			{
+				GTestedCmdLine = true;
+				// Force a log flush after each line
+				GForceLogFlush = FParse::Param( FCommandLine::Get(), TEXT("FORCELOGFLUSH") );
+			}
+			if (GForceLogFlush)
+			{
+				AsyncWriter->Flush();
+			}
+		}
+	}
+	else
+	{
+		Entry = true;
+		Serialize(Data, Verbosity, Category, Time);
+		Entry = false;
+	}
+#endif
+}
+
+void FOutputDeviceFile::Serialize( const TCHAR* Data, ELogVerbosity::Type Verbosity, const class FName& Category )
+{
+	Serialize(Data, Verbosity, Category, -1.0);
+}
+
+
+void FOutputDeviceFile::WriteRaw( const TCHAR* C )
+{
+	AsyncWriter->Serialize((uint8*)const_cast<TCHAR*>(C), FCString::Strlen(C)*sizeof(TCHAR));
+}
diff --git a/Engine/Source/Runtime/Core/Private/Misc/OutputDeviceNull.cpp b/Engine/Source/Runtime/Core/Private/Misc/OutputDeviceNull.cpp
new file mode 100644
--- /dev/null
+++ b/Engine/Source/Runtime/Core/Private/Misc/OutputDeviceNull.cpp
@@ -0,0 +1,9 @@
+// Copyright 1998-2016 Epic Games, Inc. All Rights Reserved.
+
+#include "CorePrivatePCH.h"
+#include "OutputDeviceNull.h"
+
+void FOutputDeviceNull::Serialize(const TCHAR* V, ELogVerbosity::Type Verbosity, const class FName& Category)
+{
+}
+
diff --git a/Engine/Source/Runtime/Core/Private/Misc/OutputDeviceRedirector.cpp b/Engine/Source/Runtime/Core/Private/Misc/OutputDeviceRedirector.cpp
new file mode 100644
--- /dev/null
+++ b/Engine/Source/Runtime/Core/Private/Misc/OutputDeviceRedirector.cpp
@@ -0,0 +1,258 @@
+// Copyright 1998-2016 Epic Games, Inc. All Rights Reserved.
+
+#include "CorePrivatePCH.h"
+#include "OutputDeviceRedirector.h"
+
+/*-----------------------------------------------------------------------------
+	FOutputDeviceRedirector.
+-----------------------------------------------------------------------------*/
+
+/** Initialization constructor. */
+FOutputDeviceRedirector::FOutputDeviceRedirector()
+:	MasterThreadID(FPlatformTLS::GetCurrentThreadId())
+,	bEnableBacklog(false)
+{
+}
+
+FOutputDeviceRedirector* FOutputDeviceRedirector::Get()
+{
+	static FOutputDeviceRedirector Singleton;
+	return &Singleton;
+}
+
+/**
+ * Adds an output device to the chain of redirections.	
+ *
+ * @param OutputDevice	output device to add
+ */
+void FOutputDeviceRedirector::AddOutputDevice( FOutputDevice* OutputDevice )
+{
+	FScopeLock ScopeLock( &SynchronizationObject );
+
+	if( OutputDevice )
+	{
+		OutputDevices.AddUnique( OutputDevice );
+	}
+}
+
+/**
+ * Removes an output device from the chain of redirections.	
+ *
+ * @param OutputDevice	output device to remove
+ */
+void FOutputDeviceRedirector::RemoveOutputDevice( FOutputDevice* OutputDevice )
+{
+	FScopeLock ScopeLock( &SynchronizationObject );
+	OutputDevices.Remove( OutputDevice );
+}
+
+/**
+ * Returns whether an output device is currently in the list of redirectors.
+ *
+ * @param	OutputDevice	output device to check the list against
+ * @return	true if messages are currently redirected to the the passed in output device, false otherwise
+ */
+bool FOutputDeviceRedirector::IsRedirectingTo( FOutputDevice* OutputDevice )
+{
+	FScopeLock ScopeLock( &SynchronizationObject );
+
+	return OutputDevices.Find( OutputDevice ) == INDEX_NONE ? false : true;
+}
+
+/**
+ * The unsynchronized version of FlushThreadedLogs.
+ * Assumes that the caller holds a lock on SynchronizationObject.
+ */
+void FOutputDeviceRedirector::UnsynchronizedFlushThreadedLogs( bool bUseAllDevices )
+{
+	for(int32 LineIndex = 0;LineIndex < BufferedLines.Num();LineIndex++)
+	{
+		FBufferedLine BufferedLine = BufferedLines[LineIndex];
+
+		for( int32 OutputDeviceIndex=0; OutputDeviceIndex<OutputDevices.Num(); OutputDeviceIndex++ )
+		{
+			FOutputDevice* OutputDevice = OutputDevices[OutputDeviceIndex];
+			if( OutputDevice->CanBeUsedOnAnyThread() || bUseAllDevices )
+			{
+				OutputDevice->Serialize( *BufferedLine.Data, BufferedLine.Verbosity, BufferedLine.Category, BufferedLine.Time );
+			}
+		}
+	}
+
+	BufferedLines.Empty();
+}
+
+/**
+ * Flushes lines buffered by secondary threads.
+ */
+void FOutputDeviceRedirector::FlushThreadedLogs()
+{
+	SCOPE_CYCLE_COUNTER(STAT_FlushThreadedLogs);
+	// Acquire a lock on SynchronizationObject and call the unsynchronized worker function.
+	FScopeLock ScopeLock( &SynchronizationObject );
+	check(IsInGameThread());
+	UnsynchronizedFlushThreadedLogs( true );
+}
+
+void FOutputDeviceRedirector::PanicFlushThreadedLogs()
+{
+//	SCOPE_CYCLE_COUNTER(STAT_FlushThreadedLogs);
+	// Acquire a lock on SynchronizationObject and call the unsynchronized worker function.
+	FScopeLock ScopeLock( &SynchronizationObject );
+	
+	// Flush threaded logs, but use the safe version.
+	UnsynchronizedFlushThreadedLogs( false );
+
+	// Flush devices.
+	for (int32 OutputDeviceIndex = 0; OutputDeviceIndex<OutputDevices.Num(); OutputDeviceIndex++)
+	{
+		FOutputDevice* OutputDevice = OutputDevices[OutputDeviceIndex];
+		if (OutputDevice->CanBeUsedOnAnyThread())
+		{
+			OutputDevice->Flush();
+		}
+	}
+
+	BufferedLines.Empty();
+}
+
+/**
+ * Serializes the current backlog to the specified output device.
+ * @param OutputDevice	- Output device that will receive the current backlog
+ */
+void FOutputDeviceRedirector::SerializeBacklog( FOutputDevice* OutputDevice )
+{
+	FScopeLock ScopeLock( &SynchronizationObject );
+
+	for (int32 LineIndex = 0; LineIndex < BacklogLines.Num(); LineIndex++)
+	{
+		const FBufferedLine& BacklogLine = BacklogLines[ LineIndex ];
+		OutputDevice->Serialize( *BacklogLine.Data, BacklogLine.Verbosity, BacklogLine.Category, BacklogLine.Time );
+	}
+}
+
+/**
+ * Enables or disables the backlog.
+ * @param bEnable	- Starts saving a backlog if true, disables and discards any backlog if false
+ */
+void FOutputDeviceRedirector::EnableBacklog( bool bEnable )
+{
+	FScopeLock ScopeLock( &SynchronizationObject );
+
+	bEnableBacklog = bEnable;
+	if ( bEnableBacklog == false )
+	{
+		BacklogLines.Empty();
+	}
+}
+
+/**
+ * Sets the current thread to be the master thread that prints directly
+ * (isn't queued up)
+ */
+void FOutputDeviceRedirector::SetCurrentThreadAsMasterThread()
+{
+	FScopeLock ScopeLock( &SynchronizationObject );
+
+	// make sure anything queued up is flushed out, this may be called from a background thread, so use the safe version.
+	UnsynchronizedFlushThreadedLogs( false );
+
+	// set the current thread as the master thread
+	MasterThreadID = FPlatformTLS::GetCurrentThreadId();
+}
+
+void FOutputDeviceRedirector::Serialize( const TCHAR* Data, ELogVerbosity::Type Verbosity, const class FName& Category, const double Time )
+{
+	const double RealTime = Time == -1.0f ? FPlatformTime::Seconds() - GStartTime : Time;
+
+	FScopeLock ScopeLock( &SynchronizationObject );
+
+#if PLATFORM_DESKTOP
+	// this is for errors which occur after shutdown we might be able to salvage information from stdout 
+	if ((OutputDevices.Num() == 0)&& GIsRequestingExit)
+	{
+#if PLATFORM_WINDOWS
+		_tprintf(_T("%s\n"), Data);
+#else
+		FGenericPlatformMisc::LocalPrint(Data);
+		// printf("%s\n", TCHAR_TO_ANSI(Data));
+#endif
+		return;
+	}
+#endif
+
+
+	if ( bEnableBacklog )
+	{
+		new(BacklogLines)FBufferedLine( Data, Category, Verbosity, RealTime );
+	}
+
+	if(FPlatformTLS::GetCurrentThreadId() != MasterThreadID || OutputDevices.Num() == 0)
+	{
+		new(BufferedLines)FBufferedLine( Data, Category, Verbosity, RealTime );
+	}
+	else
+	{
+		// Flush previously buffered lines from secondary threads.
+		// Since we already hold a lock on SynchronizationObject, call the unsynchronized version.
+		UnsynchronizedFlushThreadedLogs( true );
+
+		for( int32 OutputDeviceIndex=0; OutputDeviceIndex<OutputDevices.Num(); OutputDeviceIndex++ )
+		{
+			OutputDevices[OutputDeviceIndex]->Serialize( Data, Verbosity, Category, RealTime );
+		}
+	}
+}
+
+void FOutputDeviceRedirector::Serialize( const TCHAR* Data, ELogVerbosity::Type Verbosity, const class FName& Category )
+{
+	Serialize( Data, Verbosity, Category, -1.0 );
+}
+
+/**
+ * Passes on the flush request to all current output devices.
+ */
+void FOutputDeviceRedirector::Flush()
+{
+	if(FPlatformTLS::GetCurrentThreadId() == MasterThreadID)
+	{
+		FScopeLock ScopeLock( &SynchronizationObject );
+
+		// Flush previously buffered lines from secondary threads.
+		// Since we already hold a lock on SynchronizationObject, call the unsynchronized version.
+		UnsynchronizedFlushThreadedLogs( true );
+
+		for( int32 OutputDeviceIndex=0; OutputDeviceIndex<OutputDevices.Num(); OutputDeviceIndex++ )
+		{
+			OutputDevices[OutputDeviceIndex]->Flush();
+		}
+	}
+}
+
+/**
+ * Closes output device and cleans up. This can't happen in the destructor
+ * as we might have to call "delete" which cannot be done for static/ global
+ * objects.
+ */
+void FOutputDeviceRedirector::TearDown()
+{
+	check(FPlatformTLS::GetCurrentThreadId() == MasterThreadID);
+
+	FScopeLock ScopeLock( &SynchronizationObject );
+
+	// Flush previously buffered lines from secondary threads.
+	// Since we already hold a lock on SynchronizationObject, call the unsynchronized version.
+	UnsynchronizedFlushThreadedLogs( false );
+
+	for( int32 OutputDeviceIndex=0; OutputDeviceIndex<OutputDevices.Num(); OutputDeviceIndex++ )
+	{
+		OutputDevices[OutputDeviceIndex]->TearDown();
+	}
+	OutputDevices.Empty();
+}
+
+CORE_API FOutputDeviceRedirector* GetGlobalLogSingleton()
+{
+	return FOutputDeviceRedirector::Get();
+}
+
