Copying //UE4/Dev-Build to //UE4/Dev-Main (Source: //UE4/Dev-Build @ 3104200)

diff --git a/Engine/Source/Runtime/Core/Private/Misc/SlowTask.cpp b/Engine/Source/Runtime/Core/Private/Misc/SlowTask.cpp
new file mode 100644
--- /dev/null
+++ b/Engine/Source/Runtime/Core/Private/Misc/SlowTask.cpp
@@ -0,0 +1,111 @@
+// Copyright 1998-2016 Epic Games, Inc. All Rights Reserved.
+
+#include "CorePrivatePCH.h"
+#include "SlowTask.h"
+
+FSlowTask::FSlowTask(float InAmountOfWork, const FText& InDefaultMessage, bool bInEnabled, FFeedbackContext& InContext)
+	: DefaultMessage(InDefaultMessage)
+	, FrameMessage()
+	, TotalAmountOfWork(InAmountOfWork)
+	, CompletedWork(0)
+	, CurrentFrameScope(0)
+	, Visibility(ESlowTaskVisibility::Default)
+	, StartTime(FPlatformTime::Seconds())
+	, bEnabled(bInEnabled && IsInGameThread())
+	, bCreatedDialog(false)		// only set to true if we create a dialog
+	, Context(InContext)
+{
+	// If we have no work to do ourselves, create an arbitrary scope so that any actions performed underneath this still contribute to this one.
+	if (TotalAmountOfWork == 0.f)
+	{
+		TotalAmountOfWork = CurrentFrameScope = 1.f;
+	}
+}
+
+void FSlowTask::Initialize()
+{
+	if (bEnabled)
+	{
+		Context.ScopeStack->Push(this);
+	}
+}
+
+void FSlowTask::Destroy()
+{
+	if (bEnabled)
+	{
+		if (bCreatedDialog)
+		{
+			checkSlow(GIsSlowTask);
+			Context.FinalizeSlowTask();
+		}
+
+		FSlowTaskStack& Stack = *Context.ScopeStack;
+		checkSlow(Stack.Num() != 0 && Stack.Last() == this);
+
+		auto* Task = Stack.Last();
+		if (ensureMsgf(Task == this, TEXT("Out-of-order scoped task construction/destruction")))
+		{
+			Stack.Pop(false);
+		}
+		else
+		{
+			Stack.RemoveSingleSwap(this, false);
+		}
+
+		if (Stack.Num() != 0)
+		{
+			// Stop anything else contributing to the parent frame
+			auto* Parent = Stack.Last();
+			Parent->EnterProgressFrame(0, Parent->FrameMessage);
+		}
+
+	}
+}
+
+void FSlowTask::MakeDialogDelayed(float Threshold, bool bShowCancelButton, bool bAllowInPIE)
+{
+	OpenDialogThreshold = Threshold;
+	bDelayedDialogShowCancelButton = bShowCancelButton;
+	bDelayedDialogAllowInPIE = bAllowInPIE;
+}
+
+void FSlowTask::EnterProgressFrame(float ExpectedWorkThisFrame, FText Text)
+{
+	FrameMessage = Text;
+	CompletedWork += CurrentFrameScope;
+
+	const float WorkRemaining = TotalAmountOfWork - CompletedWork;
+	// Add a small threshold here because when there are a lot of tasks, numerical imprecision can add up and trigger this.
+	ensureMsgf(ExpectedWorkThisFrame <= 1.01f * TotalAmountOfWork - CompletedWork, TEXT("Work overflow in slow task. Please revise call-site to account for entire progress range."));
+	CurrentFrameScope = FMath::Min(WorkRemaining, ExpectedWorkThisFrame);
+
+	if (!bCreatedDialog && OpenDialogThreshold.IsSet() && static_cast<float>(FPlatformTime::Seconds() - StartTime) > OpenDialogThreshold.GetValue())
+	{
+		MakeDialog(bDelayedDialogShowCancelButton, bDelayedDialogAllowInPIE);
+	}
+
+	if (bEnabled)
+	{
+		Context.RequestUpdateUI(bCreatedDialog || (*Context.ScopeStack)[0] == this);
+	}
+}
+
+FText FSlowTask::GetCurrentMessage() const
+{
+	return FrameMessage.IsEmpty() ? DefaultMessage : FrameMessage;
+}
+
+void FSlowTask::MakeDialog(bool bShowCancelButton, bool bAllowInPIE)
+{
+	const bool bIsDisabledByPIE = Context.IsPlayingInEditor() && !bAllowInPIE;
+	const bool bIsDialogAllowed = bEnabled && !GIsSilent && !bIsDisabledByPIE && !IsRunningCommandlet() && IsInGameThread();
+	if (!GIsSlowTask && bIsDialogAllowed)
+	{
+		Context.StartSlowTask(GetCurrentMessage(), bShowCancelButton);
+		if (GIsSlowTask)
+		{
+			bCreatedDialog = true;
+		}
+	}
+}
diff --git a/Engine/Source/Runtime/Core/Private/Serialization/Archive.cpp b/Engine/Source/Runtime/Core/Private/Serialization/Archive.cpp
--- a/Engine/Source/Runtime/Core/Private/Serialization/Archive.cpp
+++ b/Engine/Source/Runtime/Core/Private/Serialization/Archive.cpp
@@ -10,17 +10,7 @@
 #include "NetworkVersion.h"
 #include "TargetPlatform.h"
 #include "GenericPlatformCompression.h"
-
-/*-----------------------------------------------------------------------------
-	FArchiveProxy implementation.
------------------------------------------------------------------------------*/
-
-FArchiveProxy::FArchiveProxy(FArchive& InInnerArchive)
-: FArchive    (InInnerArchive)
-, InnerArchive(InInnerArchive)
-{
-}
-
+#include "CompressedChunkInfo.h"
 
 /*-----------------------------------------------------------------------------
 	FArchive implementation.
@@ -348,18 +338,6 @@ FArchive& FNameAsStringProxyArchive::operator<<( class FName& N )
 	return *this;
 }
 
-/** 
- * FCompressedChunkInfo serialize operator.
- */
-FArchive& operator<<( FArchive& Ar, FCompressedChunkInfo& Chunk )
-{
-	// The order of serialization needs to be identical to the memory layout as the async IO code is reading it in bulk.
-	// The size of the structure also needs to match what is being serialized.
-	Ar << Chunk.CompressedSize;
-	Ar << Chunk.UncompressedSize;
-	return Ar;
-}
-
 /** Accumulative time spent in IsSaving portion of SerializeCompressed. */
 CORE_API double GArchiveSerializedCompressedSavingTime = 0;
 
@@ -911,418 +889,5 @@ VARARG_BODY( void, FArchive::Logf, const TCHAR*, VARARG_NONE )
 	FMemory::SystemFree( Buffer );
 }
 
-/*----------------------------------------------------------------------------
-	Transparent compression/ decompression archives.
-----------------------------------------------------------------------------*/
-
-/*----------------------------------------------------------------------------
-	FArchiveSaveCompressedProxy
-----------------------------------------------------------------------------*/
-
-/** 
- * Constructor, initializing all member variables and allocating temp memory.
- *
- * @param	InCompressedData [ref]	Array of bytes that is going to hold compressed data
- * @param	InCompressionFlags		Compression flags to use for compressing data
- */
-FArchiveSaveCompressedProxy::FArchiveSaveCompressedProxy( TArray<uint8>& InCompressedData, ECompressionFlags InCompressionFlags )
-:	CompressedData(InCompressedData)
-,	CompressionFlags(InCompressionFlags)
-{
-	ArIsSaving							= true;
-	ArIsPersistent						= true;
-	ArWantBinaryPropertySerialization	= true;
-	bShouldSerializeToArray				= false;
-	RawBytesSerialized					= 0;
-	CurrentIndex						= 0;
-
-	// Allocate temporary memory.
-	TmpDataStart	= (uint8*) FMemory::Malloc(LOADING_COMPRESSION_CHUNK_SIZE);
-	TmpDataEnd		= TmpDataStart + LOADING_COMPRESSION_CHUNK_SIZE;
-	TmpData			= TmpDataStart;
-}
-
-/** Destructor, flushing array if needed. Also frees temporary memory. */
-FArchiveSaveCompressedProxy::~FArchiveSaveCompressedProxy()
-{
-	// Flush is required to write out remaining tmp data to array.
-	Flush();
-	// Free temporary memory allocated.
-	FMemory::Free( TmpDataStart );
-	TmpDataStart	= NULL;
-	TmpDataEnd		= NULL;
-	TmpData			= NULL;
-}
-
-/**
- * Flushes tmp data to array.
- */
-void FArchiveSaveCompressedProxy::Flush()
-{
-	if( TmpData - TmpDataStart > 0 )
-	{
-		// This will call Serialize so we need to indicate that we want to serialize to array.
-		bShouldSerializeToArray = true;
-		SerializeCompressed( TmpDataStart, TmpData - TmpDataStart, CompressionFlags );
-		bShouldSerializeToArray = false;
-		// Buffer is drained, reset.
-		TmpData	= TmpDataStart;
-	}
-}
-
-/**
- * Serializes data to archive. This function is called recursively and determines where to serialize
- * to and how to do so based on internal state.
- *
- * @param	Data	Pointer to serialize to
- * @param	Count	Number of bytes to read
- */
-void FArchiveSaveCompressedProxy::Serialize( void* InData, int64 Count )
-{
-	uint8* SrcData = (uint8*) InData;
-	// If counter > 1 it means we're calling recursively and therefore need to write to compressed data.
-	if( bShouldSerializeToArray )
-	{
-		// Add space in array if needed and copy data there.
-		int32 BytesToAdd = CurrentIndex + Count - CompressedData.Num();
-		if( BytesToAdd > 0 )
-		{
-			CompressedData.AddUninitialized(BytesToAdd);
-		}
-		// Copy memory to array.
-		FMemory::Memcpy( &CompressedData[CurrentIndex], SrcData, Count );
-		CurrentIndex += Count;
-	}
-	// Regular call to serialize, queue for compression.
-	else
-	{
-		while( Count )
-		{
-			int32 BytesToCopy = FMath::Min<int32>( Count, (int32)(TmpDataEnd - TmpData) );
-			// Enough room in buffer to copy some data.
-			if( BytesToCopy )
-			{
-				FMemory::Memcpy( TmpData, SrcData, BytesToCopy );
-				Count -= BytesToCopy;
-				TmpData += BytesToCopy;
-				SrcData += BytesToCopy;
-				RawBytesSerialized += BytesToCopy;
-			}
-			// Tmp buffer fully exhausted, compress it.
-			else
-			{
-				// Flush existing data to array after compressing it. This will call Serialize again 
-				// so we need to handle recursion.
-				Flush();
-			}
-		}
-	}
-}
-
-/**
- * Seeking is only implemented internally for writing out compressed data and asserts otherwise.
- * 
- * @param	InPos	Position to seek to
- */
-void FArchiveSaveCompressedProxy::Seek( int64 InPos )
-{
-	// Support setting position in array.
-	if( bShouldSerializeToArray )
-	{
-		CurrentIndex = InPos;
-	}
-	else
-	{
-		UE_LOG(LogSerialization, Fatal,TEXT("Seeking not supported with FArchiveSaveCompressedProxy"));
-	}
-}
-
-/**
- * @return current position in uncompressed stream in bytes
- */
-int64 FArchiveSaveCompressedProxy::Tell()
-{
-	// If we're serializing to array, return position in array.
-	if( bShouldSerializeToArray )
-	{
-		return CurrentIndex;
-	}
-	// Return global position in raw uncompressed stream.
-	else
-	{
-		return RawBytesSerialized;
-	}
-}
-
-
-/*----------------------------------------------------------------------------
-	FArchiveLoadCompressedProxy
-----------------------------------------------------------------------------*/
-
-/** 
- * Constructor, initializing all member variables and allocating temp memory.
- *
- * @param	InCompressedData	Array of bytes that is holding compressed data
- * @param	InCompressionFlags	Compression flags that were used to compress data
- */
-FArchiveLoadCompressedProxy::FArchiveLoadCompressedProxy( const TArray<uint8>& InCompressedData, ECompressionFlags InCompressionFlags )
-:	CompressedData(InCompressedData)
-,	CompressionFlags(InCompressionFlags)
-{
-	ArIsLoading							= true;
-	ArIsPersistent						= true;
-	ArWantBinaryPropertySerialization	= true;
-	bShouldSerializeFromArray			= false;
-	RawBytesSerialized					= 0;
-	CurrentIndex						= 0;
-
-	// Allocate temporary memory.
-	TmpDataStart	= (uint8*) FMemory::Malloc(LOADING_COMPRESSION_CHUNK_SIZE);
-	TmpDataEnd		= TmpDataStart + LOADING_COMPRESSION_CHUNK_SIZE;
-	TmpData			= TmpDataEnd;
-}
-
-/** Destructor, freeing temporary memory. */
-FArchiveLoadCompressedProxy::~FArchiveLoadCompressedProxy()
-{
-	// Free temporary memory allocated.
-	FMemory::Free( TmpDataStart );
-	TmpDataStart	= NULL;
-	TmpDataEnd		= NULL;
-	TmpData			= NULL;
-}
-
-/**
- * Flushes tmp data to array.
- */
-void FArchiveLoadCompressedProxy::DecompressMoreData()
-{
-	// This will call Serialize so we need to indicate that we want to serialize from array.
-	bShouldSerializeFromArray = true;
-	SerializeCompressed( TmpDataStart, LOADING_COMPRESSION_CHUNK_SIZE /** it's ignored, but that's how much we serialize */, CompressionFlags );
-	bShouldSerializeFromArray = false;
-	// Buffer is filled again, reset.
-	TmpData = TmpDataStart;
-}
-
-/**
- * Serializes data from archive. This function is called recursively and determines where to serialize
- * from and how to do so based on internal state.
- *
- * @param	InData	Pointer to serialize to
- * @param	Count	Number of bytes to read
- */
-void FArchiveLoadCompressedProxy::Serialize( void* InData, int64 Count )
-{
-	uint8* DstData = (uint8*) InData;
-	// If counter > 1 it means we're calling recursively and therefore need to write to compressed data.
-	if( bShouldSerializeFromArray )
-	{
-		// Add space in array and copy data there.
-		check(CurrentIndex+Count<=CompressedData.Num());
-		FMemory::Memcpy( DstData, &CompressedData[CurrentIndex], Count );
-		CurrentIndex += Count;
-	}
-	// Regular call to serialize, read from temp buffer
-	else
-	{	
-		while( Count )
-		{
-			int32 BytesToCopy = FMath::Min<int32>( Count, (int32)(TmpDataEnd - TmpData) );
-			// Enough room in buffer to copy some data.
-			if( BytesToCopy )
-			{
-				// We pass in a NULL pointer when forward seeking. In that case we don't want
-				// to copy the data but only care about pointing to the proper spot.
-				if( DstData )
-				{
-					FMemory::Memcpy( DstData, TmpData, BytesToCopy );
-					DstData += BytesToCopy;
-				}
-				Count -= BytesToCopy;
-				TmpData += BytesToCopy;
-				RawBytesSerialized += BytesToCopy;
-			}
-			// Tmp buffer fully exhausted, decompress new one.
-			else
-			{
-				// Decompress more data. This will call Serialize again so we need to handle recursion.
-				DecompressMoreData();
-			}
-		}
-	}
-}
-
-/**
- * Seeks to the passed in position in the stream. This archive only supports forward seeking
- * and implements it by serializing data till it reaches the position.
- */
-void FArchiveLoadCompressedProxy::Seek( int64 InPos )
-{
-	int64 CurrentPos = Tell();
-	int64 Difference = InPos - CurrentPos;
-	// We only support forward seeking.
-	check(Difference>=0);
-	// Seek by serializing data, albeit with NULL destination so it's just decompressing data.
-	Serialize( NULL, Difference );
-}
-
-/**
- * @return current position in uncompressed stream in bytes.
- */
-int64 FArchiveLoadCompressedProxy::Tell()
-{
-	return RawBytesSerialized;
-}
-
-
-/*----------------------------------------------------------------------------
-	FLargeMemoryWriter
-----------------------------------------------------------------------------*/
-
-FLargeMemoryWriter::FLargeMemoryWriter(const int64 PreAllocateBytes, bool bIsPersistent, const FName InArchiveName)
-	: FMemoryArchive()
-	, Data(nullptr)
-	, NumBytes(0)
-	, MaxBytes(0)
-	, ArchiveName(InArchiveName)
-{
-	ArIsSaving = true;
-	ArIsPersistent = bIsPersistent;
-	GrowBuffer(PreAllocateBytes);
-}
-
-void FLargeMemoryWriter::Serialize(void* InData, int64 Num)
-{
-	UE_CLOG(!Data, LogSerialization, Fatal, TEXT("Tried to serialize data to an FLargeMemoryWriter that was already released. Archive name: %s."), *ArchiveName.ToString());
-	
-	const int64 NumBytesToAdd = Offset + Num - NumBytes;
-	if (NumBytesToAdd > 0)
-	{
-		const int64 NewByteCount = NumBytes + NumBytesToAdd;
-		if (NewByteCount > MaxBytes)
-		{
-			GrowBuffer(NewByteCount);
-		}
-
-		NumBytes = NewByteCount;
-	}
-
-	check((Offset + Num) <= NumBytes);
-
-	if (Num)
-	{
-		FMemory::Memcpy(&Data[Offset], InData, Num);
-		Offset += Num;
-	}
-}
-
-FString FLargeMemoryWriter::GetArchiveName() const
-{
-	return ArchiveName.ToString();
-}
-
-int64 FLargeMemoryWriter::TotalSize()
-{
-	return NumBytes;
-}
-
-uint8* FLargeMemoryWriter::GetData() const
-{
-	UE_CLOG(!Data, LogSerialization, Warning, TEXT("Tried to get written data from an FLargeMemoryWriter that was already released. Archive name: %s."), *ArchiveName.ToString());
-
-	return Data;
-}
 
-void FLargeMemoryWriter::ReleaseOwnership()
-{
-	Data = nullptr;
-	NumBytes = 0;
-	MaxBytes = 0;
-}
-
-FLargeMemoryWriter::~FLargeMemoryWriter()
-{
-	if (Data)
-	{
-		FMemory::Free(Data);
-	}
-}
-
-void FLargeMemoryWriter::GrowBuffer(const int64 DesiredBytes)
-{
-	int64 NewBytes = 4; // Initial alloc size
-
-	if (MaxBytes || DesiredBytes > NewBytes)
-	{
-		// Allocate slack proportional to the buffer size
-		NewBytes = FMemory::QuantizeSize(DesiredBytes + 3 * DesiredBytes / 8 + 16);
-	}
-
-	if (Data)
-	{
-		Data = (uint8*)FMemory::Realloc(Data, NewBytes);
-	}
-	else
-	{
-		Data = (uint8*)FMemory::Malloc(NewBytes);
-	}
-
-	MaxBytes = NewBytes;
-}
-
-
-/*----------------------------------------------------------------------------
-	FLargeMemoryReader
-----------------------------------------------------------------------------*/
-
-FLargeMemoryReader::FLargeMemoryReader(const uint8* InData, const int64 Num, ELargeMemoryReaderFlags InFlags, const FName InArchiveName)
-	: FMemoryArchive()
-	, bFreeOnClose((InFlags & ELargeMemoryReaderFlags::TakeOwnership) != ELargeMemoryReaderFlags::None)
-	, Data(InData)
-	, NumBytes(Num)
-	, ArchiveName(InArchiveName)
-{
-	UE_CLOG(!(InData && Num > 0), LogSerialization, Fatal, TEXT("Tried to initialize an FLargeMemoryReader with a null or empty buffer. Archive name: %s."), *ArchiveName.ToString());
-	ArIsLoading = true;
-	ArIsPersistent = (InFlags & ELargeMemoryReaderFlags::Persistent) != ELargeMemoryReaderFlags::None;
-}
-
-void FLargeMemoryReader::Serialize(void* OutData, int64 Num)
-{
-	if (Num && !ArIsError)
-	{
-		// Only serialize if we have the requested amount of data
-		if (Offset + Num <= NumBytes)
-		{
-			FMemory::Memcpy(OutData, &Data[Offset], Num);
-			Offset += Num;
-		}
-		else
-		{
-			ArIsError = true;
-		}
-	}
-}
-
-int64 FLargeMemoryReader::TotalSize()
-{
-	return NumBytes;
-}
-
-FString FLargeMemoryReader::GetArchiveName() const
-{
-	return ArchiveName.ToString();
-}
-
-FLargeMemoryReader::~FLargeMemoryReader()
-{
-	if (bFreeOnClose)
-	{
-		FMemory::Free((void*)Data);
-	}
-}
 
-/*----------------------------------------------------------------------------
-	The End
-----------------------------------------------------------------------------*/
diff --git a/Engine/Source/Runtime/Core/Private/Serialization/AsyncIOSystemBase.cpp b/Engine/Source/Runtime/Core/Private/Serialization/AsyncIOSystemBase.cpp
--- a/Engine/Source/Runtime/Core/Private/Serialization/AsyncIOSystemBase.cpp
+++ b/Engine/Source/Runtime/Core/Private/Serialization/AsyncIOSystemBase.cpp
@@ -6,6 +6,7 @@
 
 #include "CorePrivatePCH.h"
 #include "Serialization/AsyncIOSystemBase.h"
+#include "CompressedChunkInfo.h"
 
 DECLARE_STATS_GROUP_VERBOSE(TEXT("AsyncIOSystem"),STATGROUP_AsyncIO_Verbose, STATCAT_Advanced);
 
diff --git a/Engine/Source/Runtime/Core/Private/Serialization/LargeMemoryReader.cpp b/Engine/Source/Runtime/Core/Private/Serialization/LargeMemoryReader.cpp
new file mode 100644
--- /dev/null
+++ b/Engine/Source/Runtime/Core/Private/Serialization/LargeMemoryReader.cpp
@@ -0,0 +1,55 @@
+// Copyright 1998-2016 Epic Games, Inc. All Rights Reserved.
+
+#include "CorePrivatePCH.h"
+#include "Serialization/Archive.h"
+
+/*----------------------------------------------------------------------------
+	FLargeMemoryReader
+----------------------------------------------------------------------------*/
+
+FLargeMemoryReader::FLargeMemoryReader(const uint8* InData, const int64 Num, ELargeMemoryReaderFlags InFlags, const FName InArchiveName)
+	: FMemoryArchive()
+	, bFreeOnClose((InFlags & ELargeMemoryReaderFlags::TakeOwnership) != ELargeMemoryReaderFlags::None)
+	, Data(InData)
+	, NumBytes(Num)
+	, ArchiveName(InArchiveName)
+{
+	UE_CLOG(!(InData && Num > 0), LogSerialization, Fatal, TEXT("Tried to initialize an FLargeMemoryReader with a null or empty buffer. Archive name: %s."), *ArchiveName.ToString());
+	ArIsLoading = true;
+	ArIsPersistent = (InFlags & ELargeMemoryReaderFlags::Persistent) != ELargeMemoryReaderFlags::None;
+}
+
+void FLargeMemoryReader::Serialize(void* OutData, int64 Num)
+{
+	if (Num && !ArIsError)
+	{
+		// Only serialize if we have the requested amount of data
+		if (Offset + Num <= NumBytes)
+		{
+			FMemory::Memcpy(OutData, &Data[Offset], Num);
+			Offset += Num;
+		}
+		else
+		{
+			ArIsError = true;
+		}
+	}
+}
+
+int64 FLargeMemoryReader::TotalSize()
+{
+	return NumBytes;
+}
+
+FString FLargeMemoryReader::GetArchiveName() const
+{
+	return ArchiveName.ToString();
+}
+
+FLargeMemoryReader::~FLargeMemoryReader()
+{
+	if (bFreeOnClose)
+	{
+		FMemory::Free((void*)Data);
+	}
+}
diff --git a/Engine/Source/Runtime/Core/Private/Serialization/LargeMemoryWriter.cpp b/Engine/Source/Runtime/Core/Private/Serialization/LargeMemoryWriter.cpp
new file mode 100644
--- /dev/null
+++ b/Engine/Source/Runtime/Core/Private/Serialization/LargeMemoryWriter.cpp
@@ -0,0 +1,100 @@
+// Copyright 1998-2016 Epic Games, Inc. All Rights Reserved.
+
+#include "CorePrivatePCH.h"
+#include "Serialization/LargeMemoryWriter.h"
+
+/*----------------------------------------------------------------------------
+	FLargeMemoryWriter
+----------------------------------------------------------------------------*/
+
+FLargeMemoryWriter::FLargeMemoryWriter(const int64 PreAllocateBytes, bool bIsPersistent, const FName InArchiveName)
+	: FMemoryArchive()
+	, Data(nullptr)
+	, NumBytes(0)
+	, MaxBytes(0)
+	, ArchiveName(InArchiveName)
+{
+	ArIsSaving = true;
+	ArIsPersistent = bIsPersistent;
+	GrowBuffer(PreAllocateBytes);
+}
+
+void FLargeMemoryWriter::Serialize(void* InData, int64 Num)
+{
+	UE_CLOG(!Data, LogSerialization, Fatal, TEXT("Tried to serialize data to an FLargeMemoryWriter that was already released. Archive name: %s."), *ArchiveName.ToString());
+	
+	const int64 NumBytesToAdd = Offset + Num - NumBytes;
+	if (NumBytesToAdd > 0)
+	{
+		const int64 NewByteCount = NumBytes + NumBytesToAdd;
+		if (NewByteCount > MaxBytes)
+		{
+			GrowBuffer(NewByteCount);
+		}
+
+		NumBytes = NewByteCount;
+	}
+
+	check((Offset + Num) <= NumBytes);
+
+	if (Num)
+	{
+		FMemory::Memcpy(&Data[Offset], InData, Num);
+		Offset += Num;
+	}
+}
+
+FString FLargeMemoryWriter::GetArchiveName() const
+{
+	return ArchiveName.ToString();
+}
+
+int64 FLargeMemoryWriter::TotalSize()
+{
+	return NumBytes;
+}
+
+uint8* FLargeMemoryWriter::GetData() const
+{
+	UE_CLOG(!Data, LogSerialization, Warning, TEXT("Tried to get written data from an FLargeMemoryWriter that was already released. Archive name: %s."), *ArchiveName.ToString());
+
+	return Data;
+}
+
+void FLargeMemoryWriter::ReleaseOwnership()
+{
+	Data = nullptr;
+	NumBytes = 0;
+	MaxBytes = 0;
+}
+
+FLargeMemoryWriter::~FLargeMemoryWriter()
+{
+	if (Data)
+	{
+		FMemory::Free(Data);
+	}
+}
+
+void FLargeMemoryWriter::GrowBuffer(const int64 DesiredBytes)
+{
+	int64 NewBytes = 4; // Initial alloc size
+
+	if (MaxBytes || DesiredBytes > NewBytes)
+	{
+		// Allocate slack proportional to the buffer size
+		NewBytes = FMemory::QuantizeSize(DesiredBytes + 3 * DesiredBytes / 8 + 16);
+	}
+
+	if (Data)
+	{
+		Data = (uint8*)FMemory::Realloc(Data, NewBytes);
+	}
+	else
+	{
+		Data = (uint8*)FMemory::Malloc(NewBytes);
+	}
+
+	MaxBytes = NewBytes;
+}
+
diff --git a/Engine/Source/Runtime/Core/Private/Serialization/NameAsStringProxyArchive.cpp b/Engine/Source/Runtime/Core/Private/Serialization/NameAsStringProxyArchive.cpp
new file mode 100644
--- /dev/null
+++ b/Engine/Source/Runtime/Core/Private/Serialization/NameAsStringProxyArchive.cpp
@@ -0,0 +1,4 @@
+// Copyright 1998-2016 Epic Games, Inc. All Rights Reserved.
+
+#include "CorePrivatePCH.h"
+#include "Serialization/Archive.h"
diff --git a/Engine/Source/Runtime/Core/Private/Windows/WindowsPlatformOutputDevices.cpp b/Engine/Source/Runtime/Core/Private/Windows/WindowsPlatformOutputDevices.cpp
--- a/Engine/Source/Runtime/Core/Private/Windows/WindowsPlatformOutputDevices.cpp
+++ b/Engine/Source/Runtime/Core/Private/Windows/WindowsPlatformOutputDevices.cpp
@@ -131,7 +131,8 @@ void FOutputDeviceWindowsError::HandleError()
 
 	// Dump the error and flush the log.
 #if !NO_LOGGING
-	FDebug::OutputMultiLineCallstack(__FILE__, __LINE__, LogWindows.GetCategoryName(), TEXT("=== Critical error: ==="), GErrorHist, ELogVerbosity::Error);
+	extern void OutputMultiLineCallstack(const ANSICHAR* File, int32 Line, const FName& LogName, const TCHAR* Heading, TCHAR* Message, ELogVerbosity::Type Verbosity);
+	OutputMultiLineCallstack(__FILE__, __LINE__, LogWindows.GetCategoryName(), TEXT("=== Critical error: ==="), GErrorHist, ELogVerbosity::Error);
 #endif
 	GLog->PanicFlushThreadedLogs();
 
diff --git a/Engine/Source/Runtime/Core/Public/Android/AndroidCriticalSection.h b/Engine/Source/Runtime/Core/Public/Android/AndroidCriticalSection.h
new file mode 100644
--- /dev/null
+++ b/Engine/Source/Runtime/Core/Public/Android/AndroidCriticalSection.h
@@ -0,0 +1,9 @@
+// Copyright 1998-2016 Epic Games, Inc. All Rights Reserved.
+
+#pragma once
+
+#include "PThreadCriticalSection.h"
+#include "GenericPlatform/GenericPlatformCriticalSection.h"
+
+typedef FPThreadsCriticalSection FCriticalSection;
+typedef FSystemWideCriticalSectionNotImplemented FSystemWideCriticalSection;
diff --git a/Engine/Source/Runtime/Core/Public/Android/AndroidProcess.h b/Engine/Source/Runtime/Core/Public/Android/AndroidProcess.h
--- a/Engine/Source/Runtime/Core/Public/Android/AndroidProcess.h
+++ b/Engine/Source/Runtime/Core/Public/Android/AndroidProcess.h
@@ -40,4 +40,3 @@ struct CORE_API FAndroidPlatformProcess : public FGenericPlatformProcess
 
 typedef FAndroidPlatformProcess FPlatformProcess;
 
-typedef FSystemWideCriticalSectionNotImplemented FSystemWideCriticalSection;
diff --git a/Engine/Source/Runtime/Core/Public/Apple/ApplePlatformSymbolication.h b/Engine/Source/Runtime/Core/Public/Apple/ApplePlatformSymbolication.h
--- a/Engine/Source/Runtime/Core/Public/Apple/ApplePlatformSymbolication.h
+++ b/Engine/Source/Runtime/Core/Public/Apple/ApplePlatformSymbolication.h
@@ -8,7 +8,7 @@
 
 #include "GenericPlatformSymbolication.h"
 #include "ApplePlatformStackWalk.h"
-#include "Serialization/ArchiveBase.h"
+#include "Serialization/Archive.h"
 
 /**
  * Opaque symbol cache for improved symbolisation performance.
diff --git a/Engine/Source/Runtime/Core/Public/Containers/Array.h b/Engine/Source/Runtime/Core/Public/Containers/Array.h
--- a/Engine/Source/Runtime/Core/Public/Containers/Array.h
+++ b/Engine/Source/Runtime/Core/Public/Containers/Array.h
@@ -6,7 +6,8 @@
 
 #include "Containers/ContainerAllocationPolicies.h"
 #include "HAL/Platform.h"
-#include "Serialization/ArchiveBase.h"
+#include "HAL/PlatformMath.h"
+#include "Serialization/Archive.h"
 #include "Templates/EnableIf.h"
 #include "Templates/Sorting.h"
 #include "Templates/UnrealTemplate.h"
@@ -22,6 +23,7 @@
 
 #define AGRESSIVE_ARRAY_FORCEINLINE
 
+class UClass;
 
 /**
  * Generic iterator which can operate on types that expose the following:
@@ -229,207 +231,6 @@ private:
 };
 
 
-/**
- * Base dynamic array.
- * An untyped data array; mirrors a TArray's members, but doesn't need an exact C++ type for its elements.
- **/
-class FScriptArray
-	: protected FHeapAllocator::ForAnyElementType
-{
-public:
-
-	FORCEINLINE void* GetData()
-	{
-		return this->GetAllocation();
-	}
-	FORCEINLINE const void* GetData() const
-	{
-		return this->GetAllocation();
-	}
-	FORCEINLINE bool IsValidIndex(int32 i) const
-	{
-		return i>=0 && i<ArrayNum;
-	}
-	FORCEINLINE int32 Num() const
-	{
-		checkSlow(ArrayNum>=0);
-		checkSlow(ArrayMax>=ArrayNum);
-		return ArrayNum;
-	}
-	void InsertZeroed( int32 Index, int32 Count, int32 NumBytesPerElement )
-	{
-		Insert( Index, Count, NumBytesPerElement );
-		FMemory::Memzero( (uint8*)this->GetAllocation()+Index*NumBytesPerElement, Count*NumBytesPerElement );
-	}
-	void Insert( int32 Index, int32 Count, int32 NumBytesPerElement )
-	{
-		check(Count>=0);
-		check(ArrayNum>=0);
-		check(ArrayMax>=ArrayNum);
-		check(Index>=0);
-		check(Index<=ArrayNum);
-
-		const int32 OldNum = ArrayNum;
-		if( (ArrayNum+=Count)>ArrayMax )
-		{
-			ResizeGrow(OldNum, NumBytesPerElement);
-		}
-		FMemory::Memmove
-		(
-			(uint8*)this->GetAllocation() + (Index+Count )*NumBytesPerElement,
-			(uint8*)this->GetAllocation() + (Index       )*NumBytesPerElement,
-			                                               (OldNum-Index)*NumBytesPerElement
-		);
-	}
-	AGRESSIVE_ARRAY_FORCEINLINE int32 Add( int32 Count, int32 NumBytesPerElement )
-	{
-		check(Count>=0);
-		checkSlow(ArrayNum>=0);
-		checkSlow(ArrayMax>=ArrayNum);
-
-		const int32 OldNum = ArrayNum;
-		if( (ArrayNum+=Count)>ArrayMax )
-		{
-			ResizeGrow(OldNum, NumBytesPerElement);
-		}
-
-		return OldNum;
-	}
-	int32 AddZeroed( int32 Count, int32 NumBytesPerElement )
-	{
-		const int32 Index = Add( Count, NumBytesPerElement );
-		FMemory::Memzero( (uint8*)this->GetAllocation()+Index*NumBytesPerElement, Count*NumBytesPerElement );
-		return Index;
-	}
-	void Shrink( int32 NumBytesPerElement )
-	{
-		checkSlow(ArrayNum>=0);
-		checkSlow(ArrayMax>=ArrayNum);
-		if (ArrayNum != ArrayMax)
-		{
-			ResizeTo(ArrayNum, NumBytesPerElement);
-		}
-	}
-	void Empty( int32 Slack, int32 NumBytesPerElement )
-	{
-		checkSlow(Slack>=0);
-		ArrayNum = 0;
-		if (Slack != ArrayMax)
-		{
-			ResizeTo(Slack, NumBytesPerElement);
-		}
-	}
-	void SwapMemory(int32 A, int32 B, int32 NumBytesPerElement )
-	{
-		FMemory::Memswap(
-			(uint8*)this->GetAllocation()+(NumBytesPerElement*A),
-			(uint8*)this->GetAllocation()+(NumBytesPerElement*B),
-			NumBytesPerElement
-			);
-	}
-	FScriptArray()
-	:   ArrayNum( 0 )
-	,	ArrayMax( 0 )
-	{
-	}
-	void CountBytes( FArchive& Ar, int32 NumBytesPerElement  )
-	{
-		Ar.CountBytes( ArrayNum*NumBytesPerElement, ArrayMax*NumBytesPerElement );
-	}
-	/**
-	 * Returns the amount of slack in this array in elements.
-	 */
-	FORCEINLINE int32 GetSlack() const
-	{
-		return ArrayMax - ArrayNum;
-	}
-		
-	void Remove( int32 Index, int32 Count, int32 NumBytesPerElement  )
-	{
-		if (Count)
-		{
-			checkSlow(Count >= 0);
-			checkSlow(Index >= 0);
-			checkSlow(Index <= ArrayNum);
-			checkSlow(Index + Count <= ArrayNum);
-
-			// Skip memmove in the common case that there is nothing to move.
-			int32 NumToMove = ArrayNum - Index - Count;
-			if (NumToMove)
-			{
-				FMemory::Memmove
-					(
-					(uint8*)this->GetAllocation() + (Index)* NumBytesPerElement,
-					(uint8*)this->GetAllocation() + (Index + Count) * NumBytesPerElement,
-					NumToMove * NumBytesPerElement
-					);
-			}
-			ArrayNum -= Count;
-
-			ResizeShrink(NumBytesPerElement);
-			checkSlow(ArrayNum >= 0);
-			checkSlow(ArrayMax >= ArrayNum);
-		}
-	}
-
-protected:
-
-	FScriptArray( int32 InNum, int32 NumBytesPerElement  )
-	:   ArrayNum( 0 )
-	,	ArrayMax( InNum )
-
-	{
-		if (ArrayMax)
-		{
-			ResizeInit(NumBytesPerElement);
-		}
-		ArrayNum = InNum;
-	}
-	int32	  ArrayNum;
-	int32	  ArrayMax;
-
-	FORCENOINLINE void ResizeInit(int32 NumBytesPerElement)
-	{
-		ArrayMax = this->CalculateSlackReserve(ArrayMax, NumBytesPerElement);
-		this->ResizeAllocation(ArrayNum, ArrayMax, NumBytesPerElement);
-	}
-	FORCENOINLINE void ResizeGrow(int32 OldNum, int32 NumBytesPerElement)
-	{
-		ArrayMax = this->CalculateSlackGrow(ArrayNum, ArrayMax, NumBytesPerElement);
-		this->ResizeAllocation(OldNum, ArrayMax, NumBytesPerElement);
-	}
-	FORCENOINLINE void ResizeShrink(int32 NumBytesPerElement)
-	{
-		const int32 NewArrayMax = this->CalculateSlackShrink(ArrayNum, ArrayMax, NumBytesPerElement);
-		if (NewArrayMax != ArrayMax)
-		{
-			ArrayMax = NewArrayMax;
-			this->ResizeAllocation(ArrayNum, ArrayMax, NumBytesPerElement);
-		}
-	}
-	FORCENOINLINE void ResizeTo(int32 NewMax, int32 NumBytesPerElement)
-	{
-		if (NewMax)
-		{
-			NewMax = this->CalculateSlackReserve(NewMax, NumBytesPerElement);
-		}
-		if (NewMax != ArrayMax)
-		{
-			ArrayMax = NewMax;
-			this->ResizeAllocation(ArrayNum, ArrayMax, NumBytesPerElement);
-		}
-	}
-public:
-	// These should really be private, because they shouldn't be called, but there's a bunch of code
-	// that needs to be fixed first.
-	FScriptArray(const FScriptArray&) { check(false); }
-	void operator=(const FScriptArray&) { check(false); }
-};
-
-
-template<> struct TIsZeroConstructType<FScriptArray> { enum { Value = true }; };
-
-
 /**
  * TReversePredicateWrapper class used by implicit heaps. 
  * This is similar to TDereferenceWrapper from Sorting.h except it reverses the comparison at the same time
@@ -1642,7 +1443,7 @@ public:
 			// Replace the elements in the hole created by the removal with elements from the end of the array, so the range of indices used by the array is contiguous.
 			const int32 NumElementsInHole = Count;
 			const int32 NumElementsAfterHole = ArrayNum - (Index + Count);
-			const int32 NumElementsToMoveIntoHole = FMath::Min(NumElementsInHole, NumElementsAfterHole);
+			const int32 NumElementsToMoveIntoHole = FPlatformMath::Min(NumElementsInHole, NumElementsAfterHole);
 			if (NumElementsToMoveIntoHole)
 			{
 				FMemory::Memcpy(
@@ -2679,7 +2480,7 @@ public:
 
 		TDereferenceWrapper< ElementType, PREDICATE_CLASS> PredicateWrapper(Predicate);
 		SiftDown(Index, Num(), PredicateWrapper);
-		SiftUp(0, FMath::Min(Index, Num() - 1), PredicateWrapper);
+		SiftUp(0, FPlatformMath::Min(Index, Num() - 1), PredicateWrapper);
 
 #if DEBUG_HEAP
 		VerifyHeap(PredicateWrapper);
@@ -2882,866 +2683,3 @@ template <typename T,typename Allocator> void* operator new( size_t Size, TArray
 	return &Array[Index];
 }
 
-/*-----------------------------------------------------------------------------
-	MRU array.
------------------------------------------------------------------------------*/
-
-/**
- * Same as TArray except:
- * - Has an upper limit of the number of items it will store.
- * - Any item that is added to the array is moved to the top.
- */
-template<typename T, typename Allocator = FDefaultAllocator>
-class TMRUArray
-	: public TArray<T, Allocator>
-{
-public:
-	typedef TArray<T, Allocator> Super;
-
-	/** The maximum number of items we can store in this array. */
-	int32 MaxItems;
-
-	/**
-	 * Constructor.
-	 */
-	TMRUArray()
-		: Super()
-	{
-		MaxItems = 0;
-	}
-
-#if PLATFORM_COMPILER_HAS_DEFAULTED_FUNCTIONS
-
-	TMRUArray(TMRUArray&&) = default;
-	TMRUArray(const TMRUArray&) = default;
-	TMRUArray& operator=(TMRUArray&&) = default;
-	TMRUArray& operator=(const TMRUArray&) = default;
-
-#else
-
-	/**
-	 * Copy constructor.
-	 *
-	 * @param Other Other array to copy.
-	 */
-	FORCEINLINE TMRUArray(const TMRUArray& Other)
-		: Super((const Super&)Other)
-	{
-	}
-
-	/**
-	 * Assign operator.
-	 *
-	 * @param Other Other array to assign with.
-	 * @returns Reference to this object.
-	 */
-	FORCEINLINE TMRUArray& operator=(const TMRUArray& Other)
-	{
-		(Super&)*this = (const Super&)Other;
-		return *this;
-	}
-
-	/**
-	 * Move constructor.
-	 *
-	 * @param Other Other array to move.
-	 */
-	FORCEINLINE TMRUArray(TMRUArray&& Other)
-		: Super((TMRUArray&&)Other)
-	{
-	}
-
-	/**
-	 * Move assign operator.
-	 *
-	 * @param Other Other array to assign with.
-	 * @returns Reference to this object.
-	 */
-	FORCEINLINE TMRUArray& operator=(TMRUArray&& Other)
-	{
-		(Super&)*this = (Super&&)Other;
-		return *this;
-	}
-
-#endif
-	/**
-	 * Adds item to the array. Makes sure that we don't add more than the
-	 * limit.
-	 *
-	 * @param Item Item to add.
-	 * @returns Always 0.
-	 */
-	int32 Add(const T& Item)
-	{
-		const int32 idx = Super::Add(Item);
-		this->Swap(idx, 0);
-		CullArray();
-		return 0;
-	}
-
-	/**
-	 * Adds a number of zeroed elements to the array. Makes sure that we don't
-	 * add more than the limit.
-	 *
-	 * @param Count (Optional) A number of elements to add. Default is 0.
-	 * @returns Always 0.
-	 */
-	int32 AddZeroed(int32 Count = 1)
-	{
-		const int32 idx = Super::AddZeroed(Count);
-		this->Swap(idx, 0);
-		CullArray();
-		return 0;
-	}
-
-	/**
-	 * Adds unique item to the array. Makes sure that we don't add more than
-	 * the limit. If the item existed it will be removed before addition.
-	 *
-	 * @param Item Element to add.
-	 * @returns Always 0.
-	 */
-	int32 AddUnique(const T& Item)
-	{
-		// Remove any existing copies of the item.
-		this->Remove(Item);
-
-		this->InsertUninitialized(0);
-		(*this)[0] = Item;
-
-		CullArray();
-
-		return 0;
-	}
-
-	/**
-	 * Makes sure that the array never gets beyond MaxItems in size.
-	 */
-	void CullArray()
-	{
-		// 0 = no limit
-		if (!MaxItems)
-		{
-			return;
-		}
-
-		while (this->Num() > MaxItems)
-		{
-			this->RemoveAt(this->Num() - 1, 1);
-		}
-	}
-};
-
-
-template<typename T, typename Allocator>
-struct TContainerTraits<TMRUArray<T, Allocator> > : public TContainerTraitsBase<TMRUArray<T, Allocator> >
-{
-	enum { MoveWillEmptyContainer = TContainerTraitsBase<typename TMRUArray<T, Allocator>::Super>::MoveWillEmptyContainer };
-};
-
-
-/*-----------------------------------------------------------------------------
-	Indirect array.
-	Same as a TArray above, but stores pointers to the elements, to allow
-	resizing the array index without relocating the actual elements.
------------------------------------------------------------------------------*/
-template<typename T,typename Allocator = FDefaultAllocator>
-class TIndirectArray
-{
-public:
-	typedef T                        ElementType;
-	typedef TArray<void*, Allocator> InternalArrayType;
-
-	/** Default constructors. */
-#if PLATFORM_COMPILER_HAS_DEFAULTED_FUNCTIONS
-
-	TIndirectArray() = default;
-	TIndirectArray(TIndirectArray&&) = default;
-	TIndirectArray& operator=(TIndirectArray&&) = default;
-
-#else
-
-	            TIndirectArray() {}
-	FORCEINLINE TIndirectArray(TIndirectArray&& Other) : Array(MoveTemp(Other.Array)) {}
-	FORCEINLINE TIndirectArray& operator=(TIndirectArray&& Other) { Array = MoveTemp(Other.Array); return *this; }
-
-#endif
-
-	/**
-	 * Copy constructor.
-	 *
-	 * @param Other Other array to copy from.
-	 */
-	TIndirectArray(const TIndirectArray& Other)
-	{
-		for (auto& Item : Other)
-		{
-			Add(new T(Item));
-		}
-	}
-
-	/**
-	 * Assignment operator.
-	 *
-	 * @param Other Other array to assign with.
-	 */
-	TIndirectArray& operator=(const TIndirectArray& Other)
-	{
-		if (&Other != this)
-		{
-			Empty(Other.Num());
-			for (auto& Item : Other)
-			{
-				Add(new T(Item));
-			}
-		}
-
-		return *this;
-	}
-
-	/** Destructor. */
-	~TIndirectArray()
-	{
-		Empty();
-	}
-
-	/**
-	 * Gets number of elements in array.
-	 *
-	 * @returns Number of elements in array.
-	 */
-	FORCEINLINE int32 Num() const
-	{
-		return Array.Num();
-	}
-
-	/**
-	 * Helper function for returning a typed pointer to the first array entry.
-	 *
-	 * @returns Pointer to first array entry or nullptr if this->ArrayMax == 0.
-	 */
-	FORCEINLINE T** GetData()
-	{
-		return (T**)Array.GetData();
-	}
-
-	/**
-	 * Helper function for returning a typed pointer to the first array entry.
-	 *
-	 * @returns Pointer to first array entry or nullptr if this->ArrayMax == 0.
-	 */
-	FORCEINLINE const T** GetData() const
-	{
-		return (const T**)Array.GetData();
-	}
-
-	/** 
-	 * Helper function returning the size of the inner type.
-	 *
-	 * @returns Size in bytes of array type.
-	 */
-	uint32 GetTypeSize() const
-	{
-		return sizeof(T*);
-	}
-
-	/**
-	 * Bracket array access operator.
-	 *
-	 * @param Index Position of element to return.
-	 * @returns Reference to element in array at given position.
-	 */
-	FORCEINLINE T& operator[](int32 Index)
-	{
-		return *(T*)Array[Index];
-	}
-
-	/**
-	 * Bracket array access operator.
-	 *
-	 * Const version.
-	 *
-	 * @param Index Position of element to return.
-	 * @returns Reference to element in array at given position.
-	 */
-	FORCEINLINE const T& operator[](int32 Index) const
-	{
-		return *(T*)Array[Index];
-	}
-
-	/**
-	 * Returns n-th last element from the array.
-	 *
-	 * @param IndexFromTheEnd (Optional) Index from the end of array (default = 0).
-	 * @returns Reference to n-th last element from the array.
-	 */
-	FORCEINLINE ElementType& Last(int32 IndexFromTheEnd = 0)
-	{
-		return *(T*)Array.Last(IndexFromTheEnd);
-	}
-
-	/**
-	 * Returns n-th last element from the array.
-	 *
-	 * Const version.
-	 *
-	 * @param IndexFromTheEnd (Optional) Index from the end of array (default = 0).
-	 * @returns Reference to n-th last element from the array.
-	 */
-	FORCEINLINE const ElementType& Last(int32 IndexFromTheEnd = 0) const
-	{
-		return *(T*)Array.Last(IndexFromTheEnd);
-	}
-
-	/**
-	 * Shrinks the array's used memory to smallest possible to store elements
-	 * currently in it.
-	 */
-	void Shrink()
-	{
-		Array.Shrink();
-	}
-
-	/**
-	 * Resets the array to the new given size. It calls the destructors on held
-	 * items.
-	 *
-	 * @param NewSize (Optional) The expected usage size after calling this
-	 *		function. Default is 0.
-	 */
-	void Reset(int32 NewSize = 0)
-	{
-		DestructAndFreeItems();
-		Array.Reset(NewSize);
-	}
-
-	/**
-	 * Special serialize function passing the owning UObject along as required
-	 * by FUnytpedBulkData serialization.
-	 *
-	 * @param Ar Archive to serialize with.
-	 * @param Owner UObject this structure is serialized within.
-	 */
-	void Serialize(FArchive& Ar, UObject* Owner)
-	{
-		CountBytes(Ar);
-		if (Ar.IsLoading())
-		{
-			// Load array.
-			int32 NewNum;
-			Ar << NewNum;
-			Empty(NewNum);
-			for (int32 Index = 0; Index < NewNum; Index++)
-			{
-				new(*this) T;
-			}
-			for (int32 Index = 0; Index < NewNum; Index++)
-			{
-				(*this)[Index].Serialize(Ar, Owner, Index);
-			}
-		}
-		else
-		{
-			// Save array.
-			int32 Num = Array.Num();
-			Ar << Num;
-			for (int32 Index = 0; Index < Num; Index++)
-			{
-				(*this)[Index].Serialize(Ar, Owner, Index);
-			}
-		}
-	}
-
-	/**
-	 * Serialization operator for TIndirectArray.
-	 *
-	 * @param Ar Archive to serialize with.
-	 * @param A Array to serialize.
-	 * @returns Passing down serializing archive.
-	 */
-	friend FArchive& operator<<(FArchive& Ar, TIndirectArray& A)
-	{
-		A.CountBytes(Ar);
-		if (Ar.IsLoading())
-		{
-			// Load array.
-			int32 NewNum;
-			Ar << NewNum;
-			A.Empty(NewNum);
-			for (int32 Index = 0; Index < NewNum; Index++)
-			{
-				Ar << *new(A)T;
-			}
-		}
-		else
-		{
-			// Save array.
-			int32 Num = A.Num();
-			Ar << Num;
-			for (int32 Index = 0; Index < Num; Index++)
-			{
-				Ar << A[Index];
-			}
-		}
-		return Ar;
-	}
-
-	/**
-	 * Count bytes needed to serialize this array.
-	 *
-	 * @param Ar Archive to count for.
-	 */
-	void CountBytes(FArchive& Ar)
-	{
-		Array.CountBytes(Ar);
-	}
-
-	/**
-	 * Removes an element (or elements) at given location optionally shrinking
-	 * the array.
-	 *
-	 * @param Index Location in array of the element to remove.
-	 * @param Count (Optional) Number of elements to remove. Default is 1.
-	 * @param bAllowShrinking (Optional) Tells if this call can shrink array if
-	 *                        suitable after remove. Default is true.
-	 */
-	void RemoveAt(int32 Index, int32 Count = 1, bool bAllowShrinking = true)
-	{
-		check(Index >= 0);
-		check(Index <= Array.Num());
-		check(Index + Count <= Array.Num());
-		T** Element = GetData() + Index;
-		for (int32 ElementId = Count; ElementId; --ElementId)
-		{
-			// We need a typedef here because VC won't compile the destructor call below if T itself has a member called T
-			typedef T IndirectArrayDestructElementType;
-
-			(*Element)->IndirectArrayDestructElementType::~IndirectArrayDestructElementType();
-			FMemory::Free(*Element);
-			++Element;
-		}
-		Array.RemoveAt(Index, Count, bAllowShrinking);
-	}
-
-	/**
-	 * Removes an element (or elements) at given location optionally shrinking
-	 * the array.
-	 *
-	 * This version is much more efficient than RemoveAt (O(Count) instead of
-	 * O(ArrayNum)), but does not preserve the order.
-	 *
-	 * @param Index Location in array of the element to remove.
-	 * @param Count (Optional) Number of elements to remove. Default is 1.
-	 * @param bAllowShrinking (Optional) Tells if this call can shrink array if
-	 *                        suitable after remove. Default is true.
-	 */
-	void RemoveAtSwap(int32 Index, int32 Count = 1, bool bAllowShrinking = true)
-	{
-		check(Index >= 0);
-		check(Index <= Array.Num());
-		check(Index + Count <= Array.Num());
-		T** Element = GetData() + Index;
-		for (int32 ElementId = Count; ElementId; --ElementId)
-		{
-			// We need a typedef here because VC won't compile the destructor call below if T itself has a member called T
-			typedef T IndirectArrayDestructElementType;
-
-			(*Element)->IndirectArrayDestructElementType::~IndirectArrayDestructElementType();
-			FMemory::Free(*Element);
-			++Element;
-		}
-		Array.RemoveAtSwap(Index, Count, bAllowShrinking);
-	}
-
-	/**
-	 * Element-wise array element swap.
-	 *
-	 * This version is doing more sanity checks than SwapMemory.
-	 *
-	 * @param FirstIndexToSwap Position of the first element to swap.
-	 * @param SecondIndexToSwap Position of the second element to swap.
-	 */
-	void Swap(int32 FirstIndexToSwap, int32 SecondIndexToSwap)
-	{
-		Array.Swap(FirstIndexToSwap, SecondIndexToSwap);
-	}
-
-	/**
-	 * Empties the array. It calls the destructors on held items.
-	 *
-	 * @param Slack (Optional) The expected usage size after empty operation.
-	 *              Default is 0.
-	 */
-	void Empty(int32 Slack = 0)
-	{
-		DestructAndFreeItems();
-		Array.Empty(Slack);
-	}
-
-	/**
-	 * Adds a new item to the end of the array, possibly reallocating the
-	 * whole array to fit.
-	 *
-	 * @param Item The item to add.
-	 * @returns Index to the new item.
-	 */
-	FORCEINLINE int32 Add(T* Item)
-	{
-		return Array.Add(Item);
-	}
-
-	/**
-	 * Inserts a given element into the array at given location.
-	 *
-	 * @param Item The element to insert.
-	 * @param Index Tells where to insert the new elements.
-	 */
-	FORCEINLINE void Insert(T* Item, int32 Index)
-	{
-		Array.Insert(Item, Index);
-	}
-
-	/**
-	 * Reserves memory such that the array can contain at least Number elements.
-	 *
-	 * @param Number The number of elements that the array should be able to
-	 *               contain after allocation.
-	 */
-	FORCEINLINE void Reserve(int32 Number)
-	{
-		Array.Reserve(Number);
-	}
-
-	/**
-	 * Tests if index is valid, i.e. greater than zero and less than number of
-	 * elements in array.
-	 *
-	 * @param Index Index to test.
-	 * @returns True if index is valid. False otherwise.
-	 */
-	FORCEINLINE bool IsValidIndex(int32 Index) const
-	{
-		return Array.IsValidIndex(Index);
-	}
-
-	/** 
-	 * Helper function to return the amount of memory allocated by this container 
-	 *
-	 * @returns Number of bytes allocated by this container.
-	 */
-	SIZE_T GetAllocatedSize() const
-	{
-		return Array.Max() * sizeof(T*) + Array.Num() * sizeof(T);
-	}
-
-	// Iterators
-	typedef TIndexedContainerIterator<      TIndirectArray,       ElementType, int32> TIterator;
-	typedef TIndexedContainerIterator<const TIndirectArray, const ElementType, int32> TConstIterator;
-
-	/**
-	 * Creates an iterator for the contents of this array.
-	 *
-	 * @returns The iterator.
-	 */
-	TIterator CreateIterator()
-	{
-		return TIterator(*this);
-	}
-
-	/**
-	 * Creates a const iterator for the contents of this array.
-	 *
-	 * @returns The const iterator.
-	 */
-	TConstIterator CreateConstIterator() const
-	{
-		return TConstIterator(*this);
-	}
-
-private:
-
-	/**
-	 * Calls destructor and frees memory on every element in the array.
-	 */
-	void DestructAndFreeItems()
-	{
-		T** Element = GetData();
-		for (int32 Index = Array.Num(); Index; --Index)
-		{
-			// We need a typedef here because VC won't compile the destructor call below if T itself has a member called T
-			typedef T IndirectArrayDestructElementType;
-
-			(*Element)->IndirectArrayDestructElementType::~IndirectArrayDestructElementType();
-			FMemory::Free(*Element);
-			++Element;
-		}
-	}
-
-	/**
-	 * DO NOT USE DIRECTLY
-	 * STL-like iterators to enable range-based for loop support.
-	 */
-	FORCEINLINE friend TDereferencingIterator<      ElementType, typename InternalArrayType::RangedForIteratorType     > begin(      TIndirectArray& IndirectArray) { return TDereferencingIterator<      ElementType, typename InternalArrayType::RangedForIteratorType     >(begin(IndirectArray.Array)); }
-	FORCEINLINE friend TDereferencingIterator<const ElementType, typename InternalArrayType::RangedForConstIteratorType> begin(const TIndirectArray& IndirectArray) { return TDereferencingIterator<const ElementType, typename InternalArrayType::RangedForConstIteratorType>(begin(IndirectArray.Array)); }
-	FORCEINLINE friend TDereferencingIterator<      ElementType, typename InternalArrayType::RangedForIteratorType     > end  (      TIndirectArray& IndirectArray) { return TDereferencingIterator<      ElementType, typename InternalArrayType::RangedForIteratorType     >(end  (IndirectArray.Array)); }
-	FORCEINLINE friend TDereferencingIterator<const ElementType, typename InternalArrayType::RangedForConstIteratorType> end  (const TIndirectArray& IndirectArray) { return TDereferencingIterator<const ElementType, typename InternalArrayType::RangedForConstIteratorType>(end  (IndirectArray.Array)); }
-
-	InternalArrayType Array;
-};
-
-
-template<typename T, typename Allocator>
-struct TContainerTraits<TIndirectArray<T, Allocator> >
-	: public TContainerTraitsBase<TIndirectArray<T, Allocator> >
-{
-	enum { MoveWillEmptyContainer = TContainerTraitsBase<typename TIndirectArray<T, Allocator>::InternalArrayType>::MoveWillEmptyContainer };
-};
-
-
-template <typename T,typename Allocator> void* operator new( size_t Size, TIndirectArray<T,Allocator>& Array )
-{
-	check(Size == sizeof(T));
-	const int32 Index = Array.Add((T*)FMemory::Malloc(Size));
-	return &Array[Index];
-}
-
-
-template <typename T,typename Allocator> void* operator new( size_t Size, TIndirectArray<T,Allocator>& Array, int32 Index )
-{
-	check(Size == sizeof(T));
-	Array.Insert((T*)FMemory::Malloc(Size), Index);
-	return &Array[Index];
-}
-
-
-/*-----------------------------------------------------------------------------
-	Transactional array.
------------------------------------------------------------------------------*/
-
-// NOTE: Right now, you can't use a custom allocation policy with transactional arrays. If
-// you need to do it, you will have to fix up FTransaction::FObjectRecord to use the correct TArray<Allocator>.
-template< typename T >
-class TTransArray
-	: public TArray<T>
-{
-public:
-
-	typedef TArray<T> Super;
-
-	// Constructors.
-	explicit TTransArray( UObject* InOwner )
-	:	Owner( InOwner )
-	{
-		checkSlow(Owner);
-	}
-
-	TTransArray( UObject* InOwner, const Super& Other )
-	:	Super( Other )
-	,	Owner( InOwner )
-	{
-		checkSlow(Owner);
-	}
-
-#if PLATFORM_COMPILER_HAS_DEFAULTED_FUNCTIONS
-
-	TTransArray(TTransArray&&) = default;
-	TTransArray(const TTransArray&) = default;
-	TTransArray& operator=(TTransArray&&) = default;
-	TTransArray& operator=(const TTransArray&) = default;
-
-#else
-
-	FORCEINLINE TTransArray(const TTransArray& Other)
-		: Super((const Super&)Other)
-		, Owner(Other.Owner)
-	{
-	}
-
-	FORCEINLINE TTransArray& operator=(const TTransArray& Other)
-	{
-		(Super&)*this = (const Super&)Other;
-		Owner         = Other.Owner;
-		return *this;
-	}
-
-	FORCEINLINE TTransArray(TTransArray&& Other)
-		: Super((TTransArray&&)Other)
-		, Owner(Other.Owner)
-	{
-	}
-
-	FORCEINLINE TTransArray& operator=(TTransArray&& Other)
-	{
-		(Super&)*this = (Super&&)Other;
-		Owner         = Other.Owner;
-		return *this;
-	}
-
-#endif
-
-	// Add, Insert, Remove, Empty interface.
-	int32 AddUninitialized( int32 Count=1 )
-	{
-		const int32 Index = Super::AddUninitialized( Count );
-		if( GUndo )
-		{
-			GUndo->SaveArray( Owner, (FScriptArray*)this, Index, Count, 1, sizeof(T), DefaultConstructItem, SerializeItem, DestructItem );
-		}
-		return Index;
-	}
-
-	void InsertUninitialized( int32 Index, int32 Count=1 )
-	{
-		Super::InsertUninitialized( Index, Count );
-		if( GUndo )
-		{
-			GUndo->SaveArray( Owner, (FScriptArray*)this, Index, Count, 1, sizeof(T), DefaultConstructItem, SerializeItem, DestructItem );
-		}
-	}
-
-	void RemoveAt( int32 Index, int32 Count=1 )
-	{
-		if( GUndo )
-		{
-			GUndo->SaveArray( Owner, (FScriptArray*)this, Index, Count, -1, sizeof(T), DefaultConstructItem, SerializeItem, DestructItem );
-		}
-		Super::RemoveAt( Index, Count );
-	}
-
-	void Empty( int32 Slack=0 )
-	{
-		if( GUndo )
-		{
-			GUndo->SaveArray( Owner, (FScriptArray*)this, 0, this->ArrayNum, -1, sizeof(T), DefaultConstructItem, SerializeItem, DestructItem );
-		}
-		Super::Empty( Slack );
-	}
-
-	// Functions dependent on Add, Remove.
-	void AssignButKeepOwner( const Super& Other )
-	{
-		(Super&)*this = Other;
-	}
-
-	// Functions dependent on Add, Remove.
-	void AssignButKeepOwner( Super&& Other )
-	{
-		(Super&)*this = MoveTemp(Other);
-	}
-
-	int32 Add( const T& Item )
-	{
-		new(*this) T(Item);
-		return this->Num() - 1;
-	}
-
-	int32 AddZeroed( int32 n=1 )
-	{
-		const int32 Index = AddUninitialized(n);
-		FMemory::Memzero(this->GetData() + Index, n*sizeof(T));
-		return Index;
-	}
-
-	int32 AddUnique( const T& Item )
-	{
-		for( int32 Index=0; Index<this->ArrayNum; Index++ )
-		{
-			if( (*this)[Index]==Item )
-			{
-				return Index;
-			}
-		}
-		return Add( Item );
-	}
-
-	int32 Remove( const T& Item )
-	{
-		this->CheckAddress(&Item);
-
-		const int32 OriginalNum=this->ArrayNum;
-		for( int32 Index=0; Index<this->ArrayNum; Index++ )
-		{
-			if( (*this)[Index]==Item )
-			{
-				RemoveAt( Index-- );
-			}
-		}
-		return OriginalNum - this->ArrayNum;
-	}
-
-	// TTransArray interface.
-	UObject* GetOwner() const
-	{
-		return Owner;
-	}
-
-	void SetOwner( UObject* NewOwner )
-	{
-		Owner = NewOwner;
-	}
-
-	void ModifyItem( int32 Index )
-	{
-		if( GUndo )
-		{
-			GUndo->SaveArray( Owner, (FScriptArray*)this, Index, 1, 0, sizeof(T), DefaultConstructItem, SerializeItem, DestructItem );
-		}
-	}
-
-	void ModifyAllItems()
-	{
-		if( GUndo )
-		{
-			GUndo->SaveArray( Owner, (FScriptArray*)this, 0, this->Num(), 0, sizeof(T), DefaultConstructItem, SerializeItem, DestructItem );
-		}
-	}
-
-	friend FArchive& operator<<( FArchive& Ar, TTransArray& A )
-	{
-		Ar << A.Owner;
-		Ar << (Super&)A;
-		return Ar;
-	}
-
-protected:
-
-	static void DefaultConstructItem( void* TPtr )
-	{
-		new (TPtr) T;
-	}
-	static void SerializeItem( FArchive& Ar, void* TPtr )
-	{
-		Ar << *(T*)TPtr;
-	}
-	static void DestructItem( void* TPtr )
-	{
-		((T*)TPtr)->~T();
-	}
-	UObject* Owner;
-};
-
-
-template<typename T>
-struct TContainerTraits<TTransArray<T> > : public TContainerTraitsBase<TTransArray<T> >
-{
-	enum { MoveWillEmptyContainer = TContainerTraitsBase<typename TTransArray<T>::Super>::MoveWillEmptyContainer };
-};
-
-
-//
-// Transactional array operator news.
-//
-template <typename T> void* operator new( size_t Size, TTransArray<T>& Array )
-{
-	check(Size == sizeof(T));
-	const int32 Index = Array.AddUninitialized();
-	return &Array[Index];
-}
-
-
-template <typename T> void* operator new( size_t Size, TTransArray<T>& Array, int32 Index )
-{
-	check(Size == sizeof(T));
-	Array.Insert(Index);
-	return &Array[Index];
-}
diff --git a/Engine/Source/Runtime/Core/Public/Containers/ContainerAllocationPolicies.h b/Engine/Source/Runtime/Core/Public/Containers/ContainerAllocationPolicies.h
--- a/Engine/Source/Runtime/Core/Public/Containers/ContainerAllocationPolicies.h
+++ b/Engine/Source/Runtime/Core/Public/Containers/ContainerAllocationPolicies.h
@@ -2,9 +2,11 @@
 
 #pragma once
 
+#include "HAL/PlatformMath.h"
 #include "MemoryBase.h"
 #include "Misc/OutputDevice.h"
 #include "Templates/MemoryOps.h"
+#include "Math/NumericLimits.h"
 
 
 /** Used to determine the alignment of an element type. */
@@ -860,7 +862,7 @@ public:
 	{
 		if(NumHashedElements >= MinNumberOfHashedElements)
 		{
-			return FMath::RoundUpToPowerOfTwo(NumHashedElements / AverageNumberOfElementsPerHashBucket + BaseNumberOfHashBuckets);
+			return FPlatformMath::RoundUpToPowerOfTwo(NumHashedElements / AverageNumberOfElementsPerHashBucket + BaseNumberOfHashBuckets);
 		}
 
 		return 1;
@@ -890,7 +892,7 @@ public:
 	/** Computes the number of hash buckets to use for a given number of elements. */
 	static FORCEINLINE uint32 GetNumberOfHashBuckets(uint32 NumHashedElements)
 	{
-		const uint32 NumDesiredHashBuckets = FMath::RoundUpToPowerOfTwo(NumHashedElements / AverageNumberOfElementsPerHashBucket);
+		const uint32 NumDesiredHashBuckets = FPlatformMath::RoundUpToPowerOfTwo(NumHashedElements / AverageNumberOfElementsPerHashBucket);
 		if (NumDesiredHashBuckets < NumInlineHashBuckets)
 		{
 			return NumInlineHashBuckets;
diff --git a/Engine/Source/Runtime/Core/Public/Containers/IndirectArray.h b/Engine/Source/Runtime/Core/Public/Containers/IndirectArray.h
new file mode 100644
--- /dev/null
+++ b/Engine/Source/Runtime/Core/Public/Containers/IndirectArray.h
@@ -0,0 +1,479 @@
+// Copyright 1998-2016 Epic Games, Inc. All Rights Reserved.
+
+#pragma once
+
+#include "Array.h"
+
+/*-----------------------------------------------------------------------------
+	Indirect array.
+	Same as a TArray, but stores pointers to the elements, to allow
+	resizing the array index without relocating the actual elements.
+-----------------------------------------------------------------------------*/
+
+template<typename T,typename Allocator = FDefaultAllocator>
+class TIndirectArray
+{
+public:
+	typedef T                        ElementType;
+	typedef TArray<void*, Allocator> InternalArrayType;
+
+	/** Default constructors. */
+#if PLATFORM_COMPILER_HAS_DEFAULTED_FUNCTIONS
+
+	TIndirectArray() = default;
+	TIndirectArray(TIndirectArray&&) = default;
+	TIndirectArray& operator=(TIndirectArray&&) = default;
+
+#else
+
+	            TIndirectArray() {}
+	FORCEINLINE TIndirectArray(TIndirectArray&& Other) : Array(MoveTemp(Other.Array)) {}
+	FORCEINLINE TIndirectArray& operator=(TIndirectArray&& Other) { Array = MoveTemp(Other.Array); return *this; }
+
+#endif
+
+	/**
+	 * Copy constructor.
+	 *
+	 * @param Other Other array to copy from.
+	 */
+	TIndirectArray(const TIndirectArray& Other)
+	{
+		for (auto& Item : Other)
+		{
+			Add(new T(Item));
+		}
+	}
+
+	/**
+	 * Assignment operator.
+	 *
+	 * @param Other Other array to assign with.
+	 */
+	TIndirectArray& operator=(const TIndirectArray& Other)
+	{
+		if (&Other != this)
+		{
+			Empty(Other.Num());
+			for (auto& Item : Other)
+			{
+				Add(new T(Item));
+			}
+		}
+
+		return *this;
+	}
+
+	/** Destructor. */
+	~TIndirectArray()
+	{
+		Empty();
+	}
+
+	/**
+	 * Gets number of elements in array.
+	 *
+	 * @returns Number of elements in array.
+	 */
+	FORCEINLINE int32 Num() const
+	{
+		return Array.Num();
+	}
+
+	/**
+	 * Helper function for returning a typed pointer to the first array entry.
+	 *
+	 * @returns Pointer to first array entry or nullptr if this->ArrayMax == 0.
+	 */
+	FORCEINLINE T** GetData()
+	{
+		return (T**)Array.GetData();
+	}
+
+	/**
+	 * Helper function for returning a typed pointer to the first array entry.
+	 *
+	 * @returns Pointer to first array entry or nullptr if this->ArrayMax == 0.
+	 */
+	FORCEINLINE const T** GetData() const
+	{
+		return (const T**)Array.GetData();
+	}
+
+	/** 
+	 * Helper function returning the size of the inner type.
+	 *
+	 * @returns Size in bytes of array type.
+	 */
+	uint32 GetTypeSize() const
+	{
+		return sizeof(T*);
+	}
+
+	/**
+	 * Bracket array access operator.
+	 *
+	 * @param Index Position of element to return.
+	 * @returns Reference to element in array at given position.
+	 */
+	FORCEINLINE T& operator[](int32 Index)
+	{
+		return *(T*)Array[Index];
+	}
+
+	/**
+	 * Bracket array access operator.
+	 *
+	 * Const version.
+	 *
+	 * @param Index Position of element to return.
+	 * @returns Reference to element in array at given position.
+	 */
+	FORCEINLINE const T& operator[](int32 Index) const
+	{
+		return *(T*)Array[Index];
+	}
+
+	/**
+	 * Returns n-th last element from the array.
+	 *
+	 * @param IndexFromTheEnd (Optional) Index from the end of array (default = 0).
+	 * @returns Reference to n-th last element from the array.
+	 */
+	FORCEINLINE ElementType& Last(int32 IndexFromTheEnd = 0)
+	{
+		return *(T*)Array.Last(IndexFromTheEnd);
+	}
+
+	/**
+	 * Returns n-th last element from the array.
+	 *
+	 * Const version.
+	 *
+	 * @param IndexFromTheEnd (Optional) Index from the end of array (default = 0).
+	 * @returns Reference to n-th last element from the array.
+	 */
+	FORCEINLINE const ElementType& Last(int32 IndexFromTheEnd = 0) const
+	{
+		return *(T*)Array.Last(IndexFromTheEnd);
+	}
+
+	/**
+	 * Shrinks the array's used memory to smallest possible to store elements
+	 * currently in it.
+	 */
+	void Shrink()
+	{
+		Array.Shrink();
+	}
+
+	/**
+	 * Resets the array to the new given size. It calls the destructors on held
+	 * items.
+	 *
+	 * @param NewSize (Optional) The expected usage size after calling this
+	 *		function. Default is 0.
+	 */
+	void Reset(int32 NewSize = 0)
+	{
+		DestructAndFreeItems();
+		Array.Reset(NewSize);
+	}
+
+	/**
+	 * Special serialize function passing the owning UObject along as required
+	 * by FUnytpedBulkData serialization.
+	 *
+	 * @param Ar Archive to serialize with.
+	 * @param Owner UObject this structure is serialized within.
+	 */
+	void Serialize(FArchive& Ar, UObject* Owner)
+	{
+		CountBytes(Ar);
+		if (Ar.IsLoading())
+		{
+			// Load array.
+			int32 NewNum;
+			Ar << NewNum;
+			Empty(NewNum);
+			for (int32 Index = 0; Index < NewNum; Index++)
+			{
+				new(*this) T;
+			}
+			for (int32 Index = 0; Index < NewNum; Index++)
+			{
+				(*this)[Index].Serialize(Ar, Owner, Index);
+			}
+		}
+		else
+		{
+			// Save array.
+			int32 Num = Array.Num();
+			Ar << Num;
+			for (int32 Index = 0; Index < Num; Index++)
+			{
+				(*this)[Index].Serialize(Ar, Owner, Index);
+			}
+		}
+	}
+
+	/**
+	 * Serialization operator for TIndirectArray.
+	 *
+	 * @param Ar Archive to serialize with.
+	 * @param A Array to serialize.
+	 * @returns Passing down serializing archive.
+	 */
+	friend FArchive& operator<<(FArchive& Ar, TIndirectArray& A)
+	{
+		A.CountBytes(Ar);
+		if (Ar.IsLoading())
+		{
+			// Load array.
+			int32 NewNum;
+			Ar << NewNum;
+			A.Empty(NewNum);
+			for (int32 Index = 0; Index < NewNum; Index++)
+			{
+				Ar << *new(A)T;
+			}
+		}
+		else
+		{
+			// Save array.
+			int32 Num = A.Num();
+			Ar << Num;
+			for (int32 Index = 0; Index < Num; Index++)
+			{
+				Ar << A[Index];
+			}
+		}
+		return Ar;
+	}
+
+	/**
+	 * Count bytes needed to serialize this array.
+	 *
+	 * @param Ar Archive to count for.
+	 */
+	void CountBytes(FArchive& Ar)
+	{
+		Array.CountBytes(Ar);
+	}
+
+	/**
+	 * Removes an element (or elements) at given location optionally shrinking
+	 * the array.
+	 *
+	 * @param Index Location in array of the element to remove.
+	 * @param Count (Optional) Number of elements to remove. Default is 1.
+	 * @param bAllowShrinking (Optional) Tells if this call can shrink array if
+	 *                        suitable after remove. Default is true.
+	 */
+	void RemoveAt(int32 Index, int32 Count = 1, bool bAllowShrinking = true)
+	{
+		check(Index >= 0);
+		check(Index <= Array.Num());
+		check(Index + Count <= Array.Num());
+		T** Element = GetData() + Index;
+		for (int32 ElementId = Count; ElementId; --ElementId)
+		{
+			// We need a typedef here because VC won't compile the destructor call below if T itself has a member called T
+			typedef T IndirectArrayDestructElementType;
+
+			(*Element)->IndirectArrayDestructElementType::~IndirectArrayDestructElementType();
+			FMemory::Free(*Element);
+			++Element;
+		}
+		Array.RemoveAt(Index, Count, bAllowShrinking);
+	}
+
+	/**
+	 * Removes an element (or elements) at given location optionally shrinking
+	 * the array.
+	 *
+	 * This version is much more efficient than RemoveAt (O(Count) instead of
+	 * O(ArrayNum)), but does not preserve the order.
+	 *
+	 * @param Index Location in array of the element to remove.
+	 * @param Count (Optional) Number of elements to remove. Default is 1.
+	 * @param bAllowShrinking (Optional) Tells if this call can shrink array if
+	 *                        suitable after remove. Default is true.
+	 */
+	void RemoveAtSwap(int32 Index, int32 Count = 1, bool bAllowShrinking = true)
+	{
+		check(Index >= 0);
+		check(Index <= Array.Num());
+		check(Index + Count <= Array.Num());
+		T** Element = GetData() + Index;
+		for (int32 ElementId = Count; ElementId; --ElementId)
+		{
+			// We need a typedef here because VC won't compile the destructor call below if T itself has a member called T
+			typedef T IndirectArrayDestructElementType;
+
+			(*Element)->IndirectArrayDestructElementType::~IndirectArrayDestructElementType();
+			FMemory::Free(*Element);
+			++Element;
+		}
+		Array.RemoveAtSwap(Index, Count, bAllowShrinking);
+	}
+
+	/**
+	 * Element-wise array element swap.
+	 *
+	 * This version is doing more sanity checks than SwapMemory.
+	 *
+	 * @param FirstIndexToSwap Position of the first element to swap.
+	 * @param SecondIndexToSwap Position of the second element to swap.
+	 */
+	void Swap(int32 FirstIndexToSwap, int32 SecondIndexToSwap)
+	{
+		Array.Swap(FirstIndexToSwap, SecondIndexToSwap);
+	}
+
+	/**
+	 * Empties the array. It calls the destructors on held items.
+	 *
+	 * @param Slack (Optional) The expected usage size after empty operation.
+	 *              Default is 0.
+	 */
+	void Empty(int32 Slack = 0)
+	{
+		DestructAndFreeItems();
+		Array.Empty(Slack);
+	}
+
+	/**
+	 * Adds a new item to the end of the array, possibly reallocating the
+	 * whole array to fit.
+	 *
+	 * @param Item The item to add.
+	 * @returns Index to the new item.
+	 */
+	FORCEINLINE int32 Add(T* Item)
+	{
+		return Array.Add(Item);
+	}
+
+	/**
+	 * Inserts a given element into the array at given location.
+	 *
+	 * @param Item The element to insert.
+	 * @param Index Tells where to insert the new elements.
+	 */
+	FORCEINLINE void Insert(T* Item, int32 Index)
+	{
+		Array.Insert(Item, Index);
+	}
+
+	/**
+	 * Reserves memory such that the array can contain at least Number elements.
+	 *
+	 * @param Number The number of elements that the array should be able to
+	 *               contain after allocation.
+	 */
+	FORCEINLINE void Reserve(int32 Number)
+	{
+		Array.Reserve(Number);
+	}
+
+	/**
+	 * Tests if index is valid, i.e. greater than zero and less than number of
+	 * elements in array.
+	 *
+	 * @param Index Index to test.
+	 * @returns True if index is valid. False otherwise.
+	 */
+	FORCEINLINE bool IsValidIndex(int32 Index) const
+	{
+		return Array.IsValidIndex(Index);
+	}
+
+	/** 
+	 * Helper function to return the amount of memory allocated by this container 
+	 *
+	 * @returns Number of bytes allocated by this container.
+	 */
+	SIZE_T GetAllocatedSize() const
+	{
+		return Array.Max() * sizeof(T*) + Array.Num() * sizeof(T);
+	}
+
+	// Iterators
+	typedef TIndexedContainerIterator<      TIndirectArray,       ElementType, int32> TIterator;
+	typedef TIndexedContainerIterator<const TIndirectArray, const ElementType, int32> TConstIterator;
+
+	/**
+	 * Creates an iterator for the contents of this array.
+	 *
+	 * @returns The iterator.
+	 */
+	TIterator CreateIterator()
+	{
+		return TIterator(*this);
+	}
+
+	/**
+	 * Creates a const iterator for the contents of this array.
+	 *
+	 * @returns The const iterator.
+	 */
+	TConstIterator CreateConstIterator() const
+	{
+		return TConstIterator(*this);
+	}
+
+private:
+
+	/**
+	 * Calls destructor and frees memory on every element in the array.
+	 */
+	void DestructAndFreeItems()
+	{
+		T** Element = GetData();
+		for (int32 Index = Array.Num(); Index; --Index)
+		{
+			// We need a typedef here because VC won't compile the destructor call below if T itself has a member called T
+			typedef T IndirectArrayDestructElementType;
+
+			(*Element)->IndirectArrayDestructElementType::~IndirectArrayDestructElementType();
+			FMemory::Free(*Element);
+			++Element;
+		}
+	}
+
+	/**
+	 * DO NOT USE DIRECTLY
+	 * STL-like iterators to enable range-based for loop support.
+	 */
+	FORCEINLINE friend TDereferencingIterator<      ElementType, typename InternalArrayType::RangedForIteratorType     > begin(      TIndirectArray& IndirectArray) { return TDereferencingIterator<      ElementType, typename InternalArrayType::RangedForIteratorType     >(begin(IndirectArray.Array)); }
+	FORCEINLINE friend TDereferencingIterator<const ElementType, typename InternalArrayType::RangedForConstIteratorType> begin(const TIndirectArray& IndirectArray) { return TDereferencingIterator<const ElementType, typename InternalArrayType::RangedForConstIteratorType>(begin(IndirectArray.Array)); }
+	FORCEINLINE friend TDereferencingIterator<      ElementType, typename InternalArrayType::RangedForIteratorType     > end  (      TIndirectArray& IndirectArray) { return TDereferencingIterator<      ElementType, typename InternalArrayType::RangedForIteratorType     >(end  (IndirectArray.Array)); }
+	FORCEINLINE friend TDereferencingIterator<const ElementType, typename InternalArrayType::RangedForConstIteratorType> end  (const TIndirectArray& IndirectArray) { return TDereferencingIterator<const ElementType, typename InternalArrayType::RangedForConstIteratorType>(end  (IndirectArray.Array)); }
+
+	InternalArrayType Array;
+};
+
+
+template<typename T, typename Allocator>
+struct TContainerTraits<TIndirectArray<T, Allocator> >
+	: public TContainerTraitsBase<TIndirectArray<T, Allocator> >
+{
+	enum { MoveWillEmptyContainer = TContainerTraitsBase<typename TIndirectArray<T, Allocator>::InternalArrayType>::MoveWillEmptyContainer };
+};
+
+
+template <typename T,typename Allocator> void* operator new( size_t Size, TIndirectArray<T,Allocator>& Array )
+{
+	check(Size == sizeof(T));
+	const int32 Index = Array.Add((T*)FMemory::Malloc(Size));
+	return &Array[Index];
+}
+
+
+template <typename T,typename Allocator> void* operator new( size_t Size, TIndirectArray<T,Allocator>& Array, int32 Index )
+{
+	check(Size == sizeof(T));
+	Array.Insert((T*)FMemory::Malloc(Size), Index);
+	return &Array[Index];
+}
diff --git a/Engine/Source/Runtime/Core/Public/Containers/MRUArray.h b/Engine/Source/Runtime/Core/Public/Containers/MRUArray.h
new file mode 100644
--- /dev/null
+++ b/Engine/Source/Runtime/Core/Public/Containers/MRUArray.h
@@ -0,0 +1,164 @@
+// Copyright 1998-2016 Epic Games, Inc. All Rights Reserved.
+
+#pragma once
+
+#include "Array.h"
+
+/*-----------------------------------------------------------------------------
+	MRU array.
+-----------------------------------------------------------------------------*/
+
+/**
+ * Same as TArray except:
+ * - Has an upper limit of the number of items it will store.
+ * - Any item that is added to the array is moved to the top.
+ */
+template<typename T, typename Allocator = FDefaultAllocator>
+class TMRUArray
+	: public TArray<T, Allocator>
+{
+public:
+	typedef TArray<T, Allocator> Super;
+
+	/** The maximum number of items we can store in this array. */
+	int32 MaxItems;
+
+	/**
+	 * Constructor.
+	 */
+	TMRUArray()
+		: Super()
+	{
+		MaxItems = 0;
+	}
+
+#if PLATFORM_COMPILER_HAS_DEFAULTED_FUNCTIONS
+
+	TMRUArray(TMRUArray&&) = default;
+	TMRUArray(const TMRUArray&) = default;
+	TMRUArray& operator=(TMRUArray&&) = default;
+	TMRUArray& operator=(const TMRUArray&) = default;
+
+#else
+
+	/**
+	 * Copy constructor.
+	 *
+	 * @param Other Other array to copy.
+	 */
+	FORCEINLINE TMRUArray(const TMRUArray& Other)
+		: Super((const Super&)Other)
+	{
+	}
+
+	/**
+	 * Assign operator.
+	 *
+	 * @param Other Other array to assign with.
+	 * @returns Reference to this object.
+	 */
+	FORCEINLINE TMRUArray& operator=(const TMRUArray& Other)
+	{
+		(Super&)*this = (const Super&)Other;
+		return *this;
+	}
+
+	/**
+	 * Move constructor.
+	 *
+	 * @param Other Other array to move.
+	 */
+	FORCEINLINE TMRUArray(TMRUArray&& Other)
+		: Super((TMRUArray&&)Other)
+	{
+	}
+
+	/**
+	 * Move assign operator.
+	 *
+	 * @param Other Other array to assign with.
+	 * @returns Reference to this object.
+	 */
+	FORCEINLINE TMRUArray& operator=(TMRUArray&& Other)
+	{
+		(Super&)*this = (Super&&)Other;
+		return *this;
+	}
+
+#endif
+	/**
+	 * Adds item to the array. Makes sure that we don't add more than the
+	 * limit.
+	 *
+	 * @param Item Item to add.
+	 * @returns Always 0.
+	 */
+	int32 Add(const T& Item)
+	{
+		const int32 idx = Super::Add(Item);
+		this->Swap(idx, 0);
+		CullArray();
+		return 0;
+	}
+
+	/**
+	 * Adds a number of zeroed elements to the array. Makes sure that we don't
+	 * add more than the limit.
+	 *
+	 * @param Count (Optional) A number of elements to add. Default is 0.
+	 * @returns Always 0.
+	 */
+	int32 AddZeroed(int32 Count = 1)
+	{
+		const int32 idx = Super::AddZeroed(Count);
+		this->Swap(idx, 0);
+		CullArray();
+		return 0;
+	}
+
+	/**
+	 * Adds unique item to the array. Makes sure that we don't add more than
+	 * the limit. If the item existed it will be removed before addition.
+	 *
+	 * @param Item Element to add.
+	 * @returns Always 0.
+	 */
+	int32 AddUnique(const T& Item)
+	{
+		// Remove any existing copies of the item.
+		this->Remove(Item);
+
+		this->InsertUninitialized(0);
+		(*this)[0] = Item;
+
+		CullArray();
+
+		return 0;
+	}
+
+	/**
+	 * Makes sure that the array never gets beyond MaxItems in size.
+	 */
+	void CullArray()
+	{
+		// 0 = no limit
+		if (!MaxItems)
+		{
+			return;
+		}
+
+		while (this->Num() > MaxItems)
+		{
+			this->RemoveAt(this->Num() - 1, 1);
+		}
+	}
+};
+
+
+template<typename T, typename Allocator>
+struct TContainerTraits<TMRUArray<T, Allocator> > : public TContainerTraitsBase<TMRUArray<T, Allocator> >
+{
+	enum { MoveWillEmptyContainer = TContainerTraitsBase<typename TMRUArray<T, Allocator>::Super>::MoveWillEmptyContainer };
+};
+
+
diff --git a/Engine/Source/Runtime/Core/Public/Containers/ScriptArray.h b/Engine/Source/Runtime/Core/Public/Containers/ScriptArray.h
new file mode 100644
--- /dev/null
+++ b/Engine/Source/Runtime/Core/Public/Containers/ScriptArray.h
@@ -0,0 +1,207 @@
+// Copyright 1998-2016 Epic Games, Inc. All Rights Reserved.
+
+#pragma once
+
+#include <initializer_list>
+
+#include "Array.h"
+
+/**
+ * Base dynamic array.
+ * An untyped data array; mirrors a TArray's members, but doesn't need an exact C++ type for its elements.
+ **/
+class FScriptArray
+	: protected FHeapAllocator::ForAnyElementType
+{
+public:
+
+	FORCEINLINE void* GetData()
+	{
+		return this->GetAllocation();
+	}
+	FORCEINLINE const void* GetData() const
+	{
+		return this->GetAllocation();
+	}
+	FORCEINLINE bool IsValidIndex(int32 i) const
+	{
+		return i>=0 && i<ArrayNum;
+	}
+	FORCEINLINE int32 Num() const
+	{
+		checkSlow(ArrayNum>=0);
+		checkSlow(ArrayMax>=ArrayNum);
+		return ArrayNum;
+	}
+	void InsertZeroed( int32 Index, int32 Count, int32 NumBytesPerElement )
+	{
+		Insert( Index, Count, NumBytesPerElement );
+		FMemory::Memzero( (uint8*)this->GetAllocation()+Index*NumBytesPerElement, Count*NumBytesPerElement );
+	}
+	void Insert( int32 Index, int32 Count, int32 NumBytesPerElement )
+	{
+		check(Count>=0);
+		check(ArrayNum>=0);
+		check(ArrayMax>=ArrayNum);
+		check(Index>=0);
+		check(Index<=ArrayNum);
+
+		const int32 OldNum = ArrayNum;
+		if( (ArrayNum+=Count)>ArrayMax )
+		{
+			ResizeGrow(OldNum, NumBytesPerElement);
+		}
+		FMemory::Memmove
+		(
+			(uint8*)this->GetAllocation() + (Index+Count )*NumBytesPerElement,
+			(uint8*)this->GetAllocation() + (Index       )*NumBytesPerElement,
+			                                               (OldNum-Index)*NumBytesPerElement
+		);
+	}
+	AGRESSIVE_ARRAY_FORCEINLINE int32 Add( int32 Count, int32 NumBytesPerElement )
+	{
+		check(Count>=0);
+		checkSlow(ArrayNum>=0);
+		checkSlow(ArrayMax>=ArrayNum);
+
+		const int32 OldNum = ArrayNum;
+		if( (ArrayNum+=Count)>ArrayMax )
+		{
+			ResizeGrow(OldNum, NumBytesPerElement);
+		}
+
+		return OldNum;
+	}
+	int32 AddZeroed( int32 Count, int32 NumBytesPerElement )
+	{
+		const int32 Index = Add( Count, NumBytesPerElement );
+		FMemory::Memzero( (uint8*)this->GetAllocation()+Index*NumBytesPerElement, Count*NumBytesPerElement );
+		return Index;
+	}
+	void Shrink( int32 NumBytesPerElement )
+	{
+		checkSlow(ArrayNum>=0);
+		checkSlow(ArrayMax>=ArrayNum);
+		if (ArrayNum != ArrayMax)
+		{
+			ResizeTo(ArrayNum, NumBytesPerElement);
+		}
+	}
+	void Empty( int32 Slack, int32 NumBytesPerElement )
+	{
+		checkSlow(Slack>=0);
+		ArrayNum = 0;
+		if (Slack != ArrayMax)
+		{
+			ResizeTo(Slack, NumBytesPerElement);
+		}
+	}
+	void SwapMemory(int32 A, int32 B, int32 NumBytesPerElement )
+	{
+		FMemory::Memswap(
+			(uint8*)this->GetAllocation()+(NumBytesPerElement*A),
+			(uint8*)this->GetAllocation()+(NumBytesPerElement*B),
+			NumBytesPerElement
+			);
+	}
+	FScriptArray()
+	:   ArrayNum( 0 )
+	,	ArrayMax( 0 )
+	{
+	}
+	void CountBytes( FArchive& Ar, int32 NumBytesPerElement  )
+	{
+		Ar.CountBytes( ArrayNum*NumBytesPerElement, ArrayMax*NumBytesPerElement );
+	}
+	/**
+	 * Returns the amount of slack in this array in elements.
+	 */
+	FORCEINLINE int32 GetSlack() const
+	{
+		return ArrayMax - ArrayNum;
+	}
+		
+	void Remove( int32 Index, int32 Count, int32 NumBytesPerElement  )
+	{
+		if (Count)
+		{
+			checkSlow(Count >= 0);
+			checkSlow(Index >= 0);
+			checkSlow(Index <= ArrayNum);
+			checkSlow(Index + Count <= ArrayNum);
+
+			// Skip memmove in the common case that there is nothing to move.
+			int32 NumToMove = ArrayNum - Index - Count;
+			if (NumToMove)
+			{
+				FMemory::Memmove
+					(
+					(uint8*)this->GetAllocation() + (Index)* NumBytesPerElement,
+					(uint8*)this->GetAllocation() + (Index + Count) * NumBytesPerElement,
+					NumToMove * NumBytesPerElement
+					);
+			}
+			ArrayNum -= Count;
+
+			ResizeShrink(NumBytesPerElement);
+			checkSlow(ArrayNum >= 0);
+			checkSlow(ArrayMax >= ArrayNum);
+		}
+	}
+
+protected:
+
+	FScriptArray( int32 InNum, int32 NumBytesPerElement  )
+	:   ArrayNum( 0 )
+	,	ArrayMax( InNum )
+
+	{
+		if (ArrayMax)
+		{
+			ResizeInit(NumBytesPerElement);
+		}
+		ArrayNum = InNum;
+	}
+	int32	  ArrayNum;
+	int32	  ArrayMax;
+
+	FORCENOINLINE void ResizeInit(int32 NumBytesPerElement)
+	{
+		ArrayMax = this->CalculateSlackReserve(ArrayMax, NumBytesPerElement);
+		this->ResizeAllocation(ArrayNum, ArrayMax, NumBytesPerElement);
+	}
+	FORCENOINLINE void ResizeGrow(int32 OldNum, int32 NumBytesPerElement)
+	{
+		ArrayMax = this->CalculateSlackGrow(ArrayNum, ArrayMax, NumBytesPerElement);
+		this->ResizeAllocation(OldNum, ArrayMax, NumBytesPerElement);
+	}
+	FORCENOINLINE void ResizeShrink(int32 NumBytesPerElement)
+	{
+		const int32 NewArrayMax = this->CalculateSlackShrink(ArrayNum, ArrayMax, NumBytesPerElement);
+		if (NewArrayMax != ArrayMax)
+		{
+			ArrayMax = NewArrayMax;
+			this->ResizeAllocation(ArrayNum, ArrayMax, NumBytesPerElement);
+		}
+	}
+	FORCENOINLINE void ResizeTo(int32 NewMax, int32 NumBytesPerElement)
+	{
+		if (NewMax)
+		{
+			NewMax = this->CalculateSlackReserve(NewMax, NumBytesPerElement);
+		}
+		if (NewMax != ArrayMax)
+		{
+			ArrayMax = NewMax;
+			this->ResizeAllocation(ArrayNum, ArrayMax, NumBytesPerElement);
+		}
+	}
+public:
+	// These should really be private, because they shouldn't be called, but there's a bunch of code
+	// that needs to be fixed first.
+	FScriptArray(const FScriptArray&) { check(false); }
+	void operator=(const FScriptArray&) { check(false); }
+};
+
+
+template<> struct TIsZeroConstructType<FScriptArray> { enum { Value = true }; };
diff --git a/Engine/Source/Runtime/Core/Public/Containers/StringConv.h b/Engine/Source/Runtime/Core/Public/Containers/StringConv.h
--- a/Engine/Source/Runtime/Core/Public/Containers/StringConv.h
+++ b/Engine/Source/Runtime/Core/Public/Containers/StringConv.h
@@ -8,6 +8,7 @@
 #include "Containers/ContainerAllocationPolicies.h"
 #include "HAL/Platform.h"
 #include "HAL/PlatformString.h"
+#include "Misc/CString.h"
 
 
 #define DEFAULT_STRING_CONVERSION_SIZE 128u
diff --git a/Engine/Source/Runtime/Core/Public/Containers/TransArray.h b/Engine/Source/Runtime/Core/Public/Containers/TransArray.h
new file mode 100644
--- /dev/null
+++ b/Engine/Source/Runtime/Core/Public/Containers/TransArray.h
@@ -0,0 +1,242 @@
+// Copyright 1998-2016 Epic Games, Inc. All Rights Reserved.
+
+#pragma once
+
+#include <initializer_list>
+
+#include "Containers/Array.h"
+#include "Misc/ITransaction.h"
+
+extern CORE_API ITransaction* GUndo;
+
+/*-----------------------------------------------------------------------------
+	Transactional array.
+-----------------------------------------------------------------------------*/
+
+// NOTE: Right now, you can't use a custom allocation policy with transactional arrays. If
+// you need to do it, you will have to fix up FTransaction::FObjectRecord to use the correct TArray<Allocator>.
+template< typename T >
+class TTransArray
+	: public TArray<T>
+{
+public:
+
+	typedef TArray<T> Super;
+
+	// Constructors.
+	explicit TTransArray( UObject* InOwner )
+	:	Owner( InOwner )
+	{
+		checkSlow(Owner);
+	}
+
+	TTransArray( UObject* InOwner, const Super& Other )
+	:	Super( Other )
+	,	Owner( InOwner )
+	{
+		checkSlow(Owner);
+	}
+
+#if PLATFORM_COMPILER_HAS_DEFAULTED_FUNCTIONS
+
+	TTransArray(TTransArray&&) = default;
+	TTransArray(const TTransArray&) = default;
+	TTransArray& operator=(TTransArray&&) = default;
+	TTransArray& operator=(const TTransArray&) = default;
+
+#else
+
+	FORCEINLINE TTransArray(const TTransArray& Other)
+		: Super((const Super&)Other)
+		, Owner(Other.Owner)
+	{
+	}
+
+	FORCEINLINE TTransArray& operator=(const TTransArray& Other)
+	{
+		(Super&)*this = (const Super&)Other;
+		Owner         = Other.Owner;
+		return *this;
+	}
+
+	FORCEINLINE TTransArray(TTransArray&& Other)
+		: Super((TTransArray&&)Other)
+		, Owner(Other.Owner)
+	{
+	}
+
+	FORCEINLINE TTransArray& operator=(TTransArray&& Other)
+	{
+		(Super&)*this = (Super&&)Other;
+		Owner         = Other.Owner;
+		return *this;
+	}
+
+#endif
+
+	// Add, Insert, Remove, Empty interface.
+	int32 AddUninitialized( int32 Count=1 )
+	{
+		const int32 Index = Super::AddUninitialized( Count );
+		if( GUndo )
+		{
+			GUndo->SaveArray( Owner, (FScriptArray*)this, Index, Count, 1, sizeof(T), DefaultConstructItem, SerializeItem, DestructItem );
+		}
+		return Index;
+	}
+
+	void InsertUninitialized( int32 Index, int32 Count=1 )
+	{
+		Super::InsertUninitialized( Index, Count );
+		if( GUndo )
+		{
+			GUndo->SaveArray( Owner, (FScriptArray*)this, Index, Count, 1, sizeof(T), DefaultConstructItem, SerializeItem, DestructItem );
+		}
+	}
+
+	void RemoveAt( int32 Index, int32 Count=1 )
+	{
+		if( GUndo )
+		{
+			GUndo->SaveArray( Owner, (FScriptArray*)this, Index, Count, -1, sizeof(T), DefaultConstructItem, SerializeItem, DestructItem );
+		}
+		Super::RemoveAt( Index, Count );
+	}
+
+	void Empty( int32 Slack=0 )
+	{
+		if( GUndo )
+		{
+			GUndo->SaveArray( Owner, (FScriptArray*)this, 0, this->ArrayNum, -1, sizeof(T), DefaultConstructItem, SerializeItem, DestructItem );
+		}
+		Super::Empty( Slack );
+	}
+
+	// Functions dependent on Add, Remove.
+	void AssignButKeepOwner( const Super& Other )
+	{
+		(Super&)*this = Other;
+	}
+
+	// Functions dependent on Add, Remove.
+	void AssignButKeepOwner( Super&& Other )
+	{
+		(Super&)*this = MoveTemp(Other);
+	}
+
+	int32 Add( const T& Item )
+	{
+		new(*this) T(Item);
+		return this->Num() - 1;
+	}
+
+	int32 AddZeroed( int32 n=1 )
+	{
+		const int32 Index = AddUninitialized(n);
+		FMemory::Memzero(this->GetData() + Index, n*sizeof(T));
+		return Index;
+	}
+
+	int32 AddUnique( const T& Item )
+	{
+		for( int32 Index=0; Index<this->ArrayNum; Index++ )
+		{
+			if( (*this)[Index]==Item )
+			{
+				return Index;
+			}
+		}
+		return Add( Item );
+	}
+
+	int32 Remove( const T& Item )
+	{
+		this->CheckAddress(&Item);
+
+		const int32 OriginalNum=this->ArrayNum;
+		for( int32 Index=0; Index<this->ArrayNum; Index++ )
+		{
+			if( (*this)[Index]==Item )
+			{
+				RemoveAt( Index-- );
+			}
+		}
+		return OriginalNum - this->ArrayNum;
+	}
+
+	// TTransArray interface.
+	UObject* GetOwner() const
+	{
+		return Owner;
+	}
+
+	void SetOwner( UObject* NewOwner )
+	{
+		Owner = NewOwner;
+	}
+
+	void ModifyItem( int32 Index )
+	{
+		if( GUndo )
+		{
+			GUndo->SaveArray( Owner, (FScriptArray*)this, Index, 1, 0, sizeof(T), DefaultConstructItem, SerializeItem, DestructItem );
+		}
+	}
+
+	void ModifyAllItems()
+	{
+		if( GUndo )
+		{
+			GUndo->SaveArray( Owner, (FScriptArray*)this, 0, this->Num(), 0, sizeof(T), DefaultConstructItem, SerializeItem, DestructItem );
+		}
+	}
+
+	friend FArchive& operator<<( FArchive& Ar, TTransArray& A )
+	{
+		Ar << A.Owner;
+		Ar << (Super&)A;
+		return Ar;
+	}
+
+protected:
+
+	static void DefaultConstructItem( void* TPtr )
+	{
+		new (TPtr) T;
+	}
+	static void SerializeItem( FArchive& Ar, void* TPtr )
+	{
+		Ar << *(T*)TPtr;
+	}
+	static void DestructItem( void* TPtr )
+	{
+		((T*)TPtr)->~T();
+	}
+	UObject* Owner;
+};
+
+
+template<typename T>
+struct TContainerTraits<TTransArray<T> > : public TContainerTraitsBase<TTransArray<T> >
+{
+	enum { MoveWillEmptyContainer = TContainerTraitsBase<typename TTransArray<T>::Super>::MoveWillEmptyContainer };
+};
+
+
+//
+// Transactional array operator news.
+//
+template <typename T> void* operator new( size_t Size, TTransArray<T>& Array )
+{
+	check(Size == sizeof(T));
+	const int32 Index = Array.AddUninitialized();
+	return &Array[Index];
+}
+
+
+template <typename T> void* operator new( size_t Size, TTransArray<T>& Array, int32 Index )
+{
+	check(Size == sizeof(T));
+	Array.Insert(Index);
+	return &Array[Index];
+}
diff --git a/Engine/Source/Runtime/Core/Public/Containers/Union.h b/Engine/Source/Runtime/Core/Public/Containers/Union.h
--- a/Engine/Source/Runtime/Core/Public/Containers/Union.h
+++ b/Engine/Source/Runtime/Core/Public/Containers/Union.h
@@ -6,7 +6,7 @@
 #include "UnrealTypeTraits.h"
 #include "Containers/ContainerAllocationPolicies.h"
 #include "UnrealTemplate.h"
-#include "ArchiveBase.h"
+#include "Archive.h"
 
 
 CORE_API DECLARE_LOG_CATEGORY_EXTERN(LogUnion, Log, All);
diff --git a/Engine/Source/Runtime/Core/Public/Core.h b/Engine/Source/Runtime/Core/Public/Core.h
--- a/Engine/Source/Runtime/Core/Public/Core.h
+++ b/Engine/Source/Runtime/Core/Public/Core.h
@@ -16,8 +16,11 @@
 #include "PlatformIncludes.h"					// Include the main and misc platform headers
 #include "PlatformFilemanager.h"				// Platform file manager.
 #include "AssertionMacros.h"					// Various assertion macros
+#include "LogMacros.h"							// Various log macros
 #include "UObject/UnrealNames.h"				// EName definition.
 #include "OutputDevice.h"						// Output devices, logf, debugf, etc
+#include "Misc/MessageDialog.h"
+#include "Misc/Exec.h"
 #include "NumericLimits.h"						// Numeric limits
 #include "UnrealMathUtility.h"					// FMath
 #include "UnrealTypeTraits.h"					// Type Traits
@@ -52,17 +55,26 @@
 #include "ObjectVersion.h"				// Object version info.
 #include "TypeHash.h"					// Type hashing functions
 #include "EnumAsByte.h"					// Common template definitions.
-#include "ArchiveBase.h"				// Archive class.
+#include "Archive.h"					// Archive class.
+#include "ArchiveProxy.h"
+#include "NameAsStringProxyArchive.h"
 #include "ScopedPointer.h"				// Scoped pointer definitions.
 #include "Sorting.h"					// Sorting definitions.
 #include "Array.h"						// Dynamic array definitions.
+#include "ScriptArray.h"
+#include "MRUArray.h"
+#include "TransArray.h"
+#include "IndirectArray.h"
 #include "ArrayBuilder.h"				// Builder template for arrays.
 #include "BitArray.h"					// Bit array definition.
-#include "Bitstreams.h"					// Bit stream archiver.
+#include "BitReader.h"					// Bit stream archiver.
+#include "BitWriter.h"					// Bit stream archiver.
 #include "SparseArray.h"				// Sparse array definitions.
 #include "UnrealString.h"				// Dynamic string definitions.
 #include "NameTypes.h"					// Global name subsystem.
 #include "CoreMisc.h"					// Low level utility code.
+#include "CommandLine.h"
+#include "FileHelper.h"
 #include "Paths.h"						// Path helper functions
 #include "StaticArray.h"                // Static array definition.
 #include "StaticBitArray.h"             // Static bit array definition.
@@ -84,11 +96,31 @@
 #include "SHMath.h"						// SH math functions.
 #include "RandomStream.h"				// Random stream definitions.
 #include "OutputDevices.h"				// Output devices
+#include "OutputDeviceRedirector.h"
+#include "OutputDeviceNull.h"
+#include "OutputDeviceMemory.h"
+#include "OutputDeviceFile.h"
+#include "OutputDeviceDebug.h"
+#include "OutputDeviceArchiveWrapper.h"
+#include "OutputDeviceAnsiError.h"
+#include "BufferedOutputDevice.h"
+#include "LogScopedVerbosityOverride.h"
 #include "CoreStats.h"
 #include "TimeGuard.h"
 #include "MemStack.h"					// Stack based memory management.
 #include "AsyncWork.h"					// Async threaded work
-#include "Archive.h"					// Utility archive classes
+#include "MemoryArchive.h"
+#include "MemoryWriter.h"
+#include "LargeMemoryWriter.h"
+#include "LargeMemoryReader.h"
+#include "BufferArchive.h"
+#include "MemoryReader.h"
+#include "ArrayReader.h"
+#include "ArrayWriter.h"
+#include "BufferReader.h"
+#include "BufferWriter.h"
+#include "ArchiveSaveCompressedProxy.h"
+#include "ArchiveLoadCompressedProxy.h"
 #include "IOBase.h"						// base IO declarations, FIOManager, FIOSystem
 #include "Variant.h"
 #include "WildcardString.h"
@@ -100,6 +132,8 @@
 #include "ConfigCacheIni.h"				// The configuration cache declarations
 #include "IConsoleManager.h"
 #include "FeedbackContext.h"
+#include "SlowTask.h"
+#include "ScopedSlowTask.h"
 #include "AutomationTest.h"				// Automation testing functionality for game, editor, etc.
 #include "CallbackDevice.h"				// Base class for callback devices.
 #include "ObjectThumbnail.h"			// Object thumbnails
diff --git a/Engine/Source/Runtime/Core/Public/CoreGlobals.h b/Engine/Source/Runtime/Core/Public/CoreGlobals.h
--- a/Engine/Source/Runtime/Core/Public/CoreGlobals.h
+++ b/Engine/Source/Runtime/Core/Public/CoreGlobals.h
@@ -4,21 +4,63 @@
 #include "HAL/Platform.h"
 #include "Misc/CoreMiscDefines.h"
 #include "Misc/OutputDevice.h"
+#include "Logging/LogMacros.h"
 
 class FConfigCacheIni;
 class FExec;
 class FName;
 class FOutputDeviceConsole;
+class FOutputDeviceRedirector;
 class FReloadObjectArc;
 class FString;
 class FText;
 class ITransaction;
+class FRunnableThread;
 
 struct FScriptTraceStackNode;
 
+CORE_API DECLARE_LOG_CATEGORY_EXTERN(LogHAL, Log, All);
+CORE_API DECLARE_LOG_CATEGORY_EXTERN(LogMac, Log, All);
+CORE_API DECLARE_LOG_CATEGORY_EXTERN(LogLinux, Log, All);
+CORE_API DECLARE_LOG_CATEGORY_EXTERN(LogIOS, Log, All);
+CORE_API DECLARE_LOG_CATEGORY_EXTERN(LogAndroid, Log, All);
+CORE_API DECLARE_LOG_CATEGORY_EXTERN(LogPS4, Log, All);
+CORE_API DECLARE_LOG_CATEGORY_EXTERN(LogWindows, Log, All);
+CORE_API DECLARE_LOG_CATEGORY_EXTERN(LogSerialization, Log, All);
+CORE_API DECLARE_LOG_CATEGORY_EXTERN(LogUnrealMath, Log, All);
+CORE_API DECLARE_LOG_CATEGORY_EXTERN(LogUnrealMatrix, Log, All);
+CORE_API DECLARE_LOG_CATEGORY_EXTERN(LogContentComparisonCommandlet, Log, All);
+CORE_API DECLARE_LOG_CATEGORY_EXTERN(LogNetPackageMap, Warning, All);
+CORE_API DECLARE_LOG_CATEGORY_EXTERN(LogNetSerialization, Warning, All);
+CORE_API DECLARE_LOG_CATEGORY_EXTERN(LogMemory, Log, All);
+CORE_API DECLARE_LOG_CATEGORY_EXTERN(LogProfilingDebugging, Log, All);
+CORE_API DECLARE_LOG_CATEGORY_EXTERN(LogCore, Log, All);
+CORE_API DECLARE_LOG_CATEGORY_EXTERN(LogOutputDevice, Log, All);
+
+CORE_API DECLARE_LOG_CATEGORY_EXTERN(LogSHA, Warning, All);
+CORE_API DECLARE_LOG_CATEGORY_EXTERN(LogStats, Log, All);
+CORE_API DECLARE_LOG_CATEGORY_EXTERN(LogStreaming, Display, All);
+CORE_API DECLARE_LOG_CATEGORY_EXTERN(LogInit, Log, All);
+CORE_API DECLARE_LOG_CATEGORY_EXTERN(LogExit, Log, All);
+CORE_API DECLARE_LOG_CATEGORY_EXTERN(LogExec, Warning, All);
+CORE_API DECLARE_LOG_CATEGORY_EXTERN(LogScript, Warning, All);
+CORE_API DECLARE_LOG_CATEGORY_EXTERN(LogLocalization, Error, All);
+CORE_API DECLARE_LOG_CATEGORY_EXTERN(LogLongPackageNames, Log, All);
+CORE_API DECLARE_LOG_CATEGORY_EXTERN(LogProcess, Log, All);
+CORE_API DECLARE_LOG_CATEGORY_EXTERN(LogLoad, Log, All);
+
+// Temporary log category, generally you should not check things in that use this
+CORE_API DECLARE_LOG_CATEGORY_EXTERN(LogTemp, Log, All);
+
+CORE_API FOutputDeviceRedirector* GetGlobalLogSingleton();
+
+#define GLog GetGlobalLogSingleton()
 extern CORE_API FConfigCacheIni* GConfig;
 extern CORE_API ITransaction* GUndo;
 extern CORE_API FOutputDeviceConsole* GLogConsole;
+CORE_API extern class FOutputDeviceError*			GError;
+CORE_API extern class FFeedbackContext*				GWarn;
+
 
 extern CORE_API TCHAR GErrorHist[16384];
 
@@ -338,3 +380,48 @@ extern CORE_API bool GEnableVREditorHacks;
 CORE_API void EnsureRetrievingVTablePtrDuringCtor(const TCHAR* CtorSignature);
 #endif // WITH_HOT_RELOAD_CTORS
 
+/** @return True if called from the game thread. */
+FORCEINLINE bool IsInGameThread()
+{
+	if(GIsGameThreadIdInitialized)
+	{
+		const uint32 CurrentThreadId = FPlatformTLS::GetCurrentThreadId();
+		return CurrentThreadId == GGameThreadId || CurrentThreadId == GSlateLoadingThreadId;
+	}
+
+	return true;
+}
+
+/** @return True if called from the audio thread, and not merely a thread calling audio functions. */
+extern CORE_API bool IsInAudioThread();
+
+/** Thread used for audio */
+extern CORE_API FRunnableThread* GAudioThread;
+
+/** @return True if called from the slate thread, and not merely a thread calling slate functions. */
+extern CORE_API bool IsInSlateThread();
+
+/** @return True if called from the rendering thread, or if called from ANY thread during single threaded rendering */
+extern CORE_API bool IsInRenderingThread();
+
+/** @return True if called from the rendering thread, or if called from ANY thread that isn't the game thread, except that during single threaded rendering the game thread is ok too.*/
+extern CORE_API bool IsInParallelRenderingThread();
+
+/** @return True if called from the rendering thread. */
+// Unlike IsInRenderingThread, this will always return false if we are running single threaded. It only returns true if this is actually a separate rendering thread. Mostly useful for checks
+extern CORE_API bool IsInActualRenderingThread();
+
+/** @return True if called from the async loading thread if it's enabled, otherwise if called from game thread while is async loading code. */
+extern CORE_API bool (*IsInAsyncLoadingThread)();
+
+/** Thread used for rendering */
+extern CORE_API FRunnableThread* GRenderingThread;
+
+/** Whether the rendering thread is suspended (not even processing the tickables) */
+extern CORE_API int32 GIsRenderingThreadSuspended;
+
+/** @return True if called from the RHI thread, or if called from ANY thread during single threaded rendering */
+extern CORE_API bool IsInRHIThread();
+
+/** Thread used for RHI */
+extern CORE_API FRunnableThread* GRHIThread;
diff --git a/Engine/Source/Runtime/Core/Public/GenericPlatform/GenericPlatformCriticalSection.h b/Engine/Source/Runtime/Core/Public/GenericPlatform/GenericPlatformCriticalSection.h
new file mode 100644
--- /dev/null
+++ b/Engine/Source/Runtime/Core/Public/GenericPlatform/GenericPlatformCriticalSection.h
@@ -0,0 +1,28 @@
+// Copyright 1998-2016 Epic Games, Inc. All Rights Reserved.
+
+#pragma once
+
+/** Platforms that don't need a working FSystemWideCriticalSection can just typedef this one */
+class FSystemWideCriticalSectionNotImplemented
+{
+public:
+	/** Construct a named, system-wide critical section and attempt to get access/ownership of it */
+	explicit FSystemWideCriticalSectionNotImplemented(const FString& Name, FTimespan Timeout = FTimespan::Zero());
+
+	/** Destructor releases system-wide critical section if it is currently owned */
+	~FSystemWideCriticalSectionNotImplemented() {}
+
+	/**
+	 * Does the calling thread have ownership of the system-wide critical section?
+	 *
+	 * @return True if the system-wide lock is obtained.
+	 */
+	bool IsValid() const { return false; }
+
+	/** Releases system-wide critical section if it is currently owned */
+	void Release() {}
+
+private:
+	FSystemWideCriticalSectionNotImplemented(const FSystemWideCriticalSectionNotImplemented&);
+	FSystemWideCriticalSectionNotImplemented& operator=(const FSystemWideCriticalSectionNotImplemented&);
+};
diff --git a/Engine/Source/Runtime/Core/Public/GenericPlatform/GenericPlatformProcess.h b/Engine/Source/Runtime/Core/Public/GenericPlatform/GenericPlatformProcess.h
--- a/Engine/Source/Runtime/Core/Public/GenericPlatform/GenericPlatformProcess.h
+++ b/Engine/Source/Runtime/Core/Public/GenericPlatform/GenericPlatformProcess.h
@@ -568,33 +568,3 @@ struct CORE_API FGenericPlatformProcess
 };
 
 
-#if PLATFORM_USE_PTHREADS
-	// if we are using PThreads, we can set this up at the generic layer
-	#include "PThreadCriticalSection.h"
-	typedef FPThreadsCriticalSection FCriticalSection;
-#endif
-
-/** Platforms that don't need a working FSystemWideCriticalSection can just typedef this one */
-class FSystemWideCriticalSectionNotImplemented
-{
-public:
-	/** Construct a named, system-wide critical section and attempt to get access/ownership of it */
-	explicit FSystemWideCriticalSectionNotImplemented(const FString& Name, FTimespan Timeout = FTimespan::Zero());
-
-	/** Destructor releases system-wide critical section if it is currently owned */
-	~FSystemWideCriticalSectionNotImplemented() {}
-
-	/**
-	 * Does the calling thread have ownership of the system-wide critical section?
-	 *
-	 * @return True if the system-wide lock is obtained.
-	 */
-	bool IsValid() const { return false; }
-
-	/** Releases system-wide critical section if it is currently owned */
-	void Release() {}
-
-private:
-	FSystemWideCriticalSectionNotImplemented(const FSystemWideCriticalSectionNotImplemented&);
-	FSystemWideCriticalSectionNotImplemented& operator=(const FSystemWideCriticalSectionNotImplemented&);
-};
diff --git a/Engine/Source/Runtime/Core/Public/HAL/CriticalSection.h b/Engine/Source/Runtime/Core/Public/HAL/CriticalSection.h
new file mode 100644
--- /dev/null
+++ b/Engine/Source/Runtime/Core/Public/HAL/CriticalSection.h
@@ -0,0 +1,24 @@
+// Copyright 1998-2016 Epic Games, Inc. All Rights Reserved.
+
+#pragma once
+
+#include "HAL/Platform.h"
+#include "GenericPlatform/GenericPlatformCriticalSection.h"
+
+#if PLATFORM_WINDOWS
+	#include "Windows/WindowsCriticalSection.h"
+#elif PLATFORM_PS4
+	#include "PS4/PS4CriticalSection.h"
+#elif PLATFORM_XBOXONE
+	#include "XboxOne/XboxOneCriticalSection.h"
+#elif PLATFORM_MAC
+	#include "Mac/MacCriticalSection.h"
+#elif PLATFORM_IOS
+	#include "IOS/IOSCriticalSection.h"
+#elif PLATFORM_ANDROID
+	#include "Android/AndroidCriticalSection.h"
+#elif PLATFORM_HTML5
+	#include "HTML5/HTML5CriticalSection.h"
+#elif PLATFORM_LINUX
+	#include "Linux/LinuxCriticalSection.h"
+#endif
diff --git a/Engine/Source/Runtime/Core/Public/HAL/MallocBinned2.h b/Engine/Source/Runtime/Core/Public/HAL/MallocBinned2.h
--- a/Engine/Source/Runtime/Core/Public/HAL/MallocBinned2.h
+++ b/Engine/Source/Runtime/Core/Public/HAL/MallocBinned2.h
@@ -4,6 +4,8 @@
 
 #include "AlignmentTemplates.h"
 #include "Allocators/CachedOSPageAllocator.h"
+#include "Math/NumericLimits.h"
+#include "HAL/PlatformMath.h"
 
 #define BINNED2_MAX_CACHED_OS_FREES (64)
 #if PLATFORM_64BITS
@@ -409,7 +411,7 @@ public:
 					{
 						if (Result && Ptr)
 						{
-							FMemory::Memcpy(Result, Ptr, FMath::Min<SIZE_T>(NewSize, BlockSize));
+							FMemory::Memcpy(Result, Ptr, FPlatformMath::Min<SIZE_T>(NewSize, BlockSize));
 						}
 						if (Ptr)
 						{
@@ -458,7 +460,7 @@ public:
 	FORCEINLINE virtual SIZE_T QuantizeSize(SIZE_T Count, uint32 Alignment) override
 	{
 		static_assert(DEFAULT_ALIGNMENT <= BINNED2_MINIMUM_ALIGNMENT, "DEFAULT_ALIGNMENT is assumed to be zero"); // used below
-		checkSlow(FMath::IsPowerOfTwo(Alignment));
+		checkSlow((Alignment & (Alignment - 1)) == 0); // Check the alignment is a power of two
 		SIZE_T SizeOut;
 		if ((Count <= BINNED2_MAX_SMALL_POOL_SIZE) & (Alignment <= BINNED2_MINIMUM_ALIGNMENT)) // one branch, not two
 		{
@@ -466,7 +468,7 @@ public:
 		}
 		else
 		{
-			Alignment = FMath::Max<uint32>(Alignment, OsAllocationGranularity);
+			Alignment = FPlatformMath::Max<uint32>(Alignment, OsAllocationGranularity);
 			checkSlow(Alignment <= PageSize);
 			SizeOut = Align(Count, Alignment);
 		}
diff --git a/Engine/Source/Runtime/Core/Public/HAL/MemoryBase.h b/Engine/Source/Runtime/Core/Public/HAL/MemoryBase.h
--- a/Engine/Source/Runtime/Core/Public/HAL/MemoryBase.h
+++ b/Engine/Source/Runtime/Core/Public/HAL/MemoryBase.h
@@ -2,6 +2,8 @@
 
 #pragma once
 
+#include "Misc/Exec.h"
+
 enum
 {
 	// Default allocator alignment. If the default is specified, the allocator applies to engine rules.
