Copying //UE4/Dev-Build to //UE4/Dev-Main (Source: //UE4/Dev-Build @ 3104200)

diff --git a/Engine/Source/Runtime/Core/Public/HAL/OutputDevices.h b/Engine/Source/Runtime/Core/Public/HAL/OutputDevices.h
--- a/Engine/Source/Runtime/Core/Public/HAL/OutputDevices.h
+++ b/Engine/Source/Runtime/Core/Public/HAL/OutputDevices.h
@@ -2,600 +2,9 @@
 
 #pragma once
 
-#if !PLATFORM_DESKTOP
-	// don't support colorized text on consoles
-	#define SET_WARN_COLOR(Color)
-	#define SET_WARN_COLOR_AND_BACKGROUND(Color, Bkgrnd)
-	#define CLEAR_WARN_COLOR() 
-#else
+#include "Logging/LogCategory.h"
+#include "Logging/LogSuppressionInterface.h"
+#include "CoreGlobals.h"
+#include "Logging/LogScopedVerbosityOverride.h"
+#include "Logging/LogScopedCategoryAndVerbosityOverride.h"
 
-/*-----------------------------------------------------------------------------
-	Colorized text.
-
-	To use colored text from a commandlet, you use the SET_WARN_COLOR macro with
-	one of the following standard colors. Then use CLEAR_WARN_COLOR to return 
-	to default.
-
-	To use the standard colors, you simply do this:
-	SET_WARN_COLOR(COLOR_YELLOW);
-	
-	You can specify a background color by appending it to the foreground:
-	SET_WARN_COLOR, COLOR_YELLOW COLOR_DARK_RED);
-
-	This will have bright yellow text on a dark red background.
-
-	Or you can make your own in the format:
-	
-	ForegroundRed | ForegroundGreen | ForegroundBlue | ForegroundBright | BackgroundRed | BackgroundGreen | BackgroundBlue | BackgroundBright
-	where each value is either 0 or 1 (can leave off trailing 0's), so 
-	blue on bright yellow is "00101101" and red on black is "1"
-	
-	An empty string reverts to the normal gray on black.
------------------------------------------------------------------------------*/
-
-// putting them in a namespace to protect against future name conflicts
-namespace OutputDeviceColor
-{
-	const TCHAR* const COLOR_BLACK			= TEXT("0000");
-
-	const TCHAR* const COLOR_DARK_RED		= TEXT("1000");
-	const TCHAR* const COLOR_DARK_GREEN		= TEXT("0100");
-	const TCHAR* const COLOR_DARK_BLUE		= TEXT("0010");
-	const TCHAR* const COLOR_DARK_YELLOW	= TEXT("1100");
-	const TCHAR* const COLOR_DARK_CYAN		= TEXT("0110");
-	const TCHAR* const COLOR_DARK_PURPLE	= TEXT("1010");
-	const TCHAR* const COLOR_DARK_WHITE		= TEXT("1110");
-	const TCHAR* const COLOR_GRAY			= COLOR_DARK_WHITE;
-
-	const TCHAR* const COLOR_RED			= TEXT("1001");
-	const TCHAR* const COLOR_GREEN			= TEXT("0101");
-	const TCHAR* const COLOR_BLUE			= TEXT("0011");
-	const TCHAR* const COLOR_YELLOW			= TEXT("1101");
-	const TCHAR* const COLOR_CYAN			= TEXT("0111");
-	const TCHAR* const COLOR_PURPLE			= TEXT("1011");
-	const TCHAR* const COLOR_WHITE			= TEXT("1111");
-
-	const TCHAR* const COLOR_NONE			= TEXT("");
-}
-
-using namespace OutputDeviceColor;
-
-
-// let a console or (UE_BUILD_SHIPPING || UE_BUILD_TEST) define it to nothing
-#ifndef SET_WARN_COLOR
-
-/**
- * Set the console color with Color or a Color and Background color
- */
-#define SET_WARN_COLOR(Color) \
-	UE_LOG(LogHAL, SetColor, TEXT("%s"), Color);
-#define SET_WARN_COLOR_AND_BACKGROUND(Color, Bkgrnd) \
-	UE_LOG(LogHAL, SetColor, TEXT("%s%s"), Color, Bkgrnd);
-
-/**
- * Return color to it's default
- */
-#define CLEAR_WARN_COLOR() \
-	UE_LOG(LogHAL, SetColor, TEXT("%s"), COLOR_NONE);
-
-#endif
-#endif
-
-enum class EByteOrderMark : int8
-{
-	UTF8,
-	Unspecified,
-};
-
-/*-----------------------------------------------------------------------------
-	Log Suppression
------------------------------------------------------------------------------*/
-
-/** Interface to the log suppression system **/
-class FLogSuppressionInterface
-{
-public:
-	/** Singleton, returns a reference the global log suppression implementation. **/
-	static CORE_API FLogSuppressionInterface& Get();
-	/** Used by FLogCategoryBase to register itself with the global category table **/
-	virtual void AssociateSuppress(struct FLogCategoryBase* Destination) = 0;
-	/** Used by FLogCategoryBase to unregister itself from the global category table **/
-	virtual void DisassociateSuppress(struct FLogCategoryBase* Destination) = 0;
-	/** Called by appInit once the config files and commandline are set up. The log suppression system uses these to setup the boot time defaults. **/
-	virtual void ProcessConfigAndCommandLine() = 0;
-};
-
-/** Base class for all log categories. **/
-struct CORE_API FLogCategoryBase
-{
-	/** 
-	 * Constructor, registers with the log suppression system and sets up the default values.
-	 * @param CategoryName, name of the category
-	 * @param InDefaultVerbosity, default verbosity for the category, may ignored and overrridden by many other mechanisms
-	 * @param InCompileTimeVerbosity, mostly used to keep the working verbosity in bounds, macros elsewhere actually do the work of stripping compiled out things.
-	**/
-	FLogCategoryBase(const TCHAR *CategoryName, ELogVerbosity::Type InDefaultVerbosity, ELogVerbosity::Type InCompileTimeVerbosity);
-
-	/** Destructor, unregisters from the log suppression system **/
-	~FLogCategoryBase();
-	/** Should not generally be used directly. Tests the runtime verbosity and maybe triggers a debug break, etc. **/
-	FORCEINLINE bool IsSuppressed(ELogVerbosity::Type VerbosityLevel) const
-	{
-		if ((VerbosityLevel & ELogVerbosity::VerbosityMask) <= Verbosity)
-		{
-			return false;
-		}
-		return true;
-	}
-	/** Called just after a logging statement being allow to print. Checks a few things and maybe breaks into the debugger. **/
-	void PostTrigger(ELogVerbosity::Type VerbosityLevel);
-	FORCEINLINE FName GetCategoryName() const
-	{
-		return CategoryFName;
-	}
-
-	/** Sets up the working verbosity and clamps to the compile time verbosity. **/
-	void SetVerbosity(ELogVerbosity::Type Verbosity);
-
-private:
-	friend class FLogSuppressionImplementation;
-	friend class FLogScopedVerbosityOverride;
-	/** Internal call to get the working verbosity. **/
-	ELogVerbosity::Type GetVerbosity() const	{ return (ELogVerbosity::Type)Verbosity; }
-	/** Internal call to set up the working verbosity from the boot time default. **/
-	void ResetFromDefault();
-protected:
-
-	/** Holds the current suppression state **/
-	uint8 Verbosity;
-	/** Holds the break flag **/
-	bool DebugBreakOnLog;
-	/** Holds default suppression **/
-	uint8 DefaultVerbosity;
-	/** Holds compile time suppression **/
-	uint8 CompileTimeVerbosity;
-	/** FName for this category **/
-	FName CategoryFName;
-};
-
-/** Template for log categories that transfers the compile-time constant default and compile time verbosity to the FLogCategoryBase constructor. **/
-template<uint8 InDefaultVerbosity, uint8 InCompileTimeVerbosity>
-struct FLogCategory : public FLogCategoryBase
-{
-	static_assert((InDefaultVerbosity & ELogVerbosity::VerbosityMask) < ELogVerbosity::NumVerbosity, "Bogus default verbosity.");
-	static_assert(InCompileTimeVerbosity < ELogVerbosity::NumVerbosity, "Bogus compile time verbosity.");
-	enum
-	{
-		CompileTimeVerbosity = InCompileTimeVerbosity
-	};
-
-	FORCEINLINE FLogCategory(const TCHAR *CategoryName)
-		: FLogCategoryBase(CategoryName, ELogVerbosity::Type(InDefaultVerbosity), ELogVerbosity::Type(CompileTimeVerbosity))
-	{
-	}
-};
-
-/** 
- * Helper class that uses thread local storage to set up the default category and verbosity for the low level logging functions.
- * This is what allow a UE_LOG(LogHAL, Log, TEXT("...")); within a UE_LOG statement to know what the category and verbosity is.
- * When one of these goes out of scope, it restores the previous values.
-**/
-class CORE_API FScopedCategoryAndVerbosityOverride
-{
-public:
-	/** STructure to aggregate a category and verbosity pair **/
-	struct FOverride
-	{
-		ELogVerbosity::Type Verbosity;
-		FName Category;
-		FOverride()
-			: Verbosity(ELogVerbosity::Log)
-		{
-		}
-		FOverride(FName InCategory, ELogVerbosity::Type InVerbosity)
-			: Verbosity(InVerbosity)
-			, Category(InCategory)
-		{
-		}
-		void operator=(const FOverride& Other)
-		{
-			Verbosity = Other.Verbosity;
-			Category = Other.Category;
-		}
-	};
-private:
-	/** Backup of the category, verbosity pairs that was present when we were constructed **/
-	FOverride Backup;
-public:
-	/** Back up the existing category and varbosity pair, then sets them.*/
-	FScopedCategoryAndVerbosityOverride(FName Category, ELogVerbosity::Type Verbosity);
-
-	/** Restore the category and verbosity overrides to the previous value.*/
-	~FScopedCategoryAndVerbosityOverride();
-
-	/** Manages a TLS slot with the current overrides for category and verbosity. **/
-	static FOverride* GetTLSCurrent();
-};
-
-/*-----------------------------------------------------------------------------
-	FLogScopedVerbosityOverride
------------------------------------------------------------------------------*/
-/** 
- * Helper class that allows setting scoped verbosity for log category. 
- * Saved what was previous verbosity for the category, and recovers it when it goes out of scope. 
- * Use Macro LOG_SCOPE_VERBOSITY_OVERRIDE for this
- **/
-class CORE_API FLogScopedVerbosityOverride
-{
-private:
-	/** Backup of the category, verbosity pairs that was present when we were constructed **/
-	FLogCategoryBase * SavedCategory;
-	ELogVerbosity::Type SavedVerbosity;
-
-public:
-	/** Back up the existing verbosity for the category then sets new verbosity.*/
-	FLogScopedVerbosityOverride(FLogCategoryBase * Category, ELogVerbosity::Type Verbosity);
-
-	/** Restore the verbosity overrides for the category to the previous value.*/
-	~FLogScopedVerbosityOverride();
-};
-
-
-/*-----------------------------------------------------------------------------
-	FOutputDeviceRedirector.
------------------------------------------------------------------------------*/
-
-/** The type of lines buffered by secondary threads. */
-struct FBufferedLine
-{
-	const FString Data;
-	const FName Category;
-	const double Time;
-	const ELogVerbosity::Type Verbosity;
-
-	/** Initialization constructor. */
-	FBufferedLine( const TCHAR* InData, const class FName& InCategory, ELogVerbosity::Type InVerbosity, const double InTime = -1 )
-		: Data( InData )
-		, Category( InCategory )
-		, Time( InTime )
-		, Verbosity( InVerbosity )
-	{}
-};
-
-/**
- * Class used for output redirection to allow logs to show
- */
-class CORE_API FOutputDeviceRedirector : public FOutputDevice
-{
-private:
-	/** A FIFO of lines logged by non-master threads. */
-	TArray<FBufferedLine> BufferedLines;
-
-	/** A FIFO backlog of messages logged before the editor had a chance to intercept them. */
-	TArray<FBufferedLine> BacklogLines;
-
-	/** Array of output devices to redirect to */
-	TArray<FOutputDevice*> OutputDevices;
-
-	/** The master thread ID.  Logging from other threads will be buffered for processing by the master thread. */
-	uint32 MasterThreadID;
-
-	/** Whether backlogging is enabled. */
-	bool bEnableBacklog;
-
-	/** Object used for synchronization via a scoped lock */
-	FCriticalSection	SynchronizationObject;
-
-	/**
-	 * The unsynchronized version of FlushThreadedLogs.
-	 * Assumes that the caller holds a lock on SynchronizationObject.
-	 * @param bUseAllDevices - if true this method will use all output devices
-	 */
-	void UnsynchronizedFlushThreadedLogs( bool bUseAllDevices );
-
-public:
-
-	/** Initialization constructor. */
-	FOutputDeviceRedirector();
-
-	/**
-	 * Get the GLog singleton
-	 */
-	static FOutputDeviceRedirector* Get();
-
-	/**
-	 * Adds an output device to the chain of redirections.	
-	 *
-	 * @param OutputDevice	output device to add
-	 */
-	void AddOutputDevice( FOutputDevice* OutputDevice );
-
-	/**
-	 * Removes an output device from the chain of redirections.	
-	 *
-	 * @param OutputDevice	output device to remove
-	 */
-	void RemoveOutputDevice( FOutputDevice* OutputDevice );
-
-	/**
-	 * Returns whether an output device is currently in the list of redirectors.
-	 *
-	 * @param	OutputDevice	output device to check the list against
-	 * @return	true if messages are currently redirected to the the passed in output device, false otherwise
-	 */
-	bool IsRedirectingTo( FOutputDevice* OutputDevice );
-
-	/** Flushes lines buffered by secondary threads. */
-	virtual void FlushThreadedLogs();
-
-	/**
-	 *	Flushes lines buffered by secondary threads.
-	 *	Only used if a background thread crashed and we needed to push the callstack into the log. 
-	 */
-	virtual void PanicFlushThreadedLogs();
-
-	/**
-	 * Serializes the current backlog to the specified output device.
-	 * @param OutputDevice	- Output device that will receive the current backlog
-	 */
-	virtual void SerializeBacklog( FOutputDevice* OutputDevice );
-
-	/**
-	 * Enables or disables the backlog.
-	 * @param bEnable	- Starts saving a backlog if true, disables and discards any backlog if false
-	 */
-	virtual void EnableBacklog( bool bEnable );
-
-	/**
-	 * Sets the current thread to be the master thread that prints directly
-	 * (isn't queued up)
-	 */
-	virtual void SetCurrentThreadAsMasterThread();
-
-	/**
-	 * Serializes the passed in data via all current output devices.
-	 *
-	 * @param	Data	Text to log
-	 * @param	Event	Event name used for suppression purposes
-	 */
-	virtual void Serialize( const TCHAR* Data, ELogVerbosity::Type Verbosity, const class FName& Category, const double Time) override;
-
-	/**
-	 * Serializes the passed in data via all current output devices.
-	 *
-	 * @param	Data	Text to log
-	 * @param	Event	Event name used for suppression purposes
-	 */
-	virtual void Serialize( const TCHAR* Data, ELogVerbosity::Type Verbosity, const class FName& Category ) override;
-
-	/**
-	 * Passes on the flush request to all current output devices.
-	 */
-	void Flush() override;
-
-	/**
-	 * Closes output device and cleans up. This can't happen in the destructor
-	 * as we might have to call "delete" which cannot be done for static/ global
-	 * objects.
-	 */
-	void TearDown() override;
-};
-
-/*-----------------------------------------------------------------------------
-	FOutputDevice subclasses.
------------------------------------------------------------------------------*/
-
-/** string added to the filename of timestamped backup log files */
-#define BACKUP_LOG_FILENAME_POSTFIX TEXT("-backup-")
-
-/** Used by FOutputDeviceFile to write to a file on a separate thread */
-class FAsyncWriter;
-
-/**
- * File output device (Note: Only works if ALLOW_LOG_FILE && !NO_LOGGING is true, otherwise Serialize does nothing).
- */
-class CORE_API FOutputDeviceFile : public FOutputDevice
-{
-public:
-	/** 
-	 * Constructor, initializing member variables.
-	 *
-	 * @param InFilename	Filename to use, can be nullptr
-	 * @param bDisableBackup If true, existing files will not be backed up
-	 */
-	FOutputDeviceFile( const TCHAR* InFilename = nullptr, bool bDisableBackup=false );
-
-	/** Sets the filename that the output device writes to.  If the output device was already writing to a file, closes that file. */
-	void SetFilename(const TCHAR* InFilename);
-
-	//~ Begin FOutputDevice Interface.
-	/**
-	 * Closes output device and cleans up. This can't happen in the destructor
-	 * as we have to call "delete" which cannot be done for static/ global
-	 * objects.
-	 */
-	void TearDown() override;
-
-	/**
-	 * Flush the write cache so the file isn't truncated in case we crash right
-	 * after calling this function.
-	 */
-	void Flush() override;
-
-	virtual void Serialize( const TCHAR* Data, ELogVerbosity::Type Verbosity, const class FName& Category, const double Time ) override;
-	virtual void Serialize( const TCHAR* Data, ELogVerbosity::Type Verbosity, const class FName& Category) override;
-	virtual bool CanBeUsedOnAnyThread() const override
-	{
-		return true;
-	}
-	//~ End FOutputDevice Interface.
-
-	/** Creates a backup copy of a log file if it already exists */
-	static void CreateBackupCopy(const TCHAR* Filename);
-
-	/** Checks if the filename represents a backup copy of a log file */
-	static bool IsBackupCopy(const TCHAR* Filename);
-
-private:
-
-	/** Writes to a file on a separate thread */
-	FAsyncWriter* AsyncWriter;
-	/** Archive used by the async writer */
-	FArchive* WriterArchive;
-
-	TCHAR Filename[1024];
-	bool Opened;
-	bool Dead;
-
-	/** If true, existing files will not be backed up */
-	bool		bDisableBackup;
-	
-	void WriteRaw( const TCHAR* C );
-
-	/** Creates the async writer and its archive. Returns true if successful.  */
-	bool CreateWriter(uint32 MaxAttempts = 32);
-
-	void WriteByteOrderMarkToArchive(EByteOrderMark ByteOrderMark);
-};
-
-// Null output device.
-class CORE_API FOutputDeviceNull : public FOutputDevice
-{
-public:
-	/**
-	 * NULL implementation of Serialize.
-	 *
-	 * @param	Data	unused
-	 * @param	Event	unused
-	 */
-	virtual void Serialize( const TCHAR* V, ELogVerbosity::Type Verbosity, const class FName& Category ) override
-	{}
-};
-
-class CORE_API FOutputDeviceDebug : public FOutputDevice
-{
-public:
-	/**
-	 * Serializes the passed in data unless the current event is suppressed.
-	 *
-	 * @param	Data	Text to log
-	 * @param	Event	Event name used for suppression purposes
-	 */
-	virtual void Serialize( const TCHAR* Data, ELogVerbosity::Type Verbosity, const class FName& Category, const double Time ) override;
-
-	virtual void Serialize( const TCHAR* Data, ELogVerbosity::Type Verbosity, const class FName& Category ) override;
-
-	virtual bool CanBeUsedOnAnyThread() const override
-	{
-		return true;
-	}
-};
-
-/** Buffered output device. */
-class CORE_API FBufferedOutputDevice : public FOutputDevice
-{
-protected:
-	TArray<FBufferedLine>		BufferedLines;
-	ELogVerbosity::Type			FilterLevel;
-	FCriticalSection			SynchronizationObject;
-
-public:
-
-	FBufferedOutputDevice() : FilterLevel(ELogVerbosity::All)
-	{
-	}
-
-	void	SetVerbosity(ELogVerbosity::Type Verbosity) { FilterLevel = Verbosity; }
-	void	Serialize(const TCHAR* InData, ELogVerbosity::Type Verbosity, const class FName& Category) override;
-	void	GetContents(TArray<FBufferedLine>& DestBuffer, bool ClearDevice = true);
-
-	/** Pushes buffered lines into the specified output device. */
-	void RedirectTo( class FOutputDevice& Ar )
-	{
-		for( const FBufferedLine& BufferedLine : BufferedLines )
-		{
-			Ar.Serialize( *BufferedLine.Data, BufferedLine.Verbosity, BufferedLine.Category );
-		}
-	}
-};
-
-/*-----------------------------------------------------------------------------
-	FOutputDeviceError subclasses.
------------------------------------------------------------------------------*/
-
-class FOutputDeviceAnsiError : public FOutputDeviceError
-{
-public:
-	/** Constructor, initializing member variables */
-	FOutputDeviceAnsiError();
-
-	/**
-	 * Serializes the passed in data unless the current event is suppressed.
-	 *
-	 * @param	Data	Text to log
-	 * @param	Event	Event name used for suppression purposes
-	 */
-	virtual void Serialize( const TCHAR* Msg, ELogVerbosity::Type Verbosity, const class FName& Category ) override;
-
-	virtual bool CanBeUsedOnAnyThread() const override
-	{
-		return true;
-	}
-
-	/**
-	 * Error handling function that is being called from within the system wide global
-	 * error handler, e.g. using structured exception handling on the PC.
-	 */
-	void HandleError() override;
-
-private:
-
-	int32		ErrorPos;
-};
-
-
-/**
- * Output device wrapping any kind of FArchive.  Note: Works in any build configuration.
- */
-class CORE_API FOutputDeviceArchiveWrapper : public FOutputDevice
-{
-public:
-	/**
-	 * Constructor, initializing member variables.
-	 *
-	 * @param InArchive	Archive to use, must not be nullptr.  Does not take ownership of the archive, clean up or delete the archive independently!
-	 */
-	FOutputDeviceArchiveWrapper(FArchive* InArchive)
-		: LogAr(InArchive)
-	{
-		check(InArchive);
-	}
-
-	// FOutputDevice interface
-	virtual void Flush() override;
-	virtual void Serialize(const TCHAR* Data, ELogVerbosity::Type Verbosity, const class FName& Category) override;
-	// End of FOutputDevice interface
-
-private:
-	FArchive* LogAr;
-};
-
-CORE_API DECLARE_LOG_CATEGORY_EXTERN(LogHAL, Log, All);
-CORE_API DECLARE_LOG_CATEGORY_EXTERN(LogMac, Log, All);
-CORE_API DECLARE_LOG_CATEGORY_EXTERN(LogLinux, Log, All);
-CORE_API DECLARE_LOG_CATEGORY_EXTERN(LogIOS, Log, All);
-CORE_API DECLARE_LOG_CATEGORY_EXTERN(LogAndroid, Log, All);
-CORE_API DECLARE_LOG_CATEGORY_EXTERN(LogPS4, Log, All);
-CORE_API DECLARE_LOG_CATEGORY_EXTERN(LogWindows, Log, All);
-CORE_API DECLARE_LOG_CATEGORY_EXTERN(LogSerialization, Log, All);
-CORE_API DECLARE_LOG_CATEGORY_EXTERN(LogUnrealMath, Log, All);
-CORE_API DECLARE_LOG_CATEGORY_EXTERN(LogUnrealMatrix, Log, All);
-CORE_API DECLARE_LOG_CATEGORY_EXTERN(LogContentComparisonCommandlet, Log, All);
-CORE_API DECLARE_LOG_CATEGORY_EXTERN(LogNetPackageMap, Warning, All);
-CORE_API DECLARE_LOG_CATEGORY_EXTERN(LogNetSerialization, Warning, All);
-CORE_API DECLARE_LOG_CATEGORY_EXTERN(LogMemory, Log, All);
-CORE_API DECLARE_LOG_CATEGORY_EXTERN(LogProfilingDebugging, Log, All);
-CORE_API DECLARE_LOG_CATEGORY_EXTERN(LogCore, Log, All);
-
-// Temporary log category, generally you should not check things in that use this
-CORE_API DECLARE_LOG_CATEGORY_EXTERN(LogTemp, Log, All);
diff --git a/Engine/Source/Runtime/Core/Public/HAL/PlatformIncludes.h b/Engine/Source/Runtime/Core/Public/HAL/PlatformIncludes.h
--- a/Engine/Source/Runtime/Core/Public/HAL/PlatformIncludes.h
+++ b/Engine/Source/Runtime/Core/Public/HAL/PlatformIncludes.h
@@ -12,6 +12,7 @@
 #include "PlatformNamedPipe.h"
 #include "PlatformTime.h"
 #include "PlatformProcess.h"
+#include "CriticalSection.h"
 #include "PlatformOutputDevices.h"
 #include "PlatformAtomics.h"
 #include "PlatformTLS.h"
diff --git a/Engine/Source/Runtime/Core/Public/HAL/ThreadManager.h b/Engine/Source/Runtime/Core/Public/HAL/ThreadManager.h
new file mode 100644
--- /dev/null
+++ b/Engine/Source/Runtime/Core/Public/HAL/ThreadManager.h
@@ -0,0 +1,45 @@
+// Copyright 1998-2016 Epic Games, Inc. All Rights Reserved.
+
+#pragma once
+
+/**
+ * Manages runnables and runnable threads.
+ */
+class CORE_API FThreadManager
+{
+	/** List of thread objects to be ticked. */
+	TMap<uint32, class FRunnableThread*> Threads;
+	/** Critical section for ThreadList */
+	FCriticalSection ThreadsCritical;
+
+public:
+
+	/**
+	* Used internally to add a new thread object.
+	*
+	* @param Thread thread object.
+	* @see RemoveThread
+	*/
+	void AddThread(uint32 ThreadId, class FRunnableThread* Thread);
+
+	/**
+	* Used internally to remove thread object.
+	*
+	* @param Thread thread object to be removed.
+	* @see AddThread
+	*/
+	void RemoveThread(class FRunnableThread* Thread);
+
+	/** Ticks all fake threads and their runnable objects. */
+	void Tick();
+
+	/** Returns the name of a thread given its TLS id */
+	const FString& GetThreadName(uint32 ThreadId);
+
+	/**
+	 * Access to the singleton object.
+	 *
+	 * @return Thread manager object.
+	 */
+	static FThreadManager& Get();
+};
diff --git a/Engine/Source/Runtime/Core/Public/HAL/ThreadingBase.h b/Engine/Source/Runtime/Core/Public/HAL/ThreadingBase.h
--- a/Engine/Source/Runtime/Core/Public/HAL/ThreadingBase.h
+++ b/Engine/Source/Runtime/Core/Public/HAL/ThreadingBase.h
@@ -4,1353 +4,23 @@
 
 #include "Function.h"
 
-class FTlsAutoCleanup;
-class FRunnableThread;
+#include "Event.h"
+#include "Misc/ScopedEvent.h"
+#include "Runnable.h"
+#include "RunnableThread.h"
+#include "Misc/SingleThreadEvent.h"
+#include "Misc/SingleThreadRunnable.h"
+#include "ThreadManager.h"
+#include "Misc/IQueuedWork.h"
+#include "Misc/QueuedThreadPool.h"
+#include "ThreadSafeCounter.h"
+#include "ThreadSafeCounter64.h"
+#include "ThreadSafeBool.h"
+#include "Misc/NoopCounter.h"
+#include "Misc/MultiReaderSingleWriterGT.h"
+#include "Misc/ScopeLock.h"
+#include "TlsAutoCleanup.h"
+#include "ThreadSingleton.h"
 
-/**
- * The list of enumerated thread priorities we support
- */
-enum EThreadPriority
-{
-	TPri_Normal,
-	TPri_AboveNormal,
-	TPri_BelowNormal,
-	TPri_Highest,
-	TPri_Lowest,
-	TPri_SlightlyBelowNormal,
-};
 
 
-/**
- * Interface for waitable events.
- *
- * This interface has platform-specific implementations that are used to wait for another
- * thread to signal that it is ready for the waiting thread to do some work. It can also
- * be used for telling groups of threads to exit.
- */
-class FEvent
-{
-public:
-
-	/**
-	 * Creates the event.
-	 *
-	 * Manually reset events stay triggered until reset.
-	 * Named events share the same underlying event.
-	 *
-	 * @param bIsManualReset Whether the event requires manual reseting or not.
-	 * @return true if the event was created, false otherwise.
-	 */
-	virtual bool Create( bool bIsManualReset = false ) = 0;
-
-	/**
-	 * Whether the signaled state of this event needs to be reset manually.
-	 *
-	 * @return true if the state requires manual resetting, false otherwise.
-	 * @see Reset
-	 */
-	virtual bool IsManualReset() = 0;
-
-	/**
-	 * Triggers the event so any waiting threads are activated.
-	 *
-	 * @see IsManualReset, Reset
-	 */
-	virtual void Trigger() = 0;
-
-	/**
-	 * Resets the event to an untriggered (waitable) state.
-	 *
-	 * @see IsManualReset, Trigger
-	 */
-	virtual void Reset() = 0;
-
-	/**
-	 * Waits the specified amount of time for the event to be triggered.
-	 *
-	 * A wait time of MAX_uint32 is treated as infinite wait.
-	 *
-	 * @param WaitTime The time to wait (in milliseconds).
-	 * @param bIgnoreThreadIdleStats If true, ignores ThreadIdleStats
-	 * @return true if the event was triggered, false if the wait timed out.
-	 */
-	virtual bool Wait( uint32 WaitTime, const bool bIgnoreThreadIdleStats = false ) = 0;
-
-	/**
-	 * Waits an infinite amount of time for the event to be triggered.
-	 *
-	 * @return true if the event was triggered.
-	 */
-	bool Wait()
-	{
-		return Wait(MAX_uint32);
-	}
-
-	/**
-	 * Waits the specified amount of time for the event to be triggered.
-	 *
-	 * @param WaitTime The time to wait.
-	 * @param bIgnoreThreadIdleStats If true, ignores ThreadIdleStats
-	 * @return true if the event was triggered, false if the wait timed out.
-	 */
-	bool Wait( const FTimespan& WaitTime, const bool bIgnoreThreadIdleStats = false )
-	{
-		return Wait(WaitTime.GetTotalMilliseconds(), bIgnoreThreadIdleStats);
-	}
-
-	/** Default constructor. */
-	FEvent()
-		: EventId( 0 )
-		, EventStartCycles( 0 )
-	{}
-
-	/** Virtual destructor. */
-	virtual ~FEvent() 
-	{}
-
-	// DO NOT MODIFY THESE
-
-	/** Advances stats associated with this event. Used to monitor wait->trigger history. */
-	void AdvanceStats();
-
-protected:
-	/** Sends to the stats a special messages which encodes a wait for the event. */
-	void WaitForStats();
-
-	/** Send to the stats a special message which encodes a trigger for the event. */
-	void TriggerForStats();
-
-	/** Resets start cycles to 0. */
-	void ResetForStats();
-
-	/** Counter used to generate an unique id for the events. */
-	static uint32 EventUniqueId;
-
-	/** An unique id of this event. */
-	uint32 EventId;
-
-	/** Greater than 0, if the event called wait. */
-	uint32 EventStartCycles;
-};
-
-
-/**
- * This class is allows a simple one-shot scoped event.
- *
- * Usage:
- * {
- *		FScopedEvent MyEvent;
- *		SendReferenceOrPointerToSomeOtherThread(&MyEvent); // Other thread calls MyEvent->Trigger();
- *		// MyEvent destructor is here, we wait here.
- * }
- */
-class FScopedEvent
-{
-public:
-
-	/** Default constructor. */
-	CORE_API FScopedEvent();
-
-	/** Destructor. */
-	CORE_API ~FScopedEvent();
-
-	/** Triggers the event. */
-	void Trigger()
-	{
-		Event->Trigger();
-	}
-
-	/**
-	 * Retrieve the event, usually for passing around.
-	 *
-	 * @return The event.
-	 */
-	FEvent* Get()
-	{
-		return Event;
-	}
-
-private:
-
-	/** Holds the event. */
-	FEvent* Event;
-};
-
-
-/**
- * Interface for "runnable" objects.
- *
- * A runnable object is an object that is "run" on an arbitrary thread. The call usage pattern is
- * Init(), Run(), Exit(). The thread that is going to "run" this object always uses those calling
- * semantics. It does this on the thread that is created so that any thread specific uses (TLS, etc.)
- * are available in the contexts of those calls. A "runnable" does all initialization in Init().
- *
- * If initialization fails, the thread stops execution and returns an error code. If it succeeds,
- * Run() is called where the real threaded work is done. Upon completion, Exit() is called to allow
- * correct clean up.
- */
-class CORE_API FRunnable
-{
-public:
-
-	/**
-	 * Initializes the runnable object.
-	 *
-	 * This method is called in the context of the thread object that aggregates this, not the
-	 * thread that passes this runnable to a new thread.
-	 *
-	 * @return True if initialization was successful, false otherwise
-	 * @see Run, Stop, Exit
-	 */
-	virtual bool Init()
-	{
-		return true;
-	}
-
-	/**
-	 * Runs the runnable object.
-	 *
-	 * This is where all per object thread work is done. This is only called if the initialization was successful.
-	 *
-	 * @return The exit code of the runnable object
-	 * @see Init, Stop, Exit
-	 */
-	virtual uint32 Run() = 0;
-
-	/**
-	 * Stops the runnable object.
-	 *
-	 * This is called if a thread is requested to terminate early.
-	 * @see Init, Run, Exit
-	 */
-	virtual void Stop() { }
-
-	/**
-	 * Exits the runnable object.
-	 *
-	 * Called in the context of the aggregating thread to perform any cleanup.
-	 * @see Init, Run, Stop
-	 */
-	virtual void Exit() { }
-
-	/**
-	 * Gets single thread interface pointer used for ticking this runnable when multi-threading is disabled.
-	 * If the interface is not implemented, this runnable will not be ticked when FPlatformProcess::SupportsMultithreading() is false.
-	 *
-	 * @return Pointer to the single thread interface or nullptr if not implemented.
-	 */
-	virtual class FSingleThreadRunnable* GetSingleThreadInterface( )
-	{
-		return nullptr;
-	}
-
-	/** Virtual destructor */
-	virtual ~FRunnable() { }
-};
-
-
-/**
- * Interface for runnable threads.
- *
- * This interface specifies the methods used to manage a thread's life cycle.
- */
-class CORE_API FRunnableThread
-{
-	friend class FThreadSingletonInitializer;
-	friend class FTlsAutoCleanup;
-	friend class FThreadManager;
-
-	/** Index of TLS slot for FRunnableThread pointer. */
-	static uint32 RunnableTlsSlot;
-
-public:
-
-	/** Gets a new Tls slot for storing the runnable thread pointer. */
-	static uint32 GetTlsSlot();
-
-	/**
-	* Factory method to create a thread with the specified stack size and thread priority.
-	*
-	* @param InRunnable The runnable object to execute
-	* @param ThreadName Name of the thread
-	* @param bAutoDeleteSelf Whether to delete this object on exit
-	* @param bAutoDeleteRunnable Whether to delete the runnable object on exit
-	* @param InStackSize The size of the stack to create. 0 means use the current thread's stack size
-	* @param InThreadPri Tells the thread whether it needs to adjust its priority or not. Defaults to normal priority
-	* @return The newly created thread or nullptr if it failed
-	*/
-	DEPRECATED(4.3, "Function deprecated. Use FRunnableThread::Create without bAutoDeleteSelf and bAutoDeleteRunnable params and delete thread and runnable manually.")
-	static FRunnableThread* Create(
-		class FRunnable* InRunnable,
-		const TCHAR* ThreadName,
-		bool bAutoDeleteSelf,
-		bool bAutoDeleteRunnable = false,
-		uint32 InStackSize = 0,
-		EThreadPriority InThreadPri = TPri_Normal,
-		uint64 InThreadAffinityMask = 0);
-
-	/**
-	 * Factory method to create a thread with the specified stack size and thread priority.
-	 *
-	 * @param InRunnable The runnable object to execute
-	 * @param ThreadName Name of the thread
-	 * @param InStackSize The size of the stack to create. 0 means use the current thread's stack size
-	 * @param InThreadPri Tells the thread whether it needs to adjust its priority or not. Defaults to normal priority
-	 * @return The newly created thread or nullptr if it failed
-	 */
-	static FRunnableThread* Create(
-		class FRunnable* InRunnable,
-		const TCHAR* ThreadName,
-		uint32 InStackSize = 0,
-		EThreadPriority InThreadPri = TPri_Normal,
-		uint64 InThreadAffinityMask = FPlatformAffinity::GetNoAffinityMask() );
-
-	/**
-	 * Changes the thread priority of the currently running thread
-	 *
-	 * @param NewPriority The thread priority to change to
-	 */
-	virtual void SetThreadPriority( EThreadPriority NewPriority ) = 0;
-
-	/**
-	 * Tells the thread to either pause execution or resume depending on the
-	 * passed in value.
-	 *
-	 * @param bShouldPause Whether to pause the thread (true) or resume (false)
-	 */
-	virtual void Suspend( bool bShouldPause = true ) = 0;
-
-	/**
-	 * Tells the thread to exit. If the caller needs to know when the thread has exited, it should use the bShouldWait value.
-	 * It's highly recommended not to kill the thread without waiting for it.
-	 * Having a thread forcibly destroyed can cause leaks and deadlocks.
-	 * 
-	 * The kill method is calling Stop() on the runnable to kill the thread gracefully.
-	 * 
-	 * @param bShouldWait	If true, the call will wait infinitely for the thread to exit.					
-	 * @return Always true
-	 */
-	virtual bool Kill( bool bShouldWait = true ) = 0;
-
-	/** Halts the caller until this thread is has completed its work. */
-	virtual void WaitForCompletion() = 0;
-
-	/**
-	 * Thread ID for this thread 
-	 *
-	 * @return ID that was set by CreateThread
-	 * @see GetThreadName
-	 */
-	const uint32 GetThreadID() const
-	{
-		return ThreadID;
-	}
-
-	/**
-	 * Retrieves the given name of the thread
-	 *
-	 * @return Name that was set by CreateThread
-	 * @see GetThreadID
-	 */
-	const FString& GetThreadName() const
-	{
-		return ThreadName;
-	}
-
-	/** Default constructor. */
-	FRunnableThread();
-
-	/** Virtual destructor */
-	virtual ~FRunnableThread();
-
-protected:
-
-	/**
-	 * Creates the thread with the specified stack size and thread priority.
-	 *
-	 * @param InRunnable The runnable object to execute
-	 * @param ThreadName Name of the thread
-	 * @param InStackSize The size of the stack to create. 0 means use the current thread's stack size
-	 * @param InThreadPri Tells the thread whether it needs to adjust its priority or not. Defaults to normal priority
-	 * @return True if the thread and all of its initialization was successful, false otherwise
-	 */
-	virtual bool CreateInternal( FRunnable* InRunnable, const TCHAR* InThreadName,
-		uint32 InStackSize = 0,
-		EThreadPriority InThreadPri = TPri_Normal, uint64 InThreadAffinityMask = 0 ) = 0;
-
-	/** Stores this instance in the runnable thread TLS slot. */
-	void SetTls();
-
-	/** Deletes all FTlsAutoCleanup objects created for this thread. */
-	void FreeTls();
-
-	/**
-	 * @return a runnable thread that is executing this runnable, if return value is nullptr, it means the running thread can be game thread or a thread created outside the runnable interface
-	 */
-	static FRunnableThread* GetRunnableThread()
-	{
-		FRunnableThread* RunnableThread = (FRunnableThread*)FPlatformTLS::GetTlsValue( RunnableTlsSlot );
-		return RunnableThread;
-	}
-
-	/** Holds the name of the thread. */
-	FString ThreadName;
-
-	/** The runnable object to execute on this thread. */
-	FRunnable* Runnable;
-
-	/** Sync event to make sure that Init() has been completed before allowing the main thread to continue. */
-	FEvent* ThreadInitSyncEvent;
-
-	/** The Affinity to run the thread with. */
-	uint64 ThreadAffinityMask;
-
-	/** An array of FTlsAutoCleanup based instances that needs to be deleted before the thread will die. */
-	TArray<FTlsAutoCleanup*> TlsInstances;
-
-	/** The priority to run the thread at. */
-	EThreadPriority ThreadPriority;
-
-	/** ID set during thread creation. */
-	uint32 ThreadID;
-
-private:
-
-	/** Used by the thread manager to tick threads in single-threaded mode */
-	virtual void Tick() {}
-};
-
-
-/**
- * Fake event object used when running with only one thread.
- */
-class FSingleThreadEvent
-	: public FEvent
-{
-	/** Flag to know whether this event has been triggered. */
-	bool bTriggered;
-
-	/** Should this event reset automatically or not. */
-	bool bManualReset;
-
-public:
-
-	/** Default constructor. */
-	FSingleThreadEvent()
-		: bTriggered(false)
-		, bManualReset(false)
-	{ }
-
-public:
-
-	// FEvent Interface
-
-	virtual bool Create( bool bIsManualReset = false ) override
-	{ 
-		bManualReset = bIsManualReset;
-		return true; 
-	}
-
-	virtual bool IsManualReset() override
-	{
-		return bManualReset;
-	}
-
-	virtual void Trigger() override
-	{
-		bTriggered = true;
-	}
-
-	virtual void Reset() override
-	{
-		bTriggered = false;
-	}
-
-	virtual bool Wait( uint32 WaitTime, const bool bIgnoreThreadIdleStats = false ) override
-	{ 
-		// With only one thread it's assumed the event has been triggered
-		// before Wait is called, otherwise it would end up waiting forever or always fail.
-		check(bTriggered);
-		bTriggered = bManualReset;
-		return true; 
-	}
-};
-
-
-/** 
- * Interface for ticking runnables when there's only one thread available and 
- * multithreading is disabled.
- */
-class CORE_API FSingleThreadRunnable
-{
-public:
-
-	virtual ~FSingleThreadRunnable() { }
-
-	/* Tick function. */
-	virtual void Tick() = 0;
-};
-
-
-/**
- * Manages runnables and runnable threads.
- */
-class CORE_API FThreadManager
-{
-	/** List of thread objects to be ticked. */
-	TMap<uint32, class FRunnableThread*> Threads;
-	/** Critical section for ThreadList */
-	FCriticalSection ThreadsCritical;
-
-public:
-
-	/**
-	* Used internally to add a new thread object.
-	*
-	* @param Thread thread object.
-	* @see RemoveThread
-	*/
-	void AddThread(uint32 ThreadId, class FRunnableThread* Thread);
-
-	/**
-	* Used internally to remove thread object.
-	*
-	* @param Thread thread object to be removed.
-	* @see AddThread
-	*/
-	void RemoveThread(class FRunnableThread* Thread);
-
-	/** Ticks all fake threads and their runnable objects. */
-	void Tick();
-
-	/** Returns the name of a thread given its TLS id */
-	const FString& GetThreadName(uint32 ThreadId);
-
-	/**
-	 * Access to the singleton object.
-	 *
-	 * @return Thread manager object.
-	 */
-	static FThreadManager& Get();
-};
-
-
-/**
- * Interface for queued work objects.
- *
- * This interface is a type of runnable object that requires no per thread
- * initialization. It is meant to be used with pools of threads in an
- * abstract way that prevents the pool from needing to know any details
- * about the object being run. This allows queuing of disparate tasks and
- * servicing those tasks with a generic thread pool.
- */
-class IQueuedWork
-{
-public:
-
-	/**
-	 * This is where the real thread work is done. All work that is done for
-	 * this queued object should be done from within the call to this function.
-	 */
-	virtual void DoThreadedWork() = 0;
-
-	/**
-	 * Tells the queued work that it is being abandoned so that it can do
-	 * per object clean up as needed. This will only be called if it is being
-	 * abandoned before completion. NOTE: This requires the object to delete
-	 * itself using whatever heap it was allocated in.
-	 */
-	virtual void Abandon() = 0;
-
-public:
-
-	/**
-	 * Virtual destructor so that child implementations are guaranteed a chance
-	 * to clean up any resources they allocated.
-	 */
-	virtual ~IQueuedWork() { }
-};
-
-
-DEPRECATED(4.8, "FQueuedWork has been renamed to IQueuedWork")
-typedef IQueuedWork FQueuedWork;
-
-
-/**
- * Interface for queued thread pools.
- *
- * This interface is used by all queued thread pools. It used as a callback by
- * FQueuedThreads and is used to queue asynchronous work for callers.
- */
-class CORE_API FQueuedThreadPool
-{
-public:
-
-	/**
-	 * Creates the thread pool with the specified number of threads
-	 *
-	 * @param InNumQueuedThreads Specifies the number of threads to use in the pool
-	 * @param StackSize The size of stack the threads in the pool need (32K default)
-	 * @param ThreadPriority priority of new pool thread
-	 * @return Whether the pool creation was successful or not
-	 */
-	virtual bool Create( uint32 InNumQueuedThreads, uint32 StackSize = (32 * 1024), EThreadPriority ThreadPriority=TPri_Normal ) = 0;
-
-	/** Tells the pool to clean up all background threads */
-	virtual void Destroy() = 0;
-
-	/**
-	 * Checks to see if there is a thread available to perform the task. If not,
-	 * it queues the work for later. Otherwise it is immediately dispatched.
-	 *
-	 * @param InQueuedWork The work that needs to be done asynchronously
-	 * @see RetractQueuedWork
-	 */
-	virtual void AddQueuedWork( IQueuedWork* InQueuedWork ) = 0;
-
-	/**
-	 * Attempts to retract a previously queued task.
-	 *
-	 * @param InQueuedWork The work to try to retract
-	 * @return true if the work was retracted
-	 * @see AddQueuedWork
-	 */
-	virtual bool RetractQueuedWork( IQueuedWork* InQueuedWork ) = 0;
-
-	/**
-	 * Places a thread back into the available pool
-	 *
-	 * @param InQueuedThread The thread that is ready to be pooled
-	 * @return next job or null if there is no job available now
-	 */
-	virtual IQueuedWork* ReturnToPoolOrGetNextJob( class FQueuedThread* InQueuedThread ) = 0;
-
-public:
-
-	/** Virtual destructor. */
-	virtual ~FQueuedThreadPool() { }
-
-public:
-
-	/**
-	 * Allocates a thread pool
-	 *
-	 * @return A new thread pool.
-	 */
-	static FQueuedThreadPool* Allocate();
-
-	/**
-	 *	Stack size for threads created for the thread pool. 
-	 *	Can be overridden by other projects.
-	 *	If 0 means to use the value passed in the Create method.
-	 */
-	static uint32 OverrideStackSize;
-};
-
-
-/**
- *  Global thread pool for shared async operations
- */
-extern CORE_API FQueuedThreadPool* GThreadPool;
-
-#if USE_NEW_ASYNC_IO
-extern CORE_API FQueuedThreadPool* GIOThreadPool;
-#endif // USE_NEW_ASYNC_IO
-
-#if WITH_EDITOR
-extern CORE_API FQueuedThreadPool* GLargeThreadPool;
-#endif
-
-
-/** Thread safe counter */
-class FThreadSafeCounter
-{
-public:
-	typedef int32 IntegerType;
-
-	/**
-	 * Default constructor.
-	 *
-	 * Initializes the counter to 0.
-	 */
-	FThreadSafeCounter()
-	{
-		Counter = 0;
-	}
-
-	/**
-	 * Copy Constructor.
-	 *
-	 * If the counter in the Other parameter is changing from other threads, there are no
-	 * guarantees as to which values you will get up to the caller to not care, synchronize
-	 * or other way to make those guarantees.
-	 *
-	 * @param Other The other thread safe counter to copy
-	 */
-	FThreadSafeCounter( const FThreadSafeCounter& Other )
-	{
-		Counter = Other.GetValue();
-	}
-
-	/**
-	 * Constructor, initializing counter to passed in value.
-	 *
-	 * @param Value	Value to initialize counter to
-	 */
-	FThreadSafeCounter( int32 Value )
-	{
-		Counter = Value;
-	}
-
-	/**
-	 * Increment and return new value.	
-	 *
-	 * @return the new, incremented value
-	 * @see Add, Decrement, Reset, Set, Subtract
-	 */
-	int32 Increment()
-	{
-		return FPlatformAtomics::InterlockedIncrement(&Counter);
-	}
-
-	/**
-	 * Adds an amount and returns the old value.	
-	 *
-	 * @param Amount Amount to increase the counter by
-	 * @return the old value
-	 * @see Decrement, Increment, Reset, Set, Subtract
-	 */
-	int32 Add( int32 Amount )
-	{
-		return FPlatformAtomics::InterlockedAdd(&Counter, Amount);
-	}
-
-	/**
-	 * Decrement and return new value.
-	 *
-	 * @return the new, decremented value
-	 * @see Add, Increment, Reset, Set, Subtract
-	 */
-	int32 Decrement()
-	{
-		return FPlatformAtomics::InterlockedDecrement(&Counter);
-	}
-
-	/**
-	 * Subtracts an amount and returns the old value.	
-	 *
-	 * @param Amount Amount to decrease the counter by
-	 * @return the old value
-	 * @see Add, Decrement, Increment, Reset, Set
-	 */
-	int32 Subtract( int32 Amount )
-	{
-		return FPlatformAtomics::InterlockedAdd(&Counter, -Amount);
-	}
-
-	/**
-	 * Sets the counter to a specific value and returns the old value.
-	 *
-	 * @param Value	Value to set the counter to
-	 * @return The old value
-	 * @see Add, Decrement, Increment, Reset, Subtract
-	 */
-	int32 Set( int32 Value )
-	{
-		return FPlatformAtomics::InterlockedExchange(&Counter, Value);
-	}
-
-	/**
-	 * Resets the counter's value to zero.
-	 *
-	 * @return the old value.
-	 * @see Add, Decrement, Increment, Set, Subtract
-	 */
-	int32 Reset()
-	{
-		return FPlatformAtomics::InterlockedExchange(&Counter, 0);
-	}
-
-	/**
-	 * Gets the current value.
-	 *
-	 * @return the current value
-	 */
-	int32 GetValue() const
-	{
-		return Counter;
-	}
-
-private:
-
-	/** Hidden on purpose as usage wouldn't be thread safe. */
-	void operator=( const FThreadSafeCounter& Other ){}
-
-	/** Thread-safe counter */
-	volatile int32 Counter;
-};
-
-// This class cannot be implemented on platforms that don't define 64bit atomic functions
-#if PLATFORM_HAS_64BIT_ATOMICS
-/** Thread safe counter for 64bit ints */
-class FThreadSafeCounter64
-{
-public:
-	typedef int64 IntegerType;
-
-	/**
-	* Default constructor.
-	*
-	* Initializes the counter to 0.
-	*/
-	FThreadSafeCounter64()
-	{
-		Counter = 0;
-	}
-
-	/**
-	* Copy Constructor.
-	*
-	* If the counter in the Other parameter is changing from other threads, there are no
-	* guarantees as to which values you will get up to the caller to not care, synchronize
-	* or other way to make those guarantees.
-	*
-	* @param Other The other thread safe counter to copy
-	*/
-	FThreadSafeCounter64(const FThreadSafeCounter& Other)
-	{
-		Counter = Other.GetValue();
-	}
-
-	/** Assignment has the same caveats as the copy ctor. */
-	FThreadSafeCounter64& operator=(const FThreadSafeCounter64& Other)
-	{
-		if (this == &Other) return *this;
-		Set(Other.GetValue());
-		return *this;
-	}
-
-
-	/**
-	* Constructor, initializing counter to passed in value.
-	*
-	* @param Value	Value to initialize counter to
-	*/
-	FThreadSafeCounter64(int64 Value)
-	{
-		Counter = Value;
-	}
-
-	/**
-	* Increment and return new value.
-	*
-	* @return the new, incremented value
-	* @see Add, Decrement, Reset, Set, Subtract
-	*/
-	int64 Increment()
-	{
-		return FPlatformAtomics::InterlockedIncrement(&Counter);
-	}
-
-	/**
-	* Adds an amount and returns the old value.
-	*
-	* @param Amount Amount to increase the counter by
-	* @return the old value
-	* @see Decrement, Increment, Reset, Set, Subtract
-	*/
-	int64 Add(int64 Amount)
-	{
-		return FPlatformAtomics::InterlockedAdd(&Counter, Amount);
-	}
-
-	/**
-	* Decrement and return new value.
-	*
-	* @return the new, decremented value
-	* @see Add, Increment, Reset, Set, Subtract
-	*/
-	int64 Decrement()
-	{
-		return FPlatformAtomics::InterlockedDecrement(&Counter);
-	}
-
-	/**
-	* Subtracts an amount and returns the old value.
-	*
-	* @param Amount Amount to decrease the counter by
-	* @return the old value
-	* @see Add, Decrement, Increment, Reset, Set
-	*/
-	int64 Subtract(int64 Amount)
-	{
-		return FPlatformAtomics::InterlockedAdd(&Counter, -Amount);
-	}
-
-	/**
-	* Sets the counter to a specific value and returns the old value.
-	*
-	* @param Value	Value to set the counter to
-	* @return The old value
-	* @see Add, Decrement, Increment, Reset, Subtract
-	*/
-	int64 Set(int64 Value)
-	{
-		return FPlatformAtomics::InterlockedExchange(&Counter, Value);
-	}
-
-	/**
-	* Resets the counter's value to zero.
-	*
-	* @return the old value.
-	* @see Add, Decrement, Increment, Set, Subtract
-	*/
-	int64 Reset()
-	{
-		return FPlatformAtomics::InterlockedExchange(&Counter, 0);
-	}
-
-	/**
-	* Gets the current value.
-	*
-	* @return the current value
-	*/
-	int64 GetValue() const
-	{
-		return Counter;
-	}
-
-private:
-	/** Thread-safe counter */
-	volatile int64 Counter;
-};
-#endif
-
-/**
- * Thread safe bool, wraps FThreadSafeCounter
- */
-class FThreadSafeBool
-	: private FThreadSafeCounter
-{
-public:
-	/**
-	 * Constructor optionally takes value to initialize with, otherwise initializes false
-	 * @param bValue	Value to initialize with
-	 */
-	FThreadSafeBool(bool bValue = false)
-		: FThreadSafeCounter(bValue ? 1 : 0)
-	{}
-
-	/**
-	 * Operator to use this struct as a bool with thread safety
-	 */
-	FORCEINLINE operator bool() const
-	{
-		return GetValue() != 0;
-	}
-
-	/**
-	 * Operator to set the bool value with thread safety
-	 */
-	FORCEINLINE bool operator=(bool bNewValue)
-	{
-		Set(bNewValue ? 1 : 0);
-		return bNewValue;
-	}
-
-	/**
-	 * Sets a new value atomically, and returns the old value.
-	 *
-	 * @param bNewValue   Value to set
-	 * @return The old value
-	 */
-	FORCEINLINE bool AtomicSet(bool bNewValue)
-	{
-		return Set(bNewValue ? 1 : 0) != 0;
-	}
-};
-
-
-/** Fake Thread safe counter, used to avoid cluttering code with ifdefs when counters are only used for debugging. */
-class FNoopCounter
-{
-public:
-
-	FNoopCounter() { }
-	FNoopCounter( const FNoopCounter& Other ) { }
-	FNoopCounter( int32 Value ) { }
-
-	int32 Increment()
-	{
-		return 0;
-	}
-
-	int32 Add( int32 Amount )
-	{
-		return 0;
-	}
-
-	int32 Decrement()
-	{
-		return 0;
-	}
-
-	int32 Subtract( int32 Amount )
-	{
-		return 0;
-	}
-
-	int32 Set( int32 Value )
-	{
-		return 0;
-	}
-
-	int32 Reset()
-	{
-		return 0;
-	}
-
-	int32 GetValue() const
-	{
-		return 0;
-	}
-};
-
-/**
- * Object synchronizing read access (from any thread) and write access (only from game thread).
- * Multiple simultaneous reads from different threads are allowed. Allows multiple writes from
- * reentrant calls on main thread. Locking write doesn't prevent reads from main thread.
- **/
-class FMultiReaderSingleWriterGT
-{
-public:
-	/** Protect data from modification while reading. If issued on game thread, doesn't wait for write to finish. */
-	void LockRead();
-
-	/** Ends protecting data from modification while reading. */
-	void UnlockRead();
-
-	/** Protect data from reading while modifying. Can be called only on game thread. Reentrant. */
-	void LockWrite();
-
-	/** Ends protecting data from reading while modifying. */
-	void UnlockWrite();
-
-private:
-	struct FMRSWCriticalSection
-	{
-		FMRSWCriticalSection()
-			: Action(0)
-			, ReadCounter(0)
-			, WriteCounter(0)
-		{ }
-
-		int32 Action;
-		FThreadSafeCounter ReadCounter;
-		FThreadSafeCounter WriteCounter;
-	};
-
-	FMRSWCriticalSection CriticalSection;
-
-	static const int32 WritingAction = -1;
-	static const int32 NoAction = 0;
-	static const int32 ReadingAction = 1;
-
-	bool CanRead()
-	{
-		return FPlatformAtomics::InterlockedCompareExchange(&CriticalSection.Action, ReadingAction, NoAction) == ReadingAction;
-	}
-	
-	bool CanWrite()
-	{
-		return FPlatformAtomics::InterlockedCompareExchange(&CriticalSection.Action, WritingAction, NoAction) == WritingAction;
-	}
-};
-
-class FReadScopeLock
-{
-public:
-
-	/**
-	* Constructor that performs a lock on the synchronization object
-	*
-	* @param InSynchObject The synchronization object to manage
-	*/
-	FReadScopeLock(FMultiReaderSingleWriterGT* InSynchObject)
-		: SynchObject(InSynchObject)
-	{
-		check(SynchObject);
-		SynchObject->LockRead();
-	}
-
-	/** Destructor that performs a release on the synchronization object. */
-	~FReadScopeLock()
-	{
-		check(SynchObject);
-		SynchObject->UnlockRead();
-	}
-private:
-
-	/** Default constructor (hidden on purpose). */
-	FReadScopeLock();
-
-	/** Copy constructor( hidden on purpose). */
-	FReadScopeLock(FReadScopeLock* InScopeLock);
-
-	/** Assignment operator (hidden on purpose). */
-	FReadScopeLock& operator=(FReadScopeLock& InScopeLock)
-	{
-		return *this;
-	}
-
-private:
-
-	// Holds the synchronization object to aggregate and scope manage.
-	FMultiReaderSingleWriterGT* SynchObject;
-};
-
-class FWriteScopeLock
-{
-public:
-
-	/**
-	* Constructor that performs a lock on the synchronization object
-	*
-	* @param InSynchObject The synchronization object to manage
-	*/
-	FWriteScopeLock(FMultiReaderSingleWriterGT* InSynchObject)
-		: SynchObject(InSynchObject)
-	{
-		check(SynchObject);
-		SynchObject->LockWrite();
-	}
-
-	/** Destructor that performs a release on the synchronization object. */
-	~FWriteScopeLock()
-	{
-		check(SynchObject);
-		SynchObject->UnlockWrite();
-	}
-private:
-
-	/** Default constructor (hidden on purpose). */
-	FWriteScopeLock();
-
-	/** Copy constructor( hidden on purpose). */
-	FWriteScopeLock(FWriteScopeLock* InScopeLock);
-
-	/** Assignment operator (hidden on purpose). */
-	FWriteScopeLock& operator=(FWriteScopeLock& InScopeLock)
-	{
-		return *this;
-	}
-
-private:
-
-	// Holds the synchronization object to aggregate and scope manage.
-	FMultiReaderSingleWriterGT* SynchObject;
-};
-
-/**
- * Implements a scope lock.
- *
- * This is a utility class that handles scope level locking. It's very useful
- * to keep from causing deadlocks due to exceptions being caught and knowing
- * about the number of locks a given thread has on a resource. Example:
- *
- * <code>
- *	{
- *		// Synchronize thread access to the following data
- *		FScopeLock ScopeLock(SynchObject);
- *		// Access data that is shared among multiple threads
- *		...
- *		// When ScopeLock goes out of scope, other threads can access data
- *	}
- * </code>
- */
-class FScopeLock
-{
-public:
-
-	/**
-	 * Constructor that performs a lock on the synchronization object
-	 *
-	 * @param InSynchObject The synchronization object to manage
-	 */
-	FScopeLock( FCriticalSection* InSynchObject )
-		: SynchObject(InSynchObject)
-	{
-		check(SynchObject);
-		SynchObject->Lock();
-	}
-
-	/** Destructor that performs a release on the synchronization object. */
-	~FScopeLock()
-	{
-		check(SynchObject);
-		SynchObject->Unlock();
-	}
-private:
-
-	/** Default constructor (hidden on purpose). */
-	FScopeLock();
-
-	/** Copy constructor( hidden on purpose). */
-	FScopeLock(const FScopeLock& InScopeLock);
-
-	/** Assignment operator (hidden on purpose). */
-	FScopeLock& operator=( FScopeLock& InScopeLock )
-	{
-		return *this;
-	}
-
-private:
-
-	// Holds the synchronization object to aggregate and scope manage.
-	FCriticalSection* SynchObject;
-};
-
-
-/** @return True if called from the game thread. */
-FORCEINLINE bool IsInGameThread()
-{
-	if(GIsGameThreadIdInitialized)
-	{
-		const uint32 CurrentThreadId = FPlatformTLS::GetCurrentThreadId();
-		return CurrentThreadId == GGameThreadId || CurrentThreadId == GSlateLoadingThreadId;
-	}
-
-	return true;
-}
-
-/** @return True if called from the audio thread, and not merely a thread calling audio functions. */
-extern CORE_API bool IsInAudioThread();
-
-/** Thread used for audio */
-extern CORE_API FRunnableThread* GAudioThread;
-
-/** @return True if called from the slate thread, and not merely a thread calling slate functions. */
-extern CORE_API bool IsInSlateThread();
-
-/** @return True if called from the rendering thread, or if called from ANY thread during single threaded rendering */
-extern CORE_API bool IsInRenderingThread();
-
-/** @return True if called from the rendering thread, or if called from ANY thread that isn't the game thread, except that during single threaded rendering the game thread is ok too.*/
-extern CORE_API bool IsInParallelRenderingThread();
-
-/** @return True if called from the rendering thread. */
-// Unlike IsInRenderingThread, this will always return false if we are running single threaded. It only returns true if this is actually a separate rendering thread. Mostly useful for checks
-extern CORE_API bool IsInActualRenderingThread();
-
-/** @return True if called from the async loading thread if it's enabled, otherwise if called from game thread while is async loading code. */
-extern CORE_API bool (*IsInAsyncLoadingThread)();
-
-/** Thread used for rendering */
-extern CORE_API FRunnableThread* GRenderingThread;
-
-/** Whether the rendering thread is suspended (not even processing the tickables) */
-extern CORE_API int32 GIsRenderingThreadSuspended;
-
-/** @return True if called from the RHI thread, or if called from ANY thread during single threaded rendering */
-extern CORE_API bool IsInRHIThread();
-
-/** Thread used for RHI */
-extern CORE_API FRunnableThread* GRHIThread;
-
-/** Base class for objects in TLS that support auto-cleanup. */
-class CORE_API FTlsAutoCleanup
-{
-public:
-	/** Virtual destructor. */
-	virtual ~FTlsAutoCleanup()
-	{}
-
-	/** Register this instance to be auto-cleanup. */
-	void Register();
-};
-
-/** Wrapper for values to be stored in TLS that support auto-cleanup. */
-template< class T >
-class TTlsAutoCleanupValue
-	: public FTlsAutoCleanup
-{
-public:
-
-	/** Constructor. */
-	TTlsAutoCleanupValue(const T& InValue)
-		: Value(InValue)
-	{ }
-
-	/** Gets the value. */
-	T Get() const
-	{
-		return Value;
-	}
-
-private:
-
-	/** The value. */
-	T Value;
-};
-
-/**
- * Thread singleton initializer.
- */
-class FThreadSingletonInitializer
-{
-public:
-
-	/**
-	* @return an instance of a singleton for the current thread.
-	*/
-	static CORE_API FTlsAutoCleanup* Get( TFunctionRef<FTlsAutoCleanup*()> CreateInstance, uint32& TlsSlot );
-};
-
-
-/**
- * This a special version of singleton. It means that there is created only one instance for each thread.
- * Calling Get() method is thread-safe.
- */
-template < class T >
-class TThreadSingleton : public FTlsAutoCleanup
-{
-	/**
-	 * @return TLS slot that holds a TThreadSingleton.
-	 */
-	static uint32& GetTlsSlot()
-	{
-		static uint32 TlsSlot = 0;
-		return TlsSlot;
-	}
-
-protected:
-
-	/** Default constructor. */
-	TThreadSingleton()
-		: ThreadId(FPlatformTLS::GetCurrentThreadId())
-	{}
-
-	/**
-	 * @return a new instance of the thread singleton.
-	 */
-	static FTlsAutoCleanup* CreateInstance()
-	{
-		return new T();
-	}
-
-	/** Thread ID of this thread singleton. */
-	const uint32 ThreadId;
-
-public:
-
-	/**
-	 *	@return an instance of a singleton for the current thread.
-	 */
-	FORCEINLINE static T& Get()
-	{
-		return *(T*)FThreadSingletonInitializer::Get( [](){ return (FTlsAutoCleanup*)new T(); }, T::GetTlsSlot() ); //-V572
-	}
-};
-
-#define DECLARE_THREAD_SINGLETON(ClassType) \
-	EMIT_DEPRECATED_WARNING_MESSAGE("DECLARE_THREAD_SINGLETON is deprecated in 4.8. It's no longer needed and can be removed.")
\ No newline at end of file
diff --git a/Engine/Source/Runtime/Core/Public/HAL/TlsAutoCleanup.h b/Engine/Source/Runtime/Core/Public/HAL/TlsAutoCleanup.h
new file mode 100644
--- /dev/null
+++ b/Engine/Source/Runtime/Core/Public/HAL/TlsAutoCleanup.h
@@ -0,0 +1,39 @@
+// Copyright 1998-2016 Epic Games, Inc. All Rights Reserved.
+
+#pragma once
+
+/** Base class for objects in TLS that support auto-cleanup. */
+class CORE_API FTlsAutoCleanup
+{
+public:
+	/** Virtual destructor. */
+	virtual ~FTlsAutoCleanup()
+	{}
+
+	/** Register this instance to be auto-cleanup. */
+	void Register();
+};
+
+/** Wrapper for values to be stored in TLS that support auto-cleanup. */
+template< class T >
+class TTlsAutoCleanupValue
+	: public FTlsAutoCleanup
+{
+public:
+
+	/** Constructor. */
+	TTlsAutoCleanupValue(const T& InValue)
+		: Value(InValue)
+	{ }
+
+	/** Gets the value. */
+	T Get() const
+	{
+		return Value;
+	}
+
+private:
+
+	/** The value. */
+	T Value;
+};
diff --git a/Engine/Source/Runtime/Core/Public/HTML5/HTML5CriticalSection.h b/Engine/Source/Runtime/Core/Public/HTML5/HTML5CriticalSection.h
new file mode 100644
--- /dev/null
+++ b/Engine/Source/Runtime/Core/Public/HTML5/HTML5CriticalSection.h
@@ -0,0 +1,36 @@
+// Copyright 1998-2016 Epic Games, Inc. All Rights Reserved.
+
+#pragma once
+
+#include "GenericPlatform/GenericPlatformCriticalSection.h"
+
+/**
+ * @todo html5 threads: Dummy critical section
+ */
+class FHTML5CriticalSection
+{
+public:
+
+	FHTML5CriticalSection() {}
+
+	/**
+	 * Locks the critical section
+	 */
+	FORCEINLINE void Lock(void)
+	{
+	}
+
+	/**
+	 * Releases the lock on the critical seciton
+	 */
+	FORCEINLINE void Unlock(void)
+	{
+	}
+
+private:
+	FHTML5CriticalSection(const FHTML5CriticalSection&);
+	FHTML5CriticalSection& operator=(const FHTML5CriticalSection&);
+};
+
+typedef FHTML5CriticalSection FCriticalSection;
+typedef FSystemWideCriticalSectionNotImplemented FSystemWideCriticalSection;
diff --git a/Engine/Source/Runtime/Core/Public/HTML5/HTML5PlatformProcess.h b/Engine/Source/Runtime/Core/Public/HTML5/HTML5PlatformProcess.h
--- a/Engine/Source/Runtime/Core/Public/HTML5/HTML5PlatformProcess.h
+++ b/Engine/Source/Runtime/Core/Public/HTML5/HTML5PlatformProcess.h
@@ -41,35 +41,3 @@ struct CORE_API FHTML5PlatformProcess : public FGenericPlatformProcess
 };
 
 typedef FHTML5PlatformProcess FPlatformProcess;
-
-/**
- * @todo html5 threads: Dummy critical section
- */
-class FHTML5CriticalSection
-{
-public:
-
-	FHTML5CriticalSection() {}
-
-	/**
-	 * Locks the critical section
-	 */
-	FORCEINLINE void Lock(void)
-	{
-	}
-
-	/**
-	 * Releases the lock on the critical seciton
-	 */
-	FORCEINLINE void Unlock(void)
-	{
-	}
-
-private:
-	FHTML5CriticalSection(const FHTML5CriticalSection&);
-	FHTML5CriticalSection& operator=(const FHTML5CriticalSection&);
-};
-
-typedef FHTML5CriticalSection FCriticalSection;
-
-typedef FSystemWideCriticalSectionNotImplemented FSystemWideCriticalSection;
diff --git a/Engine/Source/Runtime/Core/Public/IOS/IOSPlatformProcess.h b/Engine/Source/Runtime/Core/Public/IOS/IOSPlatformProcess.h
--- a/Engine/Source/Runtime/Core/Public/IOS/IOSPlatformProcess.h
+++ b/Engine/Source/Runtime/Core/Public/IOS/IOSPlatformProcess.h
@@ -45,5 +45,3 @@ private:
 };
 
 typedef FIOSPlatformProcess FPlatformProcess;
-
-typedef FSystemWideCriticalSectionNotImplemented FSystemWideCriticalSection;
diff --git a/Engine/Source/Runtime/Core/Public/Linux/LinuxCriticalSection.h b/Engine/Source/Runtime/Core/Public/Linux/LinuxCriticalSection.h
new file mode 100644
--- /dev/null
+++ b/Engine/Source/Runtime/Core/Public/Linux/LinuxCriticalSection.h
@@ -0,0 +1,39 @@
+// Copyright 1998-2016 Epic Games, Inc. All Rights Reserved.
+
+#pragma once
+
+#include "PThreadCriticalSection.h"
+#include "GenericPlatform/GenericPlatformCriticalSection.h"
+
+/**
+  * Linux implementation of the FSystemWideCriticalSection. Uses exclusive file locking.
+  **/
+class FLinuxSystemWideCriticalSection
+{
+public:
+	/** Construct a named, system-wide critical section and attempt to get access/ownership of it */
+	explicit FLinuxSystemWideCriticalSection(const FString& InName, FTimespan InTimeout = FTimespan::Zero());
+
+	/** Destructor releases system-wide critical section if it is currently owned */
+	~FLinuxSystemWideCriticalSection();
+
+	/**
+	 * Does the calling thread have ownership of the system-wide critical section?
+	 *
+	 * @return True if the system-wide lock is obtained. WARNING: Returns true for abandoned locks so shared resources can be in undetermined states.
+	 */
+	bool IsValid() const;
+
+	/** Releases system-wide critical section if it is currently owned */
+	void Release();
+
+private:
+	FLinuxSystemWideCriticalSection(const FLinuxSystemWideCriticalSection&);
+	FLinuxSystemWideCriticalSection& operator=(const FLinuxSystemWideCriticalSection&);
+
+private:
+	int32 FileHandle;
+};
+
+typedef FPThreadsCriticalSection FCriticalSection;
+typedef FLinuxSystemWideCriticalSection FSystemWideCriticalSection;
diff --git a/Engine/Source/Runtime/Core/Public/Linux/LinuxPlatformProcess.h b/Engine/Source/Runtime/Core/Public/Linux/LinuxPlatformProcess.h
--- a/Engine/Source/Runtime/Core/Public/Linux/LinuxPlatformProcess.h
+++ b/Engine/Source/Runtime/Core/Public/Linux/LinuxPlatformProcess.h
@@ -309,34 +309,3 @@ inline bool FProcHandle::Close()
 	return false;
 }
 
-/**
-  * Linux implementation of the FSystemWideCriticalSection. Uses exclusive file locking.
-  **/
-class FLinuxSystemWideCriticalSection
-{
-public:
-	/** Construct a named, system-wide critical section and attempt to get access/ownership of it */
-	explicit FLinuxSystemWideCriticalSection(const FString& InName, FTimespan InTimeout = FTimespan::Zero());
-
-	/** Destructor releases system-wide critical section if it is currently owned */
-	~FLinuxSystemWideCriticalSection();
-
-	/**
-	 * Does the calling thread have ownership of the system-wide critical section?
-	 *
-	 * @return True if the system-wide lock is obtained. WARNING: Returns true for abandoned locks so shared resources can be in undetermined states.
-	 */
-	bool IsValid() const;
-
-	/** Releases system-wide critical section if it is currently owned */
-	void Release();
-
-private:
-	FLinuxSystemWideCriticalSection(const FLinuxSystemWideCriticalSection&);
-	FLinuxSystemWideCriticalSection& operator=(const FLinuxSystemWideCriticalSection&);
-
-private:
-	int32 FileHandle;
-};
-
-typedef FLinuxSystemWideCriticalSection FSystemWideCriticalSection;
diff --git a/Engine/Source/Runtime/Core/Public/Logging/LogCategory.h b/Engine/Source/Runtime/Core/Public/Logging/LogCategory.h
new file mode 100644
--- /dev/null
+++ b/Engine/Source/Runtime/Core/Public/Logging/LogCategory.h
@@ -0,0 +1,76 @@
+// Copyright 1998-2016 Epic Games, Inc. All Rights Reserved.
+
+#pragma once
+
+#include "UObject/NameTypes.h"
+#include "LogVerbosity.h"
+
+/** Base class for all log categories. **/
+struct CORE_API FLogCategoryBase
+{
+	/**
+	* Constructor, registers with the log suppression system and sets up the default values.
+	* @param CategoryName, name of the category
+	* @param InDefaultVerbosity, default verbosity for the category, may ignored and overrridden by many other mechanisms
+	* @param InCompileTimeVerbosity, mostly used to keep the working verbosity in bounds, macros elsewhere actually do the work of stripping compiled out things.
+	**/
+	FLogCategoryBase(const TCHAR *CategoryName, ELogVerbosity::Type InDefaultVerbosity, ELogVerbosity::Type InCompileTimeVerbosity);
+
+	/** Destructor, unregisters from the log suppression system **/
+	~FLogCategoryBase();
+	/** Should not generally be used directly. Tests the runtime verbosity and maybe triggers a debug break, etc. **/
+	FORCEINLINE bool IsSuppressed(ELogVerbosity::Type VerbosityLevel) const
+	{
+		if ((VerbosityLevel & ELogVerbosity::VerbosityMask) <= Verbosity)
+		{
+			return false;
+		}
+		return true;
+	}
+	/** Called just after a logging statement being allow to print. Checks a few things and maybe breaks into the debugger. **/
+	void PostTrigger(ELogVerbosity::Type VerbosityLevel);
+	FORCEINLINE FName GetCategoryName() const
+	{
+		return CategoryFName;
+	}
+
+	/** Sets up the working verbosity and clamps to the compile time verbosity. **/
+	void SetVerbosity(ELogVerbosity::Type Verbosity);
+
+private:
+	friend class FLogSuppressionImplementation;
+	friend class FLogScopedVerbosityOverride;
+	/** Internal call to get the working verbosity. **/
+	ELogVerbosity::Type GetVerbosity() const { return (ELogVerbosity::Type)Verbosity; }
+	/** Internal call to set up the working verbosity from the boot time default. **/
+	void ResetFromDefault();
+protected:
+
+	/** Holds the current suppression state **/
+	uint8 Verbosity;
+	/** Holds the break flag **/
+	bool DebugBreakOnLog;
+	/** Holds default suppression **/
+	uint8 DefaultVerbosity;
+	/** Holds compile time suppression **/
+	uint8 CompileTimeVerbosity;
+	/** FName for this category **/
+	FName CategoryFName;
+};
+
+/** Template for log categories that transfers the compile-time constant default and compile time verbosity to the FLogCategoryBase constructor. **/
+template<uint8 InDefaultVerbosity, uint8 InCompileTimeVerbosity>
+struct FLogCategory : public FLogCategoryBase
+{
+	static_assert((InDefaultVerbosity & ELogVerbosity::VerbosityMask) < ELogVerbosity::NumVerbosity, "Bogus default verbosity.");
+	static_assert(InCompileTimeVerbosity < ELogVerbosity::NumVerbosity, "Bogus compile time verbosity.");
+	enum
+	{
+		CompileTimeVerbosity = InCompileTimeVerbosity
+	};
+
+	FORCEINLINE FLogCategory(const TCHAR *CategoryName)
+		: FLogCategoryBase(CategoryName, ELogVerbosity::Type(InDefaultVerbosity), ELogVerbosity::Type(CompileTimeVerbosity))
+	{
+	}
+};
diff --git a/Engine/Source/Runtime/Core/Public/Logging/LogMacros.h b/Engine/Source/Runtime/Core/Public/Logging/LogMacros.h
new file mode 100644
--- /dev/null
+++ b/Engine/Source/Runtime/Core/Public/Logging/LogMacros.h
@@ -0,0 +1,336 @@
+// Copyright 1998-2016 Epic Games, Inc. All Rights Reserved.
+
+#pragma once
+
+#include "HAL/Platform.h"
+#include "HAL/PlatformMisc.h"
+#include "Misc/Build.h"
+#include "Misc/VarArgs.h"
+#include "Templates/EnableIf.h"
+#include "LogVerbosity.h"
+#include "LogCategory.h"
+#include "AssertionMacros.h"
+
+/*----------------------------------------------------------------------------
+	Logging
+----------------------------------------------------------------------------*/
+
+// Define NO_LOGGING to strip out all writing to log files, OutputDebugString(), etc.
+// This is needed for consoles that require no logging (Xbox, Xenon)
+
+/**
+ * FMsg 
+ * This struct contains functions for messaging with tools or debug logs.
+ **/
+struct CORE_API FMsg
+{
+	/** Sends a message to a remote tool. */
+	static void SendNotificationString( const TCHAR* Message );
+
+	/** Sends a formatted message to a remote tool. */
+	static void VARARGS SendNotificationStringf( const TCHAR *Format, ... );
+
+	/** Log function */
+	VARARG_DECL( static void, static void, {}, Logf, VARARG_NONE, const TCHAR*, VARARG_EXTRA(const ANSICHAR* File) VARARG_EXTRA(int32 Line) VARARG_EXTRA(const class FName& Category) VARARG_EXTRA(ELogVerbosity::Type Verbosity), VARARG_EXTRA(File) VARARG_EXTRA(Line) VARARG_EXTRA(Category) VARARG_EXTRA(Verbosity) );
+
+	/** Internal version of log function. Should be used only in logging macros, as it relies on caller to call assert on fatal error */
+	VARARG_DECL(static void, static void, {}, Logf_Internal, VARARG_NONE, const TCHAR*, VARARG_EXTRA(const ANSICHAR* File) VARARG_EXTRA(int32 Line) VARARG_EXTRA(const class FName& Category) VARARG_EXTRA(ELogVerbosity::Type Verbosity), VARARG_EXTRA(File) VARARG_EXTRA(Line) VARARG_EXTRA(Category) VARARG_EXTRA(Verbosity));
+};
+
+/*----------------------------------------------------------------------------
+	Logging suppression
+----------------------------------------------------------------------------*/
+
+#ifndef COMPILED_IN_MINIMUM_VERBOSITY
+	#define COMPILED_IN_MINIMUM_VERBOSITY VeryVerbose
+#else
+	#if !IS_MONOLITHIC
+		#error COMPILED_IN_MINIMUM_VERBOSITY can only be defined in monolithic builds.
+	#endif
+#endif
+
+#define UE_LOG_EXPAND_IS_FATAL(Verbosity, ActiveBlock, InactiveBlock) PREPROCESSOR_JOIN(UE_LOG_EXPAND_IS_FATAL_, Verbosity)(ActiveBlock, InactiveBlock)
+
+#define UE_LOG_EXPAND_IS_FATAL_Fatal(      ActiveBlock, InactiveBlock) ActiveBlock
+#define UE_LOG_EXPAND_IS_FATAL_Error(      ActiveBlock, InactiveBlock) InactiveBlock
+#define UE_LOG_EXPAND_IS_FATAL_Warning(    ActiveBlock, InactiveBlock) InactiveBlock
+#define UE_LOG_EXPAND_IS_FATAL_Display(    ActiveBlock, InactiveBlock) InactiveBlock
+#define UE_LOG_EXPAND_IS_FATAL_Log(        ActiveBlock, InactiveBlock) InactiveBlock
+#define UE_LOG_EXPAND_IS_FATAL_Verbose(    ActiveBlock, InactiveBlock) InactiveBlock
+#define UE_LOG_EXPAND_IS_FATAL_VeryVerbose(ActiveBlock, InactiveBlock) InactiveBlock
+#define UE_LOG_EXPAND_IS_FATAL_All(        ActiveBlock, InactiveBlock) InactiveBlock
+#define UE_LOG_EXPAND_IS_FATAL_SetColor(   ActiveBlock, InactiveBlock) InactiveBlock
+
+#if NO_LOGGING
+
+	struct FNoLoggingCategory {};
+	
+	// This will only log Fatal errors
+	#define UE_LOG(CategoryName, Verbosity, Format, ...) \
+	{ \
+		static_assert(IS_TCHAR_ARRAY(Format), "Formatting string must be a TCHAR array."); \
+		if (ELogVerbosity::Verbosity == ELogVerbosity::Fatal) \
+		{ \
+			LowLevelFatalErrorHandler(__FILE__, __LINE__, Format, ##__VA_ARGS__); \
+			_DebugBreakAndPromptForRemote(); \
+			FDebug::AssertFailed("", __FILE__, __LINE__, Format, ##__VA_ARGS__); \
+			UE_LOG_EXPAND_IS_FATAL(Verbosity, CA_ASSUME(false);, PREPROCESSOR_NOTHING) \
+		} \
+	}
+	// Conditional logging (fatal errors only).
+	#define UE_CLOG(Condition, CategoryName, Verbosity, Format, ...) \
+	{ \
+		static_assert(IS_TCHAR_ARRAY(Format), "Formatting string must be a TCHAR array."); \
+		if (ELogVerbosity::Verbosity == ELogVerbosity::Fatal) \
+		{ \
+			if (Condition) \
+			{ \
+				LowLevelFatalErrorHandler(__FILE__, __LINE__, Format, ##__VA_ARGS__); \
+				_DebugBreakAndPromptForRemote(); \
+				FDebug::AssertFailed("", __FILE__, __LINE__, Format, ##__VA_ARGS__); \
+				UE_LOG_EXPAND_IS_FATAL(Verbosity, CA_ASSUME(false);, PREPROCESSOR_NOTHING) \
+			} \
+		} \
+	}
+
+	#define UE_LOG_ACTIVE(...)				(0)
+	#define UE_LOG_ANY_ACTIVE(...)			(0)
+	#define UE_SUPPRESS(...) {}
+	#define UE_SET_LOG_VERBOSITY(...)
+	#define DECLARE_LOG_CATEGORY_EXTERN(CategoryName, DefaultVerbosity, CompileTimeVerbosity) extern FNoLoggingCategory CategoryName
+	#define DEFINE_LOG_CATEGORY(...)
+	#define DEFINE_LOG_CATEGORY_STATIC(...)
+	#define DECLARE_LOG_CATEGORY_CLASS(...)
+	#define DEFINE_LOG_CATEGORY_CLASS(...)
+	#define UE_SECURITY_LOG(...)
+
+#else
+
+	namespace UE4Asserts_Private
+	{
+		template <int32 VerbosityToCheck, typename CategoryType>
+		FORCEINLINE
+			typename TEnableIf<
+				((VerbosityToCheck & ELogVerbosity::VerbosityMask) <= CategoryType::CompileTimeVerbosity &&
+				(VerbosityToCheck & ELogVerbosity::VerbosityMask) <= ELogVerbosity::COMPILED_IN_MINIMUM_VERBOSITY),
+				bool>::Type
+			IsLogActive(const CategoryType& Category)
+		{
+			return !Category.IsSuppressed((ELogVerbosity::Type)VerbosityToCheck);
+		}
+
+		template <int32 VerbosityToCheck, typename CategoryType>
+		FORCEINLINE
+			typename TEnableIf<
+				!((VerbosityToCheck & ELogVerbosity::VerbosityMask) <= CategoryType::CompileTimeVerbosity &&
+				(VerbosityToCheck & ELogVerbosity::VerbosityMask) <= ELogVerbosity::COMPILED_IN_MINIMUM_VERBOSITY),
+				bool>::Type
+			IsLogActive(const CategoryType& Category)
+		{
+			return false;
+		}
+	}
+
+	/** 
+	 * A predicate that returns true if the given logging category is active (logging) at a given verbosity level 
+	 * @param CategoryName name of the logging category
+	 * @param Verbosity, verbosity level to test against
+	**/
+	#define UE_LOG_ACTIVE(CategoryName, Verbosity) (::UE4Asserts_Private::IsLogActive<(int32)ELogVerbosity::Verbosity>(CategoryName))
+
+	#define UE_SET_LOG_VERBOSITY(CategoryName, Verbosity) \
+		CategoryName.SetVerbosity(ELogVerbosity::Verbosity);
+
+	/** 
+	 * A  macro that outputs a formatted message to log if a given logging category is active at a given verbosity level
+	 * @param CategoryName name of the logging category
+	 * @param Verbosity, verbosity level to test against
+	 * @param Format, format text
+	 ***/
+	#define UE_LOG(CategoryName, Verbosity, Format, ...) \
+	{ \
+		static_assert(IS_TCHAR_ARRAY(Format), "Formatting string must be a TCHAR array."); \
+		static_assert((ELogVerbosity::Verbosity & ELogVerbosity::VerbosityMask) < ELogVerbosity::NumVerbosity && ELogVerbosity::Verbosity > 0, "Verbosity must be constant and in range."); \
+		CA_CONSTANT_IF((ELogVerbosity::Verbosity & ELogVerbosity::VerbosityMask) <= ELogVerbosity::COMPILED_IN_MINIMUM_VERBOSITY && (ELogVerbosity::Warning & ELogVerbosity::VerbosityMask) <= FLogCategory##CategoryName::CompileTimeVerbosity) \
+		{ \
+			UE_LOG_EXPAND_IS_FATAL(Verbosity, PREPROCESSOR_NOTHING, if (!CategoryName.IsSuppressed(ELogVerbosity::Verbosity))) \
+			{ \
+				FMsg::Logf_Internal(__FILE__, __LINE__, CategoryName.GetCategoryName(), ELogVerbosity::Verbosity, Format, ##__VA_ARGS__); \
+				UE_LOG_EXPAND_IS_FATAL(Verbosity, \
+					{ \
+						_DebugBreakAndPromptForRemote(); \
+						FDebug::AssertFailed("", __FILE__, __LINE__, Format, ##__VA_ARGS__); \
+						CA_ASSUME(false); \
+					}, \
+					PREPROCESSOR_NOTHING \
+				) \
+			} \
+		} \
+	}
+
+	/**
+	* A  macro that outputs a formatted message to the log specifically used for security events
+	* @param NetConnection, a valid UNetConnection
+	* @param SecurityEventType, a security event type (ESecurityEvent::Type)
+	* @param Format, format text
+	***/
+	#define UE_SECURITY_LOG(NetConnection, SecurityEventType, Format, ...) \
+	{ \
+		static_assert(IS_TCHAR_ARRAY(Format), "Formatting string must be a TCHAR array."); \
+		check(NetConnection != nullptr); \
+		CA_CONSTANT_IF((ELogVerbosity::Warning & ELogVerbosity::VerbosityMask) <= ELogVerbosity::COMPILED_IN_MINIMUM_VERBOSITY && (ELogVerbosity::Warning & ELogVerbosity::VerbosityMask) <= FLogCategoryLogSecurity::CompileTimeVerbosity) \
+		{ \
+			if (!LogSecurity.IsSuppressed(ELogVerbosity::Warning)) \
+			{ \
+				FString Test = FString::Printf(TEXT("%s: %s: %s"), *(NetConnection->RemoteAddressToString()), ToString(SecurityEventType), Format); \
+				FMsg::Logf_Internal(__FILE__, __LINE__, LogSecurity.GetCategoryName(), ELogVerbosity::Warning, *Test, ##__VA_ARGS__); \
+			} \
+		} \
+	}
+
+	// Conditional logging. Will only log if Condition is met.
+	#define UE_CLOG(Condition, CategoryName, Verbosity, Format, ...) \
+	{ \
+		static_assert(IS_TCHAR_ARRAY(Format), "Formatting string must be a TCHAR array."); \
+		static_assert((ELogVerbosity::Verbosity & ELogVerbosity::VerbosityMask) < ELogVerbosity::NumVerbosity && ELogVerbosity::Verbosity > 0, "Verbosity must be constant and in range."); \
+		CA_CONSTANT_IF((ELogVerbosity::Verbosity & ELogVerbosity::VerbosityMask) <= ELogVerbosity::COMPILED_IN_MINIMUM_VERBOSITY && (ELogVerbosity::Warning & ELogVerbosity::VerbosityMask) <= FLogCategory##CategoryName::CompileTimeVerbosity) \
+		{ \
+			UE_LOG_EXPAND_IS_FATAL(Verbosity, PREPROCESSOR_NOTHING, if (!CategoryName.IsSuppressed(ELogVerbosity::Verbosity))) \
+			{ \
+				if (Condition) \
+				{ \
+					FMsg::Logf_Internal(__FILE__, __LINE__, CategoryName.GetCategoryName(), ELogVerbosity::Verbosity, Format, ##__VA_ARGS__); \
+					UE_LOG_EXPAND_IS_FATAL(Verbosity, \
+						{ \
+							_DebugBreakAndPromptForRemote(); \
+							FDebug::AssertFailed("", __FILE__, __LINE__, Format, ##__VA_ARGS__); \
+							CA_ASSUME(false); \
+						}, \
+						PREPROCESSOR_NOTHING \
+					) \
+					CA_ASSUME(true); \
+				} \
+			} \
+		} \
+	}
+
+	/** 
+	 * A macro that executes some code within a scope if a given logging category is active at a given verbosity level
+	 * Also, withing the scope of the execution, the default category and verbosity is set up for the low level logging 
+	 * functions.
+	 * @param CategoryName name of the logging category
+	 * @param Verbosity, verbosity level to test against
+	 * @param ExecuteIfUnsuppressed, code to execute if the verbosity level for this category is being displayed
+	 ***/
+	#define UE_SUPPRESS(CategoryName, Verbosity, ExecuteIfUnsuppressed) \
+	{ \
+		static_assert((ELogVerbosity::Verbosity & ELogVerbosity::VerbosityMask) < ELogVerbosity::NumVerbosity && ELogVerbosity::Verbosity > 0, "Verbosity must be constant and in range."); \
+		CA_CONSTANT_IF((ELogVerbosity::Verbosity & ELogVerbosity::VerbosityMask) <= ELogVerbosity::COMPILED_IN_MINIMUM_VERBOSITY && (ELogVerbosity::Warning & ELogVerbosity::VerbosityMask) <= FLogCategory##CategoryName::CompileTimeVerbosity) \
+		{ \
+			if (!CategoryName.IsSuppressed(ELogVerbosity::Verbosity)) \
+			{ \
+				FScopedCategoryAndVerbosityOverride TEMP__##CategoryName(CategoryName.GetCategoryName(), ELogVerbosity::Type(ELogVerbosity::Verbosity & ELogVerbosity::VerbosityMask)); \
+				ExecuteIfUnsuppressed; \
+				CategoryName.PostTrigger(ELogVerbosity::Verbosity); \
+			} \
+		} \
+	}
+
+	/** 
+	 * A macro to declare a logging category as a C++ "extern" 
+	 * @param CategoryName, category to declare
+	 * @param DefaultVerbosity, default run time verbosity
+	 * @param CompileTimeVerbosity, maximum verbosity to compile into the code
+	 **/
+	#define DECLARE_LOG_CATEGORY_EXTERN(CategoryName, DefaultVerbosity, CompileTimeVerbosity) \
+		extern struct FLogCategory##CategoryName : public FLogCategory<ELogVerbosity::DefaultVerbosity, ELogVerbosity::CompileTimeVerbosity> \
+		{ \
+			FORCEINLINE FLogCategory##CategoryName() : FLogCategory(TEXT(#CategoryName)) {} \
+		} CategoryName;
+
+	/** 
+	 * A macro to define a logging category, usually paired with DECLARE_LOG_CATEGORY_EXTERN from the header.
+	 * @param CategoryName, category to define
+	**/
+	#define DEFINE_LOG_CATEGORY(CategoryName) FLogCategory##CategoryName CategoryName;
+
+	/** 
+	 * A macro to define a logging category as a C++ "static"
+	 * @param CategoryName, category to declare
+	 * @param DefaultVerbosity, default run time verbosity
+	 * @param CompileTimeVerbosity, maximum verbosity to compile into the code
+	**/
+	#define DEFINE_LOG_CATEGORY_STATIC(CategoryName, DefaultVerbosity, CompileTimeVerbosity) \
+		static struct FLogCategory##CategoryName : public FLogCategory<ELogVerbosity::DefaultVerbosity, ELogVerbosity::CompileTimeVerbosity> \
+		{ \
+			FORCEINLINE FLogCategory##CategoryName() : FLogCategory(TEXT(#CategoryName)) {} \
+		} CategoryName;
+
+	/** 
+	 * A macro to declare a logging category as a C++ "class static" 
+	 * @param CategoryName, category to declare
+	 * @param DefaultVerbosity, default run time verbosity
+	 * @param CompileTimeVerbosity, maximum verbosity to compile into the code
+	**/
+	#define DECLARE_LOG_CATEGORY_CLASS(CategoryName, DefaultVerbosity, CompileTimeVerbosity) \
+		DEFINE_LOG_CATEGORY_STATIC(CategoryName, DefaultVerbosity, CompileTimeVerbosity)
+
+	/** 
+	 * A macro to define a logging category, usually paired with DECLARE_LOG_CATEGORY_CLASS from the header.
+	 * @param CategoryName, category to definee
+	**/
+	#define DEFINE_LOG_CATEGORY_CLASS(Class, CategoryName) Class::FLogCategory##CategoryName Class::CategoryName;
+
+#endif // NO_LOGGING
+#if PLATFORM_HTML5
+#include "HTML5/HTML5AssertionMacros.h"
+#endif 
+
+#if UE_BUILD_SHIPPING
+#define NOTIFY_CLIENT_OF_SECURITY_EVENT_IF_NOT_SHIPPING(NetConnection, SecurityPrint) ;
+#else
+#define NOTIFY_CLIENT_OF_SECURITY_EVENT_IF_NOT_SHIPPING(NetConnection, SecurityPrint) \
+	FNetControlMessage<NMT_SecurityViolation>::Send(NetConnection, SecurityPrint); \
+	NetConnection->FlushNet(true)
+#endif
+
+/**
+	* A  macro that closes the connection and logs the security event on the server and the client
+	* @param NetConnection, a valid UNetConnection
+	* @param SecurityEventType, a security event type (ESecurityEvent::Type)
+	* @param Format, format text
+***/
+#define CLOSE_CONNECTION_DUE_TO_SECURITY_VIOLATION(NetConnection, SecurityEventType, Format, ...) \
+{ \
+	static_assert(IS_TCHAR_ARRAY(Format), "Formatting string must be a TCHAR array."); \
+	check(NetConnection != nullptr); \
+	FString SecurityPrint = FString::Printf(Format, ##__VA_ARGS__); \
+	UE_SECURITY_LOG(NetConnection, SecurityEventType, Format, ##__VA_ARGS__); \
+	UE_SECURITY_LOG(NetConnection, ESecurityEvent::Closed, TEXT("Connection closed")); \
+	NetConnection->Close(); \
+	PerfCountersIncrement(TEXT("ClosedConnectionsDueToSecurityViolations")); \
+}
+
+extern CORE_API int32 GEnsureOnNANDiagnostic;
+
+// Macro to either log an error or ensure on a NaN error.
+#if DO_CHECK && !USING_CODE_ANALYSIS
+namespace UE4Asserts_Private
+{
+	CORE_API void VARARGS InternalLogNANDiagnosticMessage(const TCHAR* FormattedMsg, ...); // UE_LOG(LogCore, Error, _FormatString_, ##__VA_ARGS__);
+}
+#define logOrEnsureNanError(_FormatString_, ...) \
+	if (!GEnsureOnNANDiagnostic)\
+	{\
+		if (UE4Asserts_Private::TrueOnFirstCallOnly([]{}))\
+		{\
+			UE4Asserts_Private::InternalLogNANDiagnosticMessage(_FormatString_, ##__VA_ARGS__); \
+		}\
+	}\
+	else\
+	{\
+		ensureMsgf(!GEnsureOnNANDiagnostic, _FormatString_, ##__VA_ARGS__); \
+	}
+#else
+#define logOrEnsureNanError(_FormatString_, ...)
+#endif // DO_CHECK
\ No newline at end of file
diff --git a/Engine/Source/Runtime/Core/Public/Logging/LogScopedCategoryAndVerbosityOverride.h b/Engine/Source/Runtime/Core/Public/Logging/LogScopedCategoryAndVerbosityOverride.h
new file mode 100644
--- /dev/null
+++ b/Engine/Source/Runtime/Core/Public/Logging/LogScopedCategoryAndVerbosityOverride.h
@@ -0,0 +1,47 @@
+// Copyright 1998-2016 Epic Games, Inc. All Rights Reserved.
+
+#pragma once
+
+/** 
+ * Helper class that uses thread local storage to set up the default category and verbosity for the low level logging functions.
+ * This is what allow a UE_LOG(LogHAL, Log, TEXT("...")); within a UE_LOG statement to know what the category and verbosity is.
+ * When one of these goes out of scope, it restores the previous values.
+**/
+class CORE_API FLogScopedCategoryAndVerbosityOverride
+{
+public:
+	/** STructure to aggregate a category and verbosity pair **/
+	struct FOverride
+	{
+		ELogVerbosity::Type Verbosity;
+		FName Category;
+		FOverride()
+			: Verbosity(ELogVerbosity::Log)
+		{
+		}
+		FOverride(FName InCategory, ELogVerbosity::Type InVerbosity)
+			: Verbosity(InVerbosity)
+			, Category(InCategory)
+		{
+		}
+		void operator=(const FOverride& Other)
+		{
+			Verbosity = Other.Verbosity;
+			Category = Other.Category;
+		}
+	};
+private:
+	/** Backup of the category, verbosity pairs that was present when we were constructed **/
+	FOverride Backup;
+public:
+	/** Back up the existing category and varbosity pair, then sets them.*/
+	FLogScopedCategoryAndVerbosityOverride(FName Category, ELogVerbosity::Type Verbosity);
+
+	/** Restore the category and verbosity overrides to the previous value.*/
+	~FLogScopedCategoryAndVerbosityOverride();
+
+	/** Manages a TLS slot with the current overrides for category and verbosity. **/
+	static FOverride* GetTLSCurrent();
+};
+
+typedef FLogScopedCategoryAndVerbosityOverride FScopedCategoryAndVerbosityOverride;
\ No newline at end of file
diff --git a/Engine/Source/Runtime/Core/Public/Mac/MacCriticalSection.h b/Engine/Source/Runtime/Core/Public/Mac/MacCriticalSection.h
new file mode 100644
--- /dev/null
+++ b/Engine/Source/Runtime/Core/Public/Mac/MacCriticalSection.h
@@ -0,0 +1,38 @@
+// Copyright 1998-2016 Epic Games, Inc. All Rights Reserved.
+
+#pragma once
+
+#include "PThreadCriticalSection.h"
+
+/**
+ * Mac implementation of the FSystemWideCriticalSection. Uses exclusive file locking.
+ **/
+class FMacSystemWideCriticalSection
+{
+public:
+	/** Construct a named, system-wide critical section and attempt to get access/ownership of it */
+	explicit FMacSystemWideCriticalSection(const FString& InName, FTimespan InTimeout = FTimespan::Zero());
+
+	/** Destructor releases system-wide critical section if it is currently owned */
+	~FMacSystemWideCriticalSection();
+
+	/**
+	 * Does the calling thread have ownership of the system-wide critical section?
+	 *
+	 * @return True if the system-wide lock is obtained. WARNING: Returns true for abandoned locks so shared resources can be in undetermined states.
+	 */
+	bool IsValid() const;
+
+	/** Releases system-wide critical section if it is currently owned */
+	void Release();
+
+private:
+	FMacSystemWideCriticalSection(const FMacSystemWideCriticalSection&);
+	FMacSystemWideCriticalSection& operator=(const FMacSystemWideCriticalSection&);
+
+private:
+	int32 FileHandle;
+};
+
+typedef FPThreadsCriticalSection FCriticalSection;
+typedef FMacSystemWideCriticalSection FSystemWideCriticalSection;
diff --git a/Engine/Source/Runtime/Core/Public/Mac/MacPlatformProcess.h b/Engine/Source/Runtime/Core/Public/Mac/MacPlatformProcess.h
--- a/Engine/Source/Runtime/Core/Public/Mac/MacPlatformProcess.h
+++ b/Engine/Source/Runtime/Core/Public/Mac/MacPlatformProcess.h
@@ -140,35 +140,4 @@ struct CORE_API FMacPlatformProcess : public FGenericPlatformProcess
 	static const TCHAR* UserLogsDir();
 };
 
-/**
- * Mac implementation of the FSystemWideCriticalSection. Uses exclusive file locking.
- **/
-class FMacSystemWideCriticalSection
-{
-public:
-	/** Construct a named, system-wide critical section and attempt to get access/ownership of it */
-	explicit FMacSystemWideCriticalSection(const FString& InName, FTimespan InTimeout = FTimespan::Zero());
-
-	/** Destructor releases system-wide critical section if it is currently owned */
-	~FMacSystemWideCriticalSection();
-
-	/**
-	 * Does the calling thread have ownership of the system-wide critical section?
-	 *
-	 * @return True if the system-wide lock is obtained. WARNING: Returns true for abandoned locks so shared resources can be in undetermined states.
-	 */
-	bool IsValid() const;
-
-	/** Releases system-wide critical section if it is currently owned */
-	void Release();
-
-private:
-	FMacSystemWideCriticalSection(const FMacSystemWideCriticalSection&);
-	FMacSystemWideCriticalSection& operator=(const FMacSystemWideCriticalSection&);
-
-private:
-	int32 FileHandle;
-};
-
-typedef FMacPlatformProcess FPlatformProcess;
-typedef FMacSystemWideCriticalSection FSystemWideCriticalSection;
+typedef FMacPlatformProcess FPlatformProcess;
\ No newline at end of file
diff --git a/Engine/Source/Runtime/Core/Public/Math/Box.h b/Engine/Source/Runtime/Core/Public/Math/Box.h
--- a/Engine/Source/Runtime/Core/Public/Math/Box.h
+++ b/Engine/Source/Runtime/Core/Public/Math/Box.h
@@ -2,6 +2,8 @@
 
 #pragma once
 
+#include "Vector.h"
+#include "UnrealMathUtility.h"
 
 /**
  * Implements an axis-aligned box.
@@ -471,7 +473,6 @@ public:
  */
 template<> struct TIsPODType<FBox> { enum { Value = true }; };
 
-
 /* FBox inline functions
  *****************************************************************************/
 
@@ -598,3 +599,202 @@ FORCEINLINE FString FBox::ToString() const
 {
 	return FString::Printf(TEXT("IsValid=%s, Min=(%s), Max=(%s)"), IsValid ? TEXT("true") : TEXT("false"), *Min.ToString(), *Max.ToString());
 }
+
+/* FMath inline functions
+ *****************************************************************************/
+
+inline bool FMath::PointBoxIntersection
+	(
+	const FVector&	Point,
+	const FBox&		Box
+	)
+{
+	if(Point.X >= Box.Min.X && Point.X <= Box.Max.X &&
+		Point.Y >= Box.Min.Y && Point.Y <= Box.Max.Y &&
+		Point.Z >= Box.Min.Z && Point.Z <= Box.Max.Z)
+		return 1;
+	else
+		return 0;
+}
+
+inline bool FMath::LineBoxIntersection
+	( 
+	const FBox& Box, 
+	const FVector& Start, 
+	const FVector& End, 
+	const FVector& Direction
+	)
+{
+	return LineBoxIntersection(Box, Start, End, Direction, Direction.Reciprocal());
+}
+
+inline bool FMath::LineBoxIntersection
+	(
+	const FBox&		Box,
+	const FVector&	Start,
+	const FVector&	End,
+	const FVector&	Direction,
+	const FVector&	OneOverDirection
+	)
+{
+	FVector	Time;
+	bool	bStartIsOutside = false;
+
+	if(Start.X < Box.Min.X)
+	{
+		bStartIsOutside = true;
+		if(End.X >= Box.Min.X)
+		{
+			Time.X = (Box.Min.X - Start.X) * OneOverDirection.X;
+		}
+		else
+		{
+			return false;
+		}
+	}
+	else if(Start.X > Box.Max.X)
+	{
+		bStartIsOutside = true;
+		if(End.X <= Box.Max.X)
+		{
+			Time.X = (Box.Max.X - Start.X) * OneOverDirection.X;
+		}
+		else
+		{
+			return false;
+		}
+	}
+	else
+	{
+		Time.X = 0.0f;
+	}
+
+	if(Start.Y < Box.Min.Y)
+	{
+		bStartIsOutside = true;
+		if(End.Y >= Box.Min.Y)
+		{
+			Time.Y = (Box.Min.Y - Start.Y) * OneOverDirection.Y;
+		}
+		else
+		{
+			return false;
+		}
+	}
+	else if(Start.Y > Box.Max.Y)
+	{
+		bStartIsOutside = true;
+		if(End.Y <= Box.Max.Y)
+		{
+			Time.Y = (Box.Max.Y - Start.Y) * OneOverDirection.Y;
+		}
+		else
+		{
+			return false;
+		}
+	}
+	else
+	{
+		Time.Y = 0.0f;
+	}
+
+	if(Start.Z < Box.Min.Z)
+	{
+		bStartIsOutside = true;
+		if(End.Z >= Box.Min.Z)
+		{
+			Time.Z = (Box.Min.Z - Start.Z) * OneOverDirection.Z;
+		}
+		else
+		{
+			return false;
+		}
+	}
+	else if(Start.Z > Box.Max.Z)
+	{
+		bStartIsOutside = true;
+		if(End.Z <= Box.Max.Z)
+		{
+			Time.Z = (Box.Max.Z - Start.Z) * OneOverDirection.Z;
+		}
+		else
+		{
+			return false;
+		}
+	}
+	else
+	{
+		Time.Z = 0.0f;
+	}
+
+	if(bStartIsOutside)
+	{
+		const float	MaxTime = Max3(Time.X,Time.Y,Time.Z);
+
+		if(MaxTime >= 0.0f && MaxTime <= 1.0f)
+		{
+			const FVector Hit = Start + Direction * MaxTime;
+			const float BOX_SIDE_THRESHOLD = 0.1f;
+			if(	Hit.X > Box.Min.X - BOX_SIDE_THRESHOLD && Hit.X < Box.Max.X + BOX_SIDE_THRESHOLD &&
+				Hit.Y > Box.Min.Y - BOX_SIDE_THRESHOLD && Hit.Y < Box.Max.Y + BOX_SIDE_THRESHOLD &&
+				Hit.Z > Box.Min.Z - BOX_SIDE_THRESHOLD && Hit.Z < Box.Max.Z + BOX_SIDE_THRESHOLD)
+			{
+				return true;
+			}
+		}
+
+		return false;
+	}
+	else
+	{
+		return true;
+	}
+}
+
+/**
+ * Performs a sphere vs box intersection test using Arvo's algorithm:
+ *
+ *	for each i in (x, y, z)
+ *		if (SphereCenter(i) < BoxMin(i)) d2 += (SphereCenter(i) - BoxMin(i)) ^ 2
+ *		else if (SphereCenter(i) > BoxMax(i)) d2 += (SphereCenter(i) - BoxMax(i)) ^ 2
+ *
+ * @param Sphere the center of the sphere being tested against the AABB
+ * @param RadiusSquared the size of the sphere being tested
+ * @param AABB the box being tested against
+ *
+ * @return Whether the sphere/box intersect or not.
+ */
+FORCEINLINE bool FMath::SphereAABBIntersection(const FVector& SphereCenter,const float RadiusSquared,const FBox& AABB)
+{
+	// Accumulates the distance as we iterate axis
+	float DistSquared = 0.f;
+	// Check each axis for min/max and add the distance accordingly
+	// NOTE: Loop manually unrolled for > 2x speed up
+	if (SphereCenter.X < AABB.Min.X)
+	{
+		DistSquared += FMath::Square(SphereCenter.X - AABB.Min.X);
+	}
+	else if (SphereCenter.X > AABB.Max.X)
+	{
+		DistSquared += FMath::Square(SphereCenter.X - AABB.Max.X);
+	}
+	if (SphereCenter.Y < AABB.Min.Y)
+	{
+		DistSquared += FMath::Square(SphereCenter.Y - AABB.Min.Y);
+	}
+	else if (SphereCenter.Y > AABB.Max.Y)
+	{
+		DistSquared += FMath::Square(SphereCenter.Y - AABB.Max.Y);
+	}
+	if (SphereCenter.Z < AABB.Min.Z)
+	{
+		DistSquared += FMath::Square(SphereCenter.Z - AABB.Min.Z);
+	}
+	else if (SphereCenter.Z > AABB.Max.Z)
+	{
+		DistSquared += FMath::Square(SphereCenter.Z - AABB.Max.Z);
+	}
+	// If the distance is less than or equal to the radius, they intersect
+	return DistSquared <= RadiusSquared;
+}
+
diff --git a/Engine/Source/Runtime/Core/Public/Math/Matrix.h b/Engine/Source/Runtime/Core/Public/Math/Matrix.h
--- a/Engine/Source/Runtime/Core/Public/Math/Matrix.h
+++ b/Engine/Source/Runtime/Core/Public/Math/Matrix.h
@@ -2,6 +2,8 @@
 
 #pragma once
 
+#include "Axis.h"
+#include "Plane.h"
 
 /**
  * 4x4 matrix of floating point values.
@@ -354,6 +356,13 @@ public:
 		Dest[10] = Src[10]; // [2][2]
 		Dest[11] = Src[14]; // [3][2]
 	}	
+
+private:
+
+	/** 
+	 * Output an error message and trigger an ensure 
+	 */
+	static CORE_API void ErrorEnsure(const TCHAR* Message);
 };
 
 
@@ -404,22 +413,6 @@ struct FBasisVectorMatrix: FMatrix
 };
 
 
-FORCEINLINE FBasisVectorMatrix::FBasisVectorMatrix(const FVector& XAxis,const FVector& YAxis,const FVector& ZAxis,const FVector& Origin)
-{
-	for(uint32 RowIndex = 0;RowIndex < 3;RowIndex++)
-	{
-		M[RowIndex][0] = (&XAxis.X)[RowIndex];
-		M[RowIndex][1] = (&YAxis.X)[RowIndex];
-		M[RowIndex][2] = (&ZAxis.X)[RowIndex];
-		M[RowIndex][3] = 0.0f;
-	}
-	M[3][0] = Origin | XAxis;
-	M[3][1] = Origin | YAxis;
-	M[3][2] = Origin | ZAxis;
-	M[3][3] = 1.0f;
-}
-
-
 struct FLookAtMatrix : FMatrix
 {
 	/** 
@@ -430,26 +423,6 @@ struct FLookAtMatrix : FMatrix
 };
 
 
-FORCEINLINE FLookAtMatrix::FLookAtMatrix(const FVector& EyePosition, const FVector& LookAtPosition, const FVector& UpVector)
-{
-	const FVector ZAxis = (LookAtPosition - EyePosition).GetSafeNormal();
-	const FVector XAxis = (UpVector ^ ZAxis).GetSafeNormal();
-	const FVector YAxis = ZAxis ^ XAxis;
-
-	for (uint32 RowIndex = 0; RowIndex < 3; RowIndex++)
-	{
-		M[RowIndex][0] = (&XAxis.X)[RowIndex];
-		M[RowIndex][1] = (&YAxis.X)[RowIndex];
-		M[RowIndex][2] = (&ZAxis.X)[RowIndex];
-		M[RowIndex][3] = 0.0f;
-	}
-	M[3][0] = -EyePosition | XAxis;
-	M[3][1] = -EyePosition | YAxis;
-	M[3][2] = -EyePosition | ZAxis;
-	M[3][3] = 1.0f;
-}
-
-
 template <> struct TIsPODType<FMatrix> { enum { Value = true }; };
 
 
diff --git a/Engine/Source/Runtime/Core/Public/Math/Matrix.inl b/Engine/Source/Runtime/Core/Public/Math/Matrix.inl
--- a/Engine/Source/Runtime/Core/Public/Math/Matrix.inl
+++ b/Engine/Source/Runtime/Core/Public/Math/Matrix.inl
@@ -6,10 +6,6 @@
 
 #pragma once
 
-#include "OutputDevices.h"
-#include "AssertionMacros.h"
-
-
 /**
  * FMatrix inline functions.
  */
@@ -262,8 +258,7 @@ inline FMatrix FMatrix::InverseFast() const
 		GetScaledAxis( EAxis::Y ).IsNearlyZero(SMALL_NUMBER) && 
 		GetScaledAxis( EAxis::Z ).IsNearlyZero(SMALL_NUMBER) ) 
 	{
-		UE_LOG(LogUnrealMath, Error, TEXT("FMatrix::InverseFast(), trying to invert a NIL matrix, this results in NaNs! Use Inverse() instead."));
-		ensureMsgf(false, TEXT("FMatrix::InverseFast(), trying to invert a NIL matrix, this results in NaNs! Use Inverse() instead."));
+		ErrorEnsure(TEXT("FMatrix::InverseFast(), trying to invert a NIL matrix, this results in NaNs! Use Inverse() instead."));
 	}
 #endif
 	FMatrix Result;
@@ -759,3 +754,63 @@ inline FArchive& operator<<(FArchive& Ar,FMatrix& M)
 	return Ar;
 }
 
+
+/**
+ * FPlane inline functions.
+ */
+
+inline FPlane FPlane::TransformBy( const FMatrix& M ) const
+{
+	const FMatrix tmpTA = M.TransposeAdjoint();
+	const float DetM = M.Determinant();
+	return this->TransformByUsingAdjointT(M, DetM, tmpTA);
+}
+
+inline FPlane FPlane::TransformByUsingAdjointT( const FMatrix& M, float DetM, const FMatrix& TA ) const
+{
+	FVector newNorm = TA.TransformVector(*this).GetSafeNormal();
+
+	if(DetM < 0.f)
+	{
+		newNorm *= -1.0f;
+	}
+
+	return FPlane(M.TransformPosition(*this * W), newNorm);
+}
+
+
+FORCEINLINE FBasisVectorMatrix::FBasisVectorMatrix(const FVector& XAxis,const FVector& YAxis,const FVector& ZAxis,const FVector& Origin)
+{
+	for(uint32 RowIndex = 0;RowIndex < 3;RowIndex++)
+	{
+		M[RowIndex][0] = (&XAxis.X)[RowIndex];
+		M[RowIndex][1] = (&YAxis.X)[RowIndex];
+		M[RowIndex][2] = (&ZAxis.X)[RowIndex];
+		M[RowIndex][3] = 0.0f;
+	}
+	M[3][0] = Origin | XAxis;
+	M[3][1] = Origin | YAxis;
+	M[3][2] = Origin | ZAxis;
+	M[3][3] = 1.0f;
+}
+
+
+FORCEINLINE FLookAtMatrix::FLookAtMatrix(const FVector& EyePosition, const FVector& LookAtPosition, const FVector& UpVector)
+{
+	const FVector ZAxis = (LookAtPosition - EyePosition).GetSafeNormal();
+	const FVector XAxis = (UpVector ^ ZAxis).GetSafeNormal();
+	const FVector YAxis = ZAxis ^ XAxis;
+
+	for (uint32 RowIndex = 0; RowIndex < 3; RowIndex++)
+	{
+		M[RowIndex][0] = (&XAxis.X)[RowIndex];
+		M[RowIndex][1] = (&YAxis.X)[RowIndex];
+		M[RowIndex][2] = (&ZAxis.X)[RowIndex];
+		M[RowIndex][3] = 0.0f;
+	}
+	M[3][0] = -EyePosition | XAxis;
+	M[3][1] = -EyePosition | YAxis;
+	M[3][2] = -EyePosition | ZAxis;
+	M[3][3] = 1.0f;
+}
+
