Copying //UE4/Dev-Build to //UE4/Dev-Main (Source: //UE4/Dev-Build @ 3104200)

diff --git a/Engine/Source/Runtime/Core/Public/Math/Plane.h b/Engine/Source/Runtime/Core/Public/Math/Plane.h
--- a/Engine/Source/Runtime/Core/Public/Math/Plane.h
+++ b/Engine/Source/Runtime/Core/Public/Math/Plane.h
@@ -2,6 +2,7 @@
 
 #pragma once
 
+#include "Vector.h"
 
 /**
  * Structure for three dimensional planes.
@@ -276,6 +277,84 @@ public:
 	}
 } GCC_ALIGN(16);
 
+/* FMath inline functions
+ *****************************************************************************/
+
+inline FVector FMath::LinePlaneIntersection
+	(
+	const FVector &Point1,
+	const FVector &Point2,
+	const FPlane  &Plane
+	)
+{
+	return
+		Point1
+		+	(Point2-Point1)
+		*	((Plane.W - (Point1|Plane))/((Point2 - Point1)|Plane));
+}
+
+inline bool FMath::IntersectPlanes3( FVector& I, const FPlane& P1, const FPlane& P2, const FPlane& P3 )
+{
+	// Compute determinant, the triple product P1|(P2^P3)==(P1^P2)|P3.
+	const float Det = (P1 ^ P2) | P3;
+	if( Square(Det) < Square(0.001f) )
+	{
+		// Degenerate.
+		I = FVector::ZeroVector;
+		return 0;
+	}
+	else
+	{
+		// Compute the intersection point, guaranteed valid if determinant is nonzero.
+		I = (P1.W*(P2^P3) + P2.W*(P3^P1) + P3.W*(P1^P2)) / Det;
+	}
+	return 1;
+}
+
+inline bool FMath::IntersectPlanes2( FVector& I, FVector& D, const FPlane& P1, const FPlane& P2 )
+{
+	// Compute line direction, perpendicular to both plane normals.
+	D = P1 ^ P2;
+	const float DD = D.SizeSquared();
+	if( DD < Square(0.001f) )
+	{
+		// Parallel or nearly parallel planes.
+		D = I = FVector::ZeroVector;
+		return 0;
+	}
+	else
+	{
+		// Compute intersection.
+		I = (P1.W*(P2^D) + P2.W*(D^P1)) / DD;
+		D.Normalize();
+		return 1;
+	}
+}
+
+/* FVector inline functions
+ *****************************************************************************/
+
+inline FVector FVector::MirrorByPlane( const FPlane& Plane ) const
+{
+	return *this - Plane * (2.f * Plane.PlaneDot(*this) );
+}
+
+inline FVector FVector::PointPlaneProject(const FVector& Point, const FPlane& Plane)
+{
+	//Find the distance of X from the plane
+	//Add the distance back along the normal from the point
+	return Point - Plane.PlaneDot(Point) * Plane;
+}
+
+inline FVector FVector::PointPlaneProject(const FVector& Point, const FVector& A, const FVector& B, const FVector& C)
+{
+	//Compute the plane normal from ABC
+	FPlane Plane(A, B, C);
+
+	//Find the distance of X from the plane
+	//Add the distance back along the normal from the point
+	return Point - Plane.PlaneDot(Point) * Plane;
+}
 
 /* FPlane inline functions
  *****************************************************************************/
@@ -336,7 +415,6 @@ FORCEINLINE FPlane FPlane::Flip() const
 	return FPlane(-X, -Y, -Z, -W);
 }
 
-
 FORCEINLINE bool FPlane::operator==(const FPlane& V) const
 {
 	return (X == V.X) && (Y == V.Y) && (Z == V.Z) && (W == V.W);
diff --git a/Engine/Source/Runtime/Core/Public/Math/Quat.h b/Engine/Source/Runtime/Core/Public/Math/Quat.h
--- a/Engine/Source/Runtime/Core/Public/Math/Quat.h
+++ b/Engine/Source/Runtime/Core/Public/Math/Quat.h
@@ -2,6 +2,7 @@
 
 #pragma once
 
+#include "UnrealMathUtility.h"
 
 /**
  * Floating point quaternion that can represent a rotation about an axis in 3-D space.
@@ -1105,3 +1106,33 @@ FORCEINLINE bool FQuat::ContainsNaN() const
 
 
 template<> struct TIsPODType<FQuat> { enum { Value = true }; };
+
+/* FMath inline functions
+ *****************************************************************************/
+
+template<class U>
+FORCEINLINE_DEBUGGABLE FQuat FMath::Lerp( const FQuat& A, const FQuat& B, const U& Alpha)
+{
+	return FQuat::Slerp(A, B, Alpha);
+}
+
+template<class U>
+FORCEINLINE_DEBUGGABLE FQuat FMath::BiLerp(const FQuat& P00, const FQuat& P10, const FQuat& P01, const FQuat& P11, float FracX, float FracY)
+{
+	FQuat Result;
+
+	Result = Lerp(
+		FQuat::Slerp_NotNormalized(P00,P10,FracX),
+		FQuat::Slerp_NotNormalized(P01,P11,FracX),
+		FracY
+		);
+
+	return Result;
+}
+
+template<class U>
+FORCEINLINE_DEBUGGABLE FQuat FMath::CubicInterp( const FQuat& P0, const FQuat& T0, const FQuat& P1, const FQuat& T1, const U& A)
+{
+	return FQuat::Squad(P0, T0, P1, T1, A);
+}
+
diff --git a/Engine/Source/Runtime/Core/Public/Math/Rotator.h b/Engine/Source/Runtime/Core/Public/Math/Rotator.h
--- a/Engine/Source/Runtime/Core/Public/Math/Rotator.h
+++ b/Engine/Source/Runtime/Core/Public/Math/Rotator.h
@@ -2,6 +2,7 @@
 
 #pragma once
 
+#include "UnrealMathUtility.h"
 
 /**
  * Implements a container for rotation information.
@@ -732,3 +733,20 @@ FORCEINLINE bool FRotator::ContainsNaN() const
 
 
 template<> struct TIsPODType<FRotator> { enum { Value = true }; };
+
+
+/* FMath inline functions
+ *****************************************************************************/
+
+template<class U>
+FORCEINLINE_DEBUGGABLE FRotator FMath::Lerp(const FRotator& A, const FRotator& B, const U& Alpha)
+{
+	return A + (B - A).GetNormalized() * Alpha;
+}
+
+template<class U>
+FORCEINLINE_DEBUGGABLE FRotator FMath::LerpRange(const FRotator& A, const FRotator& B, const U& Alpha)
+{
+	// Similar to Lerp, but does not take the shortest path. Allows interpolation over more than 180 degrees.
+	return (A * (1 - Alpha) + B * Alpha).GetNormalized();
+}
diff --git a/Engine/Source/Runtime/Core/Public/Math/Sphere.h b/Engine/Source/Runtime/Core/Public/Math/Sphere.h
--- a/Engine/Source/Runtime/Core/Public/Math/Sphere.h
+++ b/Engine/Source/Runtime/Core/Public/Math/Sphere.h
@@ -2,6 +2,7 @@
 
 #pragma once
 
+#include "UnrealMathUtility.h"
 
 /**
  * Implements a basic sphere.
@@ -172,3 +173,16 @@ public:
 		return Ar;
 	}
 };
+
+/* FMath inline functions
+ *****************************************************************************/
+
+/**
+ * Converts a sphere into a point plus radius squared for the test above
+ */
+FORCEINLINE bool FMath::SphereAABBIntersection(const FSphere& Sphere,const FBox& AABB)
+{
+	float RadiusSquared = FMath::Square(Sphere.W);
+	// If the distance is less than or equal to the radius, they intersect
+	return SphereAABBIntersection(Sphere.Center,RadiusSquared,AABB);
+}
diff --git a/Engine/Source/Runtime/Core/Public/Math/UnrealMath.h b/Engine/Source/Runtime/Core/Public/Math/UnrealMath.h
--- a/Engine/Source/Runtime/Core/Public/Math/UnrealMath.h
+++ b/Engine/Source/Runtime/Core/Public/Math/UnrealMath.h
@@ -51,507 +51,3 @@
 #include "ScalarRegister.h"
 #include "Transform.h"
 #include "ConvexHull2d.h"
-
-
-// FVector2D Implementation
-FORCEINLINE FVector2D::FVector2D( const FVector& V )
-	: X(V.X), Y(V.Y)
-{ }
-
-inline FVector FVector2D::SphericalToUnitCartesian() const
-{
-	const float SinTheta = FMath::Sin(X);
-	return FVector(FMath::Cos(Y) * SinTheta, FMath::Sin(Y) * SinTheta, FMath::Cos(X));
-}
-
-// FVector Implementation
-FORCEINLINE FVector::FVector( const FVector4& V )
-	: X(V.X), Y(V.Y), Z(V.Z)
-{
-	DiagnosticCheckNaN();
-}
-
-inline FVector FVector::MirrorByPlane( const FPlane& Plane ) const
-{
-	return *this - Plane * (2.f * Plane.PlaneDot(*this) );
-}
-
-inline FVector FVector::PointPlaneProject(const FVector& Point, const FPlane& Plane)
-{
-	//Find the distance of X from the plane
-	//Add the distance back along the normal from the point
-	return Point - Plane.PlaneDot(Point) * Plane;
-}
-
-inline FVector FVector::PointPlaneProject(const FVector& Point, const FVector& A, const FVector& B, const FVector& C)
-{
-	//Compute the plane normal from ABC
-	FPlane Plane(A, B, C);
-
-	//Find the distance of X from the plane
-	//Add the distance back along the normal from the point
-	return Point - Plane.PlaneDot(Point) * Plane;
-}
-
-FORCEINLINE_DEBUGGABLE FRotator FVector::Rotation() const
-{
-	return ToOrientationRotator();
-}
-
-FORCEINLINE_DEBUGGABLE FRotator FVector4::Rotation() const
-{
-	return ToOrientationRotator();
-}
-
-inline FPlane FPlane::TransformBy( const FMatrix& M ) const
-{
-	const FMatrix tmpTA = M.TransposeAdjoint();
-	const float DetM = M.Determinant();
-	return this->TransformByUsingAdjointT(M, DetM, tmpTA);
-}
-
-inline FPlane FPlane::TransformByUsingAdjointT( const FMatrix& M, float DetM, const FMatrix& TA ) const
-{
-	FVector newNorm = TA.TransformVector(*this).GetSafeNormal();
-
-	if(DetM < 0.f)
-	{
-		newNorm *= -1.0f;
-	}
-
-	return FPlane(M.TransformPosition(*this * W), newNorm);
-}
-
-/**
- * Find the intersection of a line and an offset plane. Assumes that the
- * line and plane do indeed intersect; you must make sure they're not
- * parallel before calling.
- *
- * @param Point1 the first point defining the line
- * @param Point2 the second point defining the line
- * @param PlaneOrigin the origin of the plane
- * @param PlaneNormal the normal of the plane
- *
- * @return The point of intersection between the line and the plane.
- */
-inline FVector FMath::LinePlaneIntersection
-	(
-	const FVector &Point1,
-	const FVector &Point2,
-	const FVector &PlaneOrigin,
-	const FVector &PlaneNormal
-	)
-{
-	return
-		Point1
-		+	(Point2-Point1)
-		*	(((PlaneOrigin - Point1)|PlaneNormal) / ((Point2 - Point1)|PlaneNormal));
-}
-
-
-/**
- * Find the intersection of a line and a plane. Assumes that the line and
- * plane do indeed intersect; you must make sure they're not parallel before
- * calling.
- *
- * @param Point1 the first point defining the line
- * @param Point2 the second point defining the line
- * @param Plane the plane
- *
- * @return The point of intersection between the line and the plane.
- */
-inline FVector FMath::LinePlaneIntersection
-	(
-	const FVector &Point1,
-	const FVector &Point2,
-	const FPlane  &Plane
-	)
-{
-	return
-		Point1
-		+	(Point2-Point1)
-		*	((Plane.W - (Point1|Plane))/((Point2 - Point1)|Plane));
-}
-
-inline bool FMath::PointBoxIntersection
-	(
-	const FVector&	Point,
-	const FBox&		Box
-	)
-{
-	if(Point.X >= Box.Min.X && Point.X <= Box.Max.X &&
-		Point.Y >= Box.Min.Y && Point.Y <= Box.Max.Y &&
-		Point.Z >= Box.Min.Z && Point.Z <= Box.Max.Z)
-		return 1;
-	else
-		return 0;
-}
-
-inline bool FMath::LineBoxIntersection
-	( 
-	const FBox& Box, 
-	const FVector& Start, 
-	const FVector& End, 
-	const FVector& Direction
-	)
-{
-	return LineBoxIntersection(Box, Start, End, Direction, Direction.Reciprocal());
-}
-
-inline bool FMath::LineBoxIntersection
-	(
-	const FBox&		Box,
-	const FVector&	Start,
-	const FVector&	End,
-	const FVector&	Direction,
-	const FVector&	OneOverDirection
-	)
-{
-	FVector	Time;
-	bool	bStartIsOutside = false;
-
-	if(Start.X < Box.Min.X)
-	{
-		bStartIsOutside = true;
-		if(End.X >= Box.Min.X)
-		{
-			Time.X = (Box.Min.X - Start.X) * OneOverDirection.X;
-		}
-		else
-		{
-			return false;
-		}
-	}
-	else if(Start.X > Box.Max.X)
-	{
-		bStartIsOutside = true;
-		if(End.X <= Box.Max.X)
-		{
-			Time.X = (Box.Max.X - Start.X) * OneOverDirection.X;
-		}
-		else
-		{
-			return false;
-		}
-	}
-	else
-	{
-		Time.X = 0.0f;
-	}
-
-	if(Start.Y < Box.Min.Y)
-	{
-		bStartIsOutside = true;
-		if(End.Y >= Box.Min.Y)
-		{
-			Time.Y = (Box.Min.Y - Start.Y) * OneOverDirection.Y;
-		}
-		else
-		{
-			return false;
-		}
-	}
-	else if(Start.Y > Box.Max.Y)
-	{
-		bStartIsOutside = true;
-		if(End.Y <= Box.Max.Y)
-		{
-			Time.Y = (Box.Max.Y - Start.Y) * OneOverDirection.Y;
-		}
-		else
-		{
-			return false;
-		}
-	}
-	else
-	{
-		Time.Y = 0.0f;
-	}
-
-	if(Start.Z < Box.Min.Z)
-	{
-		bStartIsOutside = true;
-		if(End.Z >= Box.Min.Z)
-		{
-			Time.Z = (Box.Min.Z - Start.Z) * OneOverDirection.Z;
-		}
-		else
-		{
-			return false;
-		}
-	}
-	else if(Start.Z > Box.Max.Z)
-	{
-		bStartIsOutside = true;
-		if(End.Z <= Box.Max.Z)
-		{
-			Time.Z = (Box.Max.Z - Start.Z) * OneOverDirection.Z;
-		}
-		else
-		{
-			return false;
-		}
-	}
-	else
-	{
-		Time.Z = 0.0f;
-	}
-
-	if(bStartIsOutside)
-	{
-		const float	MaxTime = Max3(Time.X,Time.Y,Time.Z);
-
-		if(MaxTime >= 0.0f && MaxTime <= 1.0f)
-		{
-			const FVector Hit = Start + Direction * MaxTime;
-			const float BOX_SIDE_THRESHOLD = 0.1f;
-			if(	Hit.X > Box.Min.X - BOX_SIDE_THRESHOLD && Hit.X < Box.Max.X + BOX_SIDE_THRESHOLD &&
-				Hit.Y > Box.Min.Y - BOX_SIDE_THRESHOLD && Hit.Y < Box.Max.Y + BOX_SIDE_THRESHOLD &&
-				Hit.Z > Box.Min.Z - BOX_SIDE_THRESHOLD && Hit.Z < Box.Max.Z + BOX_SIDE_THRESHOLD)
-			{
-				return true;
-			}
-		}
-
-		return false;
-	}
-	else
-	{
-		return true;
-	}
-}
-
-inline bool FMath::LineSphereIntersection(const FVector& Start,const FVector& Dir,float Length,const FVector& Origin,float Radius)
-{
-	const FVector	EO = Start - Origin;
-	const float		v = (Dir | (Origin - Start));
-	const float		disc = Radius * Radius - ((EO | EO) - v * v);
-
-	if(disc >= 0.0f)
-	{
-		const float	Time = (v - Sqrt(disc)) / Length;
-
-		if(Time >= 0.0f && Time <= 1.0f)
-			return 1;
-		else
-			return 0;
-	}
-	else
-		return 0;
-}
-
-inline bool FMath::IntersectPlanes3( FVector& I, const FPlane& P1, const FPlane& P2, const FPlane& P3 )
-{
-	// Compute determinant, the triple product P1|(P2^P3)==(P1^P2)|P3.
-	const float Det = (P1 ^ P2) | P3;
-	if( Square(Det) < Square(0.001f) )
-	{
-		// Degenerate.
-		I = FVector::ZeroVector;
-		return 0;
-	}
-	else
-	{
-		// Compute the intersection point, guaranteed valid if determinant is nonzero.
-		I = (P1.W*(P2^P3) + P2.W*(P3^P1) + P3.W*(P1^P2)) / Det;
-	}
-	return 1;
-}
-
-inline bool FMath::IntersectPlanes2( FVector& I, FVector& D, const FPlane& P1, const FPlane& P2 )
-{
-	// Compute line direction, perpendicular to both plane normals.
-	D = P1 ^ P2;
-	const float DD = D.SizeSquared();
-	if( DD < Square(0.001f) )
-	{
-		// Parallel or nearly parallel planes.
-		D = I = FVector::ZeroVector;
-		return 0;
-	}
-	else
-	{
-		// Compute intersection.
-		I = (P1.W*(P2^D) + P2.W*(D^P1)) / DD;
-		D.Normalize();
-		return 1;
-	}
-}
-
-FORCEINLINE float FMath::GetRangePct(FVector2D const& Range, float Value)
-{
-	return (Range.X != Range.Y) ? (Value - Range.X) / (Range.Y - Range.X) : Range.X;
-}
-
-FORCEINLINE float FMath::GetRangeValue(FVector2D const& Range, float Pct)
-{
-	return Lerp<float>(Range.X, Range.Y, Pct);
-}
-
-inline FVector FMath::VRand()
-{
-	FVector Result;
-
-	float L;
-
-	do
-	{
-		// Check random vectors in the unit sphere so result is statistically uniform.
-		Result.X = FRand() * 2.f - 1.f;
-		Result.Y = FRand() * 2.f - 1.f;
-		Result.Z = FRand() * 2.f - 1.f;
-		L = Result.SizeSquared();
-	}
-	while(L > 1.0f || L < KINDA_SMALL_NUMBER);
-
-	return Result * (1.0f / Sqrt(L));
-}
-
-template<class U>
-FORCEINLINE_DEBUGGABLE FRotator FMath::Lerp(const FRotator& A, const FRotator& B, const U& Alpha)
-{
-	return A + (B - A).GetNormalized() * Alpha;
-}
-
-template<class U>
-FORCEINLINE_DEBUGGABLE FRotator FMath::LerpRange(const FRotator& A, const FRotator& B, const U& Alpha)
-{
-	// Similar to Lerp, but does not take the shortest path. Allows interpolation over more than 180 degrees.
-	return (A * (1 - Alpha) + B * Alpha).GetNormalized();
-}
-
-template<class U>
-FORCEINLINE_DEBUGGABLE FQuat FMath::Lerp( const FQuat& A, const FQuat& B, const U& Alpha)
-{
-	return FQuat::Slerp(A, B, Alpha);
-}
-
-template<class U>
-FORCEINLINE_DEBUGGABLE FQuat FMath::BiLerp(const FQuat& P00, const FQuat& P10, const FQuat& P01, const FQuat& P11, float FracX, float FracY)
-{
-	FQuat Result;
-
-	Result = Lerp(
-		FQuat::Slerp_NotNormalized(P00,P10,FracX),
-		FQuat::Slerp_NotNormalized(P01,P11,FracX),
-		FracY
-		);
-
-	return Result;
-}
-
-template<class U>
-FORCEINLINE_DEBUGGABLE FQuat FMath::CubicInterp( const FQuat& P0, const FQuat& T0, const FQuat& P1, const FQuat& T1, const U& A)
-{
-	return FQuat::Squad(P0, T0, P1, T1, A);
-}
-
-template< class U > 
-FORCEINLINE_DEBUGGABLE U FMath::CubicCRSplineInterp(const U& P0, const U& P1, const U& P2, const U& P3, const float T0, const float T1, const float T2, const float T3, const float T)
-{
-	//Based on http://www.cemyuksel.com/research/catmullrom_param/catmullrom.pdf 
-	float InvT1MinusT0 = 1.0f / (T1 - T0);
-	U L01 = ( P0 * ((T1 - T) * InvT1MinusT0) ) + ( P1 * ((T - T0) * InvT1MinusT0) );
-	float InvT2MinusT1 = 1.0f / (T2 - T1);
-	U L12 = ( P1 * ((T2 - T) * InvT2MinusT1) ) + ( P2 * ((T - T1) * InvT2MinusT1) );
-	float InvT3MinusT2 = 1.0f / (T3 - T2);
-	U L23 = ( P2 * ((T3 - T) * InvT3MinusT2) ) + ( P3 * ((T - T2) * InvT3MinusT2) );
-
-	float InvT2MinusT0 = 1.0f / (T2 - T0);
-	U L012 = ( L01 * ((T2 - T) * InvT2MinusT0) ) + ( L12 * ((T - T0) * InvT2MinusT0) );
-	float InvT3MinusT1 = 1.0f / (T3 - T1);
-	U L123 = ( L12 * ((T3 - T) * InvT3MinusT1) ) + ( L23 * ((T - T1) * InvT3MinusT1) );
-
-	return  ( ( L012 * ((T2 - T) * InvT2MinusT1) ) + ( L123 * ((T - T1) * InvT2MinusT1) ) );
-}
-
-template< class U >
-FORCEINLINE_DEBUGGABLE U FMath::CubicCRSplineInterpSafe(const U& P0, const U& P1, const U& P2, const U& P3, const float T0, const float T1, const float T2, const float T3, const float T)
-{
-	//Based on http://www.cemyuksel.com/research/catmullrom_param/catmullrom.pdf 
-	float T1MinusT0 = (T1 - T0);
-	float T2MinusT1 = (T2 - T1);
-	float T3MinusT2 = (T3 - T2);
-	float T2MinusT0 = (T2 - T0);
-	float T3MinusT1 = (T3 - T1);
-	if (FMath::IsNearlyZero(T1MinusT0) || FMath::IsNearlyZero(T2MinusT1) || FMath::IsNearlyZero(T3MinusT2) || FMath::IsNearlyZero(T2MinusT0) || FMath::IsNearlyZero(T3MinusT1))
-	{
-		//There's going to be a divide by zero here so just bail out and return P1
-		return P1;
-	}
-
-	float InvT1MinusT0 = 1.0f / T1MinusT0;
-	U L01 = (P0 * ((T1 - T) * InvT1MinusT0)) + (P1 * ((T - T0) * InvT1MinusT0));
-	float InvT2MinusT1 = 1.0f / T2MinusT1;
-	U L12 = (P1 * ((T2 - T) * InvT2MinusT1)) + (P2 * ((T - T1) * InvT2MinusT1));
-	float InvT3MinusT2 = 1.0f / T3MinusT2;
-	U L23 = (P2 * ((T3 - T) * InvT3MinusT2)) + (P3 * ((T - T2) * InvT3MinusT2));
-
-	float InvT2MinusT0 = 1.0f / T2MinusT0;
-	U L012 = (L01 * ((T2 - T) * InvT2MinusT0)) + (L12 * ((T - T0) * InvT2MinusT0));
-	float InvT3MinusT1 = 1.0f / T3MinusT1;
-	U L123 = (L12 * ((T3 - T) * InvT3MinusT1)) + (L23 * ((T - T1) * InvT3MinusT1));
-
-	return  ((L012 * ((T2 - T) * InvT2MinusT1)) + (L123 * ((T - T1) * InvT2MinusT1)));
-}
-
-/**
- * Performs a sphere vs box intersection test using Arvo's algorithm:
- *
- *	for each i in (x, y, z)
- *		if (SphereCenter(i) < BoxMin(i)) d2 += (SphereCenter(i) - BoxMin(i)) ^ 2
- *		else if (SphereCenter(i) > BoxMax(i)) d2 += (SphereCenter(i) - BoxMax(i)) ^ 2
- *
- * @param Sphere the center of the sphere being tested against the AABB
- * @param RadiusSquared the size of the sphere being tested
- * @param AABB the box being tested against
- *
- * @return Whether the sphere/box intersect or not.
- */
-FORCEINLINE bool FMath::SphereAABBIntersection(const FVector& SphereCenter,const float RadiusSquared,const FBox& AABB)
-{
-	// Accumulates the distance as we iterate axis
-	float DistSquared = 0.f;
-	// Check each axis for min/max and add the distance accordingly
-	// NOTE: Loop manually unrolled for > 2x speed up
-	if (SphereCenter.X < AABB.Min.X)
-	{
-		DistSquared += FMath::Square(SphereCenter.X - AABB.Min.X);
-	}
-	else if (SphereCenter.X > AABB.Max.X)
-	{
-		DistSquared += FMath::Square(SphereCenter.X - AABB.Max.X);
-	}
-	if (SphereCenter.Y < AABB.Min.Y)
-	{
-		DistSquared += FMath::Square(SphereCenter.Y - AABB.Min.Y);
-	}
-	else if (SphereCenter.Y > AABB.Max.Y)
-	{
-		DistSquared += FMath::Square(SphereCenter.Y - AABB.Max.Y);
-	}
-	if (SphereCenter.Z < AABB.Min.Z)
-	{
-		DistSquared += FMath::Square(SphereCenter.Z - AABB.Min.Z);
-	}
-	else if (SphereCenter.Z > AABB.Max.Z)
-	{
-		DistSquared += FMath::Square(SphereCenter.Z - AABB.Max.Z);
-	}
-	// If the distance is less than or equal to the radius, they intersect
-	return DistSquared <= RadiusSquared;
-}
-
-/**
- * Converts a sphere into a point plus radius squared for the test above
- */
-FORCEINLINE bool FMath::SphereAABBIntersection(const FSphere& Sphere,const FBox& AABB)
-{
-	float RadiusSquared = FMath::Square(Sphere.W);
-	// If the distance is less than or equal to the radius, they intersect
-	return SphereAABBIntersection(Sphere.Center,RadiusSquared,AABB);
-}
-
-
-
-
diff --git a/Engine/Source/Runtime/Core/Public/Math/UnrealMathUtility.h b/Engine/Source/Runtime/Core/Public/Math/UnrealMathUtility.h
--- a/Engine/Source/Runtime/Core/Public/Math/UnrealMathUtility.h
+++ b/Engine/Source/Runtime/Core/Public/Math/UnrealMathUtility.h
@@ -843,11 +843,57 @@ struct FMath : public FPlatformMath
 	 *	T0-3 - The interpolation parameters for the corresponding control points.		
 	 *	T - The interpolation factor in the range 0 to 1. 0 returns P1. 1 returns P2.
 	 */
-	template< class U > static U CubicCRSplineInterp(const U& P0, const U& P1, const U& P2, const U& P3, const float T0, const float T1, const float T2, const float T3, const float T);
+	template< class U > 
+	static FORCEINLINE_DEBUGGABLE U CubicCRSplineInterp(const U& P0, const U& P1, const U& P2, const U& P3, const float T0, const float T1, const float T2, const float T3, const float T)
+	{
+		//Based on http://www.cemyuksel.com/research/catmullrom_param/catmullrom.pdf 
+		float InvT1MinusT0 = 1.0f / (T1 - T0);
+		U L01 = ( P0 * ((T1 - T) * InvT1MinusT0) ) + ( P1 * ((T - T0) * InvT1MinusT0) );
+		float InvT2MinusT1 = 1.0f / (T2 - T1);
+		U L12 = ( P1 * ((T2 - T) * InvT2MinusT1) ) + ( P2 * ((T - T1) * InvT2MinusT1) );
+		float InvT3MinusT2 = 1.0f / (T3 - T2);
+		U L23 = ( P2 * ((T3 - T) * InvT3MinusT2) ) + ( P3 * ((T - T2) * InvT3MinusT2) );
+
+		float InvT2MinusT0 = 1.0f / (T2 - T0);
+		U L012 = ( L01 * ((T2 - T) * InvT2MinusT0) ) + ( L12 * ((T - T0) * InvT2MinusT0) );
+		float InvT3MinusT1 = 1.0f / (T3 - T1);
+		U L123 = ( L12 * ((T3 - T) * InvT3MinusT1) ) + ( L23 * ((T - T1) * InvT3MinusT1) );
+
+		return  ( ( L012 * ((T2 - T) * InvT2MinusT1) ) + ( L123 * ((T - T1) * InvT2MinusT1) ) );
+	}
 
 	/* Same as CubicCRSplineInterp but with additional saftey checks. If the checks fail P1 is returned. **/
-	template< class U > static U CubicCRSplineInterpSafe(const U& P0, const U& P1, const U& P2, const U& P3, const float T0, const float T1, const float T2, const float T3, const float T);
-	
+	template< class U >
+	static FORCEINLINE_DEBUGGABLE U CubicCRSplineInterpSafe(const U& P0, const U& P1, const U& P2, const U& P3, const float T0, const float T1, const float T2, const float T3, const float T)
+	{
+		//Based on http://www.cemyuksel.com/research/catmullrom_param/catmullrom.pdf 
+		float T1MinusT0 = (T1 - T0);
+		float T2MinusT1 = (T2 - T1);
+		float T3MinusT2 = (T3 - T2);
+		float T2MinusT0 = (T2 - T0);
+		float T3MinusT1 = (T3 - T1);
+		if (FMath::IsNearlyZero(T1MinusT0) || FMath::IsNearlyZero(T2MinusT1) || FMath::IsNearlyZero(T3MinusT2) || FMath::IsNearlyZero(T2MinusT0) || FMath::IsNearlyZero(T3MinusT1))
+		{
+			//There's going to be a divide by zero here so just bail out and return P1
+			return P1;
+		}
+
+		float InvT1MinusT0 = 1.0f / T1MinusT0;
+		U L01 = (P0 * ((T1 - T) * InvT1MinusT0)) + (P1 * ((T - T0) * InvT1MinusT0));
+		float InvT2MinusT1 = 1.0f / T2MinusT1;
+		U L12 = (P1 * ((T2 - T) * InvT2MinusT1)) + (P2 * ((T - T1) * InvT2MinusT1));
+		float InvT3MinusT2 = 1.0f / T3MinusT2;
+		U L23 = (P2 * ((T3 - T) * InvT3MinusT2)) + (P3 * ((T - T2) * InvT3MinusT2));
+
+		float InvT2MinusT0 = 1.0f / T2MinusT0;
+		U L012 = (L01 * ((T2 - T) * InvT2MinusT0)) + (L12 * ((T - T0) * InvT2MinusT0));
+		float InvT3MinusT1 = 1.0f / T3MinusT1;
+		U L123 = (L12 * ((T3 - T) * InvT3MinusT1)) + (L23 * ((T - T1) * InvT3MinusT1));
+
+		return  ((L012 * ((T2 - T) * InvT2MinusT1)) + (L123 * ((T - T1) * InvT2MinusT1)));
+	}
+
+
 	// Special-case interpolation
 
 	/** Interpolate a normal vector Current to Target, by interpolating the angle between those vectors with constant step. */
@@ -898,11 +944,30 @@ struct FMath : public FPlatformMath
 
 
 	/**
-	 * Find the intersection of an infinite line (defined by two points) and
-	 * a plane.  Assumes that the line and plane do indeed intersect; you must
-	 * make sure they're not parallel before calling.
+	 * Find the intersection of a line and an offset plane. Assumes that the
+	 * line and plane do indeed intersect; you must make sure they're not
+	 * parallel before calling.
+	 *
+	 * @param Point1 the first point defining the line
+	 * @param Point2 the second point defining the line
+	 * @param PlaneOrigin the origin of the plane
+	 * @param PlaneNormal the normal of the plane
+	 *
+	 * @return The point of intersection between the line and the plane.
 	 */
 	static FVector LinePlaneIntersection( const FVector &Point1, const FVector &Point2, const FVector &PlaneOrigin, const FVector &PlaneNormal);
+
+	/**
+	 * Find the intersection of a line and a plane. Assumes that the line and
+	 * plane do indeed intersect; you must make sure they're not parallel before
+	 * calling.
+	 *
+	 * @param Point1 the first point defining the line
+	 * @param Point2 the second point defining the line
+	 * @param Plane the plane
+	 *
+	 * @return The point of intersection between the line and the plane.
+	 */
 	static FVector LinePlaneIntersection( const FVector &Point1, const FVector &Point2, const FPlane  &Plane);
 
 	// @parma InOutScissorRect should be set to View.ViewRect before the call
diff --git a/Engine/Source/Runtime/Core/Public/Math/Vector.h b/Engine/Source/Runtime/Core/Public/Math/Vector.h
--- a/Engine/Source/Runtime/Core/Public/Math/Vector.h
+++ b/Engine/Source/Runtime/Core/Public/Math/Vector.h
@@ -2,6 +2,8 @@
 
 #pragma once
 
+#include "UnrealMathUtility.h"
+#include "Vector2D.h"
 
 /**
  * A vector in 3-D space composed of components (X, Y, Z) with floating point precision.
@@ -665,7 +667,7 @@ public:
 	 * @return FRotator from the Vector's direction.
 	 * @see ToOrientationRotator(), ToOrientationQuat()
 	 */
-	FRotator Rotation() const;
+	CORE_API FRotator Rotation() const;
 
 	/**
 	 * Find good arbitrary axis vectors to represent U and V axes of a plane,
@@ -1000,19 +1002,6 @@ FORCEINLINE uint32 GetTypeHash(const FVector& Vector)
 }
 
 
-/**
- * Creates a hash value from a FVector2D. 
- *
- * @param Vector the vector to create a hash value for
- * @return The hash value from the components
- */
-FORCEINLINE uint32 GetTypeHash(const FVector2D& Vector)
-{
-	// Note: this assumes there's no padding in FVector2D that could contain uncompared data.
-	return FCrc::MemCrc_DEPRECATED(&Vector,sizeof(Vector));
-}
-
-
 #if PLATFORM_LITTLE_ENDIAN
 	#define INTEL_ORDER_VECTOR(x) (x)
 #else
@@ -1960,3 +1949,73 @@ FORCEINLINE FVector ClampVector(const FVector& V, const FVector& Min, const FVec
 }
 
 template <> struct TIsPODType<FVector> { enum { Value = true }; };
+
+/* FMath inline functions
+ *****************************************************************************/
+
+inline FVector FMath::LinePlaneIntersection
+	(
+	const FVector &Point1,
+	const FVector &Point2,
+	const FVector &PlaneOrigin,
+	const FVector &PlaneNormal
+	)
+{
+	return
+		Point1
+		+	(Point2-Point1)
+		*	(((PlaneOrigin - Point1)|PlaneNormal) / ((Point2 - Point1)|PlaneNormal));
+}
+
+inline bool FMath::LineSphereIntersection(const FVector& Start,const FVector& Dir,float Length,const FVector& Origin,float Radius)
+{
+	const FVector	EO = Start - Origin;
+	const float		v = (Dir | (Origin - Start));
+	const float		disc = Radius * Radius - ((EO | EO) - v * v);
+
+	if(disc >= 0.0f)
+	{
+		const float	Time = (v - Sqrt(disc)) / Length;
+
+		if(Time >= 0.0f && Time <= 1.0f)
+			return 1;
+		else
+			return 0;
+	}
+	else
+		return 0;
+}
+
+inline FVector FMath::VRand()
+{
+	FVector Result;
+
+	float L;
+
+	do
+	{
+		// Check random vectors in the unit sphere so result is statistically uniform.
+		Result.X = FRand() * 2.f - 1.f;
+		Result.Y = FRand() * 2.f - 1.f;
+		Result.Z = FRand() * 2.f - 1.f;
+		L = Result.SizeSquared();
+	}
+	while(L > 1.0f || L < KINDA_SMALL_NUMBER);
+
+	return Result * (1.0f / Sqrt(L));
+}
+
+
+/* FVector2D inline functions
+ *****************************************************************************/
+
+FORCEINLINE FVector2D::FVector2D( const FVector& V )
+	: X(V.X), Y(V.Y)
+{ 
+}
+
+inline FVector FVector2D::SphericalToUnitCartesian() const
+{
+	const float SinTheta = FMath::Sin(X);
+	return FVector(FMath::Cos(Y) * SinTheta, FMath::Sin(Y) * SinTheta, FMath::Cos(X));
+}
diff --git a/Engine/Source/Runtime/Core/Public/Math/Vector2D.h b/Engine/Source/Runtime/Core/Public/Math/Vector2D.h
--- a/Engine/Source/Runtime/Core/Public/Math/Vector2D.h
+++ b/Engine/Source/Runtime/Core/Public/Math/Vector2D.h
@@ -2,6 +2,7 @@
 
 #pragma once
 
+#include "UnrealMathUtility.h"
 
 /**
  * A vector in 2-D space composed of components (X, Y) with floating point precision.
@@ -515,6 +516,17 @@ public:
 	inline FVector SphericalToUnitCartesian() const;
 };
 
+/**
+ * Creates a hash value from a FVector2D. 
+ *
+ * @param Vector the vector to create a hash value for
+ * @return The hash value from the components
+ */
+FORCEINLINE uint32 GetTypeHash(const FVector2D& Vector)
+{
+	// Note: this assumes there's no padding in FVector2D that could contain uncompared data.
+	return FCrc::MemCrc_DEPRECATED(&Vector,sizeof(Vector));
+}
 
 /* FVector2D inline functions
  *****************************************************************************/
@@ -544,7 +556,6 @@ FORCEINLINE FVector2D::FVector2D(EForceInit)
 {
 }
 
-
 FORCEINLINE FVector2D FVector2D::operator+(const FVector2D& V) const
 {
 	return FVector2D(X + V.X, Y + V.Y);
@@ -902,3 +913,16 @@ FORCEINLINE bool FVector2D::InitFromString(const FString& InSourceString)
 
 	return bSuccessful;
 }
+
+/* FMath inline functions
+ *****************************************************************************/
+
+FORCEINLINE float FMath::GetRangePct(FVector2D const& Range, float Value)
+{
+	return (Range.X != Range.Y) ? (Value - Range.X) / (Range.Y - Range.X) : Range.X;
+}
+
+FORCEINLINE float FMath::GetRangeValue(FVector2D const& Range, float Pct)
+{
+	return Lerp<float>(Range.X, Range.Y, Pct);
+}
diff --git a/Engine/Source/Runtime/Core/Public/Math/Vector4.h b/Engine/Source/Runtime/Core/Public/Math/Vector4.h
--- a/Engine/Source/Runtime/Core/Public/Math/Vector4.h
+++ b/Engine/Source/Runtime/Core/Public/Math/Vector4.h
@@ -2,6 +2,7 @@
 
 #pragma once
 
+#include "Vector.h"
 
 /**
  * A 4D homogeneous vector, 4x1 FLOATs, 16-byte aligned.
@@ -317,7 +318,7 @@ public:
 	 * @return FRotator from the Vector's direction.
 	 * @see ToOrientationRotator()
 	 */
-	FRotator Rotation() const;
+	CORE_API FRotator Rotation() const;
 
 	/**
 	 * Set all of the vectors coordinates.
@@ -703,4 +704,14 @@ FORCEINLINE FVector4 FVector4::operator/(const FVector4& V) const
 	return FVector4(X / V.X, Y / V.Y, Z / V.Z, W / V.W);
 }
 
-template <> struct TIsPODType<FVector4> { enum { Value = true }; };
\ No newline at end of file
+template <> struct TIsPODType<FVector4> { enum { Value = true }; };
+
+
+/* FVector inline functions
+ *****************************************************************************/
+
+FORCEINLINE FVector::FVector( const FVector4& V )
+	: X(V.X), Y(V.Y), Z(V.Z)
+{
+	DiagnosticCheckNaN();
+}
diff --git a/Engine/Source/Runtime/Core/Public/Misc/AssertionMacros.h b/Engine/Source/Runtime/Core/Public/Misc/AssertionMacros.h
--- a/Engine/Source/Runtime/Core/Public/Misc/AssertionMacros.h
+++ b/Engine/Source/Runtime/Core/Public/Misc/AssertionMacros.h
@@ -8,6 +8,51 @@
 #include "Templates/EnableIf.h"
 #include "LogVerbosity.h"
 
+/**
+ * FDebug
+ * These functions offer debugging and diagnostic functionality and its presence 
+ * depends on compiler switches.
+ **/
+struct CORE_API FDebug
+{
+	/** Logs final assert message and exits the program. */
+	static void VARARGS AssertFailed(const ANSICHAR* Expr, const ANSICHAR* File, int32 Line, const TCHAR* Format = TEXT(""), ...);
+
+	/** Records the calling of AssertFailed() */
+	static bool bHasAsserted;
+
+#if DO_CHECK || DO_GUARD_SLOW
+	/** Failed assertion handler.  Warning: May be called at library startup time. */
+	static void VARARGS LogAssertFailedMessage( const ANSICHAR* Expr, const ANSICHAR* File, int32 Line, const TCHAR* Format=TEXT(""), ... );
+	
+	/**
+	 * Called when an 'ensure' assertion fails; gathers stack data and generates and error report.
+	 *
+	 * @param	Expr	Code expression ANSI string (#code)
+	 * @param	File	File name ANSI string (__FILE__)
+	 * @param	Line	Line number (__LINE__)
+	 * @param	Msg		Informative error message text
+	 * 
+	 * Don't change the name of this function, it's used to detect ensures by the crash reporter.
+	 */
+	static void EnsureFailed( const ANSICHAR* Expr, const ANSICHAR* File, int32 Line, const TCHAR* Msg );
+
+	/**
+	 * Logs an error if bLog is true, and returns false.  Takes a formatted string.
+	 *
+	 * @param	bLog	Log if true.
+	 * @param	Expr	Code expression ANSI string (#code)
+	 * @param	File	File name ANSI string (__FILE__)
+	 * @param	Line	Line number (__LINE__)
+	 * @param	FormattedMsg	Informative error message text with variable args
+	 *
+	 * @return false in all cases.
+	 *
+	 * Note: this crazy name is to ensure that the crash reporter recognizes it, which checks for functions in the callstack starting with 'EnsureNotFalse'.
+	 */
+	static bool VARARGS OptionallyLogFormattedEnsureMessageReturningFalse(bool bLog, const ANSICHAR* Expr, const ANSICHAR* File, int32 Line, const TCHAR* FormattedMsg, ...);
+#endif // DO_CHECK || DO_GUARD_SLOW
+};
 
 /*----------------------------------------------------------------------------
 	Check, verify, etc macros
@@ -187,22 +232,9 @@ namespace UE4Asserts_Private
 	((void)sizeof(&ClassName::FunctionName), TEXT(#FunctionName))
 
 /*----------------------------------------------------------------------------
-	Stats helpers
-----------------------------------------------------------------------------*/
-
-#if STATS
-	#define STAT(x) x
-#else
-	#define STAT(x)
-#endif
-
-/*----------------------------------------------------------------------------
-	Logging
+	Low level error macros
 ----------------------------------------------------------------------------*/
 
-// Define NO_LOGGING to strip out all writing to log files, OutputDebugString(), etc.
-// This is needed for consoles that require no logging (Xbox, Xenon)
-
 struct FTCharArrayTester
 {
 	template <uint32 N>
@@ -212,317 +244,14 @@ struct FTCharArrayTester
 
 #define IS_TCHAR_ARRAY(expr) (sizeof(FTCharArrayTester::Func(expr)) == 2)
 
+/** low level fatal error handler. */
+CORE_API void VARARGS LowLevelFatalErrorHandler(const ANSICHAR* File, int32 Line, const TCHAR* Format=TEXT(""), ... );
+
 #define LowLevelFatalError(Format, ...) \
 	{ \
 		static_assert(IS_TCHAR_ARRAY(Format), "Formatting string must be a TCHAR array."); \
-		FError::LowLevelFatal(__FILE__, __LINE__, Format, ##__VA_ARGS__); \
+		LowLevelFatalErrorHandler(__FILE__, __LINE__, Format, ##__VA_ARGS__); \
 		_DebugBreakAndPromptForRemote(); \
 		FDebug::AssertFailed("", __FILE__, __LINE__, Format, ##__VA_ARGS__); \
 	}
 
-
-/*----------------------------------------------------------------------------
-	Logging suppression
-----------------------------------------------------------------------------*/
-
-#ifndef COMPILED_IN_MINIMUM_VERBOSITY
-	#define COMPILED_IN_MINIMUM_VERBOSITY VeryVerbose
-#else
-	#if !IS_MONOLITHIC
-		#error COMPILED_IN_MINIMUM_VERBOSITY can only be defined in monolithic builds.
-	#endif
-#endif
-
-#define UE_LOG_EXPAND_IS_FATAL(Verbosity, ActiveBlock, InactiveBlock) PREPROCESSOR_JOIN(UE_LOG_EXPAND_IS_FATAL_, Verbosity)(ActiveBlock, InactiveBlock)
-
-#define UE_LOG_EXPAND_IS_FATAL_Fatal(      ActiveBlock, InactiveBlock) ActiveBlock
-#define UE_LOG_EXPAND_IS_FATAL_Error(      ActiveBlock, InactiveBlock) InactiveBlock
-#define UE_LOG_EXPAND_IS_FATAL_Warning(    ActiveBlock, InactiveBlock) InactiveBlock
-#define UE_LOG_EXPAND_IS_FATAL_Display(    ActiveBlock, InactiveBlock) InactiveBlock
-#define UE_LOG_EXPAND_IS_FATAL_Log(        ActiveBlock, InactiveBlock) InactiveBlock
-#define UE_LOG_EXPAND_IS_FATAL_Verbose(    ActiveBlock, InactiveBlock) InactiveBlock
-#define UE_LOG_EXPAND_IS_FATAL_VeryVerbose(ActiveBlock, InactiveBlock) InactiveBlock
-#define UE_LOG_EXPAND_IS_FATAL_All(        ActiveBlock, InactiveBlock) InactiveBlock
-#define UE_LOG_EXPAND_IS_FATAL_SetColor(   ActiveBlock, InactiveBlock) InactiveBlock
-
-#if NO_LOGGING
-
-	struct FNoLoggingCategory {};
-	
-	// This will only log Fatal errors
-	#define UE_LOG(CategoryName, Verbosity, Format, ...) \
-	{ \
-		static_assert(IS_TCHAR_ARRAY(Format), "Formatting string must be a TCHAR array."); \
-		if (ELogVerbosity::Verbosity == ELogVerbosity::Fatal) \
-		{ \
-			FError::LowLevelFatal(__FILE__, __LINE__, Format, ##__VA_ARGS__); \
-			_DebugBreakAndPromptForRemote(); \
-			FDebug::AssertFailed("", __FILE__, __LINE__, Format, ##__VA_ARGS__); \
-			UE_LOG_EXPAND_IS_FATAL(Verbosity, CA_ASSUME(false);, PREPROCESSOR_NOTHING) \
-		} \
-	}
-	// Conditional logging (fatal errors only).
-	#define UE_CLOG(Condition, CategoryName, Verbosity, Format, ...) \
-	{ \
-		static_assert(IS_TCHAR_ARRAY(Format), "Formatting string must be a TCHAR array."); \
-		if (ELogVerbosity::Verbosity == ELogVerbosity::Fatal) \
-		{ \
-			if (Condition) \
-			{ \
-				FError::LowLevelFatal(__FILE__, __LINE__, Format, ##__VA_ARGS__); \
-				_DebugBreakAndPromptForRemote(); \
-				FDebug::AssertFailed("", __FILE__, __LINE__, Format, ##__VA_ARGS__); \
-				UE_LOG_EXPAND_IS_FATAL(Verbosity, CA_ASSUME(false);, PREPROCESSOR_NOTHING) \
-			} \
-		} \
-	}
-
-	#define UE_LOG_ACTIVE(...)				(0)
-	#define UE_LOG_ANY_ACTIVE(...)			(0)
-	#define UE_SUPPRESS(...) {}
-	#define UE_SET_LOG_VERBOSITY(...)
-	#define DECLARE_LOG_CATEGORY_EXTERN(CategoryName, DefaultVerbosity, CompileTimeVerbosity) extern FNoLoggingCategory CategoryName
-	#define DEFINE_LOG_CATEGORY(...)
-	#define DEFINE_LOG_CATEGORY_STATIC(...)
-	#define DECLARE_LOG_CATEGORY_CLASS(...)
-	#define DEFINE_LOG_CATEGORY_CLASS(...)
-	#define LOG_SCOPE_VERBOSITY_OVERRIDE(...)
-	#define UE_SECURITY_LOG(...)
-
-#else
-
-	namespace UE4Asserts_Private
-	{
-		template <int32 VerbosityToCheck, typename CategoryType>
-		FORCEINLINE
-			typename TEnableIf<
-				((VerbosityToCheck & ELogVerbosity::VerbosityMask) <= CategoryType::CompileTimeVerbosity &&
-				(VerbosityToCheck & ELogVerbosity::VerbosityMask) <= ELogVerbosity::COMPILED_IN_MINIMUM_VERBOSITY),
-				bool>::Type
-			IsLogActive(const CategoryType& Category)
-		{
-			return !Category.IsSuppressed((ELogVerbosity::Type)VerbosityToCheck);
-		}
-
-		template <int32 VerbosityToCheck, typename CategoryType>
-		FORCEINLINE
-			typename TEnableIf<
-				!((VerbosityToCheck & ELogVerbosity::VerbosityMask) <= CategoryType::CompileTimeVerbosity &&
-				(VerbosityToCheck & ELogVerbosity::VerbosityMask) <= ELogVerbosity::COMPILED_IN_MINIMUM_VERBOSITY),
-				bool>::Type
-			IsLogActive(const CategoryType& Category)
-		{
-			return false;
-		}
-	}
-
-	/** 
-	 * A predicate that returns true if the given logging category is active (logging) at a given verbosity level 
-	 * @param CategoryName name of the logging category
-	 * @param Verbosity, verbosity level to test against
-	**/
-	#define UE_LOG_ACTIVE(CategoryName, Verbosity) (::UE4Asserts_Private::IsLogActive<(int32)ELogVerbosity::Verbosity>(CategoryName))
-
-	#define UE_SET_LOG_VERBOSITY(CategoryName, Verbosity) \
-		CategoryName.SetVerbosity(ELogVerbosity::Verbosity);
-
-	/** 
-	 * A  macro that outputs a formatted message to log if a given logging category is active at a given verbosity level
-	 * @param CategoryName name of the logging category
-	 * @param Verbosity, verbosity level to test against
-	 * @param Format, format text
-	 ***/
-	#define UE_LOG(CategoryName, Verbosity, Format, ...) \
-	{ \
-		static_assert(IS_TCHAR_ARRAY(Format), "Formatting string must be a TCHAR array."); \
-		static_assert((ELogVerbosity::Verbosity & ELogVerbosity::VerbosityMask) < ELogVerbosity::NumVerbosity && ELogVerbosity::Verbosity > 0, "Verbosity must be constant and in range."); \
-		CA_CONSTANT_IF((ELogVerbosity::Verbosity & ELogVerbosity::VerbosityMask) <= ELogVerbosity::COMPILED_IN_MINIMUM_VERBOSITY && (ELogVerbosity::Warning & ELogVerbosity::VerbosityMask) <= FLogCategory##CategoryName::CompileTimeVerbosity) \
-		{ \
-			UE_LOG_EXPAND_IS_FATAL(Verbosity, PREPROCESSOR_NOTHING, if (!CategoryName.IsSuppressed(ELogVerbosity::Verbosity))) \
-			{ \
-				FMsg::Logf_Internal(__FILE__, __LINE__, CategoryName.GetCategoryName(), ELogVerbosity::Verbosity, Format, ##__VA_ARGS__); \
-				UE_LOG_EXPAND_IS_FATAL(Verbosity, \
-					{ \
-						_DebugBreakAndPromptForRemote(); \
-						FDebug::AssertFailed("", __FILE__, __LINE__, Format, ##__VA_ARGS__); \
-						CA_ASSUME(false); \
-					}, \
-					PREPROCESSOR_NOTHING \
-				) \
-			} \
-		} \
-	}
-
-	/**
-	* A  macro that outputs a formatted message to the log specifically used for security events
-	* @param NetConnection, a valid UNetConnection
-	* @param SecurityEventType, a security event type (ESecurityEvent::Type)
-	* @param Format, format text
-	***/
-	#define UE_SECURITY_LOG(NetConnection, SecurityEventType, Format, ...) \
-	{ \
-		static_assert(IS_TCHAR_ARRAY(Format), "Formatting string must be a TCHAR array."); \
-		check(NetConnection != nullptr); \
-		CA_CONSTANT_IF((ELogVerbosity::Warning & ELogVerbosity::VerbosityMask) <= ELogVerbosity::COMPILED_IN_MINIMUM_VERBOSITY && (ELogVerbosity::Warning & ELogVerbosity::VerbosityMask) <= FLogCategoryLogSecurity::CompileTimeVerbosity) \
-		{ \
-			if (!LogSecurity.IsSuppressed(ELogVerbosity::Warning)) \
-			{ \
-				FString Test = FString::Printf(TEXT("%s: %s: %s"), *(NetConnection->RemoteAddressToString()), ToString(SecurityEventType), Format); \
-				FMsg::Logf_Internal(__FILE__, __LINE__, LogSecurity.GetCategoryName(), ELogVerbosity::Warning, *Test, ##__VA_ARGS__); \
-			} \
-		} \
-	}
-
-	// Conditional logging. Will only log if Condition is met.
-	#define UE_CLOG(Condition, CategoryName, Verbosity, Format, ...) \
-	{ \
-		static_assert(IS_TCHAR_ARRAY(Format), "Formatting string must be a TCHAR array."); \
-		static_assert((ELogVerbosity::Verbosity & ELogVerbosity::VerbosityMask) < ELogVerbosity::NumVerbosity && ELogVerbosity::Verbosity > 0, "Verbosity must be constant and in range."); \
-		CA_CONSTANT_IF((ELogVerbosity::Verbosity & ELogVerbosity::VerbosityMask) <= ELogVerbosity::COMPILED_IN_MINIMUM_VERBOSITY && (ELogVerbosity::Warning & ELogVerbosity::VerbosityMask) <= FLogCategory##CategoryName::CompileTimeVerbosity) \
-		{ \
-			UE_LOG_EXPAND_IS_FATAL(Verbosity, PREPROCESSOR_NOTHING, if (!CategoryName.IsSuppressed(ELogVerbosity::Verbosity))) \
-			{ \
-				if (Condition) \
-				{ \
-					FMsg::Logf_Internal(__FILE__, __LINE__, CategoryName.GetCategoryName(), ELogVerbosity::Verbosity, Format, ##__VA_ARGS__); \
-					UE_LOG_EXPAND_IS_FATAL(Verbosity, \
-						{ \
-							_DebugBreakAndPromptForRemote(); \
-							FDebug::AssertFailed("", __FILE__, __LINE__, Format, ##__VA_ARGS__); \
-							CA_ASSUME(false); \
-						}, \
-						PREPROCESSOR_NOTHING \
-					) \
-					CA_ASSUME(true); \
-				} \
-			} \
-		} \
-	}
-
-	/** 
-	 * A macro that executes some code within a scope if a given logging category is active at a given verbosity level
-	 * Also, withing the scope of the execution, the default category and verbosity is set up for the low level logging 
-	 * functions.
-	 * @param CategoryName name of the logging category
-	 * @param Verbosity, verbosity level to test against
-	 * @param ExecuteIfUnsuppressed, code to execute if the verbosity level for this category is being displayed
-	 ***/
-	#define UE_SUPPRESS(CategoryName, Verbosity, ExecuteIfUnsuppressed) \
-	{ \
-		static_assert((ELogVerbosity::Verbosity & ELogVerbosity::VerbosityMask) < ELogVerbosity::NumVerbosity && ELogVerbosity::Verbosity > 0, "Verbosity must be constant and in range."); \
-		CA_CONSTANT_IF((ELogVerbosity::Verbosity & ELogVerbosity::VerbosityMask) <= ELogVerbosity::COMPILED_IN_MINIMUM_VERBOSITY && (ELogVerbosity::Warning & ELogVerbosity::VerbosityMask) <= FLogCategory##CategoryName::CompileTimeVerbosity) \
-		{ \
-			if (!CategoryName.IsSuppressed(ELogVerbosity::Verbosity)) \
-			{ \
-				FScopedCategoryAndVerbosityOverride TEMP__##CategoryName(CategoryName.GetCategoryName(), ELogVerbosity::Type(ELogVerbosity::Verbosity & ELogVerbosity::VerbosityMask)); \
-				ExecuteIfUnsuppressed; \
-				CategoryName.PostTrigger(ELogVerbosity::Verbosity); \
-			} \
-		} \
-	}
-
-	/** 
-	 * A macro to declare a logging category as a C++ "extern" 
-	 * @param CategoryName, category to declare
-	 * @param DefaultVerbosity, default run time verbosity
-	 * @param CompileTimeVerbosity, maximum verbosity to compile into the code
-	 **/
-	#define DECLARE_LOG_CATEGORY_EXTERN(CategoryName, DefaultVerbosity, CompileTimeVerbosity) \
-		extern struct FLogCategory##CategoryName : public FLogCategory<ELogVerbosity::DefaultVerbosity, ELogVerbosity::CompileTimeVerbosity> \
-		{ \
-			FORCEINLINE FLogCategory##CategoryName() : FLogCategory(TEXT(#CategoryName)) {} \
-		} CategoryName;
-
-	/** 
-	 * A macro to define a logging category, usually paired with DECLARE_LOG_CATEGORY_EXTERN from the header.
-	 * @param CategoryName, category to define
-	**/
-	#define DEFINE_LOG_CATEGORY(CategoryName) FLogCategory##CategoryName CategoryName;
-
-	/** 
-	 * A macro to define a logging category as a C++ "static"
-	 * @param CategoryName, category to declare
-	 * @param DefaultVerbosity, default run time verbosity
-	 * @param CompileTimeVerbosity, maximum verbosity to compile into the code
-	**/
-	#define DEFINE_LOG_CATEGORY_STATIC(CategoryName, DefaultVerbosity, CompileTimeVerbosity) \
-		static struct FLogCategory##CategoryName : public FLogCategory<ELogVerbosity::DefaultVerbosity, ELogVerbosity::CompileTimeVerbosity> \
-		{ \
-			FORCEINLINE FLogCategory##CategoryName() : FLogCategory(TEXT(#CategoryName)) {} \
-		} CategoryName;
-
-	/** 
-	 * A macro to declare a logging category as a C++ "class static" 
-	 * @param CategoryName, category to declare
-	 * @param DefaultVerbosity, default run time verbosity
-	 * @param CompileTimeVerbosity, maximum verbosity to compile into the code
-	**/
-	#define DECLARE_LOG_CATEGORY_CLASS(CategoryName, DefaultVerbosity, CompileTimeVerbosity) \
-		DEFINE_LOG_CATEGORY_STATIC(CategoryName, DefaultVerbosity, CompileTimeVerbosity)
-
-	/** 
-	 * A macro to define a logging category, usually paired with DECLARE_LOG_CATEGORY_CLASS from the header.
-	 * @param CategoryName, category to definee
-	**/
-	#define DEFINE_LOG_CATEGORY_CLASS(Class, CategoryName) Class::FLogCategory##CategoryName Class::CategoryName;
-
-	/** 
-	 * A macro to override Verbosity of the Category within the scope
-	 * @param CategoryName, category to declare
-	 * @param ScopeVerbosity, verbosity to override
-	 **/
-	#define LOG_SCOPE_VERBOSITY_OVERRIDE(CategoryName, ScopeVerbosity) FLogScopedVerbosityOverride LogCategory##CategoryName##Override(&CategoryName, ScopeVerbosity)
-
-#endif // NO_LOGGING
-#if PLATFORM_HTML5
-#include "HTML5/HTML5AssertionMacros.h"
-#endif 
-
-#if UE_BUILD_SHIPPING
-#define NOTIFY_CLIENT_OF_SECURITY_EVENT_IF_NOT_SHIPPING(NetConnection, SecurityPrint) ;
-#else
-#define NOTIFY_CLIENT_OF_SECURITY_EVENT_IF_NOT_SHIPPING(NetConnection, SecurityPrint) \
-	FNetControlMessage<NMT_SecurityViolation>::Send(NetConnection, SecurityPrint); \
-	NetConnection->FlushNet(true)
-#endif
-
-/**
-	* A  macro that closes the connection and logs the security event on the server and the client
-	* @param NetConnection, a valid UNetConnection
-	* @param SecurityEventType, a security event type (ESecurityEvent::Type)
-	* @param Format, format text
-***/
-#define CLOSE_CONNECTION_DUE_TO_SECURITY_VIOLATION(NetConnection, SecurityEventType, Format, ...) \
-{ \
-	static_assert(IS_TCHAR_ARRAY(Format), "Formatting string must be a TCHAR array."); \
-	check(NetConnection != nullptr); \
-	FString SecurityPrint = FString::Printf(Format, ##__VA_ARGS__); \
-	UE_SECURITY_LOG(NetConnection, SecurityEventType, Format, ##__VA_ARGS__); \
-	UE_SECURITY_LOG(NetConnection, ESecurityEvent::Closed, TEXT("Connection closed")); \
-	NetConnection->Close(); \
-	PerfCountersIncrement(TEXT("ClosedConnectionsDueToSecurityViolations")); \
-}
-
-extern CORE_API int32 GEnsureOnNANDiagnostic;
-
-// Macro to either log an error or ensure on a NaN error.
-#if DO_CHECK && !USING_CODE_ANALYSIS
-namespace UE4Asserts_Private
-{
-	CORE_API void VARARGS InternalLogNANDiagnosticMessage(const TCHAR* FormattedMsg, ...); // UE_LOG(LogCore, Error, _FormatString_, ##__VA_ARGS__);
-}
-#define logOrEnsureNanError(_FormatString_, ...) \
-	if (!GEnsureOnNANDiagnostic)\
-	{\
-		if (UE4Asserts_Private::TrueOnFirstCallOnly([]{}))\
-		{\
-			UE4Asserts_Private::InternalLogNANDiagnosticMessage(_FormatString_, ##__VA_ARGS__); \
-		}\
-	}\
-	else\
-	{\
-		ensureMsgf(!GEnsureOnNANDiagnostic, _FormatString_, ##__VA_ARGS__); \
-	}
-#else
-#define logOrEnsureNanError(_FormatString_, ...)
-#endif // DO_CHECK
\ No newline at end of file
diff --git a/Engine/Source/Runtime/Core/Public/Misc/BufferedOutputDevice.h b/Engine/Source/Runtime/Core/Public/Misc/BufferedOutputDevice.h
new file mode 100644
--- /dev/null
+++ b/Engine/Source/Runtime/Core/Public/Misc/BufferedOutputDevice.h
@@ -0,0 +1,34 @@
+// Copyright 1998-2016 Epic Games, Inc. All Rights Reserved.
+
+#pragma once
+
+#include "OutputDevice.h"
+
+/** Buffered output device. */
+class CORE_API FBufferedOutputDevice : public FOutputDevice
+{
+protected:
+	TArray<FBufferedLine>		BufferedLines;
+	ELogVerbosity::Type			FilterLevel;
+	FCriticalSection			SynchronizationObject;
+
+public:
+
+	FBufferedOutputDevice() : FilterLevel(ELogVerbosity::All)
+	{
+	}
+
+	void	SetVerbosity(ELogVerbosity::Type Verbosity) { FilterLevel = Verbosity; }
+	void	Serialize(const TCHAR* InData, ELogVerbosity::Type Verbosity, const class FName& Category) override;
+	void	GetContents(TArray<FBufferedLine>& DestBuffer, bool ClearDevice = true);
+
+	/** Pushes buffered lines into the specified output device. */
+	void RedirectTo(class FOutputDevice& Ar)
+	{
+		for (const FBufferedLine& BufferedLine : BufferedLines)
+		{
+			Ar.Serialize(*BufferedLine.Data, BufferedLine.Verbosity, BufferedLine.Category);
+		}
+	}
+};
+
diff --git a/Engine/Source/Runtime/Core/Public/Misc/Build.h b/Engine/Source/Runtime/Core/Public/Misc/Build.h
--- a/Engine/Source/Runtime/Core/Public/Misc/Build.h
+++ b/Engine/Source/Runtime/Core/Public/Misc/Build.h
@@ -147,7 +147,7 @@
  * Checks to see if pure virtual has actually been implemented, this is normally run as a CIS process and is set (indirectly) by UBT
  *
  * @see Core.h
- * @see ObjectBase.h
+ * @see ObjectMacros.h
  **/
 #ifndef CHECK_PUREVIRTUALS
 	#define CHECK_PUREVIRTUALS 0
diff --git a/Engine/Source/Runtime/Core/Public/Misc/CallbackDevice.h b/Engine/Source/Runtime/Core/Public/Misc/CallbackDevice.h
--- a/Engine/Source/Runtime/Core/Public/Misc/CallbackDevice.h
+++ b/Engine/Source/Runtime/Core/Public/Misc/CallbackDevice.h
@@ -3,294 +3,4 @@
 #pragma once
 
 #include "ScopedCallback.h"
-
-
-class AActor;
-
-
-/**
- * The list of events that can be fired from the engine to unknown listeners.
- */
-enum ECallbackQueryType
-{
-	CALLBACK_ModalErrorMessage,				//Sent when FMessageDialog::Open or appDebugMesgf need to display a modal
-	CALLBACK_LocalizationExportFilter,		// A query sent to find out if the provided object passes the localization export filter
-	CALLBACK_QueryCount,
-};
-
-
-// delegates for hotfixes
-namespace EHotfixDelegates
-{
-	enum Type
-	{
-		Test,
-	};
-}
-
-
-// this is an example of a hotfix arg and return value structure. Once we have other examples, it can be deleted.
-struct FTestHotFixPayload
-{
-	FString Message;
-	bool ValueToReturn;
-	bool Result;
-};
-
-// Parameters passed to CrashOverrideParamsChanged used to customize crash report client behavior/appearance
-struct FCrashOverrideParameters
-{
-	FString CrashReportClientMessageText;
-};
-
-
-class CORE_API FCoreDelegates
-{
-public:
-	//hot fix delegate
-	DECLARE_DELEGATE_TwoParams(FHotFixDelegate, void *, int32);
-
-	// Callback for object property modifications
-	DECLARE_MULTICAST_DELEGATE_OneParam(FOnActorLabelChanged, AActor*);
-
-	// delegate type for prompting the pak system to mount a new pak
-	DECLARE_DELEGATE_RetVal_ThreeParams(bool, FOnMountPak, const FString&, uint32, IPlatformFile::FDirectoryVisitor*);
-
-	// delegate type for prompting the pak system to mount a new pak
-	DECLARE_DELEGATE_RetVal_OneParam(bool, FOnUnmountPak, const FString&);
-
-	/** delegate type for opening a modal message box ( Params: EAppMsgType::Type MessageType, const FText& Text, const FText& Title ) */
-	DECLARE_DELEGATE_RetVal_ThreeParams(EAppReturnType::Type, FOnModalMessageBox, EAppMsgType::Type, const FText&, const FText&);
-
-	// Callback for PER_MODULE_BOILERPLATE macro's GObjectArrayForDebugVisualizers
-	DECLARE_DELEGATE_RetVal(class FFixedUObjectArray*, FObjectArrayForDebugVisualizersDelegate);
-
-	// Called in PER_MODULE_BOILERPLATE macro.
-	static FObjectArrayForDebugVisualizersDelegate& GetObjectArrayForDebugVisualizersDelegate();
-
-	// Callback for handling an ensure
-	DECLARE_MULTICAST_DELEGATE(FOnHandleSystemEnsure);
-
-	// Callback for handling an error
-	DECLARE_MULTICAST_DELEGATE(FOnHandleSystemError);
-
-	// Callback for handling user login/logout.  first int is UserID, second int is UserIndex
-	DECLARE_MULTICAST_DELEGATE_ThreeParams(FOnUserLoginChangedEvent, bool, int32, int32);
-
-	// Callback for handling safe frame area size changes
-	DECLARE_MULTICAST_DELEGATE(FOnSafeFrameChangedEvent);
-
-	// Callback for handling accepting invitations - generally for engine code
-	DECLARE_MULTICAST_DELEGATE_TwoParams(FOnInviteAccepted, const FString&, const FString&);
-
-	// Callback for handling the Controller connection / disconnection
-	// first param is true for a connection, false for a disconnection.
-	// second param is UserID, third is UserIndex / ControllerId.
-	DECLARE_MULTICAST_DELEGATE_ThreeParams(FOnUserControllerConnectionChange, bool, int32, int32);
-
-	// Callback for platform handling when flushing async loads.
-	DECLARE_MULTICAST_DELEGATE(FOnAsyncLoadingFlush);
-	static FOnAsyncLoadingFlush OnAsyncLoadingFlush;
-
-	// get a hotfix delegate
-	static FHotFixDelegate& GetHotfixDelegate(EHotfixDelegates::Type HotFix);
-
-	// Callback when a user logs in/out of the platform.
-	static FOnUserLoginChangedEvent OnUserLoginChangedEvent;
-
-	// Callback when controllers disconnected / reconnected
-	static FOnUserControllerConnectionChange OnControllerConnectionChange;
-
-	// Callback when a user changes the safe frame size
-	static FOnSafeFrameChangedEvent OnSafeFrameChangedEvent;
-
-	// Callback for mounting a new pak file.
-	static FOnMountPak OnMountPak;
-
-	// Callback for unmounting a pak file.
-	static FOnUnmountPak OnUnmountPak;
-
-	// Callback when an ensure has occurred
-	static FOnHandleSystemEnsure OnHandleSystemEnsure;
-	// Callback when an error (crash) has occurred
-	static FOnHandleSystemError OnHandleSystemError;
-
-	// Called when an actor label is changed
-	static FOnActorLabelChanged OnActorLabelChanged;
-
-#if WITH_EDITOR
-	// Called before the editor displays a modal window, allowing other windows the opportunity to disable themselves to avoid reentrant calls
-	static FSimpleMulticastDelegate PreModal;
-
-	// Called after the editor dismisses a modal window, allowing other windows the opportunity to disable themselves to avoid reentrant calls
-	static FSimpleMulticastDelegate PostModal;
-#endif	//WITH_EDITOR
-	
-	// Called when an error occurred.
-	static FSimpleMulticastDelegate OnShutdownAfterError;
-
-	// Called when appInit is called.
-	static FSimpleMulticastDelegate OnInit;
-
-	// Called when the application is about to exit.
-	static FSimpleMulticastDelegate OnExit;
-
-	// Called when before the application is exiting.
-	static FSimpleMulticastDelegate OnPreExit;
-
-	/** Color picker color has changed, please refresh as needed*/
-	static FSimpleMulticastDelegate ColorPickerChanged;
-
-	/** requests to open a message box */
-	static FOnModalMessageBox ModalErrorMessage;
-
-	/** Called when the user accepts an invitation to the current game */
-	static FOnInviteAccepted OnInviteAccepted;
-
-	DECLARE_MULTICAST_DELEGATE_ThreeParams(FWorldOriginOffset, class UWorld*, FIntVector, FIntVector);
-	/** called before world origin shifting */
-	static FWorldOriginOffset PreWorldOriginOffset;
-	/** called after world origin shifting */
-	static FWorldOriginOffset PostWorldOriginOffset;
-
-	/** called when the main loop would otherwise starve. */
-	DECLARE_DELEGATE(FStarvedGameLoop);
-	static FStarvedGameLoop StarvedGameLoop;
-
-	/** IOS-style application lifecycle delegates */
-	DECLARE_MULTICAST_DELEGATE(FApplicationLifetimeDelegate);
-
-	// This is called when the application is about to be deactivated (e.g., due to a phone call or SMS or the sleep button).
-	// The game should be paused if possible, etc...
-	static FApplicationLifetimeDelegate ApplicationWillDeactivateDelegate;
-	
-	// Called when the application has been reactivated (reverse any processing done in the Deactivate delegate)
-	static FApplicationLifetimeDelegate ApplicationHasReactivatedDelegate;
-
-	// This is called when the application is being backgrounded (e.g., due to switching
-	// to another app or closing it via the home button)
-	// The game should release shared resources, save state, etc..., since it can be
-	// terminated from the background state without any further warning.
-	static FApplicationLifetimeDelegate ApplicationWillEnterBackgroundDelegate; // for instance, hitting the home button
-
-	// Called when the application is returning to the foreground (reverse any processing done in the EnterBackground delegate)
-	static FApplicationLifetimeDelegate ApplicationHasEnteredForegroundDelegate;
-
-	// This *may* be called when the application is getting terminated by the OS.
-	// There is no guarantee that this will ever be called on a mobile device,
-	// save state when ApplicationWillEnterBackgroundDelegate is called instead.
-	static FApplicationLifetimeDelegate ApplicationWillTerminateDelegate;
-
-	/** IOS-style push notification delegates */
-	DECLARE_MULTICAST_DELEGATE_OneParam(FApplicationRegisteredForRemoteNotificationsDelegate, TArray<uint8>);
-	DECLARE_MULTICAST_DELEGATE_OneParam(FApplicationRegisteredForUserNotificationsDelegate, int);
-	DECLARE_MULTICAST_DELEGATE_OneParam(FApplicationFailedToRegisterForRemoteNotificationsDelegate, FString);
-	DECLARE_MULTICAST_DELEGATE_OneParam(FApplicationReceivedRemoteNotificationDelegate, FString);
-	DECLARE_MULTICAST_DELEGATE_TwoParams(FApplicationReceivedLocalNotificationDelegate, FString, int);
-
-	// called when the user grants permission to register for remote notifications
-	static FApplicationRegisteredForRemoteNotificationsDelegate ApplicationRegisteredForRemoteNotificationsDelegate;
-
-	// called when the user grants permission to register for notifications
-	static FApplicationRegisteredForUserNotificationsDelegate ApplicationRegisteredForUserNotificationsDelegate;
-
-	// called when the application fails to register for remote notifications
-	static FApplicationFailedToRegisterForRemoteNotificationsDelegate ApplicationFailedToRegisterForRemoteNotificationsDelegate;
-
-	// called when the application receives a remote notification
-	static FApplicationReceivedRemoteNotificationDelegate ApplicationReceivedRemoteNotificationDelegate;
-
-	// called when the application receives a local notification
-	static FApplicationReceivedLocalNotificationDelegate ApplicationReceivedLocalNotificationDelegate;
-
-	/** Sent when a device screen orientation changes */
-	DECLARE_MULTICAST_DELEGATE_OneParam(FApplicationReceivedOnScreenOrientationChangedNotificationDelegate, int32);
-	static FApplicationReceivedOnScreenOrientationChangedNotificationDelegate ApplicationReceivedScreenOrientationChangedNotificationDelegate;
-
-	/** Checks to see if the stat is already enabled */
-	DECLARE_MULTICAST_DELEGATE_ThreeParams(FStatCheckEnabled, const TCHAR*, bool&, bool&);
-	static FStatCheckEnabled StatCheckEnabled;
-
-	/** Sent after each stat is enabled */
-	DECLARE_MULTICAST_DELEGATE_OneParam(FStatEnabled, const TCHAR*);
-	static FStatEnabled StatEnabled;
-
-	/** Sent after each stat is disabled */
-	DECLARE_MULTICAST_DELEGATE_OneParam(FStatDisabled, const TCHAR*);
-	static FStatDisabled StatDisabled;
-
-	/** Sent when all stats need to be disabled */
-	DECLARE_MULTICAST_DELEGATE_OneParam(FStatDisableAll, const bool);
-	static FStatDisableAll StatDisableAll;
-
-	// Called when an application is notified that the application license info has been updated.
-	// The new license data should be polled and steps taken based on the results (i.e. halt application if license is no longer valid).
-	DECLARE_MULTICAST_DELEGATE(FApplicationLicenseChange);
-	static FApplicationLicenseChange ApplicationLicenseChange;
-
-	/** Sent when the platform changed its laptop mode (for convertible laptops).*/
-	DECLARE_MULTICAST_DELEGATE_OneParam(FPlatformChangedLaptopMode, EConvertibleLaptopMode);
-	static FPlatformChangedLaptopMode PlatformChangedLaptopMode;
-
-	DECLARE_DELEGATE_RetVal_OneParam(bool, FLoadStringAssetReferenceInCook, const FString&);
-	static FLoadStringAssetReferenceInCook LoadStringAssetReferenceInCook;
-
-	DECLARE_DELEGATE_RetVal_OneParam(bool, FStringAssetReferenceLoaded, const FName&);
-	static FStringAssetReferenceLoaded StringAssetReferenceLoaded;
-
-	/** Sent when the platform requests a low-level VR recentering */
-	DECLARE_MULTICAST_DELEGATE(FVRHeadsetRecenter);
-	static FVRHeadsetRecenter VRHeadsetRecenter;
-
-	/** Sent when connection to VR HMD is lost */
-	DECLARE_MULTICAST_DELEGATE(FVRHeadsetLost);
-	static FVRHeadsetLost VRHeadsetLost;
-
-	/** Sent when connection to VR HMD is restored */
-	DECLARE_MULTICAST_DELEGATE(FVRHeadsetReconnected);
-	static FVRHeadsetReconnected VRHeadsetReconnected;
-
-	/** Sent when application code changes the user activity hint string for analytics, crash reports, etc */
-	DECLARE_MULTICAST_DELEGATE_OneParam(FOnUserActivityStringChanged, const FString&);
-	static FOnUserActivityStringChanged UserActivityStringChanged;
-
-	/** Sent when application code changes the currently active game session. The exact semantics of this will vary between games but it is useful for analytics, crash reports, etc  */
-	DECLARE_MULTICAST_DELEGATE_OneParam(FOnGameSessionIDChange, const FString&);
-	static FOnGameSessionIDChange GameSessionIDChanged;
-
-	/** Sent by application code to set params that customize crash reporting behavior */
-	DECLARE_MULTICAST_DELEGATE_OneParam(FOnCrashOverrideParamsChanged, const FCrashOverrideParameters&);
-	static FOnCrashOverrideParamsChanged CrashOverrideParamsChanged;
-	
-		// Callback for platform specific very early init code.
-	DECLARE_MULTICAST_DELEGATE(FOnPreMainInit);
-	static FOnPreMainInit& GetPreMainInitDelegate();
-	
-	/** Sent when GConfig is finished initializing */
-	DECLARE_MULTICAST_DELEGATE(FConfigReadyForUse);
-	static FConfigReadyForUse ConfigReadyForUse;
-
-	/** Callback for notifications regarding changes of the rendering thread. */
-	DECLARE_MULTICAST_DELEGATE(FRenderingThreadChanged)
-
-	/** Sent just after the rendering thread has been created. */
-	static FRenderingThreadChanged PostRenderingThreadCreated;
-	/* Sent just before the rendering thread is destroyed. */
-	static FRenderingThreadChanged PreRenderingThreadDestroyed;
-
-	// Called when appInit is called.
-	static FSimpleMulticastDelegate OnFEngineLoopInitComplete;
-
-	// Callback to allow custom resolution of package names. Arguments are InRequestedName, OutResolvedName.
-	// Should return True of resolution occured.
-	DECLARE_DELEGATE_RetVal_TwoParams(bool, FResolvePackageNameDelegate, const FString&, FString&);
-	static TArray<FResolvePackageNameDelegate> PackageNameResolvers;
-
-private:
-
-	// Callbacks for hotfixes
-	static TArray<FHotFixDelegate> HotFixDelegates;
-
-	// This class is only for namespace use
-	FCoreDelegates() {}
-};
+#include "CoreDelegates.h"
diff --git a/Engine/Source/Runtime/Core/Public/Misc/CommandLine.h b/Engine/Source/Runtime/Core/Public/Misc/CommandLine.h
new file mode 100644
--- /dev/null
+++ b/Engine/Source/Runtime/Core/Public/Misc/CommandLine.h
@@ -0,0 +1,111 @@
+// Copyright 1998-2016 Epic Games, Inc. All Rights Reserved.
+
+#pragma once
+
+#include "HAL/Platform.h"
+#include "Containers/Array.h"
+#include "Containers/UnrealString.h"
+
+/*-----------------------------------------------------------------------------
+	Command line.
+-----------------------------------------------------------------------------*/
+
+struct CORE_API FCommandLine
+{
+	/** maximum size of the command line */
+	static const uint32 MaxCommandLineSize = 16384;
+
+	/** 
+	 * Returns an edited version of the executable's command line with the game name and certain other parameters removed. 
+	 */
+	static const TCHAR* Get();
+	
+	/**
+	 * Returns an edited version of the executable's command line. 
+	 */
+	static const TCHAR* GetForLogging();
+
+	/**
+	 * Returns the command line originally passed to the executable.
+	 */
+	static const TCHAR* GetOriginal();
+
+	/**
+	 * Returns an edited version of the command line originally passed to the executable.
+	 */
+	static const TCHAR* GetOriginalForLogging();
+
+	/** 
+	 * Checks if the command line has been initialized. 
+	 */
+	static bool IsInitialized();
+
+	/**
+	 * Sets CmdLine to the string given
+	 */
+	static bool Set(const TCHAR* NewCommandLine);
+
+	/**
+	 * Appends the passed string to the command line as it is (no space is added).
+	 * @param AppendString String that should be appended to the commandline.
+	 */
+	static void Append(const TCHAR* AppendString);
+
+	/**
+	 * Adds a new parameter to subprocess command line. If Param
+	 * does not start with a space, one is added.
+	 *
+	 * @param Param Command line param string.
+	 */
+	static void AddToSubprocessCommandline( const TCHAR* Param );
+
+	/** 
+	 * Returns the subprocess command line string 
+	 */
+	static const FString& GetSubprocessCommandline();
+
+	/** 
+	* Removes the executable name from the passed CmdLine, denoted by parentheses.
+	* Returns the CmdLine string without the executable name.
+	*/
+	static const TCHAR* RemoveExeName(const TCHAR* CmdLine);
+
+	/**
+	 * Parses a string into tokens, separating switches (beginning with - or /) from
+	 * other parameters
+	 *
+	 * @param	CmdLine		the string to parse
+	 * @param	Tokens		[out] filled with all parameters found in the string
+	 * @param	Switches	[out] filled with all switches found in the string
+	 */
+	static void Parse(const TCHAR* CmdLine, TArray<FString>& Tokens, TArray<FString>& Switches);
+private:
+#if WANTS_COMMANDLINE_WHITELIST
+	/** Filters both the original and current command line list for approved only args */
+	static void WhitelistCommandLines();
+	/** Filters any command line args that aren't on the approved list */
+	static TArray<FString> FilterCommandLine(TCHAR* CommandLine);
+	/** Filters any command line args that are on the to-strip list */
+	static TArray<FString> FilterCommandLineForLogging(TCHAR* CommandLine);
+	/** Rebuilds the command line using the filtered args */
+	static void BuildWhitelistCommandLine(TCHAR* CommandLine, uint32 Length, const TArray<FString>& FilteredArgs);
+	static TArray<FString> ApprovedArgs;
+	static TArray<FString> FilterArgsForLogging;
+#else
+#define WhitelistCommandLines()
+#endif
+
+	/** Flag to check if the commandline has been initialized or not. */
+	static bool bIsInitialized;
+	/** character buffer containing the command line */
+	static TCHAR CmdLine[MaxCommandLineSize];
+	/** character buffer containing the original command line */
+	static TCHAR OriginalCmdLine[MaxCommandLineSize];
+	/** character buffer containing the command line filtered for logging purposes */
+	static TCHAR LoggingCmdLine[MaxCommandLineSize];
+	/** character buffer containing the original command line filtered for logging purposes */
+	static TCHAR LoggingOriginalCmdLine[MaxCommandLineSize];
+	/** subprocess command line */
+	static FString SubprocessCommandLine;
+};
+
diff --git a/Engine/Source/Runtime/Core/Public/Misc/CoreDelegates.h b/Engine/Source/Runtime/Core/Public/Misc/CoreDelegates.h
new file mode 100644
--- /dev/null
+++ b/Engine/Source/Runtime/Core/Public/Misc/CoreDelegates.h
@@ -0,0 +1,281 @@
+// Copyright 1998-2016 Epic Games, Inc. All Rights Reserved.
+
+#pragma once
+
+class AActor;
+
+// delegates for hotfixes
+namespace EHotfixDelegates
+{
+	enum Type
+	{
+		Test,
+	};
+}
+
+
+// this is an example of a hotfix arg and return value structure. Once we have other examples, it can be deleted.
+struct FTestHotFixPayload
+{
+	FString Message;
+	bool ValueToReturn;
+	bool Result;
+};
+
+// Parameters passed to CrashOverrideParamsChanged used to customize crash report client behavior/appearance
+struct FCrashOverrideParameters
+{
+	FString CrashReportClientMessageText;
+};
+
+
+class CORE_API FCoreDelegates
+{
+public:
+	//hot fix delegate
+	DECLARE_DELEGATE_TwoParams(FHotFixDelegate, void *, int32);
+
+	// Callback for object property modifications
+	DECLARE_MULTICAST_DELEGATE_OneParam(FOnActorLabelChanged, AActor*);
+
+	// delegate type for prompting the pak system to mount a new pak
+	DECLARE_DELEGATE_RetVal_ThreeParams(bool, FOnMountPak, const FString&, uint32, IPlatformFile::FDirectoryVisitor*);
+
+	// delegate type for prompting the pak system to mount a new pak
+	DECLARE_DELEGATE_RetVal_OneParam(bool, FOnUnmountPak, const FString&);
+
+	/** delegate type for opening a modal message box ( Params: EAppMsgType::Type MessageType, const FText& Text, const FText& Title ) */
+	DECLARE_DELEGATE_RetVal_ThreeParams(EAppReturnType::Type, FOnModalMessageBox, EAppMsgType::Type, const FText&, const FText&);
+
+	// Callback for PER_MODULE_BOILERPLATE macro's GObjectArrayForDebugVisualizers
+	DECLARE_DELEGATE_RetVal(class FFixedUObjectArray*, FObjectArrayForDebugVisualizersDelegate);
+
+	// Called in PER_MODULE_BOILERPLATE macro.
+	static FObjectArrayForDebugVisualizersDelegate& GetObjectArrayForDebugVisualizersDelegate();
+
+	// Callback for handling an ensure
+	DECLARE_MULTICAST_DELEGATE(FOnHandleSystemEnsure);
+
+	// Callback for handling an error
+	DECLARE_MULTICAST_DELEGATE(FOnHandleSystemError);
+
+	// Callback for handling user login/logout.  first int is UserID, second int is UserIndex
+	DECLARE_MULTICAST_DELEGATE_ThreeParams(FOnUserLoginChangedEvent, bool, int32, int32);
+
+	// Callback for handling safe frame area size changes
+	DECLARE_MULTICAST_DELEGATE(FOnSafeFrameChangedEvent);
+
+	// Callback for handling accepting invitations - generally for engine code
+	DECLARE_MULTICAST_DELEGATE_TwoParams(FOnInviteAccepted, const FString&, const FString&);
+
+	// Callback for handling the Controller connection / disconnection
+	// first param is true for a connection, false for a disconnection.
+	// second param is UserID, third is UserIndex / ControllerId.
+	DECLARE_MULTICAST_DELEGATE_ThreeParams(FOnUserControllerConnectionChange, bool, int32, int32);
+
+	// Callback for platform handling when flushing async loads.
+	DECLARE_MULTICAST_DELEGATE(FOnAsyncLoadingFlush);
+	static FOnAsyncLoadingFlush OnAsyncLoadingFlush;
+
+	// get a hotfix delegate
+	static FHotFixDelegate& GetHotfixDelegate(EHotfixDelegates::Type HotFix);
+
+	// Callback when a user logs in/out of the platform.
+	static FOnUserLoginChangedEvent OnUserLoginChangedEvent;
+
+	// Callback when controllers disconnected / reconnected
+	static FOnUserControllerConnectionChange OnControllerConnectionChange;
+
+	// Callback when a user changes the safe frame size
+	static FOnSafeFrameChangedEvent OnSafeFrameChangedEvent;
+
+	// Callback for mounting a new pak file.
+	static FOnMountPak OnMountPak;
+
+	// Callback for unmounting a pak file.
+	static FOnUnmountPak OnUnmountPak;
+
+	// Callback when an ensure has occurred
+	static FOnHandleSystemEnsure OnHandleSystemEnsure;
+	// Callback when an error (crash) has occurred
+	static FOnHandleSystemError OnHandleSystemError;
+
+	// Called when an actor label is changed
+	static FOnActorLabelChanged OnActorLabelChanged;
+
+#if WITH_EDITOR
+	// Called before the editor displays a modal window, allowing other windows the opportunity to disable themselves to avoid reentrant calls
+	static FSimpleMulticastDelegate PreModal;
+
+	// Called after the editor dismisses a modal window, allowing other windows the opportunity to disable themselves to avoid reentrant calls
+	static FSimpleMulticastDelegate PostModal;
+#endif	//WITH_EDITOR
+	
+	// Called when an error occurred.
+	static FSimpleMulticastDelegate OnShutdownAfterError;
+
+	// Called when appInit is called.
+	static FSimpleMulticastDelegate OnInit;
+
+	// Called when the application is about to exit.
+	static FSimpleMulticastDelegate OnExit;
+
+	// Called when before the application is exiting.
+	static FSimpleMulticastDelegate OnPreExit;
+
+	/** Color picker color has changed, please refresh as needed*/
+	static FSimpleMulticastDelegate ColorPickerChanged;
+
+	/** requests to open a message box */
+	static FOnModalMessageBox ModalErrorMessage;
+
+	/** Called when the user accepts an invitation to the current game */
+	static FOnInviteAccepted OnInviteAccepted;
+
+	DECLARE_MULTICAST_DELEGATE_ThreeParams(FWorldOriginOffset, class UWorld*, FIntVector, FIntVector);
+	/** called before world origin shifting */
+	static FWorldOriginOffset PreWorldOriginOffset;
+	/** called after world origin shifting */
+	static FWorldOriginOffset PostWorldOriginOffset;
+
+	/** called when the main loop would otherwise starve. */
+	DECLARE_DELEGATE(FStarvedGameLoop);
+	static FStarvedGameLoop StarvedGameLoop;
+
+	/** IOS-style application lifecycle delegates */
+	DECLARE_MULTICAST_DELEGATE(FApplicationLifetimeDelegate);
+
+	// This is called when the application is about to be deactivated (e.g., due to a phone call or SMS or the sleep button).
+	// The game should be paused if possible, etc...
+	static FApplicationLifetimeDelegate ApplicationWillDeactivateDelegate;
+	
+	// Called when the application has been reactivated (reverse any processing done in the Deactivate delegate)
+	static FApplicationLifetimeDelegate ApplicationHasReactivatedDelegate;
+
+	// This is called when the application is being backgrounded (e.g., due to switching
+	// to another app or closing it via the home button)
+	// The game should release shared resources, save state, etc..., since it can be
+	// terminated from the background state without any further warning.
+	static FApplicationLifetimeDelegate ApplicationWillEnterBackgroundDelegate; // for instance, hitting the home button
+
+	// Called when the application is returning to the foreground (reverse any processing done in the EnterBackground delegate)
+	static FApplicationLifetimeDelegate ApplicationHasEnteredForegroundDelegate;
+
+	// This *may* be called when the application is getting terminated by the OS.
+	// There is no guarantee that this will ever be called on a mobile device,
+	// save state when ApplicationWillEnterBackgroundDelegate is called instead.
+	static FApplicationLifetimeDelegate ApplicationWillTerminateDelegate;
+
+	/** IOS-style push notification delegates */
+	DECLARE_MULTICAST_DELEGATE_OneParam(FApplicationRegisteredForRemoteNotificationsDelegate, TArray<uint8>);
+	DECLARE_MULTICAST_DELEGATE_OneParam(FApplicationRegisteredForUserNotificationsDelegate, int);
+	DECLARE_MULTICAST_DELEGATE_OneParam(FApplicationFailedToRegisterForRemoteNotificationsDelegate, FString);
+	DECLARE_MULTICAST_DELEGATE_OneParam(FApplicationReceivedRemoteNotificationDelegate, FString);
+	DECLARE_MULTICAST_DELEGATE_TwoParams(FApplicationReceivedLocalNotificationDelegate, FString, int);
+
+	// called when the user grants permission to register for remote notifications
+	static FApplicationRegisteredForRemoteNotificationsDelegate ApplicationRegisteredForRemoteNotificationsDelegate;
+
+	// called when the user grants permission to register for notifications
+	static FApplicationRegisteredForUserNotificationsDelegate ApplicationRegisteredForUserNotificationsDelegate;
+
+	// called when the application fails to register for remote notifications
+	static FApplicationFailedToRegisterForRemoteNotificationsDelegate ApplicationFailedToRegisterForRemoteNotificationsDelegate;
+
+	// called when the application receives a remote notification
+	static FApplicationReceivedRemoteNotificationDelegate ApplicationReceivedRemoteNotificationDelegate;
+
+	// called when the application receives a local notification
+	static FApplicationReceivedLocalNotificationDelegate ApplicationReceivedLocalNotificationDelegate;
+
+	/** Sent when a device screen orientation changes */
+	DECLARE_MULTICAST_DELEGATE_OneParam(FApplicationReceivedOnScreenOrientationChangedNotificationDelegate, int32);
+	static FApplicationReceivedOnScreenOrientationChangedNotificationDelegate ApplicationReceivedScreenOrientationChangedNotificationDelegate;
+
+	/** Checks to see if the stat is already enabled */
+	DECLARE_MULTICAST_DELEGATE_ThreeParams(FStatCheckEnabled, const TCHAR*, bool&, bool&);
+	static FStatCheckEnabled StatCheckEnabled;
+
+	/** Sent after each stat is enabled */
+	DECLARE_MULTICAST_DELEGATE_OneParam(FStatEnabled, const TCHAR*);
+	static FStatEnabled StatEnabled;
+
+	/** Sent after each stat is disabled */
+	DECLARE_MULTICAST_DELEGATE_OneParam(FStatDisabled, const TCHAR*);
+	static FStatDisabled StatDisabled;
+
+	/** Sent when all stats need to be disabled */
+	DECLARE_MULTICAST_DELEGATE_OneParam(FStatDisableAll, const bool);
+	static FStatDisableAll StatDisableAll;
+
+	// Called when an application is notified that the application license info has been updated.
+	// The new license data should be polled and steps taken based on the results (i.e. halt application if license is no longer valid).
+	DECLARE_MULTICAST_DELEGATE(FApplicationLicenseChange);
+	static FApplicationLicenseChange ApplicationLicenseChange;
+
+	/** Sent when the platform changed its laptop mode (for convertible laptops).*/
+	DECLARE_MULTICAST_DELEGATE_OneParam(FPlatformChangedLaptopMode, EConvertibleLaptopMode);
+	static FPlatformChangedLaptopMode PlatformChangedLaptopMode;
+
+	DECLARE_DELEGATE_RetVal_OneParam(bool, FLoadStringAssetReferenceInCook, const FString&);
+	static FLoadStringAssetReferenceInCook LoadStringAssetReferenceInCook;
+
+	DECLARE_DELEGATE_RetVal_OneParam(bool, FStringAssetReferenceLoaded, const FName&);
+	static FStringAssetReferenceLoaded StringAssetReferenceLoaded;
+
+	/** Sent when the platform requests a low-level VR recentering */
+	DECLARE_MULTICAST_DELEGATE(FVRHeadsetRecenter);
+	static FVRHeadsetRecenter VRHeadsetRecenter;
+
+	/** Sent when connection to VR HMD is lost */
+	DECLARE_MULTICAST_DELEGATE(FVRHeadsetLost);
+	static FVRHeadsetLost VRHeadsetLost;
+
+	/** Sent when connection to VR HMD is restored */
+	DECLARE_MULTICAST_DELEGATE(FVRHeadsetReconnected);
+	static FVRHeadsetReconnected VRHeadsetReconnected;
+
+	/** Sent when application code changes the user activity hint string for analytics, crash reports, etc */
+	DECLARE_MULTICAST_DELEGATE_OneParam(FOnUserActivityStringChanged, const FString&);
+	static FOnUserActivityStringChanged UserActivityStringChanged;
+
+	/** Sent when application code changes the currently active game session. The exact semantics of this will vary between games but it is useful for analytics, crash reports, etc  */
+	DECLARE_MULTICAST_DELEGATE_OneParam(FOnGameSessionIDChange, const FString&);
+	static FOnGameSessionIDChange GameSessionIDChanged;
+
+	/** Sent by application code to set params that customize crash reporting behavior */
+	DECLARE_MULTICAST_DELEGATE_OneParam(FOnCrashOverrideParamsChanged, const FCrashOverrideParameters&);
+	static FOnCrashOverrideParamsChanged CrashOverrideParamsChanged;
+	
+		// Callback for platform specific very early init code.
+	DECLARE_MULTICAST_DELEGATE(FOnPreMainInit);
+	static FOnPreMainInit& GetPreMainInitDelegate();
+	
+	/** Sent when GConfig is finished initializing */
+	DECLARE_MULTICAST_DELEGATE(FConfigReadyForUse);
+	static FConfigReadyForUse ConfigReadyForUse;
+
+	/** Callback for notifications regarding changes of the rendering thread. */
+	DECLARE_MULTICAST_DELEGATE(FRenderingThreadChanged)
+
+	/** Sent just after the rendering thread has been created. */
+	static FRenderingThreadChanged PostRenderingThreadCreated;
+	/* Sent just before the rendering thread is destroyed. */
+	static FRenderingThreadChanged PreRenderingThreadDestroyed;
+
+	// Called when appInit is called.
+	static FSimpleMulticastDelegate OnFEngineLoopInitComplete;
+
+	// Callback to allow custom resolution of package names. Arguments are InRequestedName, OutResolvedName.
+	// Should return True of resolution occured.
+	DECLARE_DELEGATE_RetVal_TwoParams(bool, FResolvePackageNameDelegate, const FString&, FString&);
+	static TArray<FResolvePackageNameDelegate> PackageNameResolvers;
+
+private:
+
+	// Callbacks for hotfixes
+	static TArray<FHotFixDelegate> HotFixDelegates;
+
+	// This class is only for namespace use
+	FCoreDelegates() {}
+};
diff --git a/Engine/Source/Runtime/Core/Public/Misc/CoreMisc.h b/Engine/Source/Runtime/Core/Public/Misc/CoreMisc.h
--- a/Engine/Source/Runtime/Core/Public/Misc/CoreMisc.h
+++ b/Engine/Source/Runtime/Core/Public/Misc/CoreMisc.h
@@ -59,209 +59,6 @@ struct CORE_API FMaintenance
 	static void DeleteOldLogs();
 };
 
-/*-----------------------------------------------------------------------------
-	Command line.
------------------------------------------------------------------------------*/
-
-struct CORE_API FCommandLine
-{
-	/** maximum size of the command line */
-	static const uint32 MaxCommandLineSize = 16384;
-
-	/** 
-	 * Returns an edited version of the executable's command line with the game name and certain other parameters removed. 
-	 */
-	static const TCHAR* Get();
-	
-	/**
-	 * Returns an edited version of the executable's command line. 
-	 */
-	static const TCHAR* GetForLogging();
-
-	/**
-	 * Returns the command line originally passed to the executable.
-	 */
-	static const TCHAR* GetOriginal();
-
-	/**
-	 * Returns an edited version of the command line originally passed to the executable.
-	 */
-	static const TCHAR* GetOriginalForLogging();
-
-	/** 
-	 * Checks if the command line has been initialized. 
-	 */
-	static bool IsInitialized();
-
-	/**
-	 * Sets CmdLine to the string given
-	 */
-	static bool Set(const TCHAR* NewCommandLine);
-
-	/**
-	 * Appends the passed string to the command line as it is (no space is added).
-	 * @param AppendString String that should be appended to the commandline.
-	 */
-	static void Append(const TCHAR* AppendString);
-
-	/**
-	 * Adds a new parameter to subprocess command line. If Param
-	 * does not start with a space, one is added.
-	 *
-	 * @param Param Command line param string.
-	 */
-	static void AddToSubprocessCommandline( const TCHAR* Param );
-
-	/** 
-	 * Returns the subprocess command line string 
-	 */
-	static const FString& GetSubprocessCommandline();
-
-	/** 
-	* Removes the executable name from the passed CmdLine, denoted by parentheses.
-	* Returns the CmdLine string without the executable name.
-	*/
-	static const TCHAR* RemoveExeName(const TCHAR* CmdLine);
-
-	/**
-	 * Parses a string into tokens, separating switches (beginning with - or /) from
-	 * other parameters
-	 *
-	 * @param	CmdLine		the string to parse
-	 * @param	Tokens		[out] filled with all parameters found in the string
-	 * @param	Switches	[out] filled with all switches found in the string
-	 */
-	static void Parse(const TCHAR* CmdLine, TArray<FString>& Tokens, TArray<FString>& Switches);
-private:
-#if WANTS_COMMANDLINE_WHITELIST
-	/** Filters both the original and current command line list for approved only args */
-	static void WhitelistCommandLines();
-	/** Filters any command line args that aren't on the approved list */
-	static TArray<FString> FilterCommandLine(TCHAR* CommandLine);
-	/** Filters any command line args that are on the to-strip list */
-	static TArray<FString> FilterCommandLineForLogging(TCHAR* CommandLine);
-	/** Rebuilds the command line using the filtered args */
-	static void BuildWhitelistCommandLine(TCHAR* CommandLine, uint32 Length, const TArray<FString>& FilteredArgs);
-	static TArray<FString> ApprovedArgs;
-	static TArray<FString> FilterArgsForLogging;
-#else
-#define WhitelistCommandLines()
-#endif
-
-	/** Flag to check if the commandline has been initialized or not. */
-	static bool bIsInitialized;
-	/** character buffer containing the command line */
-	static TCHAR CmdLine[MaxCommandLineSize];
-	/** character buffer containing the original command line */
-	static TCHAR OriginalCmdLine[MaxCommandLineSize];
-	/** character buffer containing the command line filtered for logging purposes */
-	static TCHAR LoggingCmdLine[MaxCommandLineSize];
-	/** character buffer containing the original command line filtered for logging purposes */
-	static TCHAR LoggingOriginalCmdLine[MaxCommandLineSize];
-	/** subprocess command line */
-	static FString SubprocessCommandLine;
-};
-
-/*-----------------------------------------------------------------------------
-	FFileHelper
------------------------------------------------------------------------------*/
-struct CORE_API FFileHelper
-{
-	struct EHashOptions
-	{
-		enum Type
-		{
-			/** Enable the async task for verifying the hash for the file being loaded */
-			EnableVerify		=1<<0,
-			/** A missing hash entry should trigger an error */
-			ErrorMissingHash	=1<<1
-		};
-	};
-
-	struct EEncodingOptions
-	{
-		enum Type
-		{
-			AutoDetect,
-			ForceAnsi,
-			ForceUnicode,
-			ForceUTF8,
-			ForceUTF8WithoutBOM
-		};
-	};
-
-	/**
-	 * Load a text file to an FString.
-	 * Supports all combination of ANSI/Unicode files and platforms.
-	*/
-	static void BufferToString( FString& Result, const uint8* Buffer, int32 Size );
-
-	/**
-	 * Load a binary file to a dynamic array.
-	*/
-	static bool LoadFileToArray( TArray<uint8>& Result, const TCHAR* Filename, uint32 Flags = 0 );
-
-	/**
-	 * Load a text file to an FString.
-	 * Supports all combination of ANSI/Unicode files and platforms.
-	 * @param Result string representation of the loaded file
-	 * @param Filename name of the file to load
-	 * @param VerifyFlags flags controlling the hash verification behavior ( see EHashOptions )
-	 */
-	static bool LoadFileToString( FString& Result, const TCHAR* Filename, uint32 VerifyFlags=0 );
-
-	/**
-	 * Save a binary array to a file.
-	 */
-	static bool SaveArrayToFile(TArrayView<const uint8> Array, const TCHAR* Filename, IFileManager* FileManager=&IFileManager::Get(), uint32 WriteFlags = 0);
-
-	/**
-	 * Write the FString to a file.
-	 * Supports all combination of ANSI/Unicode files and platforms.
-	 */
-	static bool SaveStringToFile( const FString& String, const TCHAR* Filename, EEncodingOptions::Type EncodingOptions=EEncodingOptions::AutoDetect, IFileManager* FileManager=&IFileManager::Get(), uint32 WriteFlags = 0 );
-
-	/**
-	 * Saves a 24/32Bit BMP file to disk
-	 * 
-	 * @param Pattern filename with path, must not be 0, if with "bmp" extension (e.g. "out.bmp") the filename stays like this, if without (e.g. "out") automatic index numbers are addended (e.g. "out00002.bmp")
-	 * @param DataWidth - Width of the bitmap supplied in Data >0
-	 * @param DataHeight - Height of the bitmap supplied in Data >0
-	 * @param Data must not be 0
-	 * @param SubRectangle optional, specifies a sub-rectangle of the source image to save out. If NULL, the whole bitmap is saved
-	 * @param FileManager must not be 0
-	 * @param OutFilename optional, if specified filename will be output
-	 * @param bInWriteAlpha opional, specifies whether to write out the alpha channel. Will force BMP V4 format.
-	 *
-	 * @return true if success
-	 */
-	static bool CreateBitmap( const TCHAR* Pattern, int32 DataWidth, int32 DataHeight, const struct FColor* Data, struct FIntRect* SubRectangle = NULL, IFileManager* FileManager = &IFileManager::Get(), FString* OutFilename = NULL, bool bInWriteAlpha = false );
-
-	/**
-	 * Generates the next unique bitmap filename with a specified extension
-	 *
-	 * @param Pattern		Filename with path, but without extension.
-	 * @oaran Extension		File extension to be appended
-	 * @param OutFilename	Reference to an FString where the newly generated filename will be placed
-	 * @param FileManager	Reference to a IFileManager (or the global instance by default)
-	 *
-	 * @return true if success
-	 */
-	static bool GenerateNextBitmapFilename(const FString& Pattern, const FString& Extension, FString& OutFilename, IFileManager* FileManager = &IFileManager::Get());
-	
-	/**
-	 *	Load the given ANSI text file to an array of strings - one FString per line of the file.
-	 *	Intended for use in simple text parsing actions
-	 *
-	 *	@param	InFilename			The text file to read, full path
-	 *	@param	InFileManager		The filemanager to use - NULL will use &IFileManager::Get()
-	 *	@param	OutStrings			The array of FStrings to fill in
-	 *
-	 *	@return	bool				true if successful, false if not
-	 */
-	static bool LoadANSITextFileToStrings(const TCHAR* InFilename, IFileManager* InFileManager, TArray<FString>& OutStrings);
-};
-
 /*-----------------------------------------------------------------------------
 	Module singletons.
 -----------------------------------------------------------------------------*/
diff --git a/Engine/Source/Runtime/Core/Public/Misc/CoreStats.h b/Engine/Source/Runtime/Core/Public/Misc/CoreStats.h
--- a/Engine/Source/Runtime/Core/Public/Misc/CoreStats.h
+++ b/Engine/Source/Runtime/Core/Public/Misc/CoreStats.h
@@ -5,17 +5,6 @@
 #include "Stats.h"
 
 
-CORE_API DECLARE_LOG_CATEGORY_EXTERN(LogSHA, Warning, All);
-CORE_API DECLARE_LOG_CATEGORY_EXTERN(LogStats, Log, All);
-CORE_API DECLARE_LOG_CATEGORY_EXTERN(LogStreaming, Display, All);
-CORE_API DECLARE_LOG_CATEGORY_EXTERN(LogInit, Log, All);
-CORE_API DECLARE_LOG_CATEGORY_EXTERN(LogExit, Log, All);
-CORE_API DECLARE_LOG_CATEGORY_EXTERN(LogExec, Warning, All);
-CORE_API DECLARE_LOG_CATEGORY_EXTERN(LogScript, Warning, All);
-CORE_API DECLARE_LOG_CATEGORY_EXTERN(LogLocalization, Error, All);
-CORE_API DECLARE_LOG_CATEGORY_EXTERN(LogLongPackageNames, Log, All);
-CORE_API DECLARE_LOG_CATEGORY_EXTERN(LogProcess, Log, All);
-CORE_API DECLARE_LOG_CATEGORY_EXTERN(LogLoad, Log, All);
 
 
 /** Memory stats */
diff --git a/Engine/Source/Runtime/Core/Public/Misc/Exec.h b/Engine/Source/Runtime/Core/Public/Misc/Exec.h
new file mode 100644
--- /dev/null
+++ b/Engine/Source/Runtime/Core/Public/Misc/Exec.h
@@ -0,0 +1,25 @@
+// Copyright 1998-2016 Epic Games, Inc. All Rights Reserved.
+
+#pragma once
+
+#include "HAL/Platform.h"
+
+// Any object that is capable of taking commands.
+class CORE_API FExec
+{
+public:
+	virtual ~FExec();
+
+	/**
+	* Exec handler
+	*
+	* @param	InWorld	World context
+	* @param	Cmd		Command to parse
+	* @param	Ar		Output device to log to
+	*
+	* @return	true if command was handled, false otherwise
+	*/
+	virtual bool Exec( class UWorld* InWorld, const TCHAR* Cmd, FOutputDevice& Ar ) PURE_VIRTUAL(FExec::Exec,return false;)
+};
+
+
