Copying //UE4/Dev-Build to //UE4/Dev-Main (Source: //UE4/Dev-Build @ 3104200)

diff --git a/Engine/Source/Runtime/Core/Public/Misc/FeedbackContext.h b/Engine/Source/Runtime/Core/Public/Misc/FeedbackContext.h
--- a/Engine/Source/Runtime/Core/Public/Misc/FeedbackContext.h
+++ b/Engine/Source/Runtime/Core/Public/Misc/FeedbackContext.h
@@ -4,17 +4,7 @@
 
 #include "UniquePtr.h"
 #include "OutputDeviceHelper.h"
-
-struct FSlowTask;
-
-
-/** A stack of feedback scopes */
-struct FSlowTaskStack : TArray<FSlowTask*>
-{
-	/** Get the fraction of work that has been completed for the specified index in the stack (0=total progress) */
-	CORE_API float GetProgressFraction(int32 Index) const;
-};
-
+#include "SlowTaskStack.h"
 
 /** A context for displaying modal warning messages. */
 class CORE_API FFeedbackContext
@@ -80,10 +70,7 @@ public:
 
 	bool	TreatWarningsAsErrors;
 
-	FFeedbackContext()
-		: TreatWarningsAsErrors(0)
-		, ScopeStack(MakeShareable(new FSlowTaskStack))
-	{}
+	FFeedbackContext();
 
 	virtual ~FFeedbackContext();
 
@@ -169,234 +156,3 @@ protected:
 		Errors.Add(InError);
 	}
 };
-
-/** Enum to specify a particular slow task section should be shown */
-enum class ESlowTaskVisibility
-{
-	/** Default visibility (inferred by some heuristic of remaining work/time open) */
-	Default,
-	/** Force this particular slow task to be visible on the UI */
-	ForceVisible,
-	/** Forcibly prevent this slow task from being shown, but still use it for work progress calculations */
-	Invisible,
-};
-
-/**
- * Data type used to store information about a currently running slow task. Direct use is not advised, use FScopedSlowTask instead
- */
-struct CORE_API FSlowTask
-{
-	/** Default message to display to the user when not overridden by a frame */
-	FText DefaultMessage;
-
-	/** Message pertaining to the current frame's work */
-	FText FrameMessage;
-
-	/** The amount of work to do in this scope */
-	float TotalAmountOfWork;
-
-	/** The amount of work we have already completed in this scope */
-	float CompletedWork;
-
-	/** The amount of work the current frame is responsible for */
-	float CurrentFrameScope;
-
-	/** The visibility of this slow task */
-	ESlowTaskVisibility Visibility;
-
-	/** The time that this scope was created */
-	double StartTime;
-
-	/** Threshold before dialog is opened */
-	TOptional<float> OpenDialogThreshold;
-
-private:
-
-	/** Boolean flag to control whether this scope actually does anything (unset for quiet operations etc) */
-	bool bEnabled;
-
-	/** Flag that specifies whether this feedback scope created a new slow task dialog */
-	bool bCreatedDialog;
-	
-	/** The feedback context that we belong to */
-	FFeedbackContext& Context;
-
-	/** Specify whether the delayed dialog should show a cancel button */
-	bool bDelayedDialogShowCancelButton : 1;
-
-	/** Specify whether the delayed dialog is allowed in PIE */
-	bool bDelayedDialogAllowInPIE : 1;
-
-	/** Prevent copying */
-	FSlowTask(const FSlowTask&);
-
-public:
-
-	/**
-	 * Construct this scope from an amount of work to do, and a message to display
-	 * @param		InAmountOfWork			Arbitrary number of work units to perform (can be a percentage or number of steps).
-	 *										0 indicates that no progress frames are to be entered in this scope (automatically enters a frame encompassing the entire scope)
-	 * @param		InDefaultMessage		A message to display to the user to describe the purpose of the scope
-	 * @param		bInVisible				When false, this scope will have no effect. Allows for proper scoped objects that are conditionally hidden.
-	 */
-	FORCEINLINE FSlowTask(float InAmountOfWork, const FText& InDefaultMessage = FText(), bool bInEnabled = true, FFeedbackContext& InContext = *GWarn)
-		: DefaultMessage(InDefaultMessage)
-		, FrameMessage()
-		, TotalAmountOfWork(InAmountOfWork)
-		, CompletedWork(0)
-		, CurrentFrameScope(0)
-		, Visibility(ESlowTaskVisibility::Default)
-		, StartTime(FPlatformTime::Seconds())
-		, bEnabled(bInEnabled && IsInGameThread())
-		, bCreatedDialog(false)		// only set to true if we create a dialog
-		, Context(InContext)
-	{
-		// If we have no work to do ourselves, create an arbitrary scope so that any actions performed underneath this still contribute to this one.
-		if (TotalAmountOfWork == 0.f)
-		{
-			TotalAmountOfWork = CurrentFrameScope = 1.f;
-		}
-	}
-
-	/** Function that initializes the scope by adding it to its context's stack */
-	FORCEINLINE void Initialize()
-	{
-		if (bEnabled)
-		{
-			Context.ScopeStack->Push(this);
-		}
-	}
-
-	/** Function that finishes any remaining work and removes itself from the global scope stack */
-	FORCEINLINE void Destroy()
-	{
-		if (bEnabled)
-		{
-			if (bCreatedDialog)
-			{
-				checkSlow(GIsSlowTask);
-				Context.FinalizeSlowTask();
-			}
-
-			FSlowTaskStack& Stack = *Context.ScopeStack;
-			checkSlow(Stack.Num() != 0 && Stack.Last() == this);
-
-			auto* Task = Stack.Last();
-			if (ensureMsgf(Task == this, TEXT("Out-of-order scoped task construction/destruction")))
-			{
-				Stack.Pop(false);
-			}
-			else
-			{
-				Stack.RemoveSingleSwap(this, false);
-			}
-
-			if (Stack.Num() != 0)
-			{
-				// Stop anything else contributing to the parent frame
-				auto* Parent = Stack.Last();
-				Parent->EnterProgressFrame(0, Parent->FrameMessage);
-			}
-
-		}
-	}
-
-	/**
-	 * Creates a new dialog for this slow task after the given time threshold. If the task completes before this time, no dialog will be shown.
-	 * @param		Threshold				Time in seconds before dialog will be shown.
-	 * @param		bShowCancelButton		Whether to show a cancel button on the dialog or not
-	 * @param		bAllowInPIE				Whether to allow this dialog in PIE. If false, this dialog will not appear during PIE sessions.
-	 */
-	FORCEINLINE void MakeDialogDelayed(float Threshold, bool bShowCancelButton = false, bool bAllowInPIE = false)
-	{
-		OpenDialogThreshold = Threshold;
-		bDelayedDialogShowCancelButton = bShowCancelButton;
-		bDelayedDialogAllowInPIE = bAllowInPIE;
-	}
-
-	/**
-	 * Creates a new dialog for this slow task, if there is currently not one open
-	 * @param		bShowCancelButton		Whether to show a cancel button on the dialog or not
-	 * @param		bAllowInPIE				Whether to allow this dialog in PIE. If false, this dialog will not appear during PIE sessions.
-	 */
-	void MakeDialog(bool bShowCancelButton = false, bool bAllowInPIE = false);
-
-	/**
-	 * Indicate that we are to enter a frame that will take up the specified amount of work. Completes any previous frames (potentially contributing to parent scopes' progress).
-	 * @param		ExpectedWorkThisFrame	The amount of work that will happen between now and the next frame, as a numerator of TotalAmountOfWork.
-	 * @param		Text					Optional text to describe this frame's purpose.
-	 */
-	FORCEINLINE void EnterProgressFrame(float ExpectedWorkThisFrame = 1.f, FText Text = FText())
-	{
-		FrameMessage = Text;
-		CompletedWork += CurrentFrameScope;
-
-		const float WorkRemaining = TotalAmountOfWork - CompletedWork;
-		// Add a small threshold here because when there are a lot of tasks, numerical imprecision can add up and trigger this.
-		ensureMsgf(ExpectedWorkThisFrame <= 1.01f * TotalAmountOfWork - CompletedWork, TEXT("Work overflow in slow task. Please revise call-site to account for entire progress range."));
-		CurrentFrameScope = FMath::Min(WorkRemaining, ExpectedWorkThisFrame);
-
-		if (!bCreatedDialog && OpenDialogThreshold.IsSet() && static_cast<float>(FPlatformTime::Seconds() - StartTime) > OpenDialogThreshold.GetValue())
-		{
-			MakeDialog(bDelayedDialogShowCancelButton, bDelayedDialogAllowInPIE);
-		}
-
-		if (bEnabled)
-		{
-			Context.RequestUpdateUI(bCreatedDialog || (*Context.ScopeStack)[0] == this);
-		}
-	}
-
-	/**
-	 * Get the frame message or default message if empty
-	 */
-	FText GetCurrentMessage() const
-	{
-		return FrameMessage.IsEmpty() ? DefaultMessage : FrameMessage;
-	}
-};
-
-
-/**
- * A scope block representing an amount of work divided up into sections.
- * Use one scope at the top of each function to give accurate feedback to the user of a slow operation's progress.
- *
- * Example Usage:
- *	void DoSlowWork()
- *	{
- *		FScopedSlowTask Progress(2.f, LOCTEXT("DoingSlowWork", "Doing Slow Work..."));
- *		// Optionally make this show a dialog if not already shown
- *		Progress.MakeDialog();
- *
- *		// Indicate that we are entering a frame representing 1 unit of work
- *		Progress.EnterProgressFrame(1.f);
- *		
- *		// DoFirstThing() can follow a similar pattern of creating a scope divided into frames. These contribute to their parent's progress frame proportionately.
- *		DoFirstThing();
- *		
- *		Progress.EnterProgressFrame(1.f);
- *		DoSecondThing();
- *	}
- *
- */
-struct FScopedSlowTask : FSlowTask
-{
-
-	/**
-	 * Construct this scope from an amount of work to do, and a message to display
-	 * @param		InAmountOfWork			Arbitrary number of work units to perform (can be a percentage or number of steps).
-	 *										0 indicates that no progress frames are to be entered in this scope (automatically enters a frame encompassing the entire scope)
-	 * @param		InDefaultMessage		A message to display to the user to describe the purpose of the scope
-	 * @param		bInEnabled				When false, this scope will have no effect. Allows for proper scoped objects that are conditionally disabled.
-	 */
-	FORCEINLINE FScopedSlowTask(float InAmountOfWork, const FText& InDefaultMessage = FText(), bool bInEnabled = true, FFeedbackContext& InContext = *GWarn)
-		: FSlowTask(InAmountOfWork, InDefaultMessage, bInEnabled, InContext)
-	{
-		Initialize();
-	}
-
-	FORCEINLINE ~FScopedSlowTask()
-	{
-		Destroy();
-	}
-};
\ No newline at end of file
diff --git a/Engine/Source/Runtime/Core/Public/Misc/FileHelper.h b/Engine/Source/Runtime/Core/Public/Misc/FileHelper.h
new file mode 100644
--- /dev/null
+++ b/Engine/Source/Runtime/Core/Public/Misc/FileHelper.h
@@ -0,0 +1,107 @@
+// Copyright 1998-2016 Epic Games, Inc. All Rights Reserved.
+
+#pragma once
+
+#include "HAL/Platform.h"
+#include "Containers/UnrealString.h"
+#include "Containers/Array.h"
+
+/*-----------------------------------------------------------------------------
+	FFileHelper
+-----------------------------------------------------------------------------*/
+struct CORE_API FFileHelper
+{
+	struct EHashOptions
+	{
+		enum Type
+		{
+			/** Enable the async task for verifying the hash for the file being loaded */
+			EnableVerify		=1<<0,
+			/** A missing hash entry should trigger an error */
+			ErrorMissingHash	=1<<1
+		};
+	};
+
+	struct EEncodingOptions
+	{
+		enum Type
+		{
+			AutoDetect,
+			ForceAnsi,
+			ForceUnicode,
+			ForceUTF8,
+			ForceUTF8WithoutBOM
+		};
+	};
+
+	/**
+	 * Load a text file to an FString.
+	 * Supports all combination of ANSI/Unicode files and platforms.
+	*/
+	static void BufferToString( FString& Result, const uint8* Buffer, int32 Size );
+
+	/**
+	 * Load a binary file to a dynamic array.
+	*/
+	static bool LoadFileToArray( TArray<uint8>& Result, const TCHAR* Filename, uint32 Flags = 0 );
+
+	/**
+	 * Load a text file to an FString.
+	 * Supports all combination of ANSI/Unicode files and platforms.
+	 * @param Result string representation of the loaded file
+	 * @param Filename name of the file to load
+	 * @param VerifyFlags flags controlling the hash verification behavior ( see EHashOptions )
+	 */
+	static bool LoadFileToString( FString& Result, const TCHAR* Filename, uint32 VerifyFlags=0 );
+
+	/**
+	 * Save a binary array to a file.
+	 */
+	static bool SaveArrayToFile(TArrayView<const uint8> Array, const TCHAR* Filename, IFileManager* FileManager=&IFileManager::Get(), uint32 WriteFlags = 0);
+
+	/**
+	 * Write the FString to a file.
+	 * Supports all combination of ANSI/Unicode files and platforms.
+	 */
+	static bool SaveStringToFile( const FString& String, const TCHAR* Filename, EEncodingOptions::Type EncodingOptions=EEncodingOptions::AutoDetect, IFileManager* FileManager=&IFileManager::Get(), uint32 WriteFlags = 0 );
+
+	/**
+	 * Saves a 24/32Bit BMP file to disk
+	 * 
+	 * @param Pattern filename with path, must not be 0, if with "bmp" extension (e.g. "out.bmp") the filename stays like this, if without (e.g. "out") automatic index numbers are addended (e.g. "out00002.bmp")
+	 * @param DataWidth - Width of the bitmap supplied in Data >0
+	 * @param DataHeight - Height of the bitmap supplied in Data >0
+	 * @param Data must not be 0
+	 * @param SubRectangle optional, specifies a sub-rectangle of the source image to save out. If NULL, the whole bitmap is saved
+	 * @param FileManager must not be 0
+	 * @param OutFilename optional, if specified filename will be output
+	 * @param bInWriteAlpha opional, specifies whether to write out the alpha channel. Will force BMP V4 format.
+	 *
+	 * @return true if success
+	 */
+	static bool CreateBitmap( const TCHAR* Pattern, int32 DataWidth, int32 DataHeight, const struct FColor* Data, struct FIntRect* SubRectangle = NULL, IFileManager* FileManager = &IFileManager::Get(), FString* OutFilename = NULL, bool bInWriteAlpha = false );
+
+	/**
+	 * Generates the next unique bitmap filename with a specified extension
+	 *
+	 * @param Pattern		Filename with path, but without extension.
+	 * @oaran Extension		File extension to be appended
+	 * @param OutFilename	Reference to an FString where the newly generated filename will be placed
+	 * @param FileManager	Reference to a IFileManager (or the global instance by default)
+	 *
+	 * @return true if success
+	 */
+	static bool GenerateNextBitmapFilename(const FString& Pattern, const FString& Extension, FString& OutFilename, IFileManager* FileManager = &IFileManager::Get());
+	
+	/**
+	 *	Load the given ANSI text file to an array of strings - one FString per line of the file.
+	 *	Intended for use in simple text parsing actions
+	 *
+	 *	@param	InFilename			The text file to read, full path
+	 *	@param	InFileManager		The filemanager to use - NULL will use &IFileManager::Get()
+	 *	@param	OutStrings			The array of FStrings to fill in
+	 *
+	 *	@return	bool				true if successful, false if not
+	 */
+	static bool LoadANSITextFileToStrings(const TCHAR* InFilename, IFileManager* InFileManager, TArray<FString>& OutStrings);
+};
diff --git a/Engine/Source/Runtime/Core/Public/Misc/MessageDialog.h b/Engine/Source/Runtime/Core/Public/Misc/MessageDialog.h
new file mode 100644
--- /dev/null
+++ b/Engine/Source/Runtime/Core/Public/Misc/MessageDialog.h
@@ -0,0 +1,32 @@
+// Copyright 1998-2016 Epic Games, Inc. All Rights Reserved.
+
+#pragma once
+
+#include "HAL/Platform.h"
+
+class FText;
+
+/** 
+ * FMessageDialog
+ * These functions open a message dialog and display the specified informations
+ * these.
+ **/
+struct CORE_API FMessageDialog
+{
+	/** Pops up a message dialog box containing the input string.
+	 * @param Message Text of message to show
+	 * @param OptTitle Optional title to use (defaults to "Message")
+	*/
+	static void Debugf( const FText& Message, const FText* OptTitle = nullptr );
+
+	/** Pops up a message dialog box containing the last system error code in string form. */
+	static void ShowLastError();
+
+	/**
+	 * Open a modal message box dialog
+	 * @param MessageType Controls buttons dialog should have
+	 * @param Message Text of message to show
+	 * @param OptTitle Optional title to use (defaults to "Message")
+	*/
+	static EAppReturnType::Type Open( EAppMsgType::Type MessageType, const FText& Message, const FText* OptTitle = nullptr );
+};
diff --git a/Engine/Source/Runtime/Core/Public/Misc/MultiReaderSingleWriterGT.h b/Engine/Source/Runtime/Core/Public/Misc/MultiReaderSingleWriterGT.h
new file mode 100644
--- /dev/null
+++ b/Engine/Source/Runtime/Core/Public/Misc/MultiReaderSingleWriterGT.h
@@ -0,0 +1,138 @@
+// Copyright 1998-2016 Epic Games, Inc. All Rights Reserved.
+
+#pragma once
+
+/**
+ * Object synchronizing read access (from any thread) and write access (only from game thread).
+ * Multiple simultaneous reads from different threads are allowed. Allows multiple writes from
+ * reentrant calls on main thread. Locking write doesn't prevent reads from main thread.
+ **/
+class FMultiReaderSingleWriterGT
+{
+public:
+	/** Protect data from modification while reading. If issued on game thread, doesn't wait for write to finish. */
+	void LockRead();
+
+	/** Ends protecting data from modification while reading. */
+	void UnlockRead();
+
+	/** Protect data from reading while modifying. Can be called only on game thread. Reentrant. */
+	void LockWrite();
+
+	/** Ends protecting data from reading while modifying. */
+	void UnlockWrite();
+
+private:
+	struct FMRSWCriticalSection
+	{
+		FMRSWCriticalSection()
+			: Action(0)
+			, ReadCounter(0)
+			, WriteCounter(0)
+		{ }
+
+		int32 Action;
+		FThreadSafeCounter ReadCounter;
+		FThreadSafeCounter WriteCounter;
+	};
+
+	FMRSWCriticalSection CriticalSection;
+
+	static const int32 WritingAction = -1;
+	static const int32 NoAction = 0;
+	static const int32 ReadingAction = 1;
+
+	bool CanRead()
+	{
+		return FPlatformAtomics::InterlockedCompareExchange(&CriticalSection.Action, ReadingAction, NoAction) == ReadingAction;
+	}
+	
+	bool CanWrite()
+	{
+		return FPlatformAtomics::InterlockedCompareExchange(&CriticalSection.Action, WritingAction, NoAction) == WritingAction;
+	}
+};
+
+class FReadScopeLock
+{
+public:
+
+	/**
+	* Constructor that performs a lock on the synchronization object
+	*
+	* @param InSynchObject The synchronization object to manage
+	*/
+	FReadScopeLock(FMultiReaderSingleWriterGT* InSynchObject)
+		: SynchObject(InSynchObject)
+	{
+		check(SynchObject);
+		SynchObject->LockRead();
+	}
+
+	/** Destructor that performs a release on the synchronization object. */
+	~FReadScopeLock()
+	{
+		check(SynchObject);
+		SynchObject->UnlockRead();
+	}
+private:
+
+	/** Default constructor (hidden on purpose). */
+	FReadScopeLock();
+
+	/** Copy constructor( hidden on purpose). */
+	FReadScopeLock(FReadScopeLock* InScopeLock);
+
+	/** Assignment operator (hidden on purpose). */
+	FReadScopeLock& operator=(FReadScopeLock& InScopeLock)
+	{
+		return *this;
+	}
+
+private:
+
+	// Holds the synchronization object to aggregate and scope manage.
+	FMultiReaderSingleWriterGT* SynchObject;
+};
+
+class FWriteScopeLock
+{
+public:
+
+	/**
+	* Constructor that performs a lock on the synchronization object
+	*
+	* @param InSynchObject The synchronization object to manage
+	*/
+	FWriteScopeLock(FMultiReaderSingleWriterGT* InSynchObject)
+		: SynchObject(InSynchObject)
+	{
+		check(SynchObject);
+		SynchObject->LockWrite();
+	}
+
+	/** Destructor that performs a release on the synchronization object. */
+	~FWriteScopeLock()
+	{
+		check(SynchObject);
+		SynchObject->UnlockWrite();
+	}
+private:
+
+	/** Default constructor (hidden on purpose). */
+	FWriteScopeLock();
+
+	/** Copy constructor( hidden on purpose). */
+	FWriteScopeLock(FWriteScopeLock* InScopeLock);
+
+	/** Assignment operator (hidden on purpose). */
+	FWriteScopeLock& operator=(FWriteScopeLock& InScopeLock)
+	{
+		return *this;
+	}
+
+private:
+
+	// Holds the synchronization object to aggregate and scope manage.
+	FMultiReaderSingleWriterGT* SynchObject;
+};
diff --git a/Engine/Source/Runtime/Core/Public/Misc/OutputDevice.h b/Engine/Source/Runtime/Core/Public/Misc/OutputDevice.h
--- a/Engine/Source/Runtime/Core/Public/Misc/OutputDevice.h
+++ b/Engine/Source/Runtime/Core/Public/Misc/OutputDevice.h
@@ -5,174 +5,89 @@
 #include "HAL/Platform.h"
 #include "HAL/PlatformMisc.h"
 #include "Misc/CoreMiscDefines.h"
-#if HACK_HEADER_GENERATOR
-	#include "Templates/IsValidVariadicFunctionArg.h"
-	#include "Templates/AndOr.h"
-#endif
-#include "LogVerbosity.h"
+#include "VarArgs.h"
 
 class FText;
 class FString;
 
-// Globals.
-#define GLog FOutputDeviceRedirector::Get()
-CORE_API extern class FOutputDeviceError*			GError;
-CORE_API extern class FFeedbackContext*				GWarn;
-
+#if !PLATFORM_DESKTOP
+	// don't support colorized text on consoles
+	#define SET_WARN_COLOR(Color)
+	#define SET_WARN_COLOR_AND_BACKGROUND(Color, Bkgrnd)
+	#define CLEAR_WARN_COLOR() 
+#else
 
 /*-----------------------------------------------------------------------------
-	VarArgs helper macros.
------------------------------------------------------------------------------*/
+Colorized text.
 
-// phantom definitions to help VAX parse our VARARG_* macros (VAX build 1440)
-#ifdef VISUAL_ASSIST_HACK
-	#define VARARG_DECL( FuncRet, StaticFuncRet, Return, FuncName, Pure, FmtType, ExtraDecl, ExtraCall ) FuncRet FuncName( ExtraDecl FmtType Fmt, ... )
-	#define VARARG_BODY( FuncRet, FuncName, FmtType, ExtraDecl ) FuncRet FuncName( ExtraDecl FmtType Fmt, ... )
-#endif // VISUAL_ASSIST_HACK
+To use colored text from a commandlet, you use the SET_WARN_COLOR macro with
+one of the following standard colors. Then use CLEAR_WARN_COLOR to return
+to default.
 
+To use the standard colors, you simply do this:
+SET_WARN_COLOR(COLOR_YELLOW);
 
-#define GET_VARARGS(msg, msgsize, len, lastarg, fmt) \
-{ \
-	va_list ap; \
-	va_start(ap, lastarg); \
-	FCString::GetVarArgs(msg, msgsize, len, fmt, ap); \
-}
-#define GET_VARARGS_WIDE(msg, msgsize, len, lastarg, fmt) \
-{ \
-	va_list ap; \
-	va_start(ap, lastarg); \
-	FCStringWide::GetVarArgs(msg, msgsize, len, fmt, ap); \
-}
-#define GET_VARARGS_ANSI(msg, msgsize, len, lastarg, fmt) \
-{ \
-	va_list ap; \
-	va_start(ap, lastarg); \
-	FCStringAnsi::GetVarArgs(msg, msgsize, len, fmt, ap); \
-}
-#define GET_VARARGS_RESULT(msg, msgsize, len, lastarg, fmt, result) \
-{ \
-	va_list ap; \
-	va_start(ap, lastarg); \
-	result = FCString::GetVarArgs(msg, msgsize, len, fmt, ap); \
-	if (result >= msgsize) \
-	{ \
-		result = -1; \
-	} \
-}
-#define GET_VARARGS_RESULT_WIDE(msg, msgsize, len, lastarg, fmt, result) \
-{ \
-	va_list ap; \
-	va_start(ap, lastarg); \
-	result = FCStringWide::GetVarArgs(msg, msgsize, len, fmt, ap); \
-	if (result >= msgsize) \
-	{ \
-		result = -1; \
-	} \
-}
-#define GET_VARARGS_RESULT_ANSI(msg, msgsize, len, lastarg, fmt, result) \
-{ \
-	va_list ap; \
-	va_start(ap, lastarg); \
-	result = FCStringAnsi::GetVarArgs(msg, msgsize, len, fmt, ap); \
-	if (result >= msgsize) \
-	{ \
-		result = -1; \
-	} \
-}
+You can specify a background color by appending it to the foreground:
+SET_WARN_COLOR, COLOR_YELLOW COLOR_DARK_RED);
 
-/*-----------------------------------------------------------------------------
-	Ugly VarArgs type checking (debug builds only).
+This will have bright yellow text on a dark red background.
+
+Or you can make your own in the format:
+
+ForegroundRed | ForegroundGreen | ForegroundBlue | ForegroundBright | BackgroundRed | BackgroundGreen | BackgroundBlue | BackgroundBright
+where each value is either 0 or 1 (can leave off trailing 0's), so
+blue on bright yellow is "00101101" and red on black is "1"
+
+An empty string reverts to the normal gray on black.
 -----------------------------------------------------------------------------*/
 
-#define VARARG_EXTRA(A) A,
-#define VARARG_NONE
-#define VARARG_PURE =0
-
-#if PLATFORM_WINDOWS
-
-static inline uint32			CheckVA(uint32 dw)		{ return dw; }
-static inline uint8			CheckVA(uint8 b)			{ return b; }
-static inline int32			CheckVA(int32 i)			{ return i; }
-static inline uint64			CheckVA(uint64 qw)		{ return qw; }
-static inline int64		CheckVA(int64 sqw)		{ return sqw; }
-static inline double		CheckVA(double d)		{ return d; }
-static inline long		CheckVA(long d)		{ return d; }
-static inline long		CheckVA(unsigned long d)		{ return d; }
-static inline TCHAR			CheckVA(TCHAR c)		{ return c; }
-static inline void*			CheckVA(ANSICHAR* s)	{ return (void*)s; }
-static inline bool		CheckVA(bool b) { return b; }
-template<class T> T*		CheckVA(T* p)			{ return p; }
-template<class T> const T*	CheckVA(const T* p)		{ return p; }
-
-#define VARARG_DECL( FuncRet, StaticFuncRet, Return, FuncName, Pure, FmtType, ExtraDecl, ExtraCall )	\
-	FuncRet FuncName##__VA( ExtraDecl FmtType Fmt, ... ) Pure;  \
-	StaticFuncRet FuncName(ExtraDecl FmtType Fmt) {Return FuncName##__VA(ExtraCall (Fmt));} \
-	template<class T1> \
-	StaticFuncRet FuncName(ExtraDecl FmtType Fmt,T1 V1) {T1 v1=CheckVA(V1);Return FuncName##__VA(ExtraCall (Fmt),(v1));} \
-	template<class T1,class T2> \
-	StaticFuncRet FuncName(ExtraDecl FmtType Fmt,T1 V1,T2 V2) {T1 v1=CheckVA(V1);T2 v2=CheckVA(V2);Return FuncName##__VA(ExtraCall (Fmt),(v1),(v2));} \
-	template<class T1,class T2,class T3> \
-	StaticFuncRet FuncName(ExtraDecl FmtType Fmt,T1 V1,T2 V2,T3 V3) {T1 v1=CheckVA(V1);T2 v2=CheckVA(V2);T3 v3=CheckVA(V3);Return FuncName##__VA(ExtraCall (Fmt),(v1),(v2),(v3));} \
-	template<class T1,class T2,class T3,class T4> \
-	StaticFuncRet FuncName(ExtraDecl FmtType Fmt,T1 V1,T2 V2,T3 V3,T4 V4) {T1 v1=CheckVA(V1);T2 v2=CheckVA(V2);T3 v3=CheckVA(V3);T4 v4=CheckVA(V4);Return FuncName##__VA(ExtraCall (Fmt),(v1),(v2),(v3),(v4));} \
-	template<class T1,class T2,class T3,class T4,class T5> \
-	StaticFuncRet FuncName(ExtraDecl FmtType Fmt,T1 V1,T2 V2,T3 V3,T4 V4,T5 V5) {T1 v1=CheckVA(V1);T2 v2=CheckVA(V2);T3 v3=CheckVA(V3);T4 v4=CheckVA(V4);T5 v5=CheckVA(V5);Return FuncName##__VA(ExtraCall (Fmt),(v1),(v2),(v3),(v4),(v5));} \
-	template<class T1,class T2,class T3,class T4,class T5,class T6> \
-	StaticFuncRet FuncName(ExtraDecl FmtType Fmt,T1 V1,T2 V2,T3 V3,T4 V4,T5 V5,T6 V6) {T1 v1=CheckVA(V1);T2 v2=CheckVA(V2);T3 v3=CheckVA(V3);T4 v4=CheckVA(V4);T5 v5=CheckVA(V5);T6 v6=CheckVA(V6);Return FuncName##__VA(ExtraCall (Fmt),(v1),(v2),(v3),(v4),(v5),(v6));} \
-	template<class T1,class T2,class T3,class T4,class T5,class T6,class T7> \
-	StaticFuncRet FuncName(ExtraDecl FmtType Fmt,T1 V1,T2 V2,T3 V3,T4 V4,T5 V5,T6 V6,T7 V7) {T1 v1=CheckVA(V1);T2 v2=CheckVA(V2);T3 v3=CheckVA(V3);T4 v4=CheckVA(V4);T5 v5=CheckVA(V5);T6 v6=CheckVA(V6);T7 v7=CheckVA(V7);Return FuncName##__VA(ExtraCall (Fmt),(v1),(v2),(v3),(v4),(v5),(v6),(v7));} \
-	template<class T1,class T2,class T3,class T4,class T5,class T6,class T7,class T8> \
-	StaticFuncRet FuncName(ExtraDecl FmtType Fmt,T1 V1,T2 V2,T3 V3,T4 V4,T5 V5,T6 V6,T7 V7,T8 V8) {T1 v1=CheckVA(V1);T2 v2=CheckVA(V2);T3 v3=CheckVA(V3);T4 v4=CheckVA(V4);T5 v5=CheckVA(V5);T6 v6=CheckVA(V6);T7 v7=CheckVA(V7);T8 v8=CheckVA(V8);Return FuncName##__VA(ExtraCall (Fmt),(v1),(v2),(v3),(v4),(v5),(v6),(v7),(v8));} \
-	template<class T1,class T2,class T3,class T4,class T5,class T6,class T7,class T8,class T9> \
-	StaticFuncRet FuncName(ExtraDecl FmtType Fmt,T1 V1,T2 V2,T3 V3,T4 V4,T5 V5,T6 V6,T7 V7,T8 V8,T9 V9) {T1 v1=CheckVA(V1);T2 v2=CheckVA(V2);T3 v3=CheckVA(V3);T4 v4=CheckVA(V4);T5 v5=CheckVA(V5);T6 v6=CheckVA(V6);T7 v7=CheckVA(V7);T8 v8=CheckVA(V8);T9 v9=CheckVA(V9);Return FuncName##__VA(ExtraCall (Fmt),(v1),(v2),(v3),(v4),(v5),(v6),(v7),(v8),(v9));} \
-	template<class T1,class T2,class T3,class T4,class T5,class T6,class T7,class T8,class T9,class T10> \
-	StaticFuncRet FuncName(ExtraDecl FmtType Fmt,T1 V1,T2 V2,T3 V3,T4 V4,T5 V5,T6 V6,T7 V7,T8 V8,T9 V9,T10 V10) {T1 v1=CheckVA(V1);T2 v2=CheckVA(V2);T3 v3=CheckVA(V3);T4 v4=CheckVA(V4);T5 v5=CheckVA(V5);T6 v6=CheckVA(V6);T7 v7=CheckVA(V7);T8 v8=CheckVA(V8);T9 v9=CheckVA(V9);T10 v10=CheckVA(V10);Return FuncName##__VA(ExtraCall (Fmt),(v1),(v2),(v3),(v4),(v5),(v6),(v7),(v8),(v9),(v10));} \
-	template<class T1,class T2,class T3,class T4,class T5,class T6,class T7,class T8,class T9,class T10,class T11> \
-	StaticFuncRet FuncName(ExtraDecl FmtType Fmt,T1 V1,T2 V2,T3 V3,T4 V4,T5 V5,T6 V6,T7 V7,T8 V8,T9 V9,T10 V10,T11 V11) {T1 v1=CheckVA(V1);T2 v2=CheckVA(V2);T3 v3=CheckVA(V3);T4 v4=CheckVA(V4);T5 v5=CheckVA(V5);T6 v6=CheckVA(V6);T7 v7=CheckVA(V7);T8 v8=CheckVA(V8);T9 v9=CheckVA(V9);T10 v10=CheckVA(V10);T11 v11=CheckVA(V11);Return FuncName##__VA(ExtraCall (Fmt),(v1),(v2),(v3),(v4),(v5),(v6),(v7),(v8),(v9),(v10),(v11));} \
-	template<class T1,class T2,class T3,class T4,class T5,class T6,class T7,class T8,class T9,class T10,class T11,class T12> \
-	StaticFuncRet FuncName(ExtraDecl FmtType Fmt,T1 V1,T2 V2,T3 V3,T4 V4,T5 V5,T6 V6,T7 V7,T8 V8,T9 V9,T10 V10,T11 V11,T12 V12) {T1 v1=CheckVA(V1);T2 v2=CheckVA(V2);T3 v3=CheckVA(V3);T4 v4=CheckVA(V4);T5 v5=CheckVA(V5);T6 v6=CheckVA(V6);T7 v7=CheckVA(V7);T8 v8=CheckVA(V8);T9 v9=CheckVA(V9);T10 v10=CheckVA(V10);T11 v11=CheckVA(V11);T12 v12=CheckVA(V12);Return FuncName##__VA(ExtraCall (Fmt),(v1),(v2),(v3),(v4),(v5),(v6),(v7),(v8),(v9),(v10),(v11),(v12));} \
-	template<class T1,class T2,class T3,class T4,class T5,class T6,class T7,class T8,class T9,class T10,class T11,class T12,class T13> \
-	StaticFuncRet FuncName(ExtraDecl FmtType Fmt,T1 V1,T2 V2,T3 V3,T4 V4,T5 V5,T6 V6,T7 V7,T8 V8,T9 V9,T10 V10,T11 V11,T12 V12,T13 V13) {T1 v1=CheckVA(V1);T2 v2=CheckVA(V2);T3 v3=CheckVA(V3);T4 v4=CheckVA(V4);T5 v5=CheckVA(V5);T6 v6=CheckVA(V6);T7 v7=CheckVA(V7);T8 v8=CheckVA(V8);T9 v9=CheckVA(V9);T10 v10=CheckVA(V10);T11 v11=CheckVA(V11);T12 v12=CheckVA(V12);T13 v13=CheckVA(V13);Return FuncName##__VA(ExtraCall (Fmt),(v1),(v2),(v3),(v4),(v5),(v6),(v7),(v8),(v9),(v10),(v11),(v12),(v13));} \
-	template<class T1,class T2,class T3,class T4,class T5,class T6,class T7,class T8,class T9,class T10,class T11,class T12,class T13,class T14> \
-	StaticFuncRet FuncName(ExtraDecl FmtType Fmt,T1 V1,T2 V2,T3 V3,T4 V4,T5 V5,T6 V6,T7 V7,T8 V8,T9 V9,T10 V10,T11 V11,T12 V12,T13 V13,T14 V14) {T1 v1=CheckVA(V1);T2 v2=CheckVA(V2);T3 v3=CheckVA(V3);T4 v4=CheckVA(V4);T5 v5=CheckVA(V5);T6 v6=CheckVA(V6);T7 v7=CheckVA(V7);T8 v8=CheckVA(V8);T9 v9=CheckVA(V9);T10 v10=CheckVA(V10);T11 v11=CheckVA(V11);T12 v12=CheckVA(V12);T13 v13=CheckVA(V13);T14 v14=CheckVA(V14);Return FuncName##__VA(ExtraCall (Fmt),(v1),(v2),(v3),(v4),(v5),(v6),(v7),(v8),(v9),(v10),(v11),(v12),(v13),(v14));} \
-	template<class T1,class T2,class T3,class T4,class T5,class T6,class T7,class T8,class T9,class T10,class T11,class T12,class T13,class T14,class T15> \
-	StaticFuncRet FuncName(ExtraDecl FmtType Fmt,T1 V1,T2 V2,T3 V3,T4 V4,T5 V5,T6 V6,T7 V7,T8 V8,T9 V9,T10 V10,T11 V11,T12 V12,T13 V13,T14 V14,T15 V15) {T1 v1=CheckVA(V1);T2 v2=CheckVA(V2);T3 v3=CheckVA(V3);T4 v4=CheckVA(V4);T5 v5=CheckVA(V5);T6 v6=CheckVA(V6);T7 v7=CheckVA(V7);T8 v8=CheckVA(V8);T9 v9=CheckVA(V9);T10 v10=CheckVA(V10);T11 v11=CheckVA(V11);T12 v12=CheckVA(V12);T13 v13=CheckVA(V13);T14 v14=CheckVA(V14);T15 v15=CheckVA(V15);Return FuncName##__VA(ExtraCall (Fmt),(v1),(v2),(v3),(v4),(v5),(v6),(v7),(v8),(v9),(v10),(v11),(v12),(v13),(v14),(v15));} \
-	template<class T1,class T2,class T3,class T4,class T5,class T6,class T7,class T8,class T9,class T10,class T11,class T12,class T13,class T14,class T15,class T16> \
-	StaticFuncRet FuncName(ExtraDecl FmtType Fmt,T1 V1,T2 V2,T3 V3,T4 V4,T5 V5,T6 V6,T7 V7,T8 V8,T9 V9,T10 V10,T11 V11,T12 V12,T13 V13,T14 V14,T15 V15,T16 V16) {T1 v1=CheckVA(V1);T2 v2=CheckVA(V2);T3 v3=CheckVA(V3);T4 v4=CheckVA(V4);T5 v5=CheckVA(V5);T6 v6=CheckVA(V6);T7 v7=CheckVA(V7);T8 v8=CheckVA(V8);T9 v9=CheckVA(V9);T10 v10=CheckVA(V10);T11 v11=CheckVA(V11);T12 v12=CheckVA(V12);T13 v13=CheckVA(V13);T14 v14=CheckVA(V14);T15 v15=CheckVA(V15);T16 v16=CheckVA(V16);Return FuncName##__VA(ExtraCall (Fmt),(v1),(v2),(v3),(v4),(v5),(v6),(v7),(v8),(v9),(v10),(v11),(v12),(v13),(v14),(v15),(v16));} \
-	template<class T1,class T2,class T3,class T4,class T5,class T6,class T7,class T8,class T9,class T10,class T11,class T12,class T13,class T14,class T15,class T16,class T17> \
-	StaticFuncRet FuncName(ExtraDecl FmtType Fmt,T1 V1,T2 V2,T3 V3,T4 V4,T5 V5,T6 V6,T7 V7,T8 V8,T9 V9,T10 V10,T11 V11,T12 V12,T13 V13,T14 V14,T15 V15,T16 V16,T17 V17) {T1 v1=CheckVA(V1);T2 v2=CheckVA(V2);T3 v3=CheckVA(V3);T4 v4=CheckVA(V4);T5 v5=CheckVA(V5);T6 v6=CheckVA(V6);T7 v7=CheckVA(V7);T8 v8=CheckVA(V8);T9 v9=CheckVA(V9);T10 v10=CheckVA(V10);T11 v11=CheckVA(V11);T12 v12=CheckVA(V12);T13 v13=CheckVA(V13);T14 v14=CheckVA(V14);T15 v15=CheckVA(V15);T16 v16=CheckVA(V16);T17 v17=CheckVA(V17);Return FuncName##__VA(ExtraCall (Fmt),(v1),(v2),(v3),(v4),(v5),(v6),(v7),(v8),(v9),(v10),(v11),(v12),(v13),(v14),(v15),(v16),(v17));} \
-	template<class T1,class T2,class T3,class T4,class T5,class T6,class T7,class T8,class T9,class T10,class T11,class T12,class T13,class T14,class T15,class T16,class T17,class T18> \
-	StaticFuncRet FuncName(ExtraDecl FmtType Fmt,T1 V1,T2 V2,T3 V3,T4 V4,T5 V5,T6 V6,T7 V7,T8 V8,T9 V9,T10 V10,T11 V11,T12 V12,T13 V13,T14 V14,T15 V15,T16 V16,T17 V17,T18 V18) {T1 v1=CheckVA(V1);T2 v2=CheckVA(V2);T3 v3=CheckVA(V3);T4 v4=CheckVA(V4);T5 v5=CheckVA(V5);T6 v6=CheckVA(V6);T7 v7=CheckVA(V7);T8 v8=CheckVA(V8);T9 v9=CheckVA(V9);T10 v10=CheckVA(V10);T11 v11=CheckVA(V11);T12 v12=CheckVA(V12);T13 v13=CheckVA(V13);T14 v14=CheckVA(V14);T15 v15=CheckVA(V15);T16 v16=CheckVA(V16);T17 v17=CheckVA(V17);T18 v18=CheckVA(V18);Return FuncName##__VA(ExtraCall (Fmt),(v1),(v2),(v3),(v4),(v5),(v6),(v7),(v8),(v9),(v10),(v11),(v12),(v13),(v14),(v15),(v16),(v17),(v18));} \
-	template<class T1,class T2,class T3,class T4,class T5,class T6,class T7,class T8,class T9,class T10,class T11,class T12,class T13,class T14,class T15,class T16,class T17,class T18,class T19> \
-	StaticFuncRet FuncName(ExtraDecl FmtType Fmt,T1 V1,T2 V2,T3 V3,T4 V4,T5 V5,T6 V6,T7 V7,T8 V8,T9 V9,T10 V10,T11 V11,T12 V12,T13 V13,T14 V14,T15 V15,T16 V16,T17 V17,T18 V18,T19 V19) {T1 v1=CheckVA(V1);T2 v2=CheckVA(V2);T3 v3=CheckVA(V3);T4 v4=CheckVA(V4);T5 v5=CheckVA(V5);T6 v6=CheckVA(V6);T7 v7=CheckVA(V7);T8 v8=CheckVA(V8);T9 v9=CheckVA(V9);T10 v10=CheckVA(V10);T11 v11=CheckVA(V11);T12 v12=CheckVA(V12);T13 v13=CheckVA(V13);T14 v14=CheckVA(V14);T15 v15=CheckVA(V15);T16 v16=CheckVA(V16);T17 v17=CheckVA(V17);T18 v18=CheckVA(V18);T19 v19=CheckVA(V19);Return FuncName##__VA(ExtraCall (Fmt),(v1),(v2),(v3),(v4),(v5),(v6),(v7),(v8),(v9),(v10),(v11),(v12),(v13),(v14),(v15),(v16),(v17),(v18),(v19));} \
-	template<class T1,class T2,class T3,class T4,class T5,class T6,class T7,class T8,class T9,class T10,class T11,class T12,class T13,class T14,class T15,class T16,class T17,class T18,class T19,class T20> \
-	StaticFuncRet FuncName(ExtraDecl FmtType Fmt,T1 V1,T2 V2,T3 V3,T4 V4,T5 V5,T6 V6,T7 V7,T8 V8,T9 V9,T10 V10,T11 V11,T12 V12,T13 V13,T14 V14,T15 V15,T16 V16,T17 V17,T18 V18,T19 V19,T20 V20) {T1 v1=CheckVA(V1);T2 v2=CheckVA(V2);T3 v3=CheckVA(V3);T4 v4=CheckVA(V4);T5 v5=CheckVA(V5);T6 v6=CheckVA(V6);T7 v7=CheckVA(V7);T8 v8=CheckVA(V8);T9 v9=CheckVA(V9);T10 v10=CheckVA(V10);T11 v11=CheckVA(V11);T12 v12=CheckVA(V12);T13 v13=CheckVA(V13);T14 v14=CheckVA(V14);T15 v15=CheckVA(V15);T16 v16=CheckVA(V16);T17 v17=CheckVA(V17);T18 v18=CheckVA(V18);T19 v19=CheckVA(V19);T20 v20=CheckVA(V20);Return FuncName##__VA(ExtraCall (Fmt),(v1),(v2),(v3),(v4),(v5),(v6),(v7),(v8),(v9),(v10),(v11),(v12),(v13),(v14),(v15),(v16),(v17),(v18),(v19),(v20));} \
-	template<class T1,class T2,class T3,class T4,class T5,class T6,class T7,class T8,class T9,class T10,class T11,class T12,class T13,class T14,class T15,class T16,class T17,class T18,class T19,class T20,class T21> \
-	StaticFuncRet FuncName(ExtraDecl FmtType Fmt,T1 V1,T2 V2,T3 V3,T4 V4,T5 V5,T6 V6,T7 V7,T8 V8,T9 V9,T10 V10,T11 V11,T12 V12,T13 V13,T14 V14,T15 V15,T16 V16,T17 V17,T18 V18,T19 V19,T20 V20, T21 V21) {T1 v1=CheckVA(V1);T2 v2=CheckVA(V2);T3 v3=CheckVA(V3);T4 v4=CheckVA(V4);T5 v5=CheckVA(V5);T6 v6=CheckVA(V6);T7 v7=CheckVA(V7);T8 v8=CheckVA(V8);T9 v9=CheckVA(V9);T10 v10=CheckVA(V10);T11 v11=CheckVA(V11);T12 v12=CheckVA(V12);T13 v13=CheckVA(V13);T14 v14=CheckVA(V14);T15 v15=CheckVA(V15);T16 v16=CheckVA(V16);T17 v17=CheckVA(V17);T18 v18=CheckVA(V18);T19 v19=CheckVA(V19);T20 v20=CheckVA(V20);T21 v21=CheckVA(V21);Return FuncName##__VA(ExtraCall (Fmt),(v1),(v2),(v3),(v4),(v5),(v6),(v7),(v8),(v9),(v10),(v11),(v12),(v13),(v14),(v15),(v16),(v17),(v18),(v19),(v20),(v21));} \
-	template<class T1,class T2,class T3,class T4,class T5,class T6,class T7,class T8,class T9,class T10,class T11,class T12,class T13,class T14,class T15,class T16,class T17,class T18,class T19,class T20,class T21,class T22> \
-	StaticFuncRet FuncName(ExtraDecl FmtType Fmt,T1 V1,T2 V2,T3 V3,T4 V4,T5 V5,T6 V6,T7 V7,T8 V8,T9 V9,T10 V10,T11 V11,T12 V12,T13 V13,T14 V14,T15 V15,T16 V16,T17 V17,T18 V18,T19 V19,T20 V20, T21 V21, T22 V22) {T1 v1=CheckVA(V1);T2 v2=CheckVA(V2);T3 v3=CheckVA(V3);T4 v4=CheckVA(V4);T5 v5=CheckVA(V5);T6 v6=CheckVA(V6);T7 v7=CheckVA(V7);T8 v8=CheckVA(V8);T9 v9=CheckVA(V9);T10 v10=CheckVA(V10);T11 v11=CheckVA(V11);T12 v12=CheckVA(V12);T13 v13=CheckVA(V13);T14 v14=CheckVA(V14);T15 v15=CheckVA(V15);T16 v16=CheckVA(V16);T17 v17=CheckVA(V17);T18 v18=CheckVA(V18);T19 v19=CheckVA(V19);T20 v20=CheckVA(V20);T21 v21=CheckVA(V21);T22 v22=CheckVA(V22);Return FuncName##__VA(ExtraCall (Fmt),(v1),(v2),(v3),(v4),(v5),(v6),(v7),(v8),(v9),(v10),(v11),(v12),(v13),(v14),(v15),(v16),(v17),(v18),(v19),(v20),(v21),(v22));} \
-	template<class T1,class T2,class T3,class T4,class T5,class T6,class T7,class T8,class T9,class T10,class T11,class T12,class T13,class T14,class T15,class T16,class T17,class T18,class T19,class T20,class T21,class T22,class T23> \
-	StaticFuncRet FuncName(ExtraDecl FmtType Fmt,T1 V1,T2 V2,T3 V3,T4 V4,T5 V5,T6 V6,T7 V7,T8 V8,T9 V9,T10 V10,T11 V11,T12 V12,T13 V13,T14 V14,T15 V15,T16 V16,T17 V17,T18 V18,T19 V19,T20 V20, T21 V21, T22 V22, T23 V23) {T1 v1=CheckVA(V1);T2 v2=CheckVA(V2);T3 v3=CheckVA(V3);T4 v4=CheckVA(V4);T5 v5=CheckVA(V5);T6 v6=CheckVA(V6);T7 v7=CheckVA(V7);T8 v8=CheckVA(V8);T9 v9=CheckVA(V9);T10 v10=CheckVA(V10);T11 v11=CheckVA(V11);T12 v12=CheckVA(V12);T13 v13=CheckVA(V13);T14 v14=CheckVA(V14);T15 v15=CheckVA(V15);T16 v16=CheckVA(V16);T17 v17=CheckVA(V17);T18 v18=CheckVA(V18);T19 v19=CheckVA(V19);T20 v20=CheckVA(V20);T21 v21=CheckVA(V21);T22 v22=CheckVA(V22);T23 v23=CheckVA(V23);Return FuncName##__VA(ExtraCall (Fmt),(v1),(v2),(v3),(v4),(v5),(v6),(v7),(v8),(v9),(v10),(v11),(v12),(v13),(v14),(v15),(v16),(v17),(v18),(v19),(v20),(v21),(v22),(v23));} \
-	template<class T1,class T2,class T3,class T4,class T5,class T6,class T7,class T8,class T9,class T10,class T11,class T12,class T13,class T14,class T15,class T16,class T17,class T18,class T19,class T20,class T21,class T22,class T23,class T24> \
-	StaticFuncRet FuncName(ExtraDecl FmtType Fmt,T1 V1,T2 V2,T3 V3,T4 V4,T5 V5,T6 V6,T7 V7,T8 V8,T9 V9,T10 V10,T11 V11,T12 V12,T13 V13,T14 V14,T15 V15,T16 V16,T17 V17,T18 V18,T19 V19,T20 V20, T21 V21, T22 V22, T23 V23, T24 V24) {T1 v1=CheckVA(V1);T2 v2=CheckVA(V2);T3 v3=CheckVA(V3);T4 v4=CheckVA(V4);T5 v5=CheckVA(V5);T6 v6=CheckVA(V6);T7 v7=CheckVA(V7);T8 v8=CheckVA(V8);T9 v9=CheckVA(V9);T10 v10=CheckVA(V10);T11 v11=CheckVA(V11);T12 v12=CheckVA(V12);T13 v13=CheckVA(V13);T14 v14=CheckVA(V14);T15 v15=CheckVA(V15);T16 v16=CheckVA(V16);T17 v17=CheckVA(V17);T18 v18=CheckVA(V18);T19 v19=CheckVA(V19);T20 v20=CheckVA(V20);T21 v21=CheckVA(V21);T22 v22=CheckVA(V22);T23 v23=CheckVA(V23);T24 v24=CheckVA(V24);Return FuncName##__VA(ExtraCall (Fmt),(v1),(v2),(v3),(v4),(v5),(v6),(v7),(v8),(v9),(v10),(v11),(v12),(v13),(v14),(v15),(v16),(v17),(v18),(v19),(v20),(v21),(v22),(v23),(v24));} \
-	template<class T1,class T2,class T3,class T4,class T5,class T6,class T7,class T8,class T9,class T10,class T11,class T12,class T13,class T14,class T15,class T16,class T17,class T18,class T19,class T20,class T21,class T22,class T23,class T24,class T25> \
-	StaticFuncRet FuncName(ExtraDecl FmtType Fmt,T1 V1,T2 V2,T3 V3,T4 V4,T5 V5,T6 V6,T7 V7,T8 V8,T9 V9,T10 V10,T11 V11,T12 V12,T13 V13,T14 V14,T15 V15,T16 V16,T17 V17,T18 V18,T19 V19,T20 V20, T21 V21, T22 V22, T23 V23, T24 V24, T25 V25) {T1 v1=CheckVA(V1);T2 v2=CheckVA(V2);T3 v3=CheckVA(V3);T4 v4=CheckVA(V4);T5 v5=CheckVA(V5);T6 v6=CheckVA(V6);T7 v7=CheckVA(V7);T8 v8=CheckVA(V8);T9 v9=CheckVA(V9);T10 v10=CheckVA(V10);T11 v11=CheckVA(V11);T12 v12=CheckVA(V12);T13 v13=CheckVA(V13);T14 v14=CheckVA(V14);T15 v15=CheckVA(V15);T16 v16=CheckVA(V16);T17 v17=CheckVA(V17);T18 v18=CheckVA(V18);T19 v19=CheckVA(V19);T20 v20=CheckVA(V20);T21 v21=CheckVA(V21);T22 v22=CheckVA(V22);T23 v23=CheckVA(V23);T24 v24=CheckVA(V24);T25 v25=CheckVA(V25);Return FuncName##__VA(ExtraCall (Fmt),(v1),(v2),(v3),(v4),(v5),(v6),(v7),(v8),(v9),(v10),(v11),(v12),(v13),(v14),(v15),(v16),(v17),(v18),(v19),(v20),(v21),(v22),(v23),(v24),(v25));} \
-	template<class T1,class T2,class T3,class T4,class T5,class T6,class T7,class T8,class T9,class T10,class T11,class T12,class T13,class T14,class T15,class T16,class T17,class T18,class T19,class T20,class T21,class T22,class T23,class T24,class T25,class T26> \
-	StaticFuncRet FuncName(ExtraDecl FmtType Fmt,T1 V1,T2 V2,T3 V3,T4 V4,T5 V5,T6 V6,T7 V7,T8 V8,T9 V9,T10 V10,T11 V11,T12 V12,T13 V13,T14 V14,T15 V15,T16 V16,T17 V17,T18 V18,T19 V19,T20 V20, T21 V21, T22 V22, T23 V23, T24 V24, T25 V25, T26 V26) {T1 v1=CheckVA(V1);T2 v2=CheckVA(V2);T3 v3=CheckVA(V3);T4 v4=CheckVA(V4);T5 v5=CheckVA(V5);T6 v6=CheckVA(V6);T7 v7=CheckVA(V7);T8 v8=CheckVA(V8);T9 v9=CheckVA(V9);T10 v10=CheckVA(V10);T11 v11=CheckVA(V11);T12 v12=CheckVA(V12);T13 v13=CheckVA(V13);T14 v14=CheckVA(V14);T15 v15=CheckVA(V15);T16 v16=CheckVA(V16);T17 v17=CheckVA(V17);T18 v18=CheckVA(V18);T19 v19=CheckVA(V19);T20 v20=CheckVA(V20);T21 v21=CheckVA(V21);T22 v22=CheckVA(V22);T23 v23=CheckVA(V23);T24 v24=CheckVA(V24);T25 v25=CheckVA(V25);T26 v26=CheckVA(V26);Return FuncName##__VA(ExtraCall (Fmt),(v1),(v2),(v3),(v4),(v5),(v6),(v7),(v8),(v9),(v10),(v11),(v12),(v13),(v14),(v15),(v16),(v17),(v18),(v19),(v20),(v21),(v22),(v23),(v24),(v25),(v26));}
-
-
-
-#define VARARG_BODY( FuncRet, FuncName, FmtType, ExtraDecl )		\
-	FuncRet FuncName##__VA( ExtraDecl  FmtType Fmt, ... )
-
-#else  // !PLATFORM_WINDOWS
-
-#define VARARG_DECL( FuncRet, StaticFuncRet, Return, FuncName, Pure, FmtType, ExtraDecl, ExtraCall )	\
-	FuncRet FuncName( ExtraDecl FmtType Fmt, ... ) Pure
-#define VARARG_BODY( FuncRet, FuncName, FmtType, ExtraDecl )		\
-	FuncRet FuncName( ExtraDecl FmtType Fmt, ... )
-
-#endif // PLATFORM_WINDOWS
+// putting them in a namespace to protect against future name conflicts
+namespace OutputDeviceColor
+{
+	const TCHAR* const COLOR_BLACK = TEXT("0000");
+
+	const TCHAR* const COLOR_DARK_RED = TEXT("1000");
+	const TCHAR* const COLOR_DARK_GREEN = TEXT("0100");
+	const TCHAR* const COLOR_DARK_BLUE = TEXT("0010");
+	const TCHAR* const COLOR_DARK_YELLOW = TEXT("1100");
+	const TCHAR* const COLOR_DARK_CYAN = TEXT("0110");
+	const TCHAR* const COLOR_DARK_PURPLE = TEXT("1010");
+	const TCHAR* const COLOR_DARK_WHITE = TEXT("1110");
+	const TCHAR* const COLOR_GRAY = COLOR_DARK_WHITE;
+
+	const TCHAR* const COLOR_RED = TEXT("1001");
+	const TCHAR* const COLOR_GREEN = TEXT("0101");
+	const TCHAR* const COLOR_BLUE = TEXT("0011");
+	const TCHAR* const COLOR_YELLOW = TEXT("1101");
+	const TCHAR* const COLOR_CYAN = TEXT("0111");
+	const TCHAR* const COLOR_PURPLE = TEXT("1011");
+	const TCHAR* const COLOR_WHITE = TEXT("1111");
+
+	const TCHAR* const COLOR_NONE = TEXT("");
+}
+
+using namespace OutputDeviceColor;
+
+
+// let a console or (UE_BUILD_SHIPPING || UE_BUILD_TEST) define it to nothing
+#ifndef SET_WARN_COLOR
+
+/**
+* Set the console color with Color or a Color and Background color
+*/
+#define SET_WARN_COLOR(Color) \
+	UE_LOG(LogHAL, SetColor, TEXT("%s"), Color);
+#define SET_WARN_COLOR_AND_BACKGROUND(Color, Bkgrnd) \
+	UE_LOG(LogHAL, SetColor, TEXT("%s%s"), Color, Bkgrnd);
+
+/**
+* Return color to it's default
+*/
+#define CLEAR_WARN_COLOR() \
+	UE_LOG(LogHAL, SetColor, TEXT("%s"), COLOR_NONE);
+
+#endif
+#endif
 
 /**
  * Enum that defines how the log times are to be displayed.
@@ -321,157 +236,5 @@ protected:
 	bool bAutoEmitLineTerminator;
 };
 
-// Error device.
-class CORE_API FOutputDeviceError : public FOutputDevice
-{
-public:
-	virtual void HandleError()=0;
-};
-
-/**
- * FMsg 
- * This struct contains functions for messaging with tools or debug logs.
- **/
-struct CORE_API FMsg
-{
-	/** Sends a message to a remote tool. */
-	static void SendNotificationString( const TCHAR* Message );
-
-	/** Sends a formatted message to a remote tool. */
-	static void VARARGS SendNotificationStringf( const TCHAR *Format, ... );
-
-	/** Log function */
-	VARARG_DECL( static void, static void, {}, Logf, VARARG_NONE, const TCHAR*, VARARG_EXTRA(const ANSICHAR* File) VARARG_EXTRA(int32 Line) VARARG_EXTRA(const class FName& Category) VARARG_EXTRA(ELogVerbosity::Type Verbosity), VARARG_EXTRA(File) VARARG_EXTRA(Line) VARARG_EXTRA(Category) VARARG_EXTRA(Verbosity) );
-
-	/** Internal version of log function. Should be used only in logging macros, as it relies on caller to call assert on fatal error */
-	VARARG_DECL(static void, static void, {}, Logf_Internal, VARARG_NONE, const TCHAR*, VARARG_EXTRA(const ANSICHAR* File) VARARG_EXTRA(int32 Line) VARARG_EXTRA(const class FName& Category) VARARG_EXTRA(ELogVerbosity::Type Verbosity), VARARG_EXTRA(File) VARARG_EXTRA(Line) VARARG_EXTRA(Category) VARARG_EXTRA(Verbosity));
-};
-
-/**
- * FDebug
- * These functions offer debugging and diagnostic functionality and its presence 
- * depends on compiler switches.
- **/
-struct CORE_API FDebug
-{
-	/** Outputs a callstack message, separating out each line into a separate log call by modifying the buffer in-place */
-	static void OutputMultiLineCallstack(const ANSICHAR* File, int32 Line, const FName& LogName, const TCHAR* Heading, TCHAR* Message, ELogVerbosity::Type Verbosity);
-
-	/** Logs final assert message and exits the program. */
-	static void VARARGS AssertFailed(const ANSICHAR* Expr, const ANSICHAR* File, int32 Line, const TCHAR* Format = TEXT(""), ...);
-
-	/** Records the calling of AssertFailed() */
-	static bool bHasAsserted;
-
-#if DO_CHECK || DO_GUARD_SLOW
-	/** Failed assertion handler.  Warning: May be called at library startup time. */
-	static void VARARGS LogAssertFailedMessage( const ANSICHAR* Expr, const ANSICHAR* File, int32 Line, const TCHAR* Format=TEXT(""), ... );
-	
-	/**
-	 * Called when an 'ensure' assertion fails; gathers stack data and generates and error report.
-	 *
-	 * @param	Expr	Code expression ANSI string (#code)
-	 * @param	File	File name ANSI string (__FILE__)
-	 * @param	Line	Line number (__LINE__)
-	 * @param	Msg		Informative error message text
-	 * 
-	 * Don't change the name of this function, it's used to detect ensures by the crash reporter.
-	 */
-	static void EnsureFailed( const ANSICHAR* Expr, const ANSICHAR* File, int32 Line, const TCHAR* Msg );
-
-	/**
-	 * Logs an error if bLog is true, and returns false.  Takes a formatted string.
-	 *
-	 * @param	bLog	Log if true.
-	 * @param	Expr	Code expression ANSI string (#code)
-	 * @param	File	File name ANSI string (__FILE__)
-	 * @param	Line	Line number (__LINE__)
-	 * @param	FormattedMsg	Informative error message text with variable args
-	 *
-	 * @return false in all cases.
-	 *
-	 * Note: this crazy name is to ensure that the crash reporter recognizes it, which checks for functions in the callstack starting with 'EnsureNotFalse'.
-	 */
-	static bool VARARGS OptionallyLogFormattedEnsureMessageReturningFalse(bool bLog, const ANSICHAR* Expr, const ANSICHAR* File, int32 Line, const TCHAR* FormattedMsg, ...);
-#endif // DO_CHECK || DO_GUARD_SLOW
-};
-
-
-/** 
- * FMessageDialog
- * These functions open a message dialog and display the specified informations
- * these.
- **/
-struct CORE_API FMessageDialog
-{
-	/** Pops up a message dialog box containing the input string.
-	 * @param Message Text of message to show
-	 * @param OptTitle Optional title to use (defaults to "Message")
-	*/
-	static void Debugf( const FText& Message, const FText* OptTitle = nullptr );
-
-	/** Pops up a message dialog box containing the last system error code in string form. */
-	static void ShowLastError();
-
-	/**
-	 * Open a modal message box dialog
-	 * @param MessageType Controls buttons dialog should have
-	 * @param Message Text of message to show
-	 * @param OptTitle Optional title to use (defaults to "Message")
-	*/
-	static EAppReturnType::Type Open( EAppMsgType::Type MessageType, const FText& Message, const FText* OptTitle = nullptr );
-};
-
-
-/** 
- * FError
- * Set of functions for error reporting 
- **/
-struct CORE_API FError
-{
-	/** low level fatal error handler. */
-	static void VARARGS LowLevelFatal(const ANSICHAR* File, int32 Line, const TCHAR* Format=TEXT(""), ... );
-
-#if HACK_HEADER_GENERATOR
-	/**
-	 * Throws a printf-formatted exception as a const TCHAR*.
-	 */
-	template <typename... Types>
-	FUNCTION_NO_RETURN_START
-	static void VARARGS Throwf(const TCHAR* Fmt, Types... Args)
-	FUNCTION_NO_RETURN_END
-	{
-		static_assert(TAnd<TIsValidVariadicFunctionArg<Types>...>::Value, "Invalid argument(s) passed to FError::Throwf");
-
-		ThrowfImpl(Fmt, Args...);
-	}
-#endif
-
-private:
-#if HACK_HEADER_GENERATOR
-	FUNCTION_NO_RETURN_START
-	static void VARARGS ThrowfImpl(const TCHAR* Fmt, ...)
-	FUNCTION_NO_RETURN_END;
-#endif
-};
-
-
-// Any object that is capable of taking commands.
-class CORE_API FExec
-{
-public:
-	virtual ~FExec() {}
-
-	/**
-	* Exec handler
-	*
-	* @param	InWorld	World context
-	* @param	Cmd		Command to parse
-	* @param	Ar		Output device to log to
-	*
-	* @return	true if command was handled, false otherwise
-	*/
-	virtual bool Exec( class UWorld* InWorld, const TCHAR* Cmd, FOutputDevice& Ar ) PURE_VIRTUAL(FExec::Exec,return false;)
-};
 
 
diff --git a/Engine/Source/Runtime/Core/Public/Misc/OutputDeviceAnsiError.h b/Engine/Source/Runtime/Core/Public/Misc/OutputDeviceAnsiError.h
new file mode 100644
--- /dev/null
+++ b/Engine/Source/Runtime/Core/Public/Misc/OutputDeviceAnsiError.h
@@ -0,0 +1,36 @@
+// Copyright 1998-2016 Epic Games, Inc. All Rights Reserved.
+
+#pragma once
+
+#include "OutputDeviceError.h"
+
+class FOutputDeviceAnsiError : public FOutputDeviceError
+{
+public:
+	/** Constructor, initializing member variables */
+	FOutputDeviceAnsiError();
+
+	/**
+	* Serializes the passed in data unless the current event is suppressed.
+	*
+	* @param	Data	Text to log
+	* @param	Event	Event name used for suppression purposes
+	*/
+	virtual void Serialize(const TCHAR* Msg, ELogVerbosity::Type Verbosity, const class FName& Category) override;
+
+	virtual bool CanBeUsedOnAnyThread() const override
+	{
+		return true;
+	}
+
+	/**
+	* Error handling function that is being called from within the system wide global
+	* error handler, e.g. using structured exception handling on the PC.
+	*/
+	void HandleError() override;
+
+private:
+
+	int32		ErrorPos;
+};
+
diff --git a/Engine/Source/Runtime/Core/Public/Misc/OutputDeviceArchiveWrapper.h b/Engine/Source/Runtime/Core/Public/Misc/OutputDeviceArchiveWrapper.h
new file mode 100644
--- /dev/null
+++ b/Engine/Source/Runtime/Core/Public/Misc/OutputDeviceArchiveWrapper.h
@@ -0,0 +1,31 @@
+// Copyright 1998-2016 Epic Games, Inc. All Rights Reserved.
+
+#pragma once
+
+#include "OutputDevice.h"
+
+/**
+* Output device wrapping any kind of FArchive.  Note: Works in any build configuration.
+*/
+class CORE_API FOutputDeviceArchiveWrapper : public FOutputDevice
+{
+public:
+	/**
+	* Constructor, initializing member variables.
+	*
+	* @param InArchive	Archive to use, must not be nullptr.  Does not take ownership of the archive, clean up or delete the archive independently!
+	*/
+	FOutputDeviceArchiveWrapper(FArchive* InArchive)
+		: LogAr(InArchive)
+	{
+		check(InArchive);
+	}
+
+	// FOutputDevice interface
+	virtual void Flush() override;
+	virtual void Serialize(const TCHAR* Data, ELogVerbosity::Type Verbosity, const class FName& Category) override;
+	// End of FOutputDevice interface
+
+private:
+	FArchive* LogAr;
+};
diff --git a/Engine/Source/Runtime/Core/Public/Misc/OutputDeviceDebug.h b/Engine/Source/Runtime/Core/Public/Misc/OutputDeviceDebug.h
new file mode 100644
--- /dev/null
+++ b/Engine/Source/Runtime/Core/Public/Misc/OutputDeviceDebug.h
@@ -0,0 +1,25 @@
+// Copyright 1998-2016 Epic Games, Inc. All Rights Reserved.
+
+#pragma once
+
+#include "OutputDevice.h"
+
+class CORE_API FOutputDeviceDebug : public FOutputDevice
+{
+public:
+	/**
+	* Serializes the passed in data unless the current event is suppressed.
+	*
+	* @param	Data	Text to log
+	* @param	Event	Event name used for suppression purposes
+	*/
+	virtual void Serialize(const TCHAR* Data, ELogVerbosity::Type Verbosity, const class FName& Category, const double Time) override;
+
+	virtual void Serialize(const TCHAR* Data, ELogVerbosity::Type Verbosity, const class FName& Category) override;
+
+	virtual bool CanBeUsedOnAnyThread() const override
+	{
+		return true;
+	}
+};
+
diff --git a/Engine/Source/Runtime/Core/Public/Misc/OutputDeviceError.h b/Engine/Source/Runtime/Core/Public/Misc/OutputDeviceError.h
new file mode 100644
--- /dev/null
+++ b/Engine/Source/Runtime/Core/Public/Misc/OutputDeviceError.h
@@ -0,0 +1,13 @@
+// Copyright 1998-2016 Epic Games, Inc. All Rights Reserved.
+
+#pragma once
+
+#include "OutputDevice.h"
+
+// Error device.
+class CORE_API FOutputDeviceError : public FOutputDevice
+{
+public:
+	virtual void HandleError()=0;
+};
+
diff --git a/Engine/Source/Runtime/Core/Public/Misc/OutputDeviceFile.h b/Engine/Source/Runtime/Core/Public/Misc/OutputDeviceFile.h
new file mode 100644
--- /dev/null
+++ b/Engine/Source/Runtime/Core/Public/Misc/OutputDeviceFile.h
@@ -0,0 +1,84 @@
+// Copyright 1998-2016 Epic Games, Inc. All Rights Reserved.
+
+#pragma once
+
+#include "OutputDevice.h"
+
+/** string added to the filename of timestamped backup log files */
+#define BACKUP_LOG_FILENAME_POSTFIX TEXT("-backup-")
+
+/** Used by FOutputDeviceFile to write to a file on a separate thread */
+class FAsyncWriter;
+
+enum class EByteOrderMark : int8
+{
+	UTF8,
+	Unspecified,
+};
+
+/**
+* File output device (Note: Only works if ALLOW_LOG_FILE && !NO_LOGGING is true, otherwise Serialize does nothing).
+*/
+class CORE_API FOutputDeviceFile : public FOutputDevice
+{
+public:
+	/**
+	* Constructor, initializing member variables.
+	*
+	* @param InFilename	Filename to use, can be nullptr
+	* @param bDisableBackup If true, existing files will not be backed up
+	*/
+	FOutputDeviceFile(const TCHAR* InFilename = nullptr, bool bDisableBackup = false);
+
+	/** Sets the filename that the output device writes to.  If the output device was already writing to a file, closes that file. */
+	void SetFilename(const TCHAR* InFilename);
+
+	//~ Begin FOutputDevice Interface.
+	/**
+	* Closes output device and cleans up. This can't happen in the destructor
+	* as we have to call "delete" which cannot be done for static/ global
+	* objects.
+	*/
+	void TearDown() override;
+
+	/**
+	* Flush the write cache so the file isn't truncated in case we crash right
+	* after calling this function.
+	*/
+	void Flush() override;
+
+	virtual void Serialize(const TCHAR* Data, ELogVerbosity::Type Verbosity, const class FName& Category, const double Time) override;
+	virtual void Serialize(const TCHAR* Data, ELogVerbosity::Type Verbosity, const class FName& Category) override;
+	virtual bool CanBeUsedOnAnyThread() const override
+	{
+		return true;
+	}
+	//~ End FOutputDevice Interface.
+
+	/** Creates a backup copy of a log file if it already exists */
+	static void CreateBackupCopy(const TCHAR* Filename);
+
+	/** Checks if the filename represents a backup copy of a log file */
+	static bool IsBackupCopy(const TCHAR* Filename);
+
+private:
+
+	/** Writes to a file on a separate thread */
+	FAsyncWriter* AsyncWriter;
+	/** Archive used by the async writer */
+	FArchive* WriterArchive;
+
+	TCHAR Filename[1024];
+	bool Opened;
+	bool Dead;
+
+	/** If true, existing files will not be backed up */
+	bool		bDisableBackup;
+
+	void WriteRaw(const TCHAR* C);
+
+	/** Creates the async writer and its archive. Returns true if successful.  */
+	bool CreateWriter(uint32 MaxAttempts = 32);
+
+	void WriteByteOrderMarkToArchive(EByteOrderMark ByteOrderMark);
+};
diff --git a/Engine/Source/Runtime/Core/Public/Misc/OutputDeviceNull.h b/Engine/Source/Runtime/Core/Public/Misc/OutputDeviceNull.h
new file mode 100644
--- /dev/null
+++ b/Engine/Source/Runtime/Core/Public/Misc/OutputDeviceNull.h
@@ -0,0 +1,19 @@
+// Copyright 1998-2016 Epic Games, Inc. All Rights Reserved.
+
+#pragma once
+
+#include "OutputDevice.h"
+
+// Null output device.
+class CORE_API FOutputDeviceNull : public FOutputDevice
+{
+public:
+	/**
+	* NULL implementation of Serialize.
+	*
+	* @param	Data	unused
+	* @param	Event	unused
+	*/
+	virtual void Serialize(const TCHAR* V, ELogVerbosity::Type Verbosity, const class FName& Category) override;
+};
+
diff --git a/Engine/Source/Runtime/Core/Public/Misc/OutputDeviceRedirector.h b/Engine/Source/Runtime/Core/Public/Misc/OutputDeviceRedirector.h
new file mode 100644
--- /dev/null
+++ b/Engine/Source/Runtime/Core/Public/Misc/OutputDeviceRedirector.h
@@ -0,0 +1,145 @@
+// Copyright 1998-2016 Epic Games, Inc. All Rights Reserved.
+
+#pragma once
+
+#include "OutputDevice.h"
+
+/*-----------------------------------------------------------------------------
+FOutputDeviceRedirector.
+-----------------------------------------------------------------------------*/
+
+/** The type of lines buffered by secondary threads. */
+struct FBufferedLine
+{
+	const FString Data;
+	const FName Category;
+	const double Time;
+	const ELogVerbosity::Type Verbosity;
+
+	/** Initialization constructor. */
+	FBufferedLine(const TCHAR* InData, const class FName& InCategory, ELogVerbosity::Type InVerbosity, const double InTime = -1)
+		: Data(InData)
+		, Category(InCategory)
+		, Time(InTime)
+		, Verbosity(InVerbosity)
+	{}
+};
+
+/**
+* Class used for output redirection to allow logs to show
+*/
+class CORE_API FOutputDeviceRedirector : public FOutputDevice
+{
+private:
+	/** A FIFO of lines logged by non-master threads. */
+	TArray<FBufferedLine> BufferedLines;
+
+	/** A FIFO backlog of messages logged before the editor had a chance to intercept them. */
+	TArray<FBufferedLine> BacklogLines;
+
+	/** Array of output devices to redirect to */
+	TArray<FOutputDevice*> OutputDevices;
+
+	/** The master thread ID.  Logging from other threads will be buffered for processing by the master thread. */
+	uint32 MasterThreadID;
+
+	/** Whether backlogging is enabled. */
+	bool bEnableBacklog;
+
+	/** Object used for synchronization via a scoped lock */
+	FCriticalSection	SynchronizationObject;
+
+	/**
+	* The unsynchronized version of FlushThreadedLogs.
+	* Assumes that the caller holds a lock on SynchronizationObject.
+	* @param bUseAllDevices - if true this method will use all output devices
+	*/
+	void UnsynchronizedFlushThreadedLogs(bool bUseAllDevices);
+
+public:
+
+	/** Initialization constructor. */
+	FOutputDeviceRedirector();
+
+	/**
+	* Get the GLog singleton
+	*/
+	static FOutputDeviceRedirector* Get();
+
+	/**
+	* Adds an output device to the chain of redirections.
+	*
+	* @param OutputDevice	output device to add
+	*/
+	void AddOutputDevice(FOutputDevice* OutputDevice);
+
+	/**
+	* Removes an output device from the chain of redirections.
+	*
+	* @param OutputDevice	output device to remove
+	*/
+	void RemoveOutputDevice(FOutputDevice* OutputDevice);
+
+	/**
+	* Returns whether an output device is currently in the list of redirectors.
+	*
+	* @param	OutputDevice	output device to check the list against
+	* @return	true if messages are currently redirected to the the passed in output device, false otherwise
+	*/
+	bool IsRedirectingTo(FOutputDevice* OutputDevice);
+
+	/** Flushes lines buffered by secondary threads. */
+	virtual void FlushThreadedLogs();
+
+	/**
+	*	Flushes lines buffered by secondary threads.
+	*	Only used if a background thread crashed and we needed to push the callstack into the log.
+	*/
+	virtual void PanicFlushThreadedLogs();
+
+	/**
+	* Serializes the current backlog to the specified output device.
+	* @param OutputDevice	- Output device that will receive the current backlog
+	*/
+	virtual void SerializeBacklog(FOutputDevice* OutputDevice);
+
+	/**
+	* Enables or disables the backlog.
+	* @param bEnable	- Starts saving a backlog if true, disables and discards any backlog if false
+	*/
+	virtual void EnableBacklog(bool bEnable);
+
+	/**
+	* Sets the current thread to be the master thread that prints directly
+	* (isn't queued up)
+	*/
+	virtual void SetCurrentThreadAsMasterThread();
+
+	/**
+	* Serializes the passed in data via all current output devices.
+	*
+	* @param	Data	Text to log
+	* @param	Event	Event name used for suppression purposes
+	*/
+	virtual void Serialize(const TCHAR* Data, ELogVerbosity::Type Verbosity, const class FName& Category, const double Time) override;
+
+	/**
+	* Serializes the passed in data via all current output devices.
+	*
+	* @param	Data	Text to log
+	* @param	Event	Event name used for suppression purposes
+	*/
+	virtual void Serialize(const TCHAR* Data, ELogVerbosity::Type Verbosity, const class FName& Category) override;
+
+	/**
+	* Passes on the flush request to all current output devices.
+	*/
+	void Flush() override;
+
+	/**
+	* Closes output device and cleans up. This can't happen in the destructor
+	* as we might have to call "delete" which cannot be done for static/ global
+	* objects.
+	*/
+	void TearDown() override;
+};
diff --git a/Engine/Source/Runtime/Core/Public/Misc/QueuedThreadPool.h b/Engine/Source/Runtime/Core/Public/Misc/QueuedThreadPool.h
new file mode 100644
--- /dev/null
+++ b/Engine/Source/Runtime/Core/Public/Misc/QueuedThreadPool.h
@@ -0,0 +1,88 @@
+// Copyright 1998-2016 Epic Games, Inc. All Rights Reserved.
+
+#pragma once
+
+/**
+ * Interface for queued thread pools.
+ *
+ * This interface is used by all queued thread pools. It used as a callback by
+ * FQueuedThreads and is used to queue asynchronous work for callers.
+ */
+class CORE_API FQueuedThreadPool
+{
+public:
+
+	/**
+	 * Creates the thread pool with the specified number of threads
+	 *
+	 * @param InNumQueuedThreads Specifies the number of threads to use in the pool
+	 * @param StackSize The size of stack the threads in the pool need (32K default)
+	 * @param ThreadPriority priority of new pool thread
+	 * @return Whether the pool creation was successful or not
+	 */
+	virtual bool Create( uint32 InNumQueuedThreads, uint32 StackSize = (32 * 1024), EThreadPriority ThreadPriority=TPri_Normal ) = 0;
+
+	/** Tells the pool to clean up all background threads */
+	virtual void Destroy() = 0;
+
+	/**
+	 * Checks to see if there is a thread available to perform the task. If not,
+	 * it queues the work for later. Otherwise it is immediately dispatched.
+	 *
+	 * @param InQueuedWork The work that needs to be done asynchronously
+	 * @see RetractQueuedWork
+	 */
+	virtual void AddQueuedWork( IQueuedWork* InQueuedWork ) = 0;
+
+	/**
+	 * Attempts to retract a previously queued task.
+	 *
+	 * @param InQueuedWork The work to try to retract
+	 * @return true if the work was retracted
+	 * @see AddQueuedWork
+	 */
+	virtual bool RetractQueuedWork( IQueuedWork* InQueuedWork ) = 0;
+
+	/**
+	 * Places a thread back into the available pool
+	 *
+	 * @param InQueuedThread The thread that is ready to be pooled
+	 * @return next job or null if there is no job available now
+	 */
+	virtual IQueuedWork* ReturnToPoolOrGetNextJob( class FQueuedThread* InQueuedThread ) = 0;
+
+public:
+
+	/** Virtual destructor. */
+	virtual ~FQueuedThreadPool() { }
+
+public:
+
+	/**
+	 * Allocates a thread pool
+	 *
+	 * @return A new thread pool.
+	 */
+	static FQueuedThreadPool* Allocate();
+
+	/**
+	 *	Stack size for threads created for the thread pool. 
+	 *	Can be overridden by other projects.
+	 *	If 0 means to use the value passed in the Create method.
+	 */
+	static uint32 OverrideStackSize;
+};
+
+
+/**
+ *  Global thread pool for shared async operations
+ */
+extern CORE_API FQueuedThreadPool* GThreadPool;
+
+#if USE_NEW_ASYNC_IO
+extern CORE_API FQueuedThreadPool* GIOThreadPool;
+#endif // USE_NEW_ASYNC_IO
+
+#if WITH_EDITOR
+extern CORE_API FQueuedThreadPool* GLargeThreadPool;
+#endif
diff --git a/Engine/Source/Runtime/Core/Public/Misc/ScopeLock.h b/Engine/Source/Runtime/Core/Public/Misc/ScopeLock.h
new file mode 100644
--- /dev/null
+++ b/Engine/Source/Runtime/Core/Public/Misc/ScopeLock.h
@@ -0,0 +1,62 @@
+// Copyright 1998-2016 Epic Games, Inc. All Rights Reserved.
+
+#pragma once
+
+/**
+ * Implements a scope lock.
+ *
+ * This is a utility class that handles scope level locking. It's very useful
+ * to keep from causing deadlocks due to exceptions being caught and knowing
+ * about the number of locks a given thread has on a resource. Example:
+ *
+ * <code>
+ *	{
+ *		// Synchronize thread access to the following data
+ *		FScopeLock ScopeLock(SynchObject);
+ *		// Access data that is shared among multiple threads
+ *		...
+ *		// When ScopeLock goes out of scope, other threads can access data
+ *	}
+ * </code>
+ */
+class FScopeLock
+{
+public:
+
+	/**
+	 * Constructor that performs a lock on the synchronization object
+	 *
+	 * @param InSynchObject The synchronization object to manage
+	 */
+	FScopeLock( FCriticalSection* InSynchObject )
+		: SynchObject(InSynchObject)
+	{
+		check(SynchObject);
+		SynchObject->Lock();
+	}
+
+	/** Destructor that performs a release on the synchronization object. */
+	~FScopeLock()
+	{
+		check(SynchObject);
+		SynchObject->Unlock();
+	}
+private:
+
+	/** Default constructor (hidden on purpose). */
+	FScopeLock();
+
+	/** Copy constructor( hidden on purpose). */
+	FScopeLock(const FScopeLock& InScopeLock);
+
+	/** Assignment operator (hidden on purpose). */
+	FScopeLock& operator=( FScopeLock& InScopeLock )
+	{
+		return *this;
+	}
+
+private:
+
+	// Holds the synchronization object to aggregate and scope manage.
+	FCriticalSection* SynchObject;
+};
diff --git a/Engine/Source/Runtime/Core/Public/Misc/SingleThreadRunnable.h b/Engine/Source/Runtime/Core/Public/Misc/SingleThreadRunnable.h
new file mode 100644
--- /dev/null
+++ b/Engine/Source/Runtime/Core/Public/Misc/SingleThreadRunnable.h
@@ -0,0 +1,17 @@
+// Copyright 1998-2016 Epic Games, Inc. All Rights Reserved.
+
+#pragma once
+
+/** 
+ * Interface for ticking runnables when there's only one thread available and 
+ * multithreading is disabled.
+ */
+class CORE_API FSingleThreadRunnable
+{
+public:
+
+	virtual ~FSingleThreadRunnable() { }
+
+	/* Tick function. */
+	virtual void Tick() = 0;
+};
diff --git a/Engine/Source/Runtime/Core/Public/Misc/VarArgs.h b/Engine/Source/Runtime/Core/Public/Misc/VarArgs.h
new file mode 100644
--- /dev/null
+++ b/Engine/Source/Runtime/Core/Public/Misc/VarArgs.h
@@ -0,0 +1,156 @@
+// Copyright 1998-2016 Epic Games, Inc. All Rights Reserved.
+
+#pragma once
+
+/*-----------------------------------------------------------------------------
+VarArgs helper macros.
+-----------------------------------------------------------------------------*/
+
+// phantom definitions to help VAX parse our VARARG_* macros (VAX build 1440)
+#ifdef VISUAL_ASSIST_HACK
+	#define VARARG_DECL( FuncRet, StaticFuncRet, Return, FuncName, Pure, FmtType, ExtraDecl, ExtraCall ) FuncRet FuncName( ExtraDecl FmtType Fmt, ... )
+	#define VARARG_BODY( FuncRet, FuncName, FmtType, ExtraDecl ) FuncRet FuncName( ExtraDecl FmtType Fmt, ... )
+#endif // VISUAL_ASSIST_HACK
+
+
+#define GET_VARARGS(msg, msgsize, len, lastarg, fmt) \
+	{ \
+		va_list ap; \
+		va_start(ap, lastarg); \
+		FCString::GetVarArgs(msg, msgsize, len, fmt, ap); \
+	}
+#define GET_VARARGS_WIDE(msg, msgsize, len, lastarg, fmt) \
+	{ \
+		va_list ap; \
+		va_start(ap, lastarg); \
+		FCStringWide::GetVarArgs(msg, msgsize, len, fmt, ap); \
+	}
+#define GET_VARARGS_ANSI(msg, msgsize, len, lastarg, fmt) \
+	{ \
+		va_list ap; \
+		va_start(ap, lastarg); \
+		FCStringAnsi::GetVarArgs(msg, msgsize, len, fmt, ap); \
+	}
+#define GET_VARARGS_RESULT(msg, msgsize, len, lastarg, fmt, result) \
+	{ \
+		va_list ap; \
+		va_start(ap, lastarg); \
+		result = FCString::GetVarArgs(msg, msgsize, len, fmt, ap); \
+		if (result >= msgsize) \
+		{ \
+			result = -1; \
+		} \
+	}
+#define GET_VARARGS_RESULT_WIDE(msg, msgsize, len, lastarg, fmt, result) \
+	{ \
+		va_list ap; \
+		va_start(ap, lastarg); \
+		result = FCStringWide::GetVarArgs(msg, msgsize, len, fmt, ap); \
+		if (result >= msgsize) \
+		{ \
+			result = -1; \
+		} \
+	}
+#define GET_VARARGS_RESULT_ANSI(msg, msgsize, len, lastarg, fmt, result) \
+	{ \
+		va_list ap; \
+		va_start(ap, lastarg); \
+		result = FCStringAnsi::GetVarArgs(msg, msgsize, len, fmt, ap); \
+		if (result >= msgsize) \
+		{ \
+			result = -1; \
+		} \
+	}
+
+/*-----------------------------------------------------------------------------
+Ugly VarArgs type checking (debug builds only).
+-----------------------------------------------------------------------------*/
+
+#define VARARG_EXTRA(A) A,
+#define VARARG_NONE
+#define VARARG_PURE =0
+
+#if PLATFORM_WINDOWS
+
+	static inline uint32			CheckVA(uint32 dw) { return dw; }
+	static inline uint8			CheckVA(uint8 b) { return b; }
+	static inline int32			CheckVA(int32 i) { return i; }
+	static inline uint64			CheckVA(uint64 qw) { return qw; }
+	static inline int64		CheckVA(int64 sqw) { return sqw; }
+	static inline double		CheckVA(double d) { return d; }
+	static inline long		CheckVA(long d) { return d; }
+	static inline long		CheckVA(unsigned long d) { return d; }
+	static inline TCHAR			CheckVA(TCHAR c) { return c; }
+	static inline void*			CheckVA(ANSICHAR* s) { return (void*)s; }
+	static inline bool		CheckVA(bool b) { return b; }
+	template<class T> T*		CheckVA(T* p) { return p; }
+	template<class T> const T*	CheckVA(const T* p) { return p; }
+
+	#define VARARG_DECL( FuncRet, StaticFuncRet, Return, FuncName, Pure, FmtType, ExtraDecl, ExtraCall )	\
+		FuncRet FuncName##__VA( ExtraDecl FmtType Fmt, ... ) Pure;  \
+		StaticFuncRet FuncName(ExtraDecl FmtType Fmt) {Return FuncName##__VA(ExtraCall (Fmt));} \
+		template<class T1> \
+		StaticFuncRet FuncName(ExtraDecl FmtType Fmt,T1 V1) {T1 v1=CheckVA(V1);Return FuncName##__VA(ExtraCall (Fmt),(v1));} \
+		template<class T1,class T2> \
+		StaticFuncRet FuncName(ExtraDecl FmtType Fmt,T1 V1,T2 V2) {T1 v1=CheckVA(V1);T2 v2=CheckVA(V2);Return FuncName##__VA(ExtraCall (Fmt),(v1),(v2));} \
+		template<class T1,class T2,class T3> \
+		StaticFuncRet FuncName(ExtraDecl FmtType Fmt,T1 V1,T2 V2,T3 V3) {T1 v1=CheckVA(V1);T2 v2=CheckVA(V2);T3 v3=CheckVA(V3);Return FuncName##__VA(ExtraCall (Fmt),(v1),(v2),(v3));} \
+		template<class T1,class T2,class T3,class T4> \
+		StaticFuncRet FuncName(ExtraDecl FmtType Fmt,T1 V1,T2 V2,T3 V3,T4 V4) {T1 v1=CheckVA(V1);T2 v2=CheckVA(V2);T3 v3=CheckVA(V3);T4 v4=CheckVA(V4);Return FuncName##__VA(ExtraCall (Fmt),(v1),(v2),(v3),(v4));} \
+		template<class T1,class T2,class T3,class T4,class T5> \
+		StaticFuncRet FuncName(ExtraDecl FmtType Fmt,T1 V1,T2 V2,T3 V3,T4 V4,T5 V5) {T1 v1=CheckVA(V1);T2 v2=CheckVA(V2);T3 v3=CheckVA(V3);T4 v4=CheckVA(V4);T5 v5=CheckVA(V5);Return FuncName##__VA(ExtraCall (Fmt),(v1),(v2),(v3),(v4),(v5));} \
+		template<class T1,class T2,class T3,class T4,class T5,class T6> \
+		StaticFuncRet FuncName(ExtraDecl FmtType Fmt,T1 V1,T2 V2,T3 V3,T4 V4,T5 V5,T6 V6) {T1 v1=CheckVA(V1);T2 v2=CheckVA(V2);T3 v3=CheckVA(V3);T4 v4=CheckVA(V4);T5 v5=CheckVA(V5);T6 v6=CheckVA(V6);Return FuncName##__VA(ExtraCall (Fmt),(v1),(v2),(v3),(v4),(v5),(v6));} \
+		template<class T1,class T2,class T3,class T4,class T5,class T6,class T7> \
+		StaticFuncRet FuncName(ExtraDecl FmtType Fmt,T1 V1,T2 V2,T3 V3,T4 V4,T5 V5,T6 V6,T7 V7) {T1 v1=CheckVA(V1);T2 v2=CheckVA(V2);T3 v3=CheckVA(V3);T4 v4=CheckVA(V4);T5 v5=CheckVA(V5);T6 v6=CheckVA(V6);T7 v7=CheckVA(V7);Return FuncName##__VA(ExtraCall (Fmt),(v1),(v2),(v3),(v4),(v5),(v6),(v7));} \
+		template<class T1,class T2,class T3,class T4,class T5,class T6,class T7,class T8> \
+		StaticFuncRet FuncName(ExtraDecl FmtType Fmt,T1 V1,T2 V2,T3 V3,T4 V4,T5 V5,T6 V6,T7 V7,T8 V8) {T1 v1=CheckVA(V1);T2 v2=CheckVA(V2);T3 v3=CheckVA(V3);T4 v4=CheckVA(V4);T5 v5=CheckVA(V5);T6 v6=CheckVA(V6);T7 v7=CheckVA(V7);T8 v8=CheckVA(V8);Return FuncName##__VA(ExtraCall (Fmt),(v1),(v2),(v3),(v4),(v5),(v6),(v7),(v8));} \
+		template<class T1,class T2,class T3,class T4,class T5,class T6,class T7,class T8,class T9> \
+		StaticFuncRet FuncName(ExtraDecl FmtType Fmt,T1 V1,T2 V2,T3 V3,T4 V4,T5 V5,T6 V6,T7 V7,T8 V8,T9 V9) {T1 v1=CheckVA(V1);T2 v2=CheckVA(V2);T3 v3=CheckVA(V3);T4 v4=CheckVA(V4);T5 v5=CheckVA(V5);T6 v6=CheckVA(V6);T7 v7=CheckVA(V7);T8 v8=CheckVA(V8);T9 v9=CheckVA(V9);Return FuncName##__VA(ExtraCall (Fmt),(v1),(v2),(v3),(v4),(v5),(v6),(v7),(v8),(v9));} \
+		template<class T1,class T2,class T3,class T4,class T5,class T6,class T7,class T8,class T9,class T10> \
+		StaticFuncRet FuncName(ExtraDecl FmtType Fmt,T1 V1,T2 V2,T3 V3,T4 V4,T5 V5,T6 V6,T7 V7,T8 V8,T9 V9,T10 V10) {T1 v1=CheckVA(V1);T2 v2=CheckVA(V2);T3 v3=CheckVA(V3);T4 v4=CheckVA(V4);T5 v5=CheckVA(V5);T6 v6=CheckVA(V6);T7 v7=CheckVA(V7);T8 v8=CheckVA(V8);T9 v9=CheckVA(V9);T10 v10=CheckVA(V10);Return FuncName##__VA(ExtraCall (Fmt),(v1),(v2),(v3),(v4),(v5),(v6),(v7),(v8),(v9),(v10));} \
+		template<class T1,class T2,class T3,class T4,class T5,class T6,class T7,class T8,class T9,class T10,class T11> \
+		StaticFuncRet FuncName(ExtraDecl FmtType Fmt,T1 V1,T2 V2,T3 V3,T4 V4,T5 V5,T6 V6,T7 V7,T8 V8,T9 V9,T10 V10,T11 V11) {T1 v1=CheckVA(V1);T2 v2=CheckVA(V2);T3 v3=CheckVA(V3);T4 v4=CheckVA(V4);T5 v5=CheckVA(V5);T6 v6=CheckVA(V6);T7 v7=CheckVA(V7);T8 v8=CheckVA(V8);T9 v9=CheckVA(V9);T10 v10=CheckVA(V10);T11 v11=CheckVA(V11);Return FuncName##__VA(ExtraCall (Fmt),(v1),(v2),(v3),(v4),(v5),(v6),(v7),(v8),(v9),(v10),(v11));} \
+		template<class T1,class T2,class T3,class T4,class T5,class T6,class T7,class T8,class T9,class T10,class T11,class T12> \
+		StaticFuncRet FuncName(ExtraDecl FmtType Fmt,T1 V1,T2 V2,T3 V3,T4 V4,T5 V5,T6 V6,T7 V7,T8 V8,T9 V9,T10 V10,T11 V11,T12 V12) {T1 v1=CheckVA(V1);T2 v2=CheckVA(V2);T3 v3=CheckVA(V3);T4 v4=CheckVA(V4);T5 v5=CheckVA(V5);T6 v6=CheckVA(V6);T7 v7=CheckVA(V7);T8 v8=CheckVA(V8);T9 v9=CheckVA(V9);T10 v10=CheckVA(V10);T11 v11=CheckVA(V11);T12 v12=CheckVA(V12);Return FuncName##__VA(ExtraCall (Fmt),(v1),(v2),(v3),(v4),(v5),(v6),(v7),(v8),(v9),(v10),(v11),(v12));} \
+		template<class T1,class T2,class T3,class T4,class T5,class T6,class T7,class T8,class T9,class T10,class T11,class T12,class T13> \
+		StaticFuncRet FuncName(ExtraDecl FmtType Fmt,T1 V1,T2 V2,T3 V3,T4 V4,T5 V5,T6 V6,T7 V7,T8 V8,T9 V9,T10 V10,T11 V11,T12 V12,T13 V13) {T1 v1=CheckVA(V1);T2 v2=CheckVA(V2);T3 v3=CheckVA(V3);T4 v4=CheckVA(V4);T5 v5=CheckVA(V5);T6 v6=CheckVA(V6);T7 v7=CheckVA(V7);T8 v8=CheckVA(V8);T9 v9=CheckVA(V9);T10 v10=CheckVA(V10);T11 v11=CheckVA(V11);T12 v12=CheckVA(V12);T13 v13=CheckVA(V13);Return FuncName##__VA(ExtraCall (Fmt),(v1),(v2),(v3),(v4),(v5),(v6),(v7),(v8),(v9),(v10),(v11),(v12),(v13));} \
+		template<class T1,class T2,class T3,class T4,class T5,class T6,class T7,class T8,class T9,class T10,class T11,class T12,class T13,class T14> \
+		StaticFuncRet FuncName(ExtraDecl FmtType Fmt,T1 V1,T2 V2,T3 V3,T4 V4,T5 V5,T6 V6,T7 V7,T8 V8,T9 V9,T10 V10,T11 V11,T12 V12,T13 V13,T14 V14) {T1 v1=CheckVA(V1);T2 v2=CheckVA(V2);T3 v3=CheckVA(V3);T4 v4=CheckVA(V4);T5 v5=CheckVA(V5);T6 v6=CheckVA(V6);T7 v7=CheckVA(V7);T8 v8=CheckVA(V8);T9 v9=CheckVA(V9);T10 v10=CheckVA(V10);T11 v11=CheckVA(V11);T12 v12=CheckVA(V12);T13 v13=CheckVA(V13);T14 v14=CheckVA(V14);Return FuncName##__VA(ExtraCall (Fmt),(v1),(v2),(v3),(v4),(v5),(v6),(v7),(v8),(v9),(v10),(v11),(v12),(v13),(v14));} \
+		template<class T1,class T2,class T3,class T4,class T5,class T6,class T7,class T8,class T9,class T10,class T11,class T12,class T13,class T14,class T15> \
+		StaticFuncRet FuncName(ExtraDecl FmtType Fmt,T1 V1,T2 V2,T3 V3,T4 V4,T5 V5,T6 V6,T7 V7,T8 V8,T9 V9,T10 V10,T11 V11,T12 V12,T13 V13,T14 V14,T15 V15) {T1 v1=CheckVA(V1);T2 v2=CheckVA(V2);T3 v3=CheckVA(V3);T4 v4=CheckVA(V4);T5 v5=CheckVA(V5);T6 v6=CheckVA(V6);T7 v7=CheckVA(V7);T8 v8=CheckVA(V8);T9 v9=CheckVA(V9);T10 v10=CheckVA(V10);T11 v11=CheckVA(V11);T12 v12=CheckVA(V12);T13 v13=CheckVA(V13);T14 v14=CheckVA(V14);T15 v15=CheckVA(V15);Return FuncName##__VA(ExtraCall (Fmt),(v1),(v2),(v3),(v4),(v5),(v6),(v7),(v8),(v9),(v10),(v11),(v12),(v13),(v14),(v15));} \
+		template<class T1,class T2,class T3,class T4,class T5,class T6,class T7,class T8,class T9,class T10,class T11,class T12,class T13,class T14,class T15,class T16> \
+		StaticFuncRet FuncName(ExtraDecl FmtType Fmt,T1 V1,T2 V2,T3 V3,T4 V4,T5 V5,T6 V6,T7 V7,T8 V8,T9 V9,T10 V10,T11 V11,T12 V12,T13 V13,T14 V14,T15 V15,T16 V16) {T1 v1=CheckVA(V1);T2 v2=CheckVA(V2);T3 v3=CheckVA(V3);T4 v4=CheckVA(V4);T5 v5=CheckVA(V5);T6 v6=CheckVA(V6);T7 v7=CheckVA(V7);T8 v8=CheckVA(V8);T9 v9=CheckVA(V9);T10 v10=CheckVA(V10);T11 v11=CheckVA(V11);T12 v12=CheckVA(V12);T13 v13=CheckVA(V13);T14 v14=CheckVA(V14);T15 v15=CheckVA(V15);T16 v16=CheckVA(V16);Return FuncName##__VA(ExtraCall (Fmt),(v1),(v2),(v3),(v4),(v5),(v6),(v7),(v8),(v9),(v10),(v11),(v12),(v13),(v14),(v15),(v16));} \
+		template<class T1,class T2,class T3,class T4,class T5,class T6,class T7,class T8,class T9,class T10,class T11,class T12,class T13,class T14,class T15,class T16,class T17> \
+		StaticFuncRet FuncName(ExtraDecl FmtType Fmt,T1 V1,T2 V2,T3 V3,T4 V4,T5 V5,T6 V6,T7 V7,T8 V8,T9 V9,T10 V10,T11 V11,T12 V12,T13 V13,T14 V14,T15 V15,T16 V16,T17 V17) {T1 v1=CheckVA(V1);T2 v2=CheckVA(V2);T3 v3=CheckVA(V3);T4 v4=CheckVA(V4);T5 v5=CheckVA(V5);T6 v6=CheckVA(V6);T7 v7=CheckVA(V7);T8 v8=CheckVA(V8);T9 v9=CheckVA(V9);T10 v10=CheckVA(V10);T11 v11=CheckVA(V11);T12 v12=CheckVA(V12);T13 v13=CheckVA(V13);T14 v14=CheckVA(V14);T15 v15=CheckVA(V15);T16 v16=CheckVA(V16);T17 v17=CheckVA(V17);Return FuncName##__VA(ExtraCall (Fmt),(v1),(v2),(v3),(v4),(v5),(v6),(v7),(v8),(v9),(v10),(v11),(v12),(v13),(v14),(v15),(v16),(v17));} \
+		template<class T1,class T2,class T3,class T4,class T5,class T6,class T7,class T8,class T9,class T10,class T11,class T12,class T13,class T14,class T15,class T16,class T17,class T18> \
+		StaticFuncRet FuncName(ExtraDecl FmtType Fmt,T1 V1,T2 V2,T3 V3,T4 V4,T5 V5,T6 V6,T7 V7,T8 V8,T9 V9,T10 V10,T11 V11,T12 V12,T13 V13,T14 V14,T15 V15,T16 V16,T17 V17,T18 V18) {T1 v1=CheckVA(V1);T2 v2=CheckVA(V2);T3 v3=CheckVA(V3);T4 v4=CheckVA(V4);T5 v5=CheckVA(V5);T6 v6=CheckVA(V6);T7 v7=CheckVA(V7);T8 v8=CheckVA(V8);T9 v9=CheckVA(V9);T10 v10=CheckVA(V10);T11 v11=CheckVA(V11);T12 v12=CheckVA(V12);T13 v13=CheckVA(V13);T14 v14=CheckVA(V14);T15 v15=CheckVA(V15);T16 v16=CheckVA(V16);T17 v17=CheckVA(V17);T18 v18=CheckVA(V18);Return FuncName##__VA(ExtraCall (Fmt),(v1),(v2),(v3),(v4),(v5),(v6),(v7),(v8),(v9),(v10),(v11),(v12),(v13),(v14),(v15),(v16),(v17),(v18));} \
+		template<class T1,class T2,class T3,class T4,class T5,class T6,class T7,class T8,class T9,class T10,class T11,class T12,class T13,class T14,class T15,class T16,class T17,class T18,class T19> \
+		StaticFuncRet FuncName(ExtraDecl FmtType Fmt,T1 V1,T2 V2,T3 V3,T4 V4,T5 V5,T6 V6,T7 V7,T8 V8,T9 V9,T10 V10,T11 V11,T12 V12,T13 V13,T14 V14,T15 V15,T16 V16,T17 V17,T18 V18,T19 V19) {T1 v1=CheckVA(V1);T2 v2=CheckVA(V2);T3 v3=CheckVA(V3);T4 v4=CheckVA(V4);T5 v5=CheckVA(V5);T6 v6=CheckVA(V6);T7 v7=CheckVA(V7);T8 v8=CheckVA(V8);T9 v9=CheckVA(V9);T10 v10=CheckVA(V10);T11 v11=CheckVA(V11);T12 v12=CheckVA(V12);T13 v13=CheckVA(V13);T14 v14=CheckVA(V14);T15 v15=CheckVA(V15);T16 v16=CheckVA(V16);T17 v17=CheckVA(V17);T18 v18=CheckVA(V18);T19 v19=CheckVA(V19);Return FuncName##__VA(ExtraCall (Fmt),(v1),(v2),(v3),(v4),(v5),(v6),(v7),(v8),(v9),(v10),(v11),(v12),(v13),(v14),(v15),(v16),(v17),(v18),(v19));} \
+		template<class T1,class T2,class T3,class T4,class T5,class T6,class T7,class T8,class T9,class T10,class T11,class T12,class T13,class T14,class T15,class T16,class T17,class T18,class T19,class T20> \
+		StaticFuncRet FuncName(ExtraDecl FmtType Fmt,T1 V1,T2 V2,T3 V3,T4 V4,T5 V5,T6 V6,T7 V7,T8 V8,T9 V9,T10 V10,T11 V11,T12 V12,T13 V13,T14 V14,T15 V15,T16 V16,T17 V17,T18 V18,T19 V19,T20 V20) {T1 v1=CheckVA(V1);T2 v2=CheckVA(V2);T3 v3=CheckVA(V3);T4 v4=CheckVA(V4);T5 v5=CheckVA(V5);T6 v6=CheckVA(V6);T7 v7=CheckVA(V7);T8 v8=CheckVA(V8);T9 v9=CheckVA(V9);T10 v10=CheckVA(V10);T11 v11=CheckVA(V11);T12 v12=CheckVA(V12);T13 v13=CheckVA(V13);T14 v14=CheckVA(V14);T15 v15=CheckVA(V15);T16 v16=CheckVA(V16);T17 v17=CheckVA(V17);T18 v18=CheckVA(V18);T19 v19=CheckVA(V19);T20 v20=CheckVA(V20);Return FuncName##__VA(ExtraCall (Fmt),(v1),(v2),(v3),(v4),(v5),(v6),(v7),(v8),(v9),(v10),(v11),(v12),(v13),(v14),(v15),(v16),(v17),(v18),(v19),(v20));} \
+		template<class T1,class T2,class T3,class T4,class T5,class T6,class T7,class T8,class T9,class T10,class T11,class T12,class T13,class T14,class T15,class T16,class T17,class T18,class T19,class T20,class T21> \
+		StaticFuncRet FuncName(ExtraDecl FmtType Fmt,T1 V1,T2 V2,T3 V3,T4 V4,T5 V5,T6 V6,T7 V7,T8 V8,T9 V9,T10 V10,T11 V11,T12 V12,T13 V13,T14 V14,T15 V15,T16 V16,T17 V17,T18 V18,T19 V19,T20 V20, T21 V21) {T1 v1=CheckVA(V1);T2 v2=CheckVA(V2);T3 v3=CheckVA(V3);T4 v4=CheckVA(V4);T5 v5=CheckVA(V5);T6 v6=CheckVA(V6);T7 v7=CheckVA(V7);T8 v8=CheckVA(V8);T9 v9=CheckVA(V9);T10 v10=CheckVA(V10);T11 v11=CheckVA(V11);T12 v12=CheckVA(V12);T13 v13=CheckVA(V13);T14 v14=CheckVA(V14);T15 v15=CheckVA(V15);T16 v16=CheckVA(V16);T17 v17=CheckVA(V17);T18 v18=CheckVA(V18);T19 v19=CheckVA(V19);T20 v20=CheckVA(V20);T21 v21=CheckVA(V21);Return FuncName##__VA(ExtraCall (Fmt),(v1),(v2),(v3),(v4),(v5),(v6),(v7),(v8),(v9),(v10),(v11),(v12),(v13),(v14),(v15),(v16),(v17),(v18),(v19),(v20),(v21));} \
+		template<class T1,class T2,class T3,class T4,class T5,class T6,class T7,class T8,class T9,class T10,class T11,class T12,class T13,class T14,class T15,class T16,class T17,class T18,class T19,class T20,class T21,class T22> \
+		StaticFuncRet FuncName(ExtraDecl FmtType Fmt,T1 V1,T2 V2,T3 V3,T4 V4,T5 V5,T6 V6,T7 V7,T8 V8,T9 V9,T10 V10,T11 V11,T12 V12,T13 V13,T14 V14,T15 V15,T16 V16,T17 V17,T18 V18,T19 V19,T20 V20, T21 V21, T22 V22) {T1 v1=CheckVA(V1);T2 v2=CheckVA(V2);T3 v3=CheckVA(V3);T4 v4=CheckVA(V4);T5 v5=CheckVA(V5);T6 v6=CheckVA(V6);T7 v7=CheckVA(V7);T8 v8=CheckVA(V8);T9 v9=CheckVA(V9);T10 v10=CheckVA(V10);T11 v11=CheckVA(V11);T12 v12=CheckVA(V12);T13 v13=CheckVA(V13);T14 v14=CheckVA(V14);T15 v15=CheckVA(V15);T16 v16=CheckVA(V16);T17 v17=CheckVA(V17);T18 v18=CheckVA(V18);T19 v19=CheckVA(V19);T20 v20=CheckVA(V20);T21 v21=CheckVA(V21);T22 v22=CheckVA(V22);Return FuncName##__VA(ExtraCall (Fmt),(v1),(v2),(v3),(v4),(v5),(v6),(v7),(v8),(v9),(v10),(v11),(v12),(v13),(v14),(v15),(v16),(v17),(v18),(v19),(v20),(v21),(v22));} \
+		template<class T1,class T2,class T3,class T4,class T5,class T6,class T7,class T8,class T9,class T10,class T11,class T12,class T13,class T14,class T15,class T16,class T17,class T18,class T19,class T20,class T21,class T22,class T23> \
+		StaticFuncRet FuncName(ExtraDecl FmtType Fmt,T1 V1,T2 V2,T3 V3,T4 V4,T5 V5,T6 V6,T7 V7,T8 V8,T9 V9,T10 V10,T11 V11,T12 V12,T13 V13,T14 V14,T15 V15,T16 V16,T17 V17,T18 V18,T19 V19,T20 V20, T21 V21, T22 V22, T23 V23) {T1 v1=CheckVA(V1);T2 v2=CheckVA(V2);T3 v3=CheckVA(V3);T4 v4=CheckVA(V4);T5 v5=CheckVA(V5);T6 v6=CheckVA(V6);T7 v7=CheckVA(V7);T8 v8=CheckVA(V8);T9 v9=CheckVA(V9);T10 v10=CheckVA(V10);T11 v11=CheckVA(V11);T12 v12=CheckVA(V12);T13 v13=CheckVA(V13);T14 v14=CheckVA(V14);T15 v15=CheckVA(V15);T16 v16=CheckVA(V16);T17 v17=CheckVA(V17);T18 v18=CheckVA(V18);T19 v19=CheckVA(V19);T20 v20=CheckVA(V20);T21 v21=CheckVA(V21);T22 v22=CheckVA(V22);T23 v23=CheckVA(V23);Return FuncName##__VA(ExtraCall (Fmt),(v1),(v2),(v3),(v4),(v5),(v6),(v7),(v8),(v9),(v10),(v11),(v12),(v13),(v14),(v15),(v16),(v17),(v18),(v19),(v20),(v21),(v22),(v23));} \
+		template<class T1,class T2,class T3,class T4,class T5,class T6,class T7,class T8,class T9,class T10,class T11,class T12,class T13,class T14,class T15,class T16,class T17,class T18,class T19,class T20,class T21,class T22,class T23,class T24> \
+		StaticFuncRet FuncName(ExtraDecl FmtType Fmt,T1 V1,T2 V2,T3 V3,T4 V4,T5 V5,T6 V6,T7 V7,T8 V8,T9 V9,T10 V10,T11 V11,T12 V12,T13 V13,T14 V14,T15 V15,T16 V16,T17 V17,T18 V18,T19 V19,T20 V20, T21 V21, T22 V22, T23 V23, T24 V24) {T1 v1=CheckVA(V1);T2 v2=CheckVA(V2);T3 v3=CheckVA(V3);T4 v4=CheckVA(V4);T5 v5=CheckVA(V5);T6 v6=CheckVA(V6);T7 v7=CheckVA(V7);T8 v8=CheckVA(V8);T9 v9=CheckVA(V9);T10 v10=CheckVA(V10);T11 v11=CheckVA(V11);T12 v12=CheckVA(V12);T13 v13=CheckVA(V13);T14 v14=CheckVA(V14);T15 v15=CheckVA(V15);T16 v16=CheckVA(V16);T17 v17=CheckVA(V17);T18 v18=CheckVA(V18);T19 v19=CheckVA(V19);T20 v20=CheckVA(V20);T21 v21=CheckVA(V21);T22 v22=CheckVA(V22);T23 v23=CheckVA(V23);T24 v24=CheckVA(V24);Return FuncName##__VA(ExtraCall (Fmt),(v1),(v2),(v3),(v4),(v5),(v6),(v7),(v8),(v9),(v10),(v11),(v12),(v13),(v14),(v15),(v16),(v17),(v18),(v19),(v20),(v21),(v22),(v23),(v24));} \
+		template<class T1,class T2,class T3,class T4,class T5,class T6,class T7,class T8,class T9,class T10,class T11,class T12,class T13,class T14,class T15,class T16,class T17,class T18,class T19,class T20,class T21,class T22,class T23,class T24,class T25> \
+		StaticFuncRet FuncName(ExtraDecl FmtType Fmt,T1 V1,T2 V2,T3 V3,T4 V4,T5 V5,T6 V6,T7 V7,T8 V8,T9 V9,T10 V10,T11 V11,T12 V12,T13 V13,T14 V14,T15 V15,T16 V16,T17 V17,T18 V18,T19 V19,T20 V20, T21 V21, T22 V22, T23 V23, T24 V24, T25 V25) {T1 v1=CheckVA(V1);T2 v2=CheckVA(V2);T3 v3=CheckVA(V3);T4 v4=CheckVA(V4);T5 v5=CheckVA(V5);T6 v6=CheckVA(V6);T7 v7=CheckVA(V7);T8 v8=CheckVA(V8);T9 v9=CheckVA(V9);T10 v10=CheckVA(V10);T11 v11=CheckVA(V11);T12 v12=CheckVA(V12);T13 v13=CheckVA(V13);T14 v14=CheckVA(V14);T15 v15=CheckVA(V15);T16 v16=CheckVA(V16);T17 v17=CheckVA(V17);T18 v18=CheckVA(V18);T19 v19=CheckVA(V19);T20 v20=CheckVA(V20);T21 v21=CheckVA(V21);T22 v22=CheckVA(V22);T23 v23=CheckVA(V23);T24 v24=CheckVA(V24);T25 v25=CheckVA(V25);Return FuncName##__VA(ExtraCall (Fmt),(v1),(v2),(v3),(v4),(v5),(v6),(v7),(v8),(v9),(v10),(v11),(v12),(v13),(v14),(v15),(v16),(v17),(v18),(v19),(v20),(v21),(v22),(v23),(v24),(v25));} \
+		template<class T1,class T2,class T3,class T4,class T5,class T6,class T7,class T8,class T9,class T10,class T11,class T12,class T13,class T14,class T15,class T16,class T17,class T18,class T19,class T20,class T21,class T22,class T23,class T24,class T25,class T26> \
+		StaticFuncRet FuncName(ExtraDecl FmtType Fmt,T1 V1,T2 V2,T3 V3,T4 V4,T5 V5,T6 V6,T7 V7,T8 V8,T9 V9,T10 V10,T11 V11,T12 V12,T13 V13,T14 V14,T15 V15,T16 V16,T17 V17,T18 V18,T19 V19,T20 V20, T21 V21, T22 V22, T23 V23, T24 V24, T25 V25, T26 V26) {T1 v1=CheckVA(V1);T2 v2=CheckVA(V2);T3 v3=CheckVA(V3);T4 v4=CheckVA(V4);T5 v5=CheckVA(V5);T6 v6=CheckVA(V6);T7 v7=CheckVA(V7);T8 v8=CheckVA(V8);T9 v9=CheckVA(V9);T10 v10=CheckVA(V10);T11 v11=CheckVA(V11);T12 v12=CheckVA(V12);T13 v13=CheckVA(V13);T14 v14=CheckVA(V14);T15 v15=CheckVA(V15);T16 v16=CheckVA(V16);T17 v17=CheckVA(V17);T18 v18=CheckVA(V18);T19 v19=CheckVA(V19);T20 v20=CheckVA(V20);T21 v21=CheckVA(V21);T22 v22=CheckVA(V22);T23 v23=CheckVA(V23);T24 v24=CheckVA(V24);T25 v25=CheckVA(V25);T26 v26=CheckVA(V26);Return FuncName##__VA(ExtraCall (Fmt),(v1),(v2),(v3),(v4),(v5),(v6),(v7),(v8),(v9),(v10),(v11),(v12),(v13),(v14),(v15),(v16),(v17),(v18),(v19),(v20),(v21),(v22),(v23),(v24),(v25),(v26));}
+
+	#define VARARG_BODY( FuncRet, FuncName, FmtType, ExtraDecl )		\
+		FuncRet FuncName##__VA( ExtraDecl  FmtType Fmt, ... )
+
+#else  // !PLATFORM_WINDOWS
+
+	#define VARARG_DECL( FuncRet, StaticFuncRet, Return, FuncName, Pure, FmtType, ExtraDecl, ExtraCall )	\
+		FuncRet FuncName( ExtraDecl FmtType Fmt, ... ) Pure
+
+	#define VARARG_BODY( FuncRet, FuncName, FmtType, ExtraDecl )		\
+		FuncRet FuncName( ExtraDecl FmtType Fmt, ... )
+
+#endif // PLATFORM_WINDOWS
