Copying //UE4/Dev-Build to //UE4/Dev-Main (Source: //UE4/Dev-Build @ 3104200)

diff --git a/Engine/Source/Runtime/Core/Public/Serialization/Archive.h b/Engine/Source/Runtime/Core/Public/Serialization/Archive.h
--- a/Engine/Source/Runtime/Core/Public/Serialization/Archive.h
+++ b/Engine/Source/Runtime/Core/Public/Serialization/Archive.h
@@ -1,703 +1,1236 @@
 // Copyright 1998-2016 Epic Games, Inc. All Rights Reserved.
 
-/*=============================================================================
-	Core utility archive classes.  Must be separate from ArchiveBase.h since 
-	UnrealTemplate.h references FArchive implementation.
-				
-=============================================================================*/
-
 #pragma once
 
-#include "EnumClassFlags.h"
+#include "Containers/EnumAsByte.h"
+#include "HAL/PlatformProperties.h"
+#include "Misc/Compression.h"
+#include "Misc/EngineVersionBase.h"
+#include "TextNamespaceFwd.h"
+
+class FAssetPtr;
+class FCustomVersionContainer;
+class FLazyObjectPtr;
+struct FStringAssetReference;
+struct FUntypedBulkData;
 
 /**
- * Base class for serializing arbitrary data in memory.
+ * Base class for archives that can be used for loading, saving, and garbage
+ * collecting in a byte order neutral way.
  */
-class FMemoryArchive : public FArchive
+class CORE_API FArchive
 {
 public:
+
+	/** Default constructor. */
+	FArchive();
+
+	/** Copy constructor. */
+	FArchive(const FArchive&);
+
 	/**
-  	 * Returns the name of the Archive.  Useful for getting the name of the package a struct or object
-	 * is in when a loading error occurs.
+	 * Copy assignment operator.
 	 *
-	 * This is overridden for the specific Archive Types
-	 **/
-	virtual FString GetArchiveName() const { return TEXT("FMemoryArchive"); }
+	 * @param ArchiveToCopy The archive to copy from.
+	 */
+	FArchive& operator=(const FArchive& ArchiveToCopy);
 
-	void Seek( int64 InPos ) final
-	{
-		Offset = InPos;
-	}
-	int64 Tell() final
-	{
-		return Offset;
-	}
+	/** Destructor. */
+	virtual ~FArchive();
 
-	using FArchive::operator<<; // For visibility of the overloads we don't override
+public:
 
-	virtual FArchive& operator<<( class FName& N ) override
+	/**
+	 * Serializes an FName value from or into this archive.
+	 *
+	 * This operator can be implemented by sub-classes that wish to serialize FName instances.
+	 *
+	 * @param Value The value to serialize.
+	 * @return This instance.
+	 */
+	virtual FArchive& operator<<(class FName& Value)
 	{
-		// Serialize the FName as a string
-		if (IsLoading())
-		{
-			FString StringName;
-			*this << StringName;
-			N = FName(*StringName);
-		}
-		else
-		{
-			FString StringName = N.ToString();
-			*this << StringName;
-		}
 		return *this;
 	}
 
-	virtual FArchive& operator<<( class UObject*& Res ) override
+	/**
+	 * Serializes an FText value from or into an archive.
+	 *
+	 * @param Ar The archive to serialize from or to.
+	 * @param Value The value to serialize.
+	 */
+	virtual FArchive& operator<<(class FText& Value);
+
+	/**
+	 * Serializes an UObject value from or into this archive.
+	 *
+	 * This operator can be implemented by sub-classes that wish to serialize UObject instances.
+	 *
+	 * @param Value The value to serialize.
+	 * @return This instance.
+	 */
+	virtual FArchive& operator<<(class UObject*& Value)
 	{
-		// Not supported through this archive
-		check(0);
 		return *this;
 	}
 
-protected:
+	/**
+	 * Serializes a lazy object pointer value from or into this archive.
+	 *
+	 * Most of the time, FLazyObjectPtrs are serialized as UObject*, but some archives need to override this.
+	 *
+	 * @param Value The value to serialize.
+	 * @return This instance.
+	 */
+	virtual FArchive& operator<<(class FLazyObjectPtr& Value);
+	
+	/**
+	 * Serializes asset pointer from or into this archive.
+	 *
+	 * Most of the time, FAssetPtrs are serialized as UObject *, but some archives need to override this.
+	 *
+	 * @param Value The asset pointer to serialize.
+	 * @return This instance.
+	 */
+	virtual FArchive& operator<<(class FAssetPtr& Value);
+
+	/**
+	 * Serializes string asset reference from or into this archive.
+	 *
+	 * @param Value String asset reference to serialize.
+	 * @return This instance.
+	 */
+	virtual FArchive& operator<<(struct FStringAssetReference& Value);
+
+	/** 
+	 * Inform the archive that a blueprint would like to force finalization, normally
+	 * this is triggered by CDO load, but if there's no CDO we force finalization.
+	 */
+	virtual void ForceBlueprintFinalization() {}
+public:
 
-	/** Marked as protected to avoid instantiating this class directly */
-	FMemoryArchive()
-		: FArchive(), Offset(0)
+	/**
+	 * Serializes an ANSICHAR value from or into an archive.
+	 *
+	 * @param Ar The archive to serialize from or to.
+	 * @param Value The value to serialize.
+	 */
+	FORCEINLINE friend FArchive& operator<<(FArchive& Ar, ANSICHAR& Value)
 	{
+		Ar.Serialize(&Value, 1);
+
+		return Ar;
 	}
 
-	int64				Offset;
-};
+	/**
+	 * Serializes a WIDECHAR value from or into an archive.
+	 *
+	 * @param Ar The archive to serialize from or to.
+	 * @param Value The value to serialize.
+	 */
+	FORCEINLINE friend FArchive& operator<<(FArchive& Ar, WIDECHAR& Value)
+	{
+		Ar.ByteOrderSerialize(&Value, sizeof(Value));
 
-/**
- * Archive for storing arbitrary data to the specified memory location
- */
-class FMemoryWriter : public FMemoryArchive
-{
-public:
-	FMemoryWriter( TArray<uint8>& InBytes, bool bIsPersistent = false, bool bSetOffset = false, const FName InArchiveName = NAME_None )
-	: FMemoryArchive()
-	, Bytes(InBytes)
-	, ArchiveName(InArchiveName)
-	{
-		ArIsSaving		= true;
-		ArIsPersistent	= bIsPersistent;
-		if (bSetOffset)
-		{
-			Offset = InBytes.Num();
-		}
+		return Ar;
 	}
 
-	virtual void Serialize(void* Data, int64 Num) override
+	/**
+	 * Serializes an unsigned 8-bit integer value from or into an archive.
+	 *
+	 * @param Ar The archive to serialize from or to.
+	 * @param Value The value to serialize.
+	 */
+	FORCEINLINE friend FArchive& operator<<(FArchive& Ar, uint8& Value)
 	{
-		const int64 NumBytesToAdd = Offset + Num - Bytes.Num();
-		if( NumBytesToAdd > 0 )
-		{
-			const int64 NewArrayCount = Bytes.Num() + NumBytesToAdd;
-			if( NewArrayCount >= MAX_int32 )
-			{
-				UE_LOG( LogSerialization, Fatal, TEXT( "FMemoryWriter does not support data larger than 2GB. Archive name: %s." ), *ArchiveName.ToString() );
-			}
-
-			Bytes.AddUninitialized( (int32)NumBytesToAdd );
-		}
+		Ar.Serialize(&Value, 1);
 
-		check((Offset + Num) <= Bytes.Num());
-		
-		if( Num )
-		{
-			FMemory::Memcpy( &Bytes[Offset], Data, Num );
-			Offset+=Num;
-		}
+		return Ar;
 	}
+
 	/**
-  	 * Returns the name of the Archive.  Useful for getting the name of the package a struct or object
-	 * is in when a loading error occurs.
+	 * Serializes an enumeration value from or into an archive.
 	 *
-	 * This is overridden for the specific Archive Types
-	 **/
-	virtual FString GetArchiveName() const override { return TEXT("FMemoryWriter"); }
-
-	int64 TotalSize() override
+	 * @param Ar The archive to serialize from or to.
+	 * @param Value The value to serialize.
+	 */
+	template<class TEnum>
+	FORCEINLINE friend FArchive& operator<<(FArchive& Ar, TEnumAsByte<TEnum>& Value)
 	{
-		return Bytes.Num();
-	}
+		Ar.Serialize(&Value, 1);
 
-protected:
+		return Ar;
+	}
 
-	TArray<uint8>&	Bytes;
+	/**
+	 * Serializes a signed 8-bit integer value from or into an archive.
+	 *
+	 * @param Ar The archive to serialize from or to.
+	 * @param Value The value to serialize.
+	 */
+	FORCEINLINE friend FArchive& operator<<(FArchive& Ar, int8& Value)
+	{
+		Ar.Serialize(&Value, 1);
 
-	/** Archive name, used to debugging, by default set to NAME_None. */
-	const FName ArchiveName;
-};
+		return Ar;
+	}
 
-/**
-* Archive for storing a large amount of arbitrary data to memory
-*/
-class CORE_API FLargeMemoryWriter : public FMemoryArchive
-{
-public:
-	
-	FLargeMemoryWriter(const int64 PreAllocateBytes = 0, bool bIsPersistent = false, const FName InArchiveName = NAME_None);
+	/**
+	 * Serializes an unsigned 16-bit integer value from or into an archive.
+	 *
+	 * @param Ar The archive to serialize from or to.
+	 * @param Value The value to serialize.
+	 */
+	FORCEINLINE friend FArchive& operator<<(FArchive& Ar, uint16& Value)
+	{
+		Ar.ByteOrderSerialize(&Value, sizeof(Value));
 
-	virtual void Serialize(void* InData, int64 Num) override;
+		return Ar;
+	}
 
 	/**
-	* Returns the name of the Archive.  Useful for getting the name of the package a struct or object
-	* is in when a loading error occurs.
-	*
-	* This is overridden for the specific Archive Types
-	**/
-	virtual FString GetArchiveName() const override;
+	 * Serializes a signed 16-bit integer value from or into an archive.
+	 *
+	 * @param Ar The archive to serialize from or to.
+	 * @param Value The value to serialize.
+	 */
+	FORCEINLINE friend FArchive& operator<<(FArchive& Ar, int16& Value)
+	{
+		Ar.ByteOrderSerialize(&Value, sizeof(Value));
+
+		return Ar;
+	}
 
 	/**
-	 * Gets the total size of the data written
+	 * Serializes an unsigned 32-bit integer value from or into an archive.
+	 *
+	 * @param Ar The archive to serialize from or to.
+	 * @param Value The value to serialize.
 	 */
-	virtual int64 TotalSize() override;
-	
+	FORCEINLINE friend FArchive& operator<<(FArchive& Ar, uint32& Value)
+	{
+		Ar.ByteOrderSerialize(&Value, sizeof(Value));
+
+		return Ar;
+	}
+
 	/**
-	 * Returns the written data. To release this archive's ownership of the data, call ReleaseOwnership()
+	 * Serializes a Boolean value from or into an archive.
+	 *
+	 * @param Ar The archive to serialize from or to.
+	 * @param Value The value to serialize.
 	 */
-	uint8* GetData() const;
+	FORCEINLINE friend FArchive& operator<<( FArchive& Ar, bool& D )
+	{
+		// Serialize bool as if it were UBOOL (legacy, 32 bit int).
+		uint32 OldUBoolValue = D ? 1 : 0;
+		Ar.ByteOrderSerialize( &OldUBoolValue, sizeof(OldUBoolValue) );
+		D = !!OldUBoolValue;
+		return Ar;
+	}
 
-	/** 
-	 * Releases ownership of the written data
-	 */
-	void ReleaseOwnership();
-	
 	/**
-	 * Destructor
+	 * Serializes a signed 32-bit integer value from or into an archive.
+	 *
+	 * @param Ar The archive to serialize from or to.
+	 * @param Value The value to serialize.
 	 */
-	~FLargeMemoryWriter();
-
-private:
+	FORCEINLINE friend FArchive& operator<<(FArchive& Ar, int32& Value)
+	{
+		Ar.ByteOrderSerialize(&Value, sizeof(Value));
 
-	/** Non-copyable */
-	FLargeMemoryWriter(const FLargeMemoryWriter&) = delete;
-	FLargeMemoryWriter& operator=(const FLargeMemoryWriter&) = delete;
+		return Ar;
+	}
 
-	/** Memory owned by this archive. Ownership can be released by calling ReleaseOwnership() */
-	uint8* Data;
+#if PLATFORM_COMPILER_DISTINGUISHES_INT_AND_LONG
+	/**
+	 * Serializes a long integer value from or into an archive.
+	 *
+	 * @param Ar The archive to serialize from or to.
+	 * @param Value The value to serialize.
+	 */
+	FORCEINLINE friend FArchive& operator<<(FArchive& Ar, long& Value)
+	{
+		Ar.ByteOrderSerialize(&Value, sizeof(Value));
 
-	/** Number of bytes currently written to our data buffer */
-	int64 NumBytes;
+		return Ar;
+	}	
+#endif
 
-	/** Number of bytes currently allocated for our data buffer */
-	int64 MaxBytes;
+	/**
+	 * Serializes a single precision floating point value from or into an archive.
+	 *
+	 * @param Ar The archive to serialize from or to.
+	 * @param Value The value to serialize.
+	 */
+	FORCEINLINE friend FArchive& operator<<( FArchive& Ar, float& Value)
+	{
+		Ar.ByteOrderSerialize(&Value, sizeof(Value));
 
-	/** Resizes the data buffer to at least the desired new size with some slack */
-	void GrowBuffer(const int64 DesiredBytes);
+		return Ar;
+	}
 
-	/** Archive name, used for debugging, by default set to NAME_None. */
-	const FName ArchiveName;
-};
+	/**
+	 * Serializes a double precision floating point value from or into an archive.
+	 *
+	 * @param Ar The archive to serialize from or to.
+	 * @param Value The value to serialize.
+	 */
+	FORCEINLINE friend FArchive& operator<<(FArchive& Ar, double& Value)
+	{
+		Ar.ByteOrderSerialize(&Value, sizeof(Value));
 
-enum class ELargeMemoryReaderFlags : uint8
-{
-	None =				0x0,	// No Flags
-	TakeOwnership =		0x1,	// Archive will take ownership of the passed-in memory and free it on destruction
-	Persistent =		0x2,	// Archive will be set as persistent when constructed
-};
+		return Ar;
+	}
 
-ENUM_CLASS_FLAGS(ELargeMemoryReaderFlags);
+	/**
+	 * Serializes a unsigned 64-bit integer value from or into an archive.
+	 *
+	 * @param Ar The archive to serialize from or to.
+	 * @param Value The value to serialize.
+	 */
+	FORCEINLINE friend FArchive& operator<<(FArchive &Ar, uint64& Value)
+	{
+		Ar.ByteOrderSerialize(&Value, sizeof(Value));
 
-/**
-* Archive for reading a large amount of arbitrary data from memory
-*/
-class CORE_API FLargeMemoryReader : public FMemoryArchive
-{
-public:
+		return Ar;
+	}
 
-	FLargeMemoryReader(const uint8* InData, const int64 Num, ELargeMemoryReaderFlags InFlags = ELargeMemoryReaderFlags::None, const FName InArchiveName = NAME_None);
+	/**
+	 * Serializes a signed 64-bit integer value from or into an archive.
+	 *
+	 * @param Ar The archive to serialize from or to.
+	 * @param Value The value to serialize.
+	 */
+	/*FORCEINLINE*/friend FArchive& operator<<(FArchive& Ar, int64& Value)
+	{
+		Ar.ByteOrderSerialize(&Value, sizeof(Value));
 
-	virtual void Serialize(void* OutData, int64 Num) override;
+		return Ar;
+	}
 
 	/**
-	* Gets the total size of the data buffer
-	*/
-	virtual int64 TotalSize() override;
+	 * Serializes an FIntRect value from or into an archive.
+	 *
+	 * @param Ar The archive to serialize from or to.
+	 * @param Value The value to serialize.
+	 */
+	friend FArchive& operator<<(FArchive& Ar, struct FIntRect& Value);
 
 	/**
-	* Returns the name of the Archive.  Useful for getting the name of the package a struct or object
-	* is in when a loading error occurs.
-	*
-	* This is overridden for the specific Archive Types
-	**/
-	virtual FString GetArchiveName() const override;
+	 * Serializes an FString value from or into an archive.
+	 *
+	 * @param Ar The archive to serialize from or to.
+	 * @param Value The value to serialize.
+	 */
+	friend CORE_API FArchive& operator<<(FArchive& Ar, FString& Value);
 
-	/**
-	* Destructor
-	*/
-	~FLargeMemoryReader();
+public:
 
-private:
+	virtual void Serialize(void* V, int64 Length) { }
 
-	/** Non-copyable */
-	FLargeMemoryReader(const FLargeMemoryReader&) = delete;
-	FLargeMemoryReader& operator=(const FLargeMemoryReader&) = delete;
+	virtual void SerializeBits(void* V, int64 LengthBits)
+	{
+		Serialize(V, (LengthBits + 7) / 8);
 
-	/** Whether the data buffer should be freed when this archive is closed */
-	const uint8 bFreeOnClose : 1;
-	
-	/** Data buffer we are reading from */
-	const uint8* Data;
+		if (IsLoading() && (LengthBits % 8) != 0)
+		{
+			((uint8*)V)[LengthBits / 8] &= ((1 << (LengthBits & 7)) - 1);
+		}
+	}
+
+	virtual void SerializeInt(uint32& Value, uint32 Max)
+	{
+		ByteOrderSerialize(&Value, sizeof(Value));
+	}
 
-	/** Total size of the data buffer */
-	const int64 NumBytes;
+	/** Packs int value into bytes of 7 bits with 8th bit for 'more' */
+	virtual void SerializeIntPacked(uint32& Value);
 
-	/** Name of the archive this buffer is for */
-	const FName ArchiveName;
-};
+	virtual void Preload(UObject* Object) { }
 
+	virtual void CountBytes(SIZE_T InNum, SIZE_T InMax) { }
 
-/**
- * Buffer archiver.
- */
-class FBufferArchive : public FMemoryWriter, public TArray<uint8>
-{
-public:
-	FBufferArchive( bool bIsPersistent = false, const FName InArchiveName = NAME_None )
-	: FMemoryWriter( (TArray<uint8>&)*this, bIsPersistent, false, InArchiveName )
-	{}
 	/**
   	 * Returns the name of the Archive.  Useful for getting the name of the package a struct or object
 	 * is in when a loading error occurs.
 	 *
 	 * This is overridden for the specific Archive Types
-	 **/
-	virtual FString GetArchiveName() const { return *FString::Printf( TEXT("FBufferArchive %s"), *ArchiveName.ToString()); }
-};
+	 */
+	virtual FString GetArchiveName() const;
 
-/**
- * Archive for reading arbitrary data from the specified memory location
- */
-class FMemoryReader final : public FMemoryArchive
-{
-public:
 	/**
-  	 * Returns the name of the Archive.  Useful for getting the name of the package a struct or object
-	 * is in when a loading error occurs.
+	 * If this archive is a FLinkerLoad or FLinkerSave, returns a pointer to the ULinker portion.
 	 *
-	 * This is overridden for the specific Archive Types
-	 **/
-	virtual FString GetArchiveName() const { return TEXT("FMemoryReader"); }
-
-	int64 TotalSize()
+	 * @return The linker, or nullptr if the archive is not a linker.
+	 */
+	virtual class FLinker* GetLinker()
 	{
-		return FMath::Min((int64)Bytes.Num(), LimitSize);
+		return nullptr;
 	}
 
-	void Serialize( void* Data, int64 Num )
+	virtual int64 Tell()
 	{
-		if (Num && !ArIsError)
-		{
-			// Only serialize if we have the requested amount of data
-			if (Offset + Num <= TotalSize())
-			{
-				FMemory::Memcpy( Data, &Bytes[Offset], Num );
-				Offset += Num;
-			}
-			else
-			{
-				ArIsError = true;
-			}
-		}
+		return INDEX_NONE;
 	}
 
-	FMemoryReader( const TArray<uint8>& InBytes, bool bIsPersistent = false )
-	: FMemoryArchive()
-	, Bytes(InBytes)
-	, LimitSize(INT64_MAX)
+	virtual int64 TotalSize()
 	{
-		ArIsLoading		= true;
-		ArIsPersistent	= bIsPersistent;
+		return INDEX_NONE;
 	}
 
-	/** With this method it's possible to attach data behind some serialized data. */
-	void SetLimitSize(int64 NewLimitSize)
+	virtual bool AtEnd()
 	{
-		LimitSize = NewLimitSize;
+		int64 Pos = Tell();
+
+		return ((Pos != INDEX_NONE) && (Pos >= TotalSize()));
 	}
 
-protected:
+	virtual void Seek(int64 InPos) { }
 
-	const TArray<uint8>& Bytes;
-	int64 LimitSize;
-};
+	/**
+	 * Attaches/ associates the passed in bulk data object with the linker.
+	 *
+	 * @param	Owner		UObject owning the bulk data
+	 * @param	BulkData	Bulk data object to associate
+	 */
+	virtual void AttachBulkData(UObject* Owner, FUntypedBulkData* BulkData) { }
 
-/**
- * Archive objects that are also a TArray. Since FBufferArchive is already the 
- * writer version, we just typedef to the old, outdated name
- */
-typedef FBufferArchive FArrayWriter;
+	/**
+	 * Detaches the passed in bulk data object from the linker.
+	 *
+	 * @param	BulkData	Bulk data object to detach
+	 * @param	bEnsureBulkDataIsLoaded	Whether to ensure that the bulk data is loaded before detaching
+	 */
+	virtual void DetachBulkData(FUntypedBulkData* BulkData, bool bEnsureBulkDataIsLoaded) { }
 
-class FArrayReader final : public FMemoryArchive, public TArray<uint8>
-{
-public:
-	FArrayReader( bool bIsPersistent=false )
+	/**
+	 * Hint the archive that the region starting at passed in offset and spanning the passed in size
+	 * is going to be read soon and should be precached.
+	 *
+	 * The function returns whether the precache operation has completed or not which is an important
+	 * hint for code knowing that it deals with potential async I/O. The archive is free to either not 
+	 * implement this function or only partially precache so it is required that given sufficient time
+	 * the function will return true. Archives not based on async I/O should always return true.
+	 *
+	 * This function will not change the current archive position.
+	 *
+	 * @param	PrecacheOffset	Offset at which to begin precaching.
+	 * @param	PrecacheSize	Number of bytes to precache
+	 * @return	false if precache operation is still pending, true otherwise
+	 */
+	virtual bool Precache(int64 PrecacheOffset, int64 PrecacheSize)
 	{
-		ArIsLoading = true;
-		ArIsPersistent = bIsPersistent;
+		return true;
 	}
 
 	/**
-	* Returns the name of the Archive.  Useful for getting the name of the package a struct or object
-	* is in when a loading error occurs.
-	*
-	* This is overridden for the specific Archive Types
-	**/
-	virtual FString GetArchiveName() const { return TEXT("FArrayReader"); }
+	 * Flushes cache and frees internal data.
+	 */
+
+	virtual void FlushCache() { }
 
-	int64 TotalSize()
+	/**
+	 * Sets mapping from offsets/ sizes that are going to be used for seeking and serialization to what
+	 * is actually stored on disk. If the archive supports dealing with compression in this way it is 
+	 * going to return true.
+	 *
+	 * @param	CompressedChunks	Pointer to array containing information about [un]compressed chunks
+	 * @param	CompressionFlags	Flags determining compression format associated with mapping
+	 *
+	 * @return true if archive supports translating offsets & uncompressing on read, false otherwise
+	 */
+	virtual bool SetCompressionMap(TArray<struct FCompressedChunk>* CompressedChunks, ECompressionFlags CompressionFlags)
 	{
-		return (int64)Num();
+		return false;
 	}
 
-	void Serialize(void* Data, int64 Count)
+	virtual void Flush() { }
+
+	virtual bool Close()
 	{
-		if (Count && !ArIsError)
-		{
-			// Only serialize if we have the requested amount of data
-			if (Offset + Count <= Num())
-			{
-				FMemory::Memcpy(Data, &((*this)[Offset]), Count);
-				Offset += Count;
-			}
-			else
-			{
-				ArIsError = true;
-			}
-		}
+		return !ArIsError;
 	}
-};
 
-/**
-* Similar to FMemoryReader, but able to internally
-* manage the memory for the buffer.
-*/
-class FBufferReaderBase : public FArchive
-{
-public:
-	/**
-	* Constructor
-	*
-	* @param Data Buffer to use as the source data to read from
-	* @param Size Size of Data
-	* @param bInFreeOnClose If true, Data will be FMemory::Free'd when this archive is closed
-	* @param bIsPersistent Uses this value for ArIsPersistent
-	* @param bInSHAVerifyOnClose It true, an async SHA verification will be done on the Data buffer (bInFreeOnClose will be passed on to the async task)
-	*/
-	FBufferReaderBase(void* Data, int64 Size, bool bInFreeOnClose, bool bIsPersistent = false)
-		: ReaderData(Data)
-		, ReaderPos(0)
-		, ReaderSize(Size)
-		, bFreeOnClose(bInFreeOnClose)
+	virtual bool GetError()
 	{
-		ArIsLoading = true;
-		ArIsPersistent = bIsPersistent;
+		return ArIsError;
+	}
+
+	void SetError() 
+	{ 
+		ArIsError = true; 
 	}
 
-	~FBufferReaderBase()
+	/**
+	 * Serializes and compresses/ uncompresses data. This is a shared helper function for compression
+	 * support. The data is saved in a way compatible with FIOSystem::LoadCompressedData.
+	 *
+	 * @param	V		Data pointer to serialize data from/ to
+	 * @param	Length	Length of source data if we're saving, unused otherwise
+	 * @param	Flags	Flags to control what method to use for [de]compression and optionally control memory vs speed when compressing
+	 * @param	bTreatBufferAsFileReader true if V is actually an FArchive, which is used when saving to read data - helps to avoid single huge allocations of source data
+	 * @param	bUsePlatformBitWindow use a platform specific bitwindow setting
+	 */
+	void SerializeCompressed(void* V, int64 Length, ECompressionFlags Flags, bool bTreatBufferAsFileReader = false, bool bUsePlatformBitWindow = false);
+
+
+	FORCEINLINE bool IsByteSwapping()
 	{
-		Close();
+#if PLATFORM_LITTLE_ENDIAN
+		bool SwapBytes = ArForceByteSwapping;
+#else
+		bool SwapBytes = ArIsPersistent;
+#endif
+		return SwapBytes;
 	}
-	bool Close()
+
+	// Used to do byte swapping on small items. This does not happen usually, so we don't want it inline
+	void ByteSwap(void* V, int32 Length);
+
+	FORCEINLINE FArchive& ByteOrderSerialize(void* V, int32 Length)
 	{
-		if (bFreeOnClose)
+		Serialize(V, Length);
+		if (IsByteSwapping())
 		{
-			FMemory::Free(ReaderData);
-			ReaderData = nullptr;
+			// Transferring between memory and file, so flip the byte order.
+			ByteSwap(V, Length);
 		}
-		return !ArIsError;
+		return *this;
 	}
-	void Serialize(void* Data, int64 Num) final
+
+	/** Sets a flag indicating that this archive contains code. */
+	void ThisContainsCode()
 	{
-		check(ReaderPos >= 0);
-		check(ReaderPos + Num <= ReaderSize);
-		FMemory::Memcpy(Data, (uint8*)ReaderData + ReaderPos, Num);
-		ReaderPos += Num;
+		ArContainsCode = true;
 	}
-	int64 Tell() final
+
+	/** Sets a flag indicating that this archive contains a ULevel or UWorld object. */
+	void ThisContainsMap() 
 	{
-		return ReaderPos;
+		ArContainsMap = true;
 	}
-	int64 TotalSize() final
+
+	/** Sets a flag indicating that this archive contains data required to be gathered for localization. */
+	void ThisRequiresLocalizationGather()
 	{
-		return ReaderSize;
+		ArRequiresLocalizationGather = true;
 	}
-	void Seek(int64 InPos) final 
+
+	/** Sets a flag indicating that this archive is currently serializing class/struct defaults. */
+	void StartSerializingDefaults() 
 	{
-		check(InPos >= 0);
-		check(InPos <= ReaderSize);
-		ReaderPos = InPos;
+		ArSerializingDefaults++;
 	}
-	bool AtEnd() final 
+
+	/** Indicate that this archive is no longer serializing class/struct defaults. */
+	void StopSerializingDefaults() 
 	{
-		return ReaderPos >= ReaderSize;
+		ArSerializingDefaults--;
 	}
+
 	/**
-	* Returns the name of the Archive.  Useful for getting the name of the package a struct or object
-	* is in when a loading error occurs.
-	*
-	* This is overridden for the specific Archive Types
-	**/
-	virtual FString GetArchiveName() const { return TEXT("FBufferReaderBase"); }
-protected:
-	void*		ReaderData;
-	int64		ReaderPos;
-	int64		ReaderSize;
-	bool	bFreeOnClose;
-};
+	 * Called when an object begins serializing property data using script serialization.
+	 */
+	virtual void MarkScriptSerializationStart(const UObject* Obj) { }
 
-/**
- * Similar to FMemoryReader, but able to internally
- * manage the memory for the buffer.
- */
-class FBufferReader final : public FBufferReaderBase
-{
-public:
 	/**
-	 * Constructor
-	 * 
-	 * @param Data Buffer to use as the source data to read from
-	 * @param Size Size of Data
-	 * @param bInFreeOnClose If true, Data will be FMemory::Free'd when this archive is closed
-	 * @param bIsPersistent Uses this value for ArIsPersistent
-	 * @param bInSHAVerifyOnClose It true, an async SHA verification will be done on the Data buffer (bInFreeOnClose will be passed on to the async task)
+	 * Called when an object stops serializing property data using script serialization.
 	 */
-	FBufferReader( void* Data, int64 Size, bool bInFreeOnClose, bool bIsPersistent = false )
-		: FBufferReaderBase(Data, Size, bInFreeOnClose, bIsPersistent)
+	virtual void MarkScriptSerializationEnd(const UObject* Obj) { }
+
+	virtual void IndicateSerializationMismatch() { }
+
+	// Logf implementation for convenience.
+	VARARG_DECL(void, void, {}, Logf, VARARG_NONE, const TCHAR*, VARARG_NONE, VARARG_NONE);
+
+	FORCEINLINE int32 UE4Ver() const
+	{
+		return ArUE4Ver;
+	}
+	
+	FORCEINLINE int32 LicenseeUE4Ver() const
 	{
+		return ArLicenseeUE4Ver;
 	}
 
-	virtual FString GetArchiveName() const { return TEXT("FBufferReader"); }
-};
+	FORCEINLINE FEngineVersionBase EngineVer() const
+	{
+		return ArEngineVer;
+	}
 
+	FORCEINLINE uint32 EngineNetVer() const
+	{
+		return ArEngineNetVer;
+	}
 
+	FORCEINLINE uint32 GameNetVer() const
+	{
+		return ArGameNetVer;
+	}
 
-/**
-* Similar to FMemoryWriter, but able to internally
-* manage the memory for the buffer.
-*/
-class FBufferWriter final : public FArchive
-{
-public:
 	/**
-	* Constructor
-	*
-	* @param Data Buffer to use as the source data to read from
-	* @param Size Size of Data
-	* @param bInFreeOnClose If true, Data will be FMemory::Free'd when this archive is closed
-	* @param bIsPersistent Uses this value for ArIsPersistent
-	*/
-	FBufferWriter(void* Data, int64 Size, bool bInFreeOnClose, bool bIsPersistent = false)
-		: WriterData(Data)
-		, WriterPos(0)
-		, WriterSize(Size)
-		, bFreeOnClose(bInFreeOnClose)
+	 * Registers the custom version to the archive.  This is used to inform the archive that custom version information is about to be stored.
+	 * There is no effect when the archive is being loaded from.
+	 *
+	 * @param Guid The guid of the custom version.  This must have previously been registered with FCustomVersionRegistration.
+	 */
+	void UsingCustomVersion(const struct FGuid& Guid);
+
+	/**
+	 * Queries a custom version from the archive.  If the archive is being used to write, the custom version must have already been registered.
+	 *
+	 * @param Key The guid of the custom version to query.
+	 * @return The version number, or 0 if the custom tag isn't stored in the archive.
+	 */
+	int32 CustomVer(const struct FGuid& Key) const;
+
+	FORCEINLINE bool IsLoading() const
 	{
-		ArIsSaving = true;
-		ArIsPersistent = bIsPersistent;
+		return ArIsLoading;
 	}
 
-	~FBufferWriter()
+	FORCEINLINE bool IsSaving() const
 	{
-		Close();
+		return ArIsSaving;
 	}
-	bool Close()
+
+	FORCEINLINE bool IsTransacting() const
 	{
-		if (bFreeOnClose)
+		if (FPlatformProperties::HasEditorOnlyData())
 		{
-			FMemory::Free(WriterData);
-			WriterData = nullptr;
+			return ArIsTransacting;
+		}
+		else
+		{
+			return false;
 		}
-		return !ArIsError;
 	}
-	void Serialize(void* Data, int64 Num)
+
+	FORCEINLINE bool WantBinaryPropertySerialization() const
 	{
-		const int64 NumBytesToAdd = WriterPos + Num - WriterSize;
-		if (NumBytesToAdd > 0)
-		{
-			const int64 NewArrayCount = WriterSize + NumBytesToAdd;
-			if (NewArrayCount >= MAX_int32)
-			{
-				UE_LOG(LogSerialization, Fatal, TEXT("FBufferWriter does not support data larger than 2GB. Archive name: %s."), *GetArchiveName());
-			}
+		return ArWantBinaryPropertySerialization;
+	}
 
-			FMemory::Realloc(WriterData, NewArrayCount);
-			WriterSize = NewArrayCount;
-		}
+	FORCEINLINE bool IsForcingUnicode() const
+	{
+		return ArForceUnicode;
+	}
+
+	FORCEINLINE bool IsPersistent() const
+	{
+		return ArIsPersistent;
+	}
+
+	FORCEINLINE bool IsError() const
+	{
+		return ArIsError;
+	}
+
+	FORCEINLINE bool IsCriticalError() const
+	{
+		return ArIsCriticalError;
+	}
+
+	FORCEINLINE bool ContainsCode() const
+	{
+		return ArContainsCode;
+	}
+
+	FORCEINLINE bool ContainsMap() const
+	{
+		return ArContainsMap;
+	}
+
+	FORCEINLINE bool RequiresLocalizationGather() const
+	{
+		return ArRequiresLocalizationGather;
+	}
+
+	FORCEINLINE bool ForceByteSwapping() const
+	{
+		return ArForceByteSwapping;
+	}
+
+	FORCEINLINE bool IsSerializingDefaults() const
+	{
+		return (ArSerializingDefaults > 0) ? true : false;
+	}
+
+	FORCEINLINE bool IsIgnoringArchetypeRef() const
+	{
+		return ArIgnoreArchetypeRef;
+	}
+
+	FORCEINLINE bool DoDelta() const
+	{
+		return !ArNoDelta;
+	}
+
+	FORCEINLINE bool IsIgnoringOuterRef() const
+	{
+		return ArIgnoreOuterRef;
+	}
 
-		check(WriterPos >= 0);
-		check((WriterPos + Num) <= WriterSize);
-		FMemory::Memcpy((uint8*)WriterData + WriterPos, Data, Num);
-		WriterPos += Num;
+	FORCEINLINE bool IsIgnoringClassGeneratedByRef() const
+	{
+		return ArIgnoreClassGeneratedByRef;
+	}
+
+	FORCEINLINE bool IsIgnoringClassRef() const
+	{
+		return ArIgnoreClassRef;
+	}
+
+	FORCEINLINE bool IsAllowingLazyLoading() const
+	{
+		return ArAllowLazyLoading;
+	}
+
+	FORCEINLINE bool IsObjectReferenceCollector() const
+	{
+		return ArIsObjectReferenceCollector;
+	}
+
+	FORCEINLINE bool IsModifyingWeakAndStrongReferences() const
+	{
+		return ArIsModifyingWeakAndStrongReferences;
+	}
+
+	FORCEINLINE bool IsCountingMemory() const
+	{
+		return ArIsCountingMemory;
+	}
+
+	FORCEINLINE uint32 GetPortFlags() const
+	{
+		return ArPortFlags;
 	}
-	int64 Tell()
+
+	FORCEINLINE bool HasAnyPortFlags(uint32 Flags) const
 	{
-		return WriterPos;
+		return ((ArPortFlags & Flags) != 0);
 	}
-	int64 TotalSize()
+
+	FORCEINLINE bool HasAllPortFlags(uint32 Flags) const
+	{
+		return ((ArPortFlags & Flags) == Flags);
+	}
+
+	FORCEINLINE uint32 GetDebugSerializationFlags() const
 	{
-		return WriterSize;
+#if WITH_EDITOR
+		return ArDebugSerializationFlags;
+#else
+		return 0;
+#endif
 	}
-	void Seek(int64 InPos)
+
+	FORCEINLINE bool ShouldSkipBulkData() const
 	{
-		check(InPos >= 0);
-		check(InPos <= WriterSize);
-		WriterPos = InPos;
+		return ArShouldSkipBulkData;
 	}
-	bool AtEnd()
+
+	FORCEINLINE int64 GetMaxSerializeSize() const
 	{
-		return WriterPos >= WriterSize;
+		return ArMaxSerializeSize;
 	}
+
 	/**
-	* Returns the name of the Archive.  Useful for getting the name of the package a struct or object
-	* is in when a loading error occurs.
+	 * Sets the archive version number. Used by the code that makes sure that FLinkerLoad's 
+	 * internal archive versions match the file reader it creates.
+	 *
+	 * @param UE4Ver	new version number
+	 */
+	void SetUE4Ver(int32 InVer)
+	{
+		ArUE4Ver = InVer;
+	}
+
+	/**
+	 * Sets the archive licensee version number. Used by the code that makes sure that FLinkerLoad's 
+	 * internal archive versions match the file reader it creates.
+	 *
+	 * @param Ver	new version number
+	 */
+	void SetLicenseeUE4Ver(int32 InVer)
+	{
+		ArLicenseeUE4Ver = InVer;
+	}
+
+	/**
+	* Sets the archive engine version. Used by the code that makes sure that FLinkerLoad's
+	* internal archive versions match the file reader it creates.
 	*
-	* This is overridden for the specific Archive Types
-	**/
-	virtual FString GetArchiveName() const { return TEXT("FBufferWriter"); }
+	* @param InVer	new version number
+	*/
+	void SetEngineVer(const FEngineVersionBase& InVer)
+	{
+		ArEngineVer = InVer;
+	}
 
-	void* GetWriterData()
+	/**
+	* Sets the archive engine network version.
+	*/
+	void SetEngineNetVer(const uint32 InEngineNetVer)
 	{
-		return WriterData;
+		ArEngineNetVer = InEngineNetVer;
 	}
-protected:
-	void* WriterData;
-	int64 WriterPos;
-	int64 WriterSize;
-	bool bFreeOnClose;
-};
-/*----------------------------------------------------------------------------
-	FArchiveSaveCompressedProxy.
-----------------------------------------------------------------------------*/
 
-/**
- * FArchive Proxy to transparently write out compressed data to an array.
- */
-class CORE_API FArchiveSaveCompressedProxy : public FArchive
-{
-public:
-	/** 
-	 * Constructor, initializing all member variables and allocating temp memory.
+	/**
+	* Sets the archive game network version.
+	*/
+	void SetGameNetVer(const uint32 InGameNetVer)
+	{
+		ArGameNetVer = InGameNetVer;
+	}
+
+	/**
+	 * Gets the custom version numbers for this archive.
 	 *
-	 * @param	InCompressedData [ref]	Array of bytes that is going to hold compressed data
-	 * @param	InCompressionFlags		Compression flags to use for compressing data
+	 * @return The container of custom versions in the archive.
 	 */
-	FArchiveSaveCompressedProxy( TArray<uint8>& InCompressedData, ECompressionFlags InCompressionFlags );
+	const FCustomVersionContainer& GetCustomVersions() const;
 
-	/** Destructor, flushing array if needed. Also frees temporary memory. */
-	virtual ~FArchiveSaveCompressedProxy();
+	/**
+	 * Sets the custom version numbers for this archive.
+	 *
+	 * @param CustomVersionContainer - The container of custom versions to copy into the archive.
+	 */
+	void SetCustomVersions(const FCustomVersionContainer& CustomVersionContainer);
+
+	/** Resets the custom version numbers for this archive. */
+	void ResetCustomVersions();
+
+	/**
+	 * Sets a specific custom version
+	 *
+	 * @param Key - The guid of the custom version to query.
+	 * @param Version - The version number to set key to
+	 * @param FriendlyName - Friendly name corresponding to the key
+	 */
+	void SetCustomVersion(const struct FGuid& Key, int32 Version, FName FriendlyName);
+
+	/**
+	 * Toggle saving as Unicode. This is needed when we need to make sure ANSI strings are saved as Unicode
+	 *
+	 * @param Enabled	set to true to force saving as Unicode
+	 */
+	void SetForceUnicode(bool Enabled)
+	{
+		ArForceUnicode = Enabled;
+	}
 
 	/**
-	 * Flushes tmp data to array.
+	 * Toggle byte order swapping. This is needed in rare cases when we already know that the data
+	 * swapping has already occurred or if we know that it will be handled later.
+	 *
+	 * @param Enabled	set to true to enable byte order swapping
 	 */
-	virtual void Flush();
+	void SetByteSwapping(bool Enabled)
+	{
+		ArForceByteSwapping = Enabled;
+	}
 
 	/**
-	 * Serializes data to archive. This function is called recursively and determines where to serialize
-	 * to and how to do so based on internal state.
+	 * Sets the archive's property serialization modifier flags
 	 *
-	 * @param	Data	Pointer to serialize to
-	 * @param	Count	Number of bytes to read
+	 * @param	InPortFlags		the new flags to use for property serialization
 	 */
-	virtual void Serialize( void* Data, int64 Count );
+	void SetPortFlags(uint32 InPortFlags)
+	{
+		ArPortFlags = InPortFlags;
+	}
 
 	/**
-	 * Seeking is only implemented internally for writing out compressed data and asserts otherwise.
+	 * Sets the archives custom serialization modifier flags (nothing to do with PortFlags or Custom versions)
 	 * 
-	 * @param	InPos	Position to seek to
+	 * @param InCustomFlags the new flags to use for custom serialization
 	 */
-	virtual void Seek( int64 InPos );
+	void SetDebugSerializationFlags(uint32 InCustomFlags)
+	{
+#if WITH_EDITOR
+		ArDebugSerializationFlags = InCustomFlags;
+#endif
+	}
 
 	/**
-	 * @return current position in uncompressed stream in bytes.
+	 * Returns if an async close operation has finished or not, as well as if there was an error.
+	 *
+	 * @param bHasError true if there was an error.
+	 * @return true if the close operation is complete (or if Close is not an async operation).
 	 */
-	virtual int64 Tell();
+	virtual bool IsCloseComplete(bool& bHasError)
+	{
+		bHasError = false;
+		
+		return true;
+	}
 
-private:
-	/** Array to write compressed data to.					*/
-	TArray<uint8>&	CompressedData;
-	/** Current index in array.								*/
-	int32				CurrentIndex;
-	/** Pointer to start of temporary buffer.				*/
-	uint8*			TmpDataStart;
-	/** Pointer to end of temporary buffer.					*/
-	uint8*			TmpDataEnd;
-	/** Pointer to current position in temporary buffer.	*/
-	uint8*			TmpData;
-	/** Whether to serialize to temporary buffer of array.	*/
-	bool			bShouldSerializeToArray;
-	/** Number of raw (uncompressed) bytes serialized.		*/
-	int64		RawBytesSerialized;
-	/** Flags to use for compression.						*/
-	ECompressionFlags CompressionFlags;
-};
+	/**
+	 * Indicates whether this archive is filtering editor-only on save or contains data that had editor-only content stripped.
+	 *
+	 * @return true if the archive filters editor-only content, false otherwise.
+	 */
+	virtual bool IsFilterEditorOnly()
+	{
+		return ArIsFilterEditorOnly;
+	}
+
+	/**
+	 * Sets a flag indicating that this archive needs to filter editor-only content.
+	 *
+	 * @param InFilterEditorOnly Whether to filter editor-only content.
+	 */
+	virtual void SetFilterEditorOnly(bool InFilterEditorOnly)
+	{
+		ArIsFilterEditorOnly = InFilterEditorOnly;
+	}
 
-/*----------------------------------------------------------------------------
-	FArchiveLoadCompressedProxy.
-----------------------------------------------------------------------------*/
+	/**
+	 * Indicates whether this archive is saving or loading game state
+	 *
+	 * @return true if the archive is dealing with save games, false otherwise.
+	 */
+	virtual bool IsSaveGame()
+	{
+		return ArIsSaveGame;
+	}
 
-/**
- * FArchive Proxy to transparently load compressed data from an array.
- */
-class CORE_API FArchiveLoadCompressedProxy : public FArchive
-{
-public:
-	/** 
-	 * Constructor, initializing all member variables and allocating temp memory.
+	/**
+	 * Checks whether the archive is used for cooking.
 	 *
-	 * @param	InCompressedData	Array of bytes that is holding compressed data
-	 * @param	InCompressionFlags	Compression flags that were used to compress data
+	 * @return true if the archive is used for cooking, false otherwise.
 	 */
-	FArchiveLoadCompressedProxy( const TArray<uint8>& InCompressedData, ECompressionFlags InCompressionFlags );
+	FORCEINLINE bool IsCooking() const	
+	{
+		check(!CookingTargetPlatform || (!IsLoading() && !IsTransacting() && IsSaving()));
 
-	/** Destructor, freeing temporary memory. */
-	virtual ~FArchiveLoadCompressedProxy();
+		return !!CookingTargetPlatform;
+	}
 
 	/**
-	 * Serializes data from archive. This function is called recursively and determines where to serialize
-	 * from and how to do so based on internal state.
+	 * Returns the cooking target platform.
 	 *
-	 * @param	Data	Pointer to serialize to
-	 * @param	Count	Number of bytes to read
+	 * @return Target platform.
 	 */
-	virtual void Serialize( void* Data, int64 Count );
+	FORCEINLINE const class ITargetPlatform* CookingTarget() const 
+	{
+		return CookingTargetPlatform;
+	}
 
 	/**
-	 * Seeks to the passed in position in the stream. This archive only supports forward seeking
-	 * and implements it by serializing data till it reaches the position.
+	 * Sets the cooking target platform.
+	 *
+	 * @param InCookingTarget The target platform to set.
 	 */
-	virtual void Seek( int64 InPos );
+	FORCEINLINE void SetCookingTarget(const class ITargetPlatform* InCookingTarget) 
+	{
+		CookingTargetPlatform = InCookingTarget;
+	}
+
+	/**
+	 * Checks whether the archive is used to resolve out-of-date enum indexes
+	 * If function returns true, the archive should be called only for objects containing user defined enum
+	 *
+	 * @return true if the archive is used to resolve out-of-date enum indexes
+	 */
+	virtual bool UseToResolveEnumerators() const
+	{
+		return false;
+	}
+
+	/**
+	 * Checks whether the archive wants to skip the property independent of the other flags
+	 */
+	virtual bool ShouldSkipProperty(const class UProperty* InProperty) const
+	{
+		return false;
+	}
+
+	/**
+	 * Sets the property that is currently being serialized
+	 * 
+	 * @param InProperty Pointer to the property that is currently being serialized
+	 */
+	FORCEINLINE void SetSerializedProperty(class UProperty* InProperty)
+	{
+		SerializedProperty = InProperty;
+	}
+
+#if WITH_EDITORONLY_DATA
+	/** Pushes editor-only marker to the stack of currently serialized properties */
+	FORCEINLINE void PushEditorOnlyProperty()
+	{
+		EditorOnlyPropertyStack++;
+	}
+	/** Pops editor-only marker from the stack of currently serialized properties */
+	FORCEINLINE void PopEditorOnlyProperty()
+	{
+		EditorOnlyPropertyStack--;
+		check(EditorOnlyPropertyStack >= 0);
+	}
+	/** Returns true if the stack of currently serialized properties contains an editor-only property */
+	FORCEINLINE bool IsEditorOnlyPropertyOnTheStack() const
+	{
+		return EditorOnlyPropertyStack > 0;
+	}
+#endif
+
+	/**
+	 * Gets the property that is currently being serialized
+	 *
+	 * @return Pointer to the property that is currently being serialized
+	 */
+	FORCEINLINE class UProperty* GetSerializedProperty() const
+	{
+		return SerializedProperty;
+	}
+
+#if USE_STABLE_LOCALIZATION_KEYS
+	/**
+	 * Set the localization namespace that this archive should use when serializing text properties.
+	 * This is typically the namespace used by the package being serialized (if serializing a package, or an object within a package).
+	 */
+	virtual void SetLocalizationNamespace(const FString& InLocalizationNamespace);
+
+	/**
+	 * Get the localization namespace that this archive should use when serializing text properties.
+	 * This is typically the namespace used by the package being serialized (if serializing a package, or an object within a package).
+	 */
+	virtual FString GetLocalizationNamespace() const;
+#endif // USE_STABLE_LOCALIZATION_KEYS
+
+protected:
+
+	/** Resets all of the base archive members. */
+	void Reset();
+
+private:
+
+	/** Copies all of the members except CustomVersionContainer */
+	void CopyTrivialFArchiveStatusMembers(const FArchive& ArchiveStatusToCopy);
+
+protected:
+
+	/** Holds the archive version. */
+	int32 ArUE4Ver;
+	
+	/** Holds the archive version for licensees. */
+	int32 ArLicenseeUE4Ver;
+
+	/** Holds the engine version. */
+	FEngineVersionBase ArEngineVer;
+
+	/** Holds the engine network protocol version. */
+	uint32 ArEngineNetVer;
+
+	/** Holds the game network protocol version. */
+	uint32 ArGameNetVer;
+
+private:
 
 	/**
-	 * @return current position in uncompressed stream in bytes.
+	 * All the custom versions stored in the archive.
+	 * Stored as a pointer to a heap-allocated object because of a 3-way dependency between TArray, FCustomVersionContainer and FArchive, which is too much work to change right now.
 	 */
-	virtual int64 Tell();
+	FCustomVersionContainer* CustomVersionContainer;
+
+public:
+
+	/** Whether this archive is for loading data. */
+	bool ArIsLoading;
+	
+	/** Whether this archive is for saving data. */
+	bool ArIsSaving;
+	
+	/** Whether archive is transacting. */
+	bool ArIsTransacting;
+	
+	/** Whether this archive wants properties to be serialized in binary form instead of tagged. */
+	bool ArWantBinaryPropertySerialization;
+	
+	/** Whether this archive wants to always save strings in unicode format */
+	bool ArForceUnicode;
+	
+	/** Whether this archive saves to persistent storage. */
+	bool ArIsPersistent;
+			
+	/** Whether this archive contains errors. */
+	bool ArIsError;
+
+	/** Whether this archive contains critical errors. */
+	bool ArIsCriticalError;
+	
+	/** Quickly tell if an archive contains script code. */
+	bool ArContainsCode;
+	
+	/** Used to determine whether FArchive contains a level or world. */
+	bool ArContainsMap;
+
+	/** Used to determine whether FArchive contains data required to be gathered for localization. */
+	bool ArRequiresLocalizationGather;
+	
+	/** Whether we should forcefully swap bytes. */
+	bool ArForceByteSwapping;
+	
+	/** If true, we will not serialize the ObjectArchetype reference in UObject. */
+	bool ArIgnoreArchetypeRef;
+	
+	/** If true, we will not serialize the ObjectArchetype reference in UObject. */
+	bool ArNoDelta;
+
+	/** If true, we will not serialize the Outer reference in UObject. */
+	bool ArIgnoreOuterRef;
+
+	/** If true, we will not serialize ClassGeneratedBy reference in UClass. */
+	bool ArIgnoreClassGeneratedByRef;
+
+	/** If true, UObject::Serialize will skip serialization of the Class property. */
+	bool ArIgnoreClassRef;
+	
+	/** Whether to allow lazy loading. */
+	bool ArAllowLazyLoading;
+	
+	/** Whether this archive only cares about serializing object references. */
+	bool ArIsObjectReferenceCollector;
+	
+	/** Whether a reference collector is modifying the references and wants both weak and strong ones */
+	bool ArIsModifyingWeakAndStrongReferences;
+	
+	/** Whether this archive is counting memory and therefore wants e.g. TMaps to be serialized. */
+	bool ArIsCountingMemory;
+	
+	/** Whether bulk data serialization should be skipped or not. */
+	bool ArShouldSkipBulkData;
+
+	/** Whether editor only properties are being filtered from the archive (or has been filtered). */
+	bool ArIsFilterEditorOnly;
+
+	/** Whether this archive is saving/loading game state */
+	bool ArIsSaveGame;
 
+	/** Whether we are currently serializing defaults. > 0 means yes, <= 0 means no. */
+	int32 ArSerializingDefaults;
+
+	/** Modifier flags that be used when serializing UProperties */
+	uint32 ArPortFlags;
+			
+	/** Max size of data that this archive is allowed to serialize. */
+	int64 ArMaxSerializeSize;
+
+	/** Custom property list attribute. If the flag below is set, only these properties will be iterated during serialization. If NULL, then no properties will be iterated. */
+	const struct FCustomPropertyListNode* ArCustomPropertyList;
+
+	/** Set TRUE to use the custom property list attribute for serialization. */
+	bool ArUseCustomPropertyList;
+
+	class FScopeSetDebugSerializationFlags
+	{
+	private:
+#if WITH_EDITOR
+		uint32 PreviousFlags;
+		FArchive& Ar;
+#endif
+	public:
+		/**
+		 * Initializes an object which will set flags for the scope of this code
+		 * 
+		 * @param NewFlags new flags to set 
+		 * @param Remove should we add these flags or remove them default is to add
+		 */
+#if WITH_EDITOR
+		FScopeSetDebugSerializationFlags(FArchive& InAr, uint32 NewFlags, bool Remove = false)
+			: Ar(InAr)
+		{
+
+			PreviousFlags = Ar.GetDebugSerializationFlags();
+			if (Remove)
+			{
+				Ar.SetDebugSerializationFlags( PreviousFlags & ~NewFlags);
+			}
+			else
+			{
+				Ar.SetDebugSerializationFlags( PreviousFlags | NewFlags);
+			}
+
+		}
+		~FScopeSetDebugSerializationFlags()
+		{
+
+			Ar.SetDebugSerializationFlags( PreviousFlags);
+		}
+#else
+		FScopeSetDebugSerializationFlags(FArchive& InAr, uint32 NewFlags, bool Remove = false)
+		{}
+		~FScopeSetDebugSerializationFlags()
+		{}
+#endif
+	};
+
+#if WITH_EDITOR
+	/** Custom serialization modifier flags can be used for anything */
+	uint32 ArDebugSerializationFlags;
+	/** Debug stack storage if you want to add data to the archive for usage further down the serialization stack this should be used in conjunction with the FScopeAddDebugData struct */
+	
+	virtual void PushDebugDataString(const FName& DebugData);
+	virtual void PopDebugDataString() { }
+
+	class FScopeAddDebugData
+	{
+	private:
+		FArchive& Ar;
+	public:
+		CORE_API FScopeAddDebugData(FArchive& InAr, const FName& DebugData);
+
+		~FScopeAddDebugData()
+		{
+			Ar.PopDebugDataString();
+		}
+	};
+#endif	
 private:
+	/** Holds the cooking target platform. */
+	const ITargetPlatform* CookingTargetPlatform;
+
+	/** Holds the pointer to the property that is currently being serialized */
+	class UProperty* SerializedProperty;
+
+#if WITH_EDITORONLY_DATA
+	/** Non-zero if on the current stack of serialized properties there's an editor-only property. Needs to live in FArchive because of SerializedProperty. */
+	int32 EditorOnlyPropertyStack;
+#endif
+
+#if USE_STABLE_LOCALIZATION_KEYS
+	/**
+	 * The localization namespace that this archive should use when serializing text properties.
+	 * This is typically the namespace used by the package being serialized (if serializing a package, or an object within a package).
+	 * Stored as a pointer to a heap-allocated string because of a dependency between TArray (thus FString) and FArchive; null should be treated as an empty string.
+	 */
+	class FString* LocalizationNamespacePtr;
+
+	/** See SetLocalizationNamespace */
+	void SetBaseLocalizationNamespace(const FString& InLocalizationNamespace);
+
+	/** See GetLocalizationNamespace */
+	FString GetBaseLocalizationNamespace() const;
+#endif // USE_STABLE_LOCALIZATION_KEYS
+
 	/**
-	 * Flushes tmp data to array.
-	 */
-	void DecompressMoreData();
-
-	/** Array to write compressed data to.						*/
-	const TArray<uint8>&	CompressedData;
-	/** Current index into compressed data array.				*/
-	int32				CurrentIndex;
-	/** Pointer to start of temporary buffer.					*/
-	uint8*			TmpDataStart;
-	/** Pointer to end of temporary buffer.						*/
-	uint8*			TmpDataEnd;
-	/** Pointer to current position in temporary buffer.		*/
-	uint8*			TmpData;
-	/** Whether to serialize from temporary buffer of array.	*/
-	bool			bShouldSerializeFromArray;
-	/** Number of raw (uncompressed) bytes serialized.			*/
-	int64		RawBytesSerialized;
-	/** Flags used for compression.								*/
-	ECompressionFlags CompressionFlags;
+	 * Indicates if the custom versions container is in a 'reset' state.  This will be used to defer the choice about how to
+	 * populate the container until it is needed, where the read/write state will be known.
+	 */
+	mutable bool bCustomVersionsAreReset;
 };
 
+
+/**
+ * Template for archive constructors.
+ */
+template<class T> T Arctor(FArchive& Ar)
+{
+	T Tmp;
+	Ar << Tmp;
+
+	return Tmp;
+}
diff --git a/Engine/Source/Runtime/Core/Public/Serialization/ArchiveLoadCompressedProxy.h b/Engine/Source/Runtime/Core/Public/Serialization/ArchiveLoadCompressedProxy.h
new file mode 100644
--- /dev/null
+++ b/Engine/Source/Runtime/Core/Public/Serialization/ArchiveLoadCompressedProxy.h
@@ -0,0 +1,70 @@
+// Copyright 1998-2016 Epic Games, Inc. All Rights Reserved.
+
+#pragma once
+
+#include "Archive.h"
+
+/*----------------------------------------------------------------------------
+	FArchiveLoadCompressedProxy.
+----------------------------------------------------------------------------*/
+
+/**
+ * FArchive Proxy to transparently load compressed data from an array.
+ */
+class CORE_API FArchiveLoadCompressedProxy : public FArchive
+{
+public:
+	/** 
+	 * Constructor, initializing all member variables and allocating temp memory.
+	 *
+	 * @param	InCompressedData	Array of bytes that is holding compressed data
+	 * @param	InCompressionFlags	Compression flags that were used to compress data
+	 */
+	FArchiveLoadCompressedProxy( const TArray<uint8>& InCompressedData, ECompressionFlags InCompressionFlags );
+
+	/** Destructor, freeing temporary memory. */
+	virtual ~FArchiveLoadCompressedProxy();
+
+	/**
+	 * Serializes data from archive. This function is called recursively and determines where to serialize
+	 * from and how to do so based on internal state.
+	 *
+	 * @param	Data	Pointer to serialize to
+	 * @param	Count	Number of bytes to read
+	 */
+	virtual void Serialize( void* Data, int64 Count );
+
+	/**
+	 * Seeks to the passed in position in the stream. This archive only supports forward seeking
+	 * and implements it by serializing data till it reaches the position.
+	 */
+	virtual void Seek( int64 InPos );
+
+	/**
+	 * @return current position in uncompressed stream in bytes.
+	 */
+	virtual int64 Tell();
+
+private:
+	/**
+	 * Flushes tmp data to array.
+	 */
+	void DecompressMoreData();
+
+	/** Array to write compressed data to.						*/
+	const TArray<uint8>&	CompressedData;
+	/** Current index into compressed data array.				*/
+	int32				CurrentIndex;
+	/** Pointer to start of temporary buffer.					*/
+	uint8*			TmpDataStart;
+	/** Pointer to end of temporary buffer.						*/
+	uint8*			TmpDataEnd;
+	/** Pointer to current position in temporary buffer.		*/
+	uint8*			TmpData;
+	/** Whether to serialize from temporary buffer of array.	*/
+	bool			bShouldSerializeFromArray;
+	/** Number of raw (uncompressed) bytes serialized.			*/
+	int64		RawBytesSerialized;
+	/** Flags used for compression.								*/
+	ECompressionFlags CompressionFlags;
+};
diff --git a/Engine/Source/Runtime/Core/Public/Serialization/ArrayWriter.h b/Engine/Source/Runtime/Core/Public/Serialization/ArrayWriter.h
new file mode 100644
--- /dev/null
+++ b/Engine/Source/Runtime/Core/Public/Serialization/ArrayWriter.h
@@ -0,0 +1,11 @@
+// Copyright 1998-2016 Epic Games, Inc. All Rights Reserved.
+
+#pragma once
+
+#include "BufferArchive.h"
+
+/**
+ * Archive objects that are also a TArray. Since FBufferArchive is already the 
+ * writer version, we just typedef to the old, outdated name
+ */
+typedef FBufferArchive FArrayWriter;
diff --git a/Engine/Source/Runtime/Core/Public/Serialization/BitReader.h b/Engine/Source/Runtime/Core/Public/Serialization/BitReader.h
new file mode 100644
--- /dev/null
+++ b/Engine/Source/Runtime/Core/Public/Serialization/BitReader.h
@@ -0,0 +1,226 @@
+// Copyright 1998-2016 Epic Games, Inc. All Rights Reserved.
+
+#pragma once
+
+#include "Archive.h"
+
+CORE_API void appBitsCpy( uint8* Dest, int32 DestBit, uint8* Src, int32 SrcBit, int32 BitCount );
+
+/*-----------------------------------------------------------------------------
+	FBitReader.
+-----------------------------------------------------------------------------*/
+
+//
+// Reads bitstreams.
+//
+struct CORE_API FBitReader : public FArchive
+{
+	friend struct FBitReaderMark;
+
+public:
+
+	FBitReader( uint8* Src = nullptr, int64 CountBits = 0 );
+	void SetData( FBitReader& Src, int64 CountBits );
+	FORCEINLINE_DEBUGGABLE void SerializeBits( void* Dest, int64 LengthBits )
+	{
+		if ( IsError() || Pos+LengthBits > Num)
+		{
+			if (!IsError())
+			{
+				SetOverflowed(LengthBits);
+				//UE_LOG( LogNetSerialization, Error, TEXT( "FBitReader::SerializeBits: Pos + LengthBits > Num" ) );
+			}
+			FMemory::Memzero( Dest, (LengthBits+7)>>3 );
+			return;
+		}
+		//for( int32 i=0; i<LengthBits; i++,Pos++ )
+		//	if( Buffer(Pos>>3) & GShift[Pos&7] )
+		//		((uint8*)Dest)[i>>3] |= GShift[i&7];
+		if( LengthBits == 1 )
+		{
+			((uint8*)Dest)[0] = 0;
+			if( Buffer[Pos>>3] & Shift(Pos&7) )
+				((uint8*)Dest)[0] |= 0x01;
+			Pos++;
+		}
+		else if (LengthBits != 0)
+		{
+			((uint8*)Dest)[((LengthBits+7)>>3) - 1] = 0;
+			appBitsCpy((uint8*)Dest, 0, Buffer.GetData(), Pos, LengthBits);
+			Pos += LengthBits;
+		}
+	}
+
+	// OutValue < ValueMax
+	FORCEINLINE_DEBUGGABLE void SerializeInt(uint32& OutValue, uint32 ValueMax)
+	{
+		if (!IsError())
+		{
+			// Use local variable to avoid Load-Hit-Store
+			uint32 Value = 0;
+			int64 LocalPos = Pos;
+			const int64 LocalNum = Num;
+
+			for (uint32 Mask=1; (Value + Mask) < ValueMax && Mask; Mask *= 2, LocalPos++)
+			{
+				if (LocalPos >= LocalNum)
+				{
+					SetOverflowed(LocalPos - Pos);
+					break;
+				}
+
+				if (Buffer[LocalPos >> 3] & Shift(LocalPos & 7))
+				{
+					Value |= Mask;
+				}
+			}
+
+			// Now write back
+			Pos = LocalPos;
+			OutValue = Value;
+		}
+	}
+
+	FORCEINLINE_DEBUGGABLE uint32 ReadInt(uint32 Max)
+	{
+		uint32 Value = 0;
+
+		SerializeInt(Value, Max);
+
+		return Value;
+	}
+
+	FORCEINLINE_DEBUGGABLE uint8 ReadBit()
+	{
+		uint8 Bit=0;
+		//SerializeBits( &Bit, 1 );
+		if ( !IsError() )
+		{
+			int64 LocalPos = Pos;
+			const int64 LocalNum = Num;
+			if (LocalPos >= LocalNum)
+			{
+				SetOverflowed(1);
+				//UE_LOG( LogNetSerialization, Error, TEXT( "FBitReader::SerializeInt: LocalPos >= LocalNum" ) );
+			}
+			else
+			{
+				Bit = !!(Buffer[LocalPos>>3] & Shift(LocalPos&7));
+				Pos++;
+			}
+		}
+		return Bit;
+	}
+
+	FORCEINLINE_DEBUGGABLE void Serialize( void* Dest, int64 LengthBytes )
+	{
+		SerializeBits( Dest, LengthBytes*8 );
+	}
+
+	FORCEINLINE_DEBUGGABLE uint8* GetData()
+	{
+		return Buffer.GetData();
+	}
+
+	FORCEINLINE_DEBUGGABLE uint8* GetDataPosChecked()
+	{
+		check(Pos % 8 == 0);
+		return &Buffer[Pos >> 3];
+	}
+
+	FORCEINLINE_DEBUGGABLE uint32 GetBytesLeft()
+	{
+		return ((Num - Pos) + 7) >> 3;
+	}
+	FORCEINLINE_DEBUGGABLE uint32 GetBitsLeft()
+	{
+		return (Num - Pos);
+	}
+	FORCEINLINE_DEBUGGABLE bool AtEnd()
+	{
+		return ArIsError || Pos>=Num;
+	}
+	FORCEINLINE_DEBUGGABLE int64 GetNumBytes()
+	{
+		return (Num+7)>>3;
+	}
+	FORCEINLINE_DEBUGGABLE int64 GetNumBits()
+	{
+		return Num;
+	}
+	FORCEINLINE_DEBUGGABLE int64 GetPosBits()
+	{
+		return Pos;
+	}
+	FORCEINLINE_DEBUGGABLE void EatByteAlign()
+	{
+		int32 PrePos = Pos;
+
+		// Skip over remaining bits in current byte
+		Pos = (Pos+7) & (~0x07);
+
+		if ( Pos > Num )
+		{
+			//UE_LOG( LogNetSerialization, Error, TEXT( "FBitReader::EatByteAlign: Pos > Num" ) );
+			SetOverflowed(Pos - PrePos);
+		}
+	}
+
+	/**
+	 * Marks this bit reader as overflowed.
+	 *
+	 * @param LengthBits	The number of bits being read at the time of overflow
+	 */
+	void SetOverflowed(int32 LengthBits);
+
+	void AppendDataFromChecked( FBitReader& Src );
+	void AppendDataFromChecked( uint8* Src, uint32 NumBits );
+	void AppendTo( TArray<uint8> &Buffer );
+
+protected:
+
+	TArray<uint8> Buffer;
+	int64 Num;
+	int64 Pos;
+
+private:
+
+	FORCEINLINE uint8 Shift(uint8 Cnt)
+	{
+		return 1<<Cnt;
+	}
+
+};
+
+
+//
+// For pushing and popping FBitWriter positions.
+//
+struct CORE_API FBitReaderMark
+{
+public:
+
+	FBitReaderMark()
+		: Pos(0)
+	{ }
+
+	FBitReaderMark( FBitReader& Reader )
+		: Pos(Reader.Pos)
+	{ }
+
+	int64 GetPos()
+	{
+		return Pos;
+	}
+
+	void Pop( FBitReader& Reader )
+	{
+		Reader.Pos = Pos;
+	}
+
+	void Copy( FBitReader& Reader, TArray<uint8> &Buffer );
+
+private:
+
+	int64 Pos;
+};
diff --git a/Engine/Source/Runtime/Core/Public/Serialization/BitWriter.h b/Engine/Source/Runtime/Core/Public/Serialization/BitWriter.h
new file mode 100644
--- /dev/null
+++ b/Engine/Source/Runtime/Core/Public/Serialization/BitWriter.h
@@ -0,0 +1,209 @@
+// Copyright 1998-2016 Epic Games, Inc. All Rights Reserved.
+
+#pragma once
+
+#include "Archive.h"
+
+/*-----------------------------------------------------------------------------
+	FBitWriter.
+-----------------------------------------------------------------------------*/
+
+//
+// Writes bitstreams.
+//
+struct CORE_API FBitWriter : public FArchive
+{
+	friend struct FBitWriterMark;
+
+public:
+	/** Default constructor. Zeros everything. */
+	FBitWriter(void);
+
+	/**
+	 * Constructor using known size the buffer needs to be
+	 */
+	FBitWriter( int64 InMaxBits, bool AllowResize = false );
+
+	void SerializeBits( void* Src, int64 LengthBits );
+
+	/**
+	 * Serializes a compressed integer - Value must be < Max
+	 *
+	 * @param Value		The value to serialize, must be < Max
+	 * @param Max		The maximum allowed value - good to aim for power-of-two
+	 */
+	void SerializeInt(uint32& Value, uint32 Max);
+
+	/**
+	 * Serializes the specified Value, but does not bounds check against ValueMax;
+	 * instead, it will wrap around if the value exceeds ValueMax (this differs from SerializeInt, which clamps)
+	 *
+	 * @param Value		The value to serialize
+	 * @param ValueMax	The maximum value to write; wraps Value if it exceeds this
+	 */
+	void WriteIntWrapped(uint32 Value, uint32 ValueMax);
+
+	void WriteBit( uint8 In );
+	void Serialize( void* Src, int64 LengthBytes );
+
+	/**
+	 * Returns a pointer to our internal buffer.
+	 */
+	FORCEINLINE uint8* GetData(void)
+	{
+#if !UE_BUILD_SHIPPING
+		// If this happens, your code has insufficient IsError() checks.
+		check(!IsError());
+#endif
+
+		return Buffer.GetData();
+	}
+
+	FORCEINLINE const uint8* GetData(void) const
+	{
+#if !UE_BUILD_SHIPPING
+		// If this happens, your code has insufficient IsError() checks.
+		check(!IsError());
+#endif
+
+		return Buffer.GetData();
+	}
+
+	FORCEINLINE const TArray<uint8>* GetBuffer(void) const
+	{
+#if !UE_BUILD_SHIPPING
+		// If this happens, your code has insufficient IsError() checks.
+		check(!IsError());
+#endif
+
+		return &Buffer;
+	}
+
+	/**
+	 * Returns the number of bytes written.
+	 */
+	FORCEINLINE int64 GetNumBytes(void) const
+	{
+		return (Num+7)>>3;
+	}
+
+	/**
+	 * Returns the number of bits written.
+	 */
+	FORCEINLINE int64 GetNumBits(void) const
+	{
+		return Num;
+	}
+
+	/**
+	 * Returns the number of bits the buffer supports.
+	 */
+	FORCEINLINE int64 GetMaxBits(void) const
+	{
+		return Max;
+	}
+
+	/**
+	 * Marks this bit writer as overflowed.
+	 *
+	 * @param LengthBits	The number of bits being written at the time of overflow
+	 */
+	void SetOverflowed(int32 LengthBits);
+
+	/**
+	 * Sets whether or not this writer intentionally allows overflows (to avoid logspam)
+	 */
+	FORCEINLINE void SetAllowOverflow(bool bInAllow)
+	{
+		bAllowOverflow = bInAllow;
+	}
+
+	FORCEINLINE bool AllowAppend(int64 LengthBits)
+	{
+		if (Num+LengthBits > Max)
+		{
+			if (AllowResize)
+			{
+				// Resize our buffer. The common case for resizing bitwriters is hitting the max and continuing to add a lot of small segments of data
+				// Though we could just allow the TArray buffer to handle the slack and resizing, we would still constantly hit the FBitWriter's max
+				// and cause this block to be executed, as well as constantly zeroing out memory inside AddZeroes (though the memory would be allocated
+				// in chunks).
+
+				Max = FMath::Max<int64>(Max<<1,Num+LengthBits);
+				int32 ByteMax = (Max+7)>>3;
+				Buffer.AddZeroed(ByteMax - Buffer.Num());
+				check((Max+7)>>3== Buffer.Num());
+				return true;
+			}
+			else
+			{
+				return false;
+			}
+		}
+		return true;
+	}
+
+	FORCEINLINE void SetAllowResize(bool NewResize)
+	{
+		AllowResize = NewResize;
+	}
+
+	/**
+	 * Resets the bit writer back to its initial state
+	 */
+	void Reset(void);
+
+	FORCEINLINE void WriteAlign()
+	{
+		Num = ( Num + 7 ) & ( ~0x07 );
+	}
+
+private:
+
+	TArray<uint8> Buffer;
+	int64   Num;
+	int64   Max;
+	bool	AllowResize;
+
+	/** Whether or not overflowing is allowed (overflows silently) */
+	bool bAllowOverflow;
+};
+
+
+//
+// For pushing and popping FBitWriter positions.
+//
+struct CORE_API FBitWriterMark
+{
+public:
+
+	FBitWriterMark()
+		: Overflowed(false)
+		, Num(0)
+	{ }
+
+	FBitWriterMark( FBitWriter& Writer )
+	{
+		Init(Writer);
+	}
+
+	int64 GetNumBits()
+	{
+		return Num;
+	}
+
+	void Init( FBitWriter& Writer)
+	{
+		Num = Writer.Num;
+		Overflowed = Writer.ArIsError;
+	}
+
+	void Pop( FBitWriter& Writer );
+	void Copy( FBitWriter& Writer, TArray<uint8> &Buffer );
+	void PopWithoutClear( FBitWriter& Writer );
+
+private:
+
+	bool			Overflowed;
+	int64			Num;
+};
diff --git a/Engine/Source/Runtime/Core/Public/Serialization/BufferArchive.h b/Engine/Source/Runtime/Core/Public/Serialization/BufferArchive.h
new file mode 100644
--- /dev/null
+++ b/Engine/Source/Runtime/Core/Public/Serialization/BufferArchive.h
@@ -0,0 +1,25 @@
+// Copyright 1998-2016 Epic Games, Inc. All Rights Reserved.
+
+#pragma once
+
+#include "Archive.h"
+#include "MemoryWriter.h"
+
+/**
+ * Buffer archiver.
+ */
+class FBufferArchive : public FMemoryWriter, public TArray<uint8>
+{
+public:
+	FBufferArchive( bool bIsPersistent = false, const FName InArchiveName = NAME_None )
+	: FMemoryWriter( (TArray<uint8>&)*this, bIsPersistent, false, InArchiveName )
+	{}
+	/**
+  	 * Returns the name of the Archive.  Useful for getting the name of the package a struct or object
+	 * is in when a loading error occurs.
+	 *
+	 * This is overridden for the specific Archive Types
+	 **/
+	virtual FString GetArchiveName() const { return *FString::Printf( TEXT("FBufferArchive %s"), *ArchiveName.ToString()); }
+};
+
diff --git a/Engine/Source/Runtime/Core/Public/Serialization/BufferWriter.h b/Engine/Source/Runtime/Core/Public/Serialization/BufferWriter.h
new file mode 100644
--- /dev/null
+++ b/Engine/Source/Runtime/Core/Public/Serialization/BufferWriter.h
@@ -0,0 +1,100 @@
+// Copyright 1998-2016 Epic Games, Inc. All Rights Reserved.
+
+#pragma once
+
+#include "Archive.h"
+
+/**
+* Similar to FMemoryWriter, but able to internally
+* manage the memory for the buffer.
+*/
+class FBufferWriter final : public FArchive
+{
+public:
+	/**
+	* Constructor
+	*
+	* @param Data Buffer to use as the source data to read from
+	* @param Size Size of Data
+	* @param bInFreeOnClose If true, Data will be FMemory::Free'd when this archive is closed
+	* @param bIsPersistent Uses this value for ArIsPersistent
+	*/
+	FBufferWriter(void* Data, int64 Size, bool bInFreeOnClose, bool bIsPersistent = false)
+		: WriterData(Data)
+		, WriterPos(0)
+		, WriterSize(Size)
+		, bFreeOnClose(bInFreeOnClose)
+	{
+		ArIsSaving = true;
+		ArIsPersistent = bIsPersistent;
+	}
+
+	~FBufferWriter()
+	{
+		Close();
+	}
+	bool Close()
+	{
+		if (bFreeOnClose)
+		{
+			FMemory::Free(WriterData);
+			WriterData = nullptr;
+		}
+		return !ArIsError;
+	}
+	void Serialize(void* Data, int64 Num)
+	{
+		const int64 NumBytesToAdd = WriterPos + Num - WriterSize;
+		if (NumBytesToAdd > 0)
+		{
+			const int64 NewArrayCount = WriterSize + NumBytesToAdd;
+			if (NewArrayCount >= MAX_int32)
+			{
+				UE_LOG(LogSerialization, Fatal, TEXT("FBufferWriter does not support data larger than 2GB. Archive name: %s."), *GetArchiveName());
+			}
+
+			FMemory::Realloc(WriterData, NewArrayCount);
+			WriterSize = NewArrayCount;
+		}
+
+		check(WriterPos >= 0);
+		check((WriterPos + Num) <= WriterSize);
+		FMemory::Memcpy((uint8*)WriterData + WriterPos, Data, Num);
+		WriterPos += Num;
+	}
+	int64 Tell()
+	{
+		return WriterPos;
+	}
+	int64 TotalSize()
+	{
+		return WriterSize;
+	}
+	void Seek(int64 InPos)
+	{
+		check(InPos >= 0);
+		check(InPos <= WriterSize);
+		WriterPos = InPos;
+	}
+	bool AtEnd()
+	{
+		return WriterPos >= WriterSize;
+	}
+	/**
+	* Returns the name of the Archive.  Useful for getting the name of the package a struct or object
+	* is in when a loading error occurs.
+	*
+	* This is overridden for the specific Archive Types
+	**/
+	virtual FString GetArchiveName() const { return TEXT("FBufferWriter"); }
+
+	void* GetWriterData()
+	{
+		return WriterData;
+	}
+protected:
+	void* WriterData;
+	int64 WriterPos;
+	int64 WriterSize;
+	bool bFreeOnClose;
+};
diff --git a/Engine/Source/Runtime/Core/Public/Serialization/CompressedChunkInfo.h b/Engine/Source/Runtime/Core/Public/Serialization/CompressedChunkInfo.h
new file mode 100644
--- /dev/null
+++ b/Engine/Source/Runtime/Core/Public/Serialization/CompressedChunkInfo.h
@@ -0,0 +1,27 @@
+// Copyright 1998-2016 Epic Games, Inc. All Rights Reserved.
+
+#pragma once
+
+#include "Archive.h"
+
+/**
+ * Implements a helper structure for compression support
+ *
+ * This structure contains information on the compressed and uncompressed size of a chunk of data.
+ */
+struct FCompressedChunkInfo
+{
+	/** Holds the data's compressed size. */
+	int64 CompressedSize;
+
+	/** Holds the data's uncompresses size. */
+	int64 UncompressedSize;
+};
+
+/**
+ * Serializes an FCompressedChunkInfo value from or into an archive.
+ *
+ * @param Ar The archive to serialize from or to.
+ * @param Value The value to serialize.
+ */
+CORE_API FArchive& operator<<(FArchive& Ar, FCompressedChunkInfo& Value);
diff --git a/Engine/Source/Runtime/Core/Public/Serialization/LargeMemoryReader.h b/Engine/Source/Runtime/Core/Public/Serialization/LargeMemoryReader.h
new file mode 100644
--- /dev/null
+++ b/Engine/Source/Runtime/Core/Public/Serialization/LargeMemoryReader.h
@@ -0,0 +1,64 @@
+// Copyright 1998-2016 Epic Games, Inc. All Rights Reserved.
+
+#pragma once
+
+#include "MemoryArchive.h"
+#include "EnumClassFlags.h"
+
+enum class ELargeMemoryReaderFlags : uint8
+{
+	None =				0x0,	// No Flags
+	TakeOwnership =		0x1,	// Archive will take ownership of the passed-in memory and free it on destruction
+	Persistent =		0x2,	// Archive will be set as persistent when constructed
+};
+
+ENUM_CLASS_FLAGS(ELargeMemoryReaderFlags);
+
+/**
+* Archive for reading a large amount of arbitrary data from memory
+*/
+class CORE_API FLargeMemoryReader : public FMemoryArchive
+{
+public:
+
+	FLargeMemoryReader(const uint8* InData, const int64 Num, ELargeMemoryReaderFlags InFlags = ELargeMemoryReaderFlags::None, const FName InArchiveName = NAME_None);
+
+	virtual void Serialize(void* OutData, int64 Num) override;
+
+	/**
+	* Gets the total size of the data buffer
+	*/
+	virtual int64 TotalSize() override;
+
+	/**
+	* Returns the name of the Archive.  Useful for getting the name of the package a struct or object
+	* is in when a loading error occurs.
+	*
+	* This is overridden for the specific Archive Types
+	**/
+	virtual FString GetArchiveName() const override;
+
+	/**
+	* Destructor
+	*/
+	~FLargeMemoryReader();
+
+private:
+
+	/** Non-copyable */
+	FLargeMemoryReader(const FLargeMemoryReader&) = delete;
+	FLargeMemoryReader& operator=(const FLargeMemoryReader&) = delete;
+
+	/** Whether the data buffer should be freed when this archive is closed */
+	const uint8 bFreeOnClose : 1;
+	
+	/** Data buffer we are reading from */
+	const uint8* Data;
+
+	/** Total size of the data buffer */
+	const int64 NumBytes;
+
+	/** Name of the archive this buffer is for */
+	const FName ArchiveName;
+};
+
diff --git a/Engine/Source/Runtime/Core/Public/Serialization/LargeMemoryWriter.h b/Engine/Source/Runtime/Core/Public/Serialization/LargeMemoryWriter.h
new file mode 100644
--- /dev/null
+++ b/Engine/Source/Runtime/Core/Public/Serialization/LargeMemoryWriter.h
@@ -0,0 +1,67 @@
+// Copyright 1998-2016 Epic Games, Inc. All Rights Reserved.
+
+#pragma once
+
+#include "MemoryArchive.h"
+
+/**
+* Archive for storing a large amount of arbitrary data to memory
+*/
+class CORE_API FLargeMemoryWriter : public FMemoryArchive
+{
+public:
+	
+	FLargeMemoryWriter(const int64 PreAllocateBytes = 0, bool bIsPersistent = false, const FName InArchiveName = NAME_None);
+
+	virtual void Serialize(void* InData, int64 Num) override;
+
+	/**
+	* Returns the name of the Archive.  Useful for getting the name of the package a struct or object
+	* is in when a loading error occurs.
+	*
+	* This is overridden for the specific Archive Types
+	**/
+	virtual FString GetArchiveName() const override;
+
+	/**
+	 * Gets the total size of the data written
+	 */
+	virtual int64 TotalSize() override;
+	
+	/**
+	 * Returns the written data. To release this archive's ownership of the data, call ReleaseOwnership()
+	 */
+	uint8* GetData() const;
+
+	/** 
+	 * Releases ownership of the written data
+	 */
+	void ReleaseOwnership();
+	
+	/**
+	 * Destructor
+	 */
+	~FLargeMemoryWriter();
+
+private:
+
+	/** Non-copyable */
+	FLargeMemoryWriter(const FLargeMemoryWriter&) = delete;
+	FLargeMemoryWriter& operator=(const FLargeMemoryWriter&) = delete;
+
+	/** Memory owned by this archive. Ownership can be released by calling ReleaseOwnership() */
+	uint8* Data;
+
+	/** Number of bytes currently written to our data buffer */
+	int64 NumBytes;
+
+	/** Number of bytes currently allocated for our data buffer */
+	int64 MaxBytes;
+
+	/** Resizes the data buffer to at least the desired new size with some slack */
+	void GrowBuffer(const int64 DesiredBytes);
+
+	/** Archive name, used for debugging, by default set to NAME_None. */
+	const FName ArchiveName;
+};
+
diff --git a/Engine/Source/Runtime/Core/Public/Serialization/MemoryReader.h b/Engine/Source/Runtime/Core/Public/Serialization/MemoryReader.h
new file mode 100644
--- /dev/null
+++ b/Engine/Source/Runtime/Core/Public/Serialization/MemoryReader.h
@@ -0,0 +1,63 @@
+// Copyright 1998-2016 Epic Games, Inc. All Rights Reserved.
+
+#pragma once
+
+#include "MemoryArchive.h"
+
+/**
+ * Archive for reading arbitrary data from the specified memory location
+ */
+class FMemoryReader final : public FMemoryArchive
+{
+public:
+	/**
+  	 * Returns the name of the Archive.  Useful for getting the name of the package a struct or object
+	 * is in when a loading error occurs.
+	 *
+	 * This is overridden for the specific Archive Types
+	 **/
+	virtual FString GetArchiveName() const { return TEXT("FMemoryReader"); }
+
+	int64 TotalSize()
+	{
+		return FMath::Min((int64)Bytes.Num(), LimitSize);
+	}
+
+	void Serialize( void* Data, int64 Num )
+	{
+		if (Num && !ArIsError)
+		{
+			// Only serialize if we have the requested amount of data
+			if (Offset + Num <= TotalSize())
+			{
+				FMemory::Memcpy( Data, &Bytes[Offset], Num );
+				Offset += Num;
+			}
+			else
+			{
+				ArIsError = true;
+			}
+		}
+	}
+
+	FMemoryReader( const TArray<uint8>& InBytes, bool bIsPersistent = false )
+	: FMemoryArchive()
+	, Bytes(InBytes)
+	, LimitSize(INT64_MAX)
+	{
+		ArIsLoading		= true;
+		ArIsPersistent	= bIsPersistent;
+	}
+
+	/** With this method it's possible to attach data behind some serialized data. */
+	void SetLimitSize(int64 NewLimitSize)
+	{
+		LimitSize = NewLimitSize;
+	}
+
+protected:
+
+	const TArray<uint8>& Bytes;
+	int64 LimitSize;
+};
+
diff --git a/Engine/Source/Runtime/Core/Public/Serialization/MemoryWriter.h b/Engine/Source/Runtime/Core/Public/Serialization/MemoryWriter.h
new file mode 100644
--- /dev/null
+++ b/Engine/Source/Runtime/Core/Public/Serialization/MemoryWriter.h
@@ -0,0 +1,68 @@
+// Copyright 1998-2016 Epic Games, Inc. All Rights Reserved.
+
+#pragma once
+
+#include "MemoryArchive.h"
+
+/**
+ * Archive for storing arbitrary data to the specified memory location
+ */
+class FMemoryWriter : public FMemoryArchive
+{
+public:
+	FMemoryWriter( TArray<uint8>& InBytes, bool bIsPersistent = false, bool bSetOffset = false, const FName InArchiveName = NAME_None )
+	: FMemoryArchive()
+	, Bytes(InBytes)
+	, ArchiveName(InArchiveName)
+	{
+		ArIsSaving		= true;
+		ArIsPersistent	= bIsPersistent;
+		if (bSetOffset)
+		{
+			Offset = InBytes.Num();
+		}
+	}
+
+	virtual void Serialize(void* Data, int64 Num) override
+	{
+		const int64 NumBytesToAdd = Offset + Num - Bytes.Num();
+		if( NumBytesToAdd > 0 )
+		{
+			const int64 NewArrayCount = Bytes.Num() + NumBytesToAdd;
+			if( NewArrayCount >= MAX_int32 )
+			{
+				UE_LOG( LogSerialization, Fatal, TEXT( "FMemoryWriter does not support data larger than 2GB. Archive name: %s." ), *ArchiveName.ToString() );
+			}
+
+			Bytes.AddUninitialized( (int32)NumBytesToAdd );
+		}
+
+		check((Offset + Num) <= Bytes.Num());
+		
+		if( Num )
+		{
+			FMemory::Memcpy( &Bytes[Offset], Data, Num );
+			Offset+=Num;
+		}
+	}
+	/**
+  	 * Returns the name of the Archive.  Useful for getting the name of the package a struct or object
+	 * is in when a loading error occurs.
+	 *
+	 * This is overridden for the specific Archive Types
+	 **/
+	virtual FString GetArchiveName() const override { return TEXT("FMemoryWriter"); }
+
+	int64 TotalSize() override
+	{
+		return Bytes.Num();
+	}
+
+protected:
+
+	TArray<uint8>&	Bytes;
+
+	/** Archive name, used to debugging, by default set to NAME_None. */
+	const FName ArchiveName;
+};
+
diff --git a/Engine/Source/Runtime/Core/Public/Stats/Stats.h b/Engine/Source/Runtime/Core/Public/Stats/Stats.h
--- a/Engine/Source/Runtime/Core/Public/Stats/Stats.h
+++ b/Engine/Source/Runtime/Core/Public/Stats/Stats.h
@@ -312,6 +312,16 @@ template<> inline const TCHAR* GetStatFormatString<int64>(void) { return TEXT("%
 template<> inline const TCHAR* GetStatFormatString<float>(void) { return TEXT("%.1f"); }
 template<> inline const TCHAR* GetStatFormatString<double>(void) { return TEXT("%.1f"); }
 
+/*----------------------------------------------------------------------------
+	Stats helpers
+----------------------------------------------------------------------------*/
+
+#if STATS
+	#define STAT(x) x
+#else
+	#define STAT(x)
+#endif
+
 #include "Stats2.h"
 
 #if STATS
diff --git a/Engine/Source/Runtime/Core/Public/Windows/WindowsCriticalSection.h b/Engine/Source/Runtime/Core/Public/Windows/WindowsCriticalSection.h
--- a/Engine/Source/Runtime/Core/Public/Windows/WindowsCriticalSection.h
+++ b/Engine/Source/Runtime/Core/Public/Windows/WindowsCriticalSection.h
@@ -99,3 +99,6 @@ private:
 private:
 	HANDLE Mutex;
 };
+
+typedef FWindowsCriticalSection FCriticalSection;
+typedef FWindowsSystemWideCriticalSection FSystemWideCriticalSection;
diff --git a/Engine/Source/Runtime/Core/Public/Windows/WindowsPlatformProcess.h b/Engine/Source/Runtime/Core/Public/Windows/WindowsPlatformProcess.h
--- a/Engine/Source/Runtime/Core/Public/Windows/WindowsPlatformProcess.h
+++ b/Engine/Source/Runtime/Core/Public/Windows/WindowsPlatformProcess.h
@@ -214,7 +214,3 @@ inline bool FProcHandle::Close()
 	return false;
 }
 
-#include "WindowsCriticalSection.h"
-
-typedef FWindowsCriticalSection FCriticalSection;
-typedef FWindowsSystemWideCriticalSection FSystemWideCriticalSection;
diff --git a/Engine/Source/Runtime/Core/Public/iOS/IOSCriticalSection.h b/Engine/Source/Runtime/Core/Public/iOS/IOSCriticalSection.h
new file mode 100644
--- /dev/null
+++ b/Engine/Source/Runtime/Core/Public/iOS/IOSCriticalSection.h
@@ -0,0 +1,9 @@
+// Copyright 1998-2016 Epic Games, Inc. All Rights Reserved.
+
+#pragma once
+
+#include "GenericPlatform/GenericPlatformCriticalSection.h"
+#include "PThreadCriticalSection.h"
+
+typedef FPThreadsCriticalSection FCriticalSection;
+typedef FSystemWideCriticalSectionNotImplemented FSystemWideCriticalSection;
diff --git a/Engine/Source/Runtime/CoreUObject/Private/CoreUObjectPrivate.h b/Engine/Source/Runtime/CoreUObject/Private/CoreUObjectPrivate.h
--- a/Engine/Source/Runtime/CoreUObject/Private/CoreUObjectPrivate.h
+++ b/Engine/Source/Runtime/CoreUObject/Private/CoreUObjectPrivate.h
@@ -8,10 +8,11 @@
 
 #include "Core.h"
 
+#include "ErrorException.h"
 #include "ModuleManager.h"
 #include "ScriptInterface.h"			// Script interface definitions.
 #include "Script.h"						// Script constants and forward declarations.
-#include "ObjectBase.h"					// Object base class.
+#include "ObjectMacros.h"					// Object base class.
 #include "UObjectAllocator.h"
 #include "UObjectGlobals.h"
 #include "UObjectMarks.h"
@@ -23,24 +24,50 @@
 #include "Object.h"
 #include "UObjectIterator.h"
 #include "CoreNet.h"					// Core networking.
-#include "ArchiveUObjectBase.h"			// UObject-related Archive classes.
+#include "ArchiveUObject.h"				// UObject-related Archive classes.
 #include "GarbageCollection.h"			// Realtime garbage collection helpers
 #include "Class.h"						// Class definition.
 #include "StructOnScope.h"
 #include "Casts.h"                      // Cast templates
 #include "LazyObjectPtr.h"				// Object pointer types
 #include "AssetPtr.h"					// Object pointer types
-#include "CoreObject.h"					// Core object class definitions.
+#include "LevelGuids.h"					// Core object class definitions.
 #include "UnrealType.h"					// Base property type.
-#include "UTextProperty.h"				// Text property type.
+#include "TextProperty.h"				// Text property type.
 #include "Stack.h"						// Script stack frame definition.
 #include "ObjectRedirector.h"			// Cross-package object redirector
 #include "UObjectAnnotation.h"
 #include "ReferenceChainSearch.h"		// Search for referencers of an objects
 #include "ArchiveUObject.h"				// UObject-related Archive classes.
+#include "AsyncFileHandle.h"
+#include "TextPackageNamespaceUtil.h"
+#include "ArchiveCountMem.h"
+#include "ObjectAndNameAsStringProxyArchive.h"
+#include "ObjectWriter.h"
+#include "ObjectReader.h"
+#include "ReloadObjectArc.h"
+#include "ArchiveReplaceArchetype.h"
+#include "ArchiveShowReferences.h"
+#include "FindReferencersArchive.h"
+#include "FindObjectReferencers.h"
+#include "ArchiveFindCulprit.h"
+#include "ArchiveObjectGraph.h"
+#include "TraceReferences.h"
+#include "ArchiveTraceRoute.h"
+#include "DuplicatedObject.h"
+#include "DuplicatedDataReader.h"
+#include "DuplicatedDataWriter.h"
+#include "ArchiveReplaceObjectRef.h"
+#include "ArchiveReplaceOrClearExternalReferences.h"
+#include "ArchiveObjectPropertyMapper.h"
+#include "ArchiveReferenceMarker.h"
+#include "ArchiveAsync.h"
+#include "ArchiveObjectCrc32.h"
 #include "PackageName.h"
 #include "BulkData.h"					// Bulk data classes
 #include "Linker.h"						// Linker.
+#include "LinkerLoad.h"
+#include "LinkerSave.h"
 #include "GCObject.h"			        // non-UObject object referencer
 #include "RedirectCollector.h"
 #include "WorldCompositionUtility.h"
diff --git a/Engine/Source/Runtime/CoreUObject/Private/Serialization/PropertyLocalizationDataGathering.cpp b/Engine/Source/Runtime/CoreUObject/Private/Serialization/PropertyLocalizationDataGathering.cpp
--- a/Engine/Source/Runtime/CoreUObject/Private/Serialization/PropertyLocalizationDataGathering.cpp
+++ b/Engine/Source/Runtime/CoreUObject/Private/Serialization/PropertyLocalizationDataGathering.cpp
@@ -3,7 +3,7 @@
 #include "CoreUObjectPrivate.h"
 #include "PropertyLocalizationDataGathering.h"
 #include "TextPackageNamespaceUtil.h"
-#include "UTextProperty.h"
+#include "TextProperty.h"
 
 FPropertyLocalizationDataGatherer::FPropertyLocalizationDataGatherer(TArray<FGatherableTextData>& InOutGatherableTextDataArray, const UPackage* const InPackage, EPropertyLocalizationGathererResultFlags& OutResultFlags)
 	: GatherableTextDataArray(InOutGatherableTextDataArray)
diff --git a/Engine/Source/Runtime/CoreUObject/Private/UObject/ErrorException.cpp b/Engine/Source/Runtime/CoreUObject/Private/UObject/ErrorException.cpp
new file mode 100644
--- /dev/null
+++ b/Engine/Source/Runtime/CoreUObject/Private/UObject/ErrorException.cpp
@@ -0,0 +1,20 @@
+// Copyright 1998-2016 Epic Games, Inc. All Rights Reserved.
+
+#include "CoreUObjectPrivate.h"
+#include "UObject/ErrorException.h"
+
+//
+// Throw a string exception with a message.
+//
+#if HACK_HEADER_GENERATOR
+void VARARGS FError::ThrowfImpl(const TCHAR* Fmt, ...)
+{
+	static TCHAR TempStr[4096];
+	GET_VARARGS( TempStr, ARRAY_COUNT(TempStr), ARRAY_COUNT(TempStr)-1, Fmt, Fmt );
+#if !PLATFORM_EXCEPTIONS_DISABLED
+	throw( TempStr );
+#else
+	UE_LOG(LogOutputDevice, Fatal, TEXT("THROW: %s"), TempStr);
+#endif
+}					
+#endif
diff --git a/Engine/Source/Runtime/CoreUObject/Private/UObject/Interface.cpp b/Engine/Source/Runtime/CoreUObject/Private/UObject/Interface.cpp
new file mode 100644
--- /dev/null
+++ b/Engine/Source/Runtime/CoreUObject/Private/UObject/Interface.cpp
@@ -0,0 +1,23 @@
+// Copyright 1998-2016 Epic Games, Inc. All Rights Reserved.
+
+#include "CoreUObjectPrivate.h"
+#include "Interface.h"
+
+#if WITH_EDITOR
+IMPLEMENT_CORE_INTRINSIC_CLASS(UInterface, UObject,
+	{
+		Class->ClassAddReferencedObjects = &UInterface::AddReferencedObjects;
+		Class->SetMetaData(TEXT("IsBlueprintBase"), TEXT("true"));
+		Class->SetMetaData(TEXT("CannotImplementInterfaceInBlueprint"), TEXT(""));
+	}
+);
+
+#else
+
+IMPLEMENT_CORE_INTRINSIC_CLASS(UInterface, UObject,
+{
+	Class->ClassAddReferencedObjects = &UInterface::AddReferencedObjects;
+}
+);
+
+#endif
\ No newline at end of file
diff --git a/Engine/Source/Runtime/CoreUObject/Private/UObject/LinkerLoad.cpp b/Engine/Source/Runtime/CoreUObject/Private/UObject/LinkerLoad.cpp
--- a/Engine/Source/Runtime/CoreUObject/Private/UObject/LinkerLoad.cpp
+++ b/Engine/Source/Runtime/CoreUObject/Private/UObject/LinkerLoad.cpp
@@ -1,5 +1,6 @@
 // Copyright 1998-2016 Epic Games, Inc. All Rights Reserved.
 #include "CoreUObjectPrivate.h"
+#include "UObject/LinkerLoad.h"
 #include "SecureHash.h"
 #include "DebuggingDefines.h"
 #include "MessageLog.h"
diff --git a/Engine/Source/Runtime/CoreUObject/Private/UObject/LinkerSave.cpp b/Engine/Source/Runtime/CoreUObject/Private/UObject/LinkerSave.cpp
--- a/Engine/Source/Runtime/CoreUObject/Private/UObject/LinkerSave.cpp
+++ b/Engine/Source/Runtime/CoreUObject/Private/UObject/LinkerSave.cpp
@@ -1,6 +1,7 @@
 // Copyright 1998-2016 Epic Games, Inc. All Rights Reserved.
 
 #include "CoreUObjectPrivate.h"
+#include "UObject/LinkerSave.h"
 #include "TextPackageNamespaceUtil.h"
 
 /*----------------------------------------------------------------------------
diff --git a/Engine/Source/Runtime/CoreUObject/Private/UObject/ObjectResource.cpp b/Engine/Source/Runtime/CoreUObject/Private/UObject/ObjectResource.cpp
--- a/Engine/Source/Runtime/CoreUObject/Private/UObject/ObjectResource.cpp
+++ b/Engine/Source/Runtime/CoreUObject/Private/UObject/ObjectResource.cpp
@@ -1,6 +1,7 @@
 // Copyright 1998-2016 Epic Games, Inc. All Rights Reserved.
 
 #include "CoreUObjectPrivate.h"
+#include "UObject/ObjectResource.h"
 
 /*-----------------------------------------------------------------------------
 	Helper functions.
diff --git a/Engine/Source/Runtime/CoreUObject/Private/UObject/PackageFileSummary.cpp b/Engine/Source/Runtime/CoreUObject/Private/UObject/PackageFileSummary.cpp
--- a/Engine/Source/Runtime/CoreUObject/Private/UObject/PackageFileSummary.cpp
+++ b/Engine/Source/Runtime/CoreUObject/Private/UObject/PackageFileSummary.cpp
@@ -1,6 +1,7 @@
 // Copyright 1998-2016 Epic Games, Inc. All Rights Reserved.
 
 #include "CoreUObjectPrivate.h"
+#include "UObject/PackageFileSummary.h"
 #include "EngineVersion.h"
 #include "EditorObjectVersion.h"
 
diff --git a/Engine/Source/Runtime/CoreUObject/Private/UObject/SavePackage.cpp b/Engine/Source/Runtime/CoreUObject/Private/UObject/SavePackage.cpp
--- a/Engine/Source/Runtime/CoreUObject/Private/UObject/SavePackage.cpp
+++ b/Engine/Source/Runtime/CoreUObject/Private/UObject/SavePackage.cpp
@@ -1,7 +1,7 @@
 ﻿// Copyright 1998-2016 Epic Games, Inc. All Rights Reserved.
 
 #include "CoreUObjectPrivate.h"
-#include "UObject/UTextProperty.h"
+#include "UObject/TextProperty.h"
 #include "TextPackageNamespaceUtil.h"
 #include "Interface.h"
 #include "TargetPlatform.h"
diff --git a/Engine/Source/Runtime/CoreUObject/Private/UObject/TextProperty.cpp b/Engine/Source/Runtime/CoreUObject/Private/UObject/TextProperty.cpp
new file mode 100644
--- /dev/null
+++ b/Engine/Source/Runtime/CoreUObject/Private/UObject/TextProperty.cpp
@@ -0,0 +1,202 @@
+// Copyright 1998-2016 Epic Games, Inc. All Rights Reserved.
+
+#include "CoreUObjectPrivate.h"
+#include "UObject/TextProperty.h"
+#include "PropertyHelper.h"
+#include "PropertyTag.h"
+#include "TextPackageNamespaceUtil.h"
+
+bool UTextProperty::ConvertFromType(const FPropertyTag& Tag, FArchive& Ar, uint8* Data, UStruct* DefaultsStruct, bool& bOutAdvanceProperty)
+{
+	// Convert serialized string to text.
+	if (Tag.Type==NAME_StrProperty) 
+	{ 
+		FString str;
+		Ar << str;
+		FText Text = FText::FromString(str);
+		Text.TextData->PersistText();
+		Text.Flags |= ETextFlag::ConvertedProperty;
+		SetPropertyValue_InContainer(Data, Text, Tag.ArrayIndex);
+		bOutAdvanceProperty = true;
+	}
+
+	// Convert serialized name to text.
+	else if (Tag.Type==NAME_NameProperty) 
+	{ 
+		FName Name;  
+		Ar << Name;
+		FText Text = FText::FromName(Name);
+		Text.Flags |= ETextFlag::ConvertedProperty;
+		SetPropertyValue_InContainer(Data, Text, Tag.ArrayIndex);
+		bOutAdvanceProperty = true;
+	}
+	else
+	{
+		bOutAdvanceProperty = false;
+	}
+
+	return bOutAdvanceProperty;
+}
+
+bool UTextProperty::Identical_Implementation(const FText& ValueA, const FText& ValueB, uint32 PortFlags)
+{
+	if (ValueA.IsCultureInvariant() != ValueB.IsCultureInvariant() || ValueA.IsTransient() != ValueB.IsTransient())
+	{
+		//A culture variant text is never equal to a culture invariant text
+		//A transient text is never equal to a non-transient text
+		return false;
+	}
+
+	if (ValueA.IsCultureInvariant() == ValueB.IsCultureInvariant() || ValueA.IsTransient() == ValueB.IsTransient())
+	{
+		//Culture invariant text don't have a namespace/key so we compare the source string
+		//Transient text don't have a namespace/key or source so we compare the display string
+		return FTextInspector::GetDisplayString(ValueA) == FTextInspector::GetDisplayString(ValueB);
+	}
+
+	if (GIsEditor)
+	{
+		return FTextInspector::GetSourceString(ValueA)->Compare(*FTextInspector::GetSourceString(ValueB), ESearchCase::CaseSensitive) == 0;
+	}
+	else
+	{
+		return	FTextInspector::GetNamespace(ValueA) == FTextInspector::GetNamespace(ValueB) &&
+			FTextInspector::GetKey(ValueA) == FTextInspector::GetKey(ValueB);
+	}
+}
+
+bool UTextProperty::Identical( const void* A, const void* B, uint32 PortFlags ) const
+{
+	const TCppType ValueA = GetPropertyValue(A);
+	if ( B )
+	{
+		const TCppType ValueB = GetPropertyValue(B);
+		return Identical_Implementation(ValueA, ValueB, PortFlags);
+	}
+
+	return FTextInspector::GetDisplayString(ValueA).IsEmpty();
+}
+
+void UTextProperty::SerializeItem( FArchive& Ar, void* Value, void const* Defaults ) const
+{
+	const TCppType PropertyValue = GetPropertyValue(Value);
+	Ar << *GetPropertyValuePtr(Value);
+}
+
+void UTextProperty::ExportTextItem( FString& ValueStr, const void* PropertyValue, const void* DefaultValue, UObject* Parent, int32 PortFlags, UObject* ExportRootScope ) const
+{
+	const FText& TextValue = GetPropertyValue(PropertyValue);
+
+	if (PortFlags & PPF_ExportCpp)
+	{
+		ValueStr += GenerateCppCodeForTextValue(TextValue, FString());
+	}
+	else if (PortFlags & PPF_PropertyWindow)
+	{
+		if (PortFlags & PPF_Delimited)
+		{
+			ValueStr += TEXT("\"");
+			ValueStr += TextValue.ToString();
+			ValueStr += TEXT("\"");
+		}
+		else
+		{
+			ValueStr += TextValue.ToString();
+		}
+	}
+	else
+	{
+		FTextStringHelper::WriteToString(ValueStr, TextValue, !!(PortFlags & PPF_Delimited));
+	}
+}
+
+const TCHAR* UTextProperty::ImportText_Internal( const TCHAR* Buffer, void* Data, int32 PortFlags, UObject* Parent, FOutputDevice* ErrorText ) const
+{
+	FText* TextPtr = GetPropertyValuePtr(Data);
+
+	FString PackageNamespace;
+#if USE_STABLE_LOCALIZATION_KEYS
+	if (GIsEditor && !(PortFlags & PPF_DuplicateForPIE))
+	{
+		PackageNamespace = TextNamespaceUtil::EnsurePackageNamespace(Parent);
+	}
+#endif // USE_STABLE_LOCALIZATION_KEYS
+
+	int32 NumCharsRead = 0;
+	if (FTextStringHelper::ReadFromString(Buffer, *TextPtr, nullptr, *PackageNamespace, &NumCharsRead, !!(PortFlags & PPF_Delimited)))
+	{
+		Buffer += NumCharsRead;
+		return Buffer;
+	}
+	
+	return nullptr;
+}
+
+FString UTextProperty::GenerateCppCodeForTextValue(const FText& InValue, const FString& Indent)
+{
+	FString CppCode;
+
+	const FString& StringValue = FTextInspector::GetDisplayString(InValue);
+
+	if (InValue.IsEmpty())
+	{
+		CppCode += TEXT("FText::GetEmpty()");
+	}
+	else if (InValue.IsCultureInvariant())
+	{
+		// Produces FText::AsCultureInvariant(TEXT("..."))
+		CppCode += TEXT("FText::AsCultureInvariant(");
+		CppCode += UStrProperty::ExportCppHardcodedText(StringValue, Indent + TEXT("\t"));
+		CppCode += TEXT(")");
+	}
+	else
+	{
+		bool bIsLocalized = false;
+		FString Namespace;
+		FString Key;
+		const FString* SourceString = FTextInspector::GetSourceString(InValue);
+
+		if (SourceString && InValue.ShouldGatherForLocalization())
+		{
+			bIsLocalized = FTextLocalizationManager::Get().FindNamespaceAndKeyFromDisplayString(FTextInspector::GetSharedDisplayString(InValue), Namespace, Key);
+		}
+
+		if (bIsLocalized)
+		{
+			// Produces FInternationalization::ForUseOnlyByLocMacroAndGraphNodeTextLiterals_CreateText(TEXT("..."), TEXT("..."), TEXT("..."))
+			CppCode += TEXT("FInternationalization::ForUseOnlyByLocMacroAndGraphNodeTextLiterals_CreateText(\n");
+
+			CppCode += Indent;
+			CppCode += TEXT("\t");
+			CppCode += UStrProperty::ExportCppHardcodedText(*SourceString, Indent + TEXT("\t\t"));
+			CppCode += TEXT(", /* Literal Text */\n");
+
+			CppCode += Indent;
+			CppCode += TEXT("\t");
+			CppCode += UStrProperty::ExportCppHardcodedText(Namespace, Indent + TEXT("\t\t"));
+			CppCode += TEXT(", /* Namespace */\n");
+
+			CppCode += Indent;
+			CppCode += TEXT("\t");
+			CppCode += UStrProperty::ExportCppHardcodedText(Key, Indent + TEXT("\t\t"));
+			CppCode += TEXT(" /* Key */\n");
+
+			CppCode += Indent;
+			CppCode += TEXT("\t)");
+		}
+		else
+		{
+			// Produces FText::FromString(TEXT("..."))
+			CppCode += TEXT("FText::FromString(");
+			CppCode += UStrProperty::ExportCppHardcodedText(StringValue, Indent + TEXT("\t"));
+			CppCode += TEXT(")");
+		}
+	}
+
+	return CppCode;
+}
+
+IMPLEMENT_CORE_INTRINSIC_CLASS(UTextProperty, UProperty,
+{
+}
+);
diff --git a/Engine/Source/Runtime/CoreUObject/Public/CoreUObject.h b/Engine/Source/Runtime/CoreUObject/Public/CoreUObject.h
--- a/Engine/Source/Runtime/CoreUObject/Public/CoreUObject.h
+++ b/Engine/Source/Runtime/CoreUObject/Public/CoreUObject.h
@@ -3,11 +3,12 @@
 #pragma once
 
 #include "Core.h"
+#include "ErrorException.h"
 #include "ModuleManager.h"
 #include "ScriptInterface.h"			// Script interface definitions.
 #include "Script.h"						// Script constants and forward declarations.
 #include "ScriptMacros.h"				// Script macro definitions
-#include "ObjectBase.h"					// Object base class.
+#include "ObjectMacros.h"				// Object macros.
 #include "UObjectAllocator.h"
 #include "UObjectGlobals.h"
 #include "UObjectMarks.h"
@@ -19,7 +20,7 @@
 #include "Object.h"
 #include "UObjectIterator.h"
 #include "CoreNet.h"					// Core networking.
-#include "ArchiveUObjectBase.h"			// UObject-related Archive classes.
+#include "ArchiveUObject.h"				// UObject-related Archive classes.
 #include "GarbageCollection.h"			// Realtime garbage collection helpers
 #include "TextBuffer.h"					// UObjects for text buffers
 #include "Class.h"						// Class definition.
@@ -29,21 +30,47 @@
 #include "LazyObjectPtr.h"				// Object pointer types
 #include "AssetPtr.h"					// Object pointer types
 #include "Interface.h"					// Interface base classes.
-#include "CoreObject.h"					// Core object class definitions.
+#include "LevelGuids.h"					// Core object class definitions.
 #include "Package.h"					// Package class definition
 #include "MetaData.h"					// Metadata for packages
 #include "UnrealType.h"					// Base property type.
-#include "UTextProperty.h"				// FText property type
+#include "TextProperty.h"				// FText property type
 #include "Stack.h"						// Script stack frame definition.
 #include "ObjectRedirector.h"			// Cross-package object redirector
 #include "UObjectAnnotation.h"
 #include "ObjectMemoryAnalyzer.h"		// Object memory usage analyzer
 #include "ReferenceChainSearch.h"		// Search for referencers of an objects
 #include "ArchiveUObject.h"				// UObject-related Archive classes.
+#include "AsyncFileHandle.h"
+#include "TextPackageNamespaceUtil.h"
+#include "ArchiveCountMem.h"
+#include "ObjectAndNameAsStringProxyArchive.h"
+#include "ObjectWriter.h"
+#include "ObjectReader.h"
+#include "ReloadObjectArc.h"
+#include "ArchiveReplaceArchetype.h"
+#include "ArchiveShowReferences.h"
+#include "FindReferencersArchive.h"
+#include "FindObjectReferencers.h"
+#include "ArchiveFindCulprit.h"
+#include "ArchiveObjectGraph.h"
+#include "TraceReferences.h"
+#include "ArchiveTraceRoute.h"
+#include "DuplicatedObject.h"
+#include "DuplicatedDataReader.h"
+#include "DuplicatedDataWriter.h"
+#include "ArchiveReplaceObjectRef.h"
+#include "ArchiveReplaceOrClearExternalReferences.h"
+#include "ArchiveObjectPropertyMapper.h"
+#include "ArchiveReferenceMarker.h"
+#include "ArchiveAsync.h"
+#include "ArchiveObjectCrc32.h"
 #include "PackageName.h"
 #include "ConstructorHelpers.h"			// Helper templates for object construction.
 #include "BulkData.h"					// Bulk data classes
 #include "Linker.h"						// Linker.
+#include "LinkerLoad.h"
+#include "LinkerSave.h"
 #include "GCObject.h"			        // non-UObject object referencer
 #include "AsyncLoading.h"				// FAsyncPackage definition
 #include "StartupPackages.h"
diff --git a/Engine/Source/Runtime/CoreUObject/Public/Misc/TextBuffer.h b/Engine/Source/Runtime/CoreUObject/Public/Misc/TextBuffer.h
--- a/Engine/Source/Runtime/CoreUObject/Public/Misc/TextBuffer.h
+++ b/Engine/Source/Runtime/CoreUObject/Public/Misc/TextBuffer.h
@@ -2,7 +2,7 @@
 
 #pragma once
 
-#include "ObjectBase.h"
+#include "ObjectMacros.h"
 
 
 /**
