When loading DLLs on Windows, search in the DLL directory stack as well as the added search directories. The game binaries directory is only pushed onto the directory stack (via PushDllDirectory()), not added added through AddDllDirectory().

diff --git a/Engine/Source/Runtime/Core/Private/Windows/WindowsPlatformProcess.cpp b/Engine/Source/Runtime/Core/Private/Windows/WindowsPlatformProcess.cpp
--- a/Engine/Source/Runtime/Core/Private/Windows/WindowsPlatformProcess.cpp
+++ b/Engine/Source/Runtime/Core/Private/Windows/WindowsPlatformProcess.cpp
@@ -42,8 +42,19 @@ void* FWindowsPlatformProcess::GetDllHandle( const TCHAR* FileName )
 {
 	check(FileName);
 
+	// Combine the explicit DLL search directories with the contents of the directory stack 
+	TArray<FString> SearchPaths;
+	for(int32 Idx = DllDirectoryStack.Num() - 1; Idx >= 0; Idx--)
+	{
+		SearchPaths.AddUnique(DllDirectoryStack[Idx]);
+	}
+	for(int32 Idx = 0; Idx < DllDirectories.Num(); Idx++)
+	{
+		SearchPaths.AddUnique(DllDirectories[Idx]);
+	}
+
 	::SetErrorMode(SEM_NOOPENFILEERRORBOX);
-	return LoadLibraryWithSearchPaths(FileName, DllDirectories);
+	return LoadLibraryWithSearchPaths(FileName, SearchPaths);
 }
 
 void FWindowsPlatformProcess::FreeDllHandle( void* DllHandle )
