Incorrect allocator was being used at runtime.  Reverted back to using binned.

diff --git a/Engine/Source/Runtime/Core/Private/Windows/WindowsPlatformMemory.cpp b/Engine/Source/Runtime/Core/Private/Windows/WindowsPlatformMemory.cpp
--- a/Engine/Source/Runtime/Core/Private/Windows/WindowsPlatformMemory.cpp
+++ b/Engine/Source/Runtime/Core/Private/Windows/WindowsPlatformMemory.cpp
@@ -76,20 +76,19 @@ FMalloc* FWindowsPlatformMemory::BaseAllocator()
 #elif USE_MALLOC_STOMP
 	// Allocator that helps track down memory stomps.
 	return new FMallocStomp();
-#else
-
-#if (WITH_EDITORONLY_DATA || IS_PROGRAM) && TBB_ALLOCATOR_ALLOWED
+#elif (WITH_EDITORONLY_DATA || IS_PROGRAM) && TBB_ALLOCATOR_ALLOWED
 	// -run is used for commandlets.  Cooker and other commandlets run faster and with less memory using malloc binned.
-	if ( (FCString::Stristr(::GetCommandLine(), TEXT("-mallocbinned")) || FCString::Stristr(::GetCommandLine(), TEXT("-run="))) &&
-		(!FCString::Stristr(::GetCommandLine(), TEXT("-malloctbb")) ) )
+	if ((FCString::Stristr(::GetCommandLine(), TEXT("-mallocbinned")) || FCString::Stristr(::GetCommandLine(), TEXT("-run="))) &&
+		(!FCString::Stristr(::GetCommandLine(), TEXT("-malloctbb"))))
 	{
-		return new FMallocBinnedRedirect((uint32)(GetConstants().PageSize&MAX_uint32), (uint64)MAX_uint32 + 1); 
+		return new FMallocBinnedRedirect((uint32)(GetConstants().PageSize&MAX_uint32), (uint64)MAX_uint32 + 1);
 	}
 	else
-#endif
 	{
 		return new FMallocTBB();
-	}
+}
+#else
+	return new FMallocBinnedRedirect((uint32)(GetConstants().PageSize&MAX_uint32), (uint64)MAX_uint32 + 1);
 #endif
 }
 
