Fix McpProfileSys DLL load error when launching the UT editor. The UnrealTournament DLL is circularly dependent on the BlueprintContext DLL, but DLL pre-loading previously allowed visiting each module once. The Windows loader behaves slightly differently, and forces the load of the UnrealTournament DLL as soon as the circular reference is encountered. Emulate this behavior by not adding the UnrealTournament DLL to the initial list of visited modules.

diff --git a/Engine/Source/Runtime/Core/Private/Windows/WindowsPlatformProcess.cpp b/Engine/Source/Runtime/Core/Private/Windows/WindowsPlatformProcess.cpp
--- a/Engine/Source/Runtime/Core/Private/Windows/WindowsPlatformProcess.cpp
+++ b/Engine/Source/Runtime/Core/Private/Windows/WindowsPlatformProcess.cpp
@@ -1321,9 +1321,9 @@ FProcHandle FWindowsPlatformProcess::OpenProcess(uint32 ProcessID)
 
 void *FWindowsPlatformProcess::LoadLibraryWithSearchPaths(const FString& FileName, const TArray<FString>& SearchPaths)
 {
-	// Create a list of files which we've already checked for imports
+	// Create a list of files which we've already checked for imports. Don't add the initial file to this list to improve the resolution of dependencies for direct circular dependencies of this
+	// module; by allowing the module to be visited twice, any mutually depended on DLLs will be visited first.
 	TArray<FString> VisitedImportNames;
-	VisitedImportNames.Add(FPaths::GetCleanFilename(FileName));
 
 	// Find a list of all the DLLs that need to be loaded
 	TArray<FString> ImportFileNames;
