Changed FString(const NSString* In) constructor to use CFStringGetBytes instead of [NSString cStringUsingEncoding]. It's faster and solves the problem with previous implementation's incorrect assumpion that UTF32 string returned from cStringUsingEncoding would be zero-terminated.

diff --git a/Engine/Source/Runtime/Core/Public/Containers/UnrealString.h b/Engine/Source/Runtime/Core/Public/Containers/UnrealString.h
--- a/Engine/Source/Runtime/Core/Public/Containers/UnrealString.h
+++ b/Engine/Source/Runtime/Core/Public/Containers/UnrealString.h
@@ -160,18 +160,20 @@ public:
 			// Convert the NSString data into the native TCHAR format for UE4
 			// This returns a buffer of bytes, but they can be safely cast to a buffer of TCHARs
 #if PLATFORM_TCHAR_IS_4_BYTES
-			const char* ConvertedStrBuffer = [In cStringUsingEncoding:NSUTF32StringEncoding];
+			const CFStringEncoding Encoding = kCFStringEncodingUTF32LE;
 #else
-			const char* ConvertedStrBuffer = [In cStringUsingEncoding:NSUTF16StringEncoding];
+			const CFStringEncoding Encoding = kCFStringEncodingUTF16LE;
 #endif
 
-			// Copy the converted string into the FString and ensure it's terminated
-			// We need to do this using the length of the converted buffer rather than the source NSString, as they may not match after the conversion
-			// (particularly when dealing with UTF-32 encodings of characters outside the BMP)
-			const int32 ConvertedStrLen = FCString::Strlen((const TCHAR*)ConvertedStrBuffer);
-			Data.AddUninitialized(ConvertedStrLen + 1);
-			FMemory::Memcpy(Data.GetData(), ConvertedStrBuffer, ConvertedStrLen * sizeof(TCHAR));
-			Data[ConvertedStrLen] = '\0';
+			CFRange Range = CFRangeMake(0, CFStringGetLength((__bridge CFStringRef)In));
+			CFIndex BytesNeeded;
+			if (CFStringGetBytes((__bridge CFStringRef)In, Range, Encoding, '?', false, NULL, 0, &BytesNeeded) > 0)
+			{
+				const size_t Length = BytesNeeded / sizeof(TCHAR);
+				Data.AddUninitialized(Length + 1);
+				CFStringGetBytes((__bridge CFStringRef)In, Range, Encoding, '?', false, (uint8*)Data.GetData(), Length * sizeof(TCHAR) + 1, NULL);
+				Data[Length] = 0;
+			}
 		}
 	}
 #endif
