fix for tvOS build warning (re-made from Release-4.13)

diff --git a/Engine/Source/Runtime/Core/Private/IOS/IOSPlatformFramePacer.cpp b/Engine/Source/Runtime/Core/Private/IOS/IOSPlatformFramePacer.cpp
--- a/Engine/Source/Runtime/Core/Private/IOS/IOSPlatformFramePacer.cpp
+++ b/Engine/Source/Runtime/Core/Private/IOS/IOSPlatformFramePacer.cpp
@@ -29,8 +29,12 @@ static TArray<FEvent*> ListeningEvents;
 {
 	NSRunLoop *runloop = [NSRunLoop currentRunLoop];
 	CADisplayLink *displayLink = [CADisplayLink displayLinkWithTarget:self selector:@selector(signal:)];
+#if defined(__IPHONE_10_0) || defined(__TVOS_10_0)
+	displayLink.preferredFramesPerSecond = 60 / FIOSPlatformRHIFramePacer::FrameInterval;
+#else
 	displayLink.frameInterval = FIOSPlatformRHIFramePacer::FrameInterval;
-    
+#endif
+
 	[displayLink addToRunLoop:runloop forMode:NSDefaultRunLoopMode];
 	[runloop run];
 }
