fix for crash when utilizing a build using tvOS SDK 10 on a tvOS 9 device

diff --git a/Engine/Source/Runtime/Core/Private/IOS/IOSPlatformFramePacer.cpp b/Engine/Source/Runtime/Core/Private/IOS/IOSPlatformFramePacer.cpp
--- a/Engine/Source/Runtime/Core/Private/IOS/IOSPlatformFramePacer.cpp
+++ b/Engine/Source/Runtime/Core/Private/IOS/IOSPlatformFramePacer.cpp
@@ -30,10 +30,17 @@ static TArray<FEvent*> ListeningEvents;
 	NSRunLoop *runloop = [NSRunLoop currentRunLoop];
 	CADisplayLink *displayLink = [CADisplayLink displayLinkWithTarget:self selector:@selector(signal:)];
 #if defined(__IPHONE_10_0) || defined(__TVOS_10_0)
-	displayLink.preferredFramesPerSecond = 60 / FIOSPlatformRHIFramePacer::FrameInterval;
-#else
-	displayLink.frameInterval = FIOSPlatformRHIFramePacer::FrameInterval;
+	if ([displayLink respondsToSelector : @selector(preferredFramesPerSecond)] == YES)
+	{
+		displayLink.preferredFramesPerSecond = 60 / FIOSPlatformRHIFramePacer::FrameInterval;
+	}
+	else
 #endif
+	{
+#if (__IPHONE_OS_VERSION_MIN_REQUIRED < __IPHONE_10_0) || (__TV_OS_VERSION_MIN_REQUIRED < __TVOS_10_0)
+		displayLink.frameInterval = FIOSPlatformRHIFramePacer::FrameInterval;
+#endif
+	}
 
 	[displayLink addToRunLoop:runloop forMode:NSDefaultRunLoopMode];
 	[runloop run];
