fix for tvOS deprecationg warnings

diff --git a/Engine/Source/Runtime/Core/Private/IOS/IOSPlatformFramePacer.cpp b/Engine/Source/Runtime/Core/Private/IOS/IOSPlatformFramePacer.cpp
--- a/Engine/Source/Runtime/Core/Private/IOS/IOSPlatformFramePacer.cpp
+++ b/Engine/Source/Runtime/Core/Private/IOS/IOSPlatformFramePacer.cpp
@@ -42,6 +42,10 @@ static TArray<FEvent*> ListeningEvents;
 
 @implementation FIOSFramePacer
 
+#ifdef __GNUC__
+#pragma GCC diagnostic push
+#pragma GCC diagnostic ignored "-Wdeprecated-declarations"
+#endif
 -(void)run:(id)param
 {
 	NSRunLoop *runloop = [NSRunLoop currentRunLoop];
@@ -62,6 +66,9 @@ static TArray<FEvent*> ListeningEvents;
 	[displayLink addToRunLoop:runloop forMode:NSDefaultRunLoopMode];
 	[runloop run];
 }
+#ifdef __GNUC__
+#pragma GCC diagnostic pop
+#endif
 
 
 -(void)signal:(id)param
