Fix for UE-36087. Make sure the mouse cursor on Windows is locked in editor mode when needed, even if it's invisible (PIE, SIE)

diff --git a/Engine/Source/Runtime/Core/Private/Windows/WindowsCursor.cpp b/Engine/Source/Runtime/Core/Private/Windows/WindowsCursor.cpp
--- a/Engine/Source/Runtime/Core/Private/Windows/WindowsCursor.cpp
+++ b/Engine/Source/Runtime/Core/Private/Windows/WindowsCursor.cpp
@@ -209,23 +209,20 @@ void FWindowsCursor::UpdateClipping()
 		// Unlock the cursor
 		::ClipCursor(nullptr);
 	}
-	else if (bIsVisible)
+	else if (bIsVisible || !IsRunningGame()) // We never want to teleport the cursor back to the center unless we're in a game, otherwise it can interfere with how Simulate mode works
 	{
 		// Lock the cursor
 		::ClipCursor(&ClipRect);
 	}
-	else
+	else if (!bIsVisible)
 	{
-		// If the cursor is not visible, assume we're in game where mouse is controlling the camera and lock it to the center of the widget.
+		// If the cursor is not visible and we're running game, assume we're in a mode where the mouse is controlling the camera and lock it to the center of the widget.
 		// This is a workaround for a problem in borderless window mode with round window corners where the cursor in the corner is in fact outside of the window area.
-		if( IsRunningGame() )	// We never want to teleport the cursor back to the center unless we're in a game, otherwise it can interfere with how Simulate mode works
-		{
-			RECT Center;
-			Center.left = ClipRect.left + (ClipRect.right - ClipRect.left) / 2;
-			Center.right = Center.left + 1;
-			Center.top = ClipRect.top + (ClipRect.bottom - ClipRect.top) / 2;
-			Center.bottom = Center.top + 1;
-			::ClipCursor(&Center);
-		}
+		RECT Center;
+		Center.left = ClipRect.left + (ClipRect.right - ClipRect.left) / 2;
+		Center.right = Center.left + 1;
+		Center.top = ClipRect.top + (ClipRect.bottom - ClipRect.top) / 2;
+		Center.bottom = Center.top + 1;
+		::ClipCursor(&Center);
 	}
 }
