Copying //UE4/Dev-Editor to //UE4/Dev-Main (Source: //UE4/Dev-Editor @ 3152045)

diff --git a/Engine/Source/Runtime/AutomationWorker/Private/AutomationWorkerModule.cpp b/Engine/Source/Runtime/AutomationWorker/Private/AutomationWorkerModule.cpp
--- a/Engine/Source/Runtime/AutomationWorker/Private/AutomationWorkerModule.cpp
+++ b/Engine/Source/Runtime/AutomationWorker/Private/AutomationWorkerModule.cpp
@@ -246,7 +246,7 @@ void FAutomationWorkerModule::SendTests( const FMessageAddress& ControllerAddres
 {
 	for( int32 TestIndex = 0; TestIndex < TestInfo.Num(); TestIndex++ )
 	{
-		MessageEndpoint->Send(new FAutomationWorkerRequestTestsReply(TestInfo[TestIndex].GetTestAsString(), TestInfo.Num()), ControllerAddress);
+		MessageEndpoint->Send(new FAutomationWorkerRequestTestsReply(TestInfo[TestIndex], TestInfo.Num()), ControllerAddress);
 	}
 	MessageEndpoint->Send(new FAutomationWorkerRequestTestsReplyComplete(), ControllerAddress);
 }
diff --git a/Engine/Source/Runtime/Core/Private/Internationalization/Internationalization.cpp b/Engine/Source/Runtime/Core/Private/Internationalization/Internationalization.cpp
--- a/Engine/Source/Runtime/Core/Private/Internationalization/Internationalization.cpp
+++ b/Engine/Source/Runtime/Core/Private/Internationalization/Internationalization.cpp
@@ -86,142 +86,107 @@ void FInternationalization::Terminate()
 }
 
 #if ENABLE_LOC_TESTING
-namespace
+FString& FInternationalization::Leetify(FString& SourceString)
 {
-	bool LeetifyInRange(FString& String, const int32 Begin, const int32 End)
+	static const TCHAR LeetifyTextStartMarker = TEXT('\x2021');
+	static const TCHAR LeetifyTextEndMarker = TEXT('\x2021');
+	static const TCHAR LeetifyArgumentStartMarker = TEXT('\x00AB');
+	static const TCHAR LeetifyArgumentEndMarker = TEXT('\x00BB');
+	static const TCHAR SourceArgumentStartMarker = TEXT('{');
+	static const TCHAR SourceArgumentEndMarker = TEXT('}');
+	static const TCHAR SourceEscapeMarker = TEXT('`');
+
+	auto LeetifyCharacter = [](const TCHAR InChar) -> TCHAR
 	{
-		bool Succeeded = false;
-		for(int32 Index = Begin; Index < End; ++Index)
+		switch (InChar)
 		{
-			switch(String[Index])
-			{
-			case TEXT('A'): { String[Index] = TEXT('4'); Succeeded = true; } break;
-			case TEXT('a'): { String[Index] = TEXT('@'); Succeeded = true; } break;
-			case TEXT('B'): { String[Index] = TEXT('8'); Succeeded = true; } break;
-			case TEXT('b'): { String[Index] = TEXT('8'); Succeeded = true; } break;
-			case TEXT('E'): { String[Index] = TEXT('3'); Succeeded = true; } break;
-			case TEXT('e'): { String[Index] = TEXT('3'); Succeeded = true; } break;
-			case TEXT('G'): { String[Index] = TEXT('9'); Succeeded = true; } break;
-			case TEXT('g'): { String[Index] = TEXT('9'); Succeeded = true; } break;
-			case TEXT('I'): { String[Index] = TEXT('1'); Succeeded = true; } break;
-			case TEXT('i'): { String[Index] = TEXT('!'); Succeeded = true; } break;
-			case TEXT('O'): { String[Index] = TEXT('0'); Succeeded = true; } break;
-			case TEXT('o'): { String[Index] = TEXT('0'); Succeeded = true; } break;
-			case TEXT('S'): { String[Index] = TEXT('5'); Succeeded = true; } break;
-			case TEXT('s'): { String[Index] = TEXT('$'); Succeeded = true; } break;
-			case TEXT('T'): { String[Index] = TEXT('7'); Succeeded = true; } break;
-			case TEXT('t'): { String[Index] = TEXT('7'); Succeeded = true; } break;
-			case TEXT('Z'): { String[Index] = TEXT('2'); Succeeded = true; } break;
-			case TEXT('z'): { String[Index] = TEXT('2'); Succeeded = true; } break;
-			}
+		case TEXT('A'): return TEXT('4');
+		case TEXT('a'): return TEXT('@');
+		case TEXT('B'): return TEXT('8');
+		case TEXT('b'): return TEXT('8');
+		case TEXT('E'): return TEXT('3');
+		case TEXT('e'): return TEXT('3');
+		case TEXT('G'): return TEXT('9');
+		case TEXT('g'): return TEXT('9');
+		case TEXT('I'): return TEXT('1');
+		case TEXT('i'): return TEXT('!');
+		case TEXT('O'): return TEXT('0');
+		case TEXT('o'): return TEXT('0');
+		case TEXT('S'): return TEXT('5');
+		case TEXT('s'): return TEXT('$');
+		case TEXT('T'): return TEXT('7');
+		case TEXT('t'): return TEXT('7');
+		case TEXT('Z'): return TEXT('2');
+		case TEXT('z'): return TEXT('2');
+		default:		return InChar;
 		}
+	};
 
-		return Succeeded;
+	if (SourceString.IsEmpty() || (SourceString.Len() >= 2 && SourceString[0] == LeetifyTextStartMarker && SourceString[SourceString.Len() - 1] == LeetifyTextEndMarker))
+	{
+		// Already leetified
+		return SourceString;
 	}
-}
 
-FString& FInternationalization::Leetify(FString& SourceString)
-{
-	// Check that the string hasn't already been Leetified
-	if( SourceString.IsEmpty() ||
-		(SourceString[ 0 ] != 0x2021 && SourceString.Find( TEXT("\x00AB"), ESearchCase::CaseSensitive, ESearchDir::FromStart, 0 ) == -1) )
-	{
-		bool Succeeded = false;
+	// We insert a start and end marker (+2), and format strings typically have <= 8 argument blocks which we'll wrap with a start and end marker (+16), so +18 should be a reasonable slack
+	FString LeetifiedString;
+	LeetifiedString.Reserve(SourceString.Len() + 18);
 
-		FString OpenBlock = TEXT("{");
-		FString CloseBlock = TEXT("}");
-		uint32 SanityLoopCheck=0xFFFF;
+	// Inject the start marker
+	LeetifiedString.AppendChar(LeetifyTextStartMarker);
 
-		int32 CurrentBlockBeginPos=-1;
-		int32 CurrentBlockEndPos=0;
-		int32 PrevBlockBeginPos=0;
-		int32 PrevBlockEndPos=-1;
-		int32 CurrArgBlock=0;
+	// Iterate and leetify each character in the source string, but don't change argument names as that will break formatting
+	{
+		bool bEscapeNextChar = false;
 
-		struct FBlockRange
+		const int32 SourceStringLen = SourceString.Len();
+		for (int32 SourceCharIndex = 0; SourceCharIndex < SourceStringLen; ++SourceCharIndex)
 		{
-			int32 BeginPos;
-			int32 EndPos;
-		};
+			const TCHAR SourceChar = SourceString[SourceCharIndex];
 
-		TArray<FBlockRange> BlockRanges;
-		while(--SanityLoopCheck > 0)
-		{
-			//Find the start of the next block. Delimited with an open brace '{'
-			++CurrentBlockBeginPos;
-			while(true)
+			if (!bEscapeNextChar && SourceChar == SourceArgumentStartMarker)
 			{
-				CurrentBlockBeginPos = SourceString.Find(OpenBlock,ESearchCase::CaseSensitive, ESearchDir::FromStart, CurrentBlockBeginPos);
-				if( CurrentBlockBeginPos == -1 )
-				{
-					break;//No block open started so we've reached the end of the format string.
-				}
+				const TCHAR* RawSourceStringPtr = SourceString.GetCharArray().GetData();
 
-				if( CurrentBlockBeginPos >= 0 && SourceString[CurrentBlockBeginPos+1] == OpenBlock[0] )
+				// Walk forward to find the end of this argument block to make sure we have a pair of tokens
+				const TCHAR* ArgumentEndPtr = FCString::Strchr(RawSourceStringPtr + SourceCharIndex + 1, SourceArgumentEndMarker);
+				if (ArgumentEndPtr)
 				{
-					//Skip past {{ literals in the format
-					CurrentBlockBeginPos += 2;
+					const int32 ArgumentEndIndex = ArgumentEndPtr - RawSourceStringPtr;
+
+					// Inject a marker before the argument block
+					LeetifiedString.AppendChar(LeetifyArgumentStartMarker);
+
+					// Copy the body of the argument, including the opening and closing tags
+					check(ArgumentEndIndex >= SourceCharIndex);
+					LeetifiedString.AppendChars(RawSourceStringPtr + SourceCharIndex, (ArgumentEndIndex - SourceCharIndex) + 1);
+
+					// Inject a marker after the end of the argument block
+					LeetifiedString.AppendChar(LeetifyArgumentEndMarker);
+
+					// Move to the end of the argument we just parsed
+					SourceCharIndex = ArgumentEndIndex;
 					continue;
 				}
-
-				break;
 			}
 
-			//No more block opening braces found so we're done.
-			if( CurrentBlockBeginPos == -1 )
+			if (SourceChar == SourceEscapeMarker)
 			{
-				break;
+				bEscapeNextChar = !bEscapeNextChar;
 			}
-
-			//Find the end of the block. Delimited with a close brace '}'
-			CurrentBlockEndPos=SourceString.Find(CloseBlock,ESearchCase::CaseSensitive, ESearchDir::FromStart,CurrentBlockBeginPos);
-
-			FBlockRange NewRange = { CurrentBlockBeginPos, CurrentBlockEndPos };
-			BlockRanges.Add(NewRange);
-
-
-			// Insertion of double arrows causes block start and end to be move later in the string, adjust for that.
-			++CurrentBlockBeginPos;
-			++CurrentBlockEndPos;
-
-			Succeeded = LeetifyInRange(SourceString, PrevBlockEndPos + 1, CurrentBlockBeginPos) || Succeeded;
-
-			PrevBlockBeginPos=CurrentBlockBeginPos;
-			PrevBlockEndPos=CurrentBlockEndPos;
-
-			++CurrArgBlock;
-		}
-
-		Succeeded = LeetifyInRange(SourceString, CurrentBlockEndPos + 1, SourceString.Len()) || Succeeded;
-
-		// Insert double arrows around parameter blocks.
-		if(BlockRanges.Num() > 0)
-		{
-			FString ResultString;
-			int32 i;
-			for(i = 0; i < BlockRanges.Num(); ++i)
+			else
 			{
-				// Append intermediate part of string.
-				int32 EndOfLastPart = (i - 1 >= 0) ? (BlockRanges[i - 1].EndPos + 1) : 0;
-				ResultString += SourceString.Mid(EndOfLastPart, BlockRanges[i].BeginPos - EndOfLastPart);
-				// Wrap block.
-				ResultString += TEXT("\x00AB");
-				// Append block.
-				ResultString += SourceString.Mid(BlockRanges[i].BeginPos, BlockRanges[i].EndPos - BlockRanges[i].BeginPos + 1);
-				// Wrap block.
-				ResultString += TEXT("\x00BB");
+				bEscapeNextChar = false;
 			}
-			ResultString += SourceString.Mid(BlockRanges[i - 1].EndPos + 1, SourceString.Len() - BlockRanges[i - 1].EndPos + 1);
-			SourceString = ResultString;
-		}
 
-		if( !Succeeded )
-		{
-			// Failed to LEETify, add something to beginning and end just to help mark the string.
-			SourceString = FString(TEXT("\x2021")) + SourceString + FString(TEXT("\x2021"));
+			LeetifiedString.AppendChar(LeetifyCharacter(SourceChar));
 		}
 	}
 
+	// Inject the end marker
+	LeetifiedString.AppendChar(LeetifyTextEndMarker);
+
+	SourceString = LeetifiedString;
 	return SourceString;
 }
 #endif
diff --git a/Engine/Source/Runtime/Core/Private/Internationalization/TextLocalizationManager.cpp b/Engine/Source/Runtime/Core/Private/Internationalization/TextLocalizationManager.cpp
--- a/Engine/Source/Runtime/Core/Private/Internationalization/TextLocalizationManager.cpp
+++ b/Engine/Source/Runtime/Core/Private/Internationalization/TextLocalizationManager.cpp
@@ -46,26 +46,26 @@ void EndInitTextLocalization()
 		// Use culture override specified on commandline.
 		else if (FParse::Value(FCommandLine::Get(), TEXT("CULTURE="), RequestedCultureName))
 		{
-			UE_LOG(LogInit, Log, TEXT("Overriding culture %s w/ command-line option."), *RequestedCultureName);
+			UE_LOG(LogInit, Log, TEXT("Overriding culture %s with command-line option."), *RequestedCultureName);
 		}
 		else
 #if WITH_EDITOR
 			// See if we've been provided a culture override in the editor
 			if(GIsEditor && GConfig->GetString( TEXT("Internationalization"), TEXT("Culture"), RequestedCultureName, GEditorSettingsIni ))
 			{
-				UE_LOG(LogInit, Log, TEXT("Overriding culture %s w/ editor configuration."), *RequestedCultureName);
+				UE_LOG(LogInit, Log, TEXT("Overriding culture %s with editor configuration."), *RequestedCultureName);
 			}
 			else
 #endif // WITH_EDITOR
 				// Use culture specified in game configuration.
 				if(GConfig->GetString( TEXT("Internationalization"), TEXT("Culture"), RequestedCultureName, GGameUserSettingsIni ))
 				{
-					UE_LOG(LogInit, Log, TEXT("Overriding culture %s w/ game configuration."), *RequestedCultureName);
+					UE_LOG(LogInit, Log, TEXT("Overriding culture %s with game configuration."), *RequestedCultureName);
 				}
 				// Use culture specified in engine configuration.
 				else if(GConfig->GetString( TEXT("Internationalization"), TEXT("Culture"), RequestedCultureName, GEngineIni ))
 				{
-					UE_LOG(LogInit, Log, TEXT("Overriding culture %s w/ engine configuration."), *RequestedCultureName);
+					UE_LOG(LogInit, Log, TEXT("Overriding culture %s with engine configuration."), *RequestedCultureName);
 				}
 				else
 				{
diff --git a/Engine/Source/Runtime/Core/Private/Misc/AutomationTest.cpp b/Engine/Source/Runtime/Core/Private/Misc/AutomationTest.cpp
--- a/Engine/Source/Runtime/Core/Private/Misc/AutomationTest.cpp
+++ b/Engine/Source/Runtime/Core/Private/Misc/AutomationTest.cpp
@@ -6,54 +6,6 @@
 
 DEFINE_LOG_CATEGORY_STATIC(LogAutomationTest, Warning, All);
 
-FString FAutomationTestInfo::GetTestAsString() const
-{
-	return FString::Printf(TEXT("%s,%s,%i,%d,%s,%s,%i,%s,%s"), *DisplayName, *TestName, TestFlags, NumParticipantsRequired, *TestParameter, *SourceFile, SourceFileLine, *AssetPath, *OpenCommand);
-}
-
-void FAutomationTestInfo::ParseStringInfo(const FString& InTestInfo)
-{
-	//Split New Test Name into string array
-	TArray<FString> Pieces;
-	InTestInfo.ParseIntoArray(Pieces, TEXT(","), false);
-
-	// We should always have at least 3 parameters
-	check(Pieces.Num() >= 3);
-
-	DisplayName = Pieces[0];
-	TestName = Pieces[1];
-	TestFlags = uint8(FCString::Atoi(*Pieces[2]));
-
-	NumParticipantsRequired = FCString::Atoi(*Pieces[3]);
-
-	// Optional Parameters
-	if ( Pieces.Num() >= 5 )
-	{
-		TestParameter = Pieces[4];
-	}
-
-	if ( Pieces.Num() >= 6 )
-	{
-		SourceFile = Pieces[5];
-	}
-
-	if ( Pieces.Num() >= 7 )
-	{
-		SourceFileLine = FCString::Atoi(*Pieces[6]);
-	}
-
-	if ( Pieces.Num() >= 8 )
-	{
-		AssetPath = Pieces[7];
-	}
-
-	if ( Pieces.Num() >= 9 )
-	{
-		OpenCommand = Pieces[8];
-	}
-}
-
-
 void FAutomationTestFramework::FAutomationTestFeedbackContext::Serialize( const TCHAR* V, ELogVerbosity::Type Verbosity, const class FName& Category )
 {
 	if (FAutomationTestFramework::Get().CachedContext)
@@ -243,7 +195,6 @@ void FAutomationTestFramework::ResetTests()
 	IFileManager::Get().DeleteDirectory(*FPaths::AutomationTransientDir(), bEnsureExists, bDeleteEntireTree);
 }
 
-
 void FAutomationTestFramework::StartTestByName( const FString& InTestToRun, const int32 InRoleIndex )
 {
 	if (GIsAutomationTesting)
@@ -687,6 +638,16 @@ void FAutomationTestFramework::AddAnalyticsItemToCurrentTest( const FString& Ana
 	}
 }
 
+bool FAutomationTestFramework::GetTreatWarningsAsErrors() const
+{
+	return AutomationTestFeedbackContext.TreatWarningsAsErrors;
+}
+
+void FAutomationTestFramework::SetTreatWarningsAsErrors(TOptional<bool> bTreatWarningsAsErrors)
+{
+	AutomationTestFeedbackContext.TreatWarningsAsErrors = bTreatWarningsAsErrors.IsSet() ? bTreatWarningsAsErrors.GetValue() : GWarn->TreatWarningsAsErrors;
+}
+
 FAutomationTestFramework::FAutomationTestFramework()
 :	CachedContext( NULL )
 ,	RequestedTestFilter(EAutomationTestFlags::SmokeFilter)
@@ -800,6 +761,7 @@ void FAutomationTestBase::GenerateTestNames(TArray<FAutomationTestInfo>& TestInf
 
 		// Add the test info to our collection
 		FAutomationTestInfo NewTestInfo(
+			CompleteBeautifiedNames,
 			CompleteBeautifiedNames,
 			CompleteTestName,
 			GetTestFlags(),
diff --git a/Engine/Source/Runtime/Core/Private/ProfilingDebugging/MallocProfiler.cpp b/Engine/Source/Runtime/Core/Private/ProfilingDebugging/MallocProfiler.cpp
--- a/Engine/Source/Runtime/Core/Private/ProfilingDebugging/MallocProfiler.cpp
+++ b/Engine/Source/Runtime/Core/Private/ProfilingDebugging/MallocProfiler.cpp
@@ -25,6 +25,8 @@ CORE_API FMallocProfiler* GMallocProfiler;
 #define	MEMORY_PROFILER_MAX_BACKTRACE_DEPTH			75
 /** Number of backtrace entries to skip											*/
 #define MEMORY_PROFILER_SKIP_NUM_BACKTRACE_ENTRIES	1
+/** Whether to track allocation tags? */
+#define MEMORY_PROFILER_INCLUDE_ALLOC_TAGS			1
 
 /** Magic value, determining that file is a memory profiler file.				*/
 #define MEMORY_PROFILER_MAGIC						0xDA15F7D8
@@ -35,8 +37,10 @@ CORE_API FMallocProfiler* GMallocProfiler;
  *		Removed NumDataFiles, no longer used.
  *
  *	6 - Added meta-data table.
+ *
+  *	7 - Added allocation tags to alloc/realloc.
  */
-#define MEMORY_PROFILER_VERSION						6
+#define MEMORY_PROFILER_VERSION						7
 
 /*=============================================================================
 	Profiler header.
@@ -75,6 +79,11 @@ struct FProfilerHeader
 	/** Number of callstack entries.					*/
 	uint64	CallStackTableEntries;
 
+	/** Offset in file for tags table.					*/
+	uint64	TagsTableOffset;
+	/** Number of tags entries.							*/
+	uint64	TagsTableEntries;
+
 	/** The file offset for module information.			*/
 	uint64	ModulesOffset;
 	/** The number of module entries.					*/
@@ -103,6 +112,8 @@ struct FProfilerHeader
 			<< Header.CallStackAddressTableEntries
 			<< Header.CallStackTableOffset
 			<< Header.CallStackTableEntries
+			<< Header.TagsTableOffset
+			<< Header.TagsTableEntries
 			<< Header.ModulesOffset
 			<< Header.ModuleEntries;
 
@@ -162,7 +173,7 @@ struct FCallStackInfo
 /**
  * Relevant information for a single malloc operation.
  *
- * 16 bytes
+ * 20 bytes
  */
 struct FProfilerAllocInfo
 {
@@ -170,6 +181,8 @@ struct FProfilerAllocInfo
 	uint64 Pointer;
 	/** Index of callstack.													*/
 	int32 CallStackIndex;
+	/** Index of tags.														*/
+	int32 TagsIndex;
 	/** Size of allocation.													*/
 	uint32 Size;
 
@@ -184,6 +197,7 @@ struct FProfilerAllocInfo
 	{
 		Ar	<< AllocInfo.Pointer
 			<< AllocInfo.CallStackIndex
+			<< AllocInfo.TagsIndex
 			<< AllocInfo.Size;
 		return Ar;
 	}
@@ -216,7 +230,7 @@ struct FProfilerFreeInfo
 /**
  * Relevant information for a single realloc operation.
  *
- * 24 bytes
+ * 28 bytes
  */
 struct FProfilerReallocInfo
 {
@@ -226,6 +240,8 @@ struct FProfilerReallocInfo
 	uint64 NewPointer;
 	/** Index of callstack.													*/
 	int32 CallStackIndex;
+	/** Index of tags.														*/
+	int32 TagsIndex;
 	/** Size of allocation.													*/
 	uint32 Size;
 
@@ -241,6 +257,7 @@ struct FProfilerReallocInfo
 		Ar	<< ReallocInfo.OldPointer
 			<< ReallocInfo.NewPointer
 			<< ReallocInfo.CallStackIndex
+			<< ReallocInfo.TagsIndex
 			<< ReallocInfo.Size;
 		return Ar;
 	}
@@ -276,6 +293,131 @@ struct FProfilerOtherInfo
 	}
 };
 
+
+#if MEMORY_PROFILER_INCLUDE_ALLOC_TAGS
+
+/*=============================================================================
+	TLS storage for allocation tags
+=============================================================================*/
+
+/** Container for all of the active tags for a single thread. This will only be accessed by the thread that owns it. */
+class FMallocProfilerTags : public FUseSystemMallocForNew, public FTlsAutoCleanup, private FNoncopyable
+{
+public:
+	void AddTag(const FName InTag)
+	{
+		const int32 ExistingTagAndCountIndex = Tags.IndexOfByPredicate([&](const FTagNameAndCount& InTagAndCount) { return InTagAndCount.TagName == InTag; });
+
+		if (ExistingTagAndCountIndex != INDEX_NONE)
+		{
+			Tags[ExistingTagAndCountIndex].RefCount += 1;
+		}
+		else
+		{
+			Tags.Add(FTagNameAndCount(InTag));
+		}
+	}
+
+	void RemoveTag(const FName InTag)
+	{
+		const int32 ExistingTagAndCountIndex = Tags.IndexOfByPredicate([&](const FTagNameAndCount& InTagAndCount) { return InTagAndCount.TagName == InTag; });
+
+		if (ExistingTagAndCountIndex != INDEX_NONE)
+		{
+			Tags[ExistingTagAndCountIndex].RefCount -= 1;
+			if (Tags[ExistingTagAndCountIndex].RefCount == 0)
+			{
+				Tags.RemoveAt(ExistingTagAndCountIndex);
+			}
+		}
+	}
+
+	FString AsString() const
+	{
+		FString FlatTags;
+		for (const FTagNameAndCount& TagAndCount : Tags)
+		{
+			FlatTags += TagAndCount.TagName.ToString();
+			FlatTags.AppendChar(TEXT(';'));
+		}
+		return FlatTags;
+	}
+
+	uint32 GetHash() const
+	{
+		uint32 Hash = 0;
+		for (const FTagNameAndCount& TagAndCount : Tags)
+		{
+			Hash = HashCombine(Hash, GetTypeHash(TagAndCount.TagName));
+		}
+		return Hash;
+	}
+
+private:
+	struct FTagNameAndCount
+	{
+		FTagNameAndCount(FName InTagName)
+			: TagName(InTagName)
+			, RefCount(1)
+		{
+		}
+
+		FName TagName;
+		int32 RefCount;
+	};
+
+	static const int32 MaxNumTags = 32;
+
+	TArray<FTagNameAndCount, TFixedAllocator<MaxNumTags>> Tags;
+};
+
+/** TLS management. Handles the lifespans of the per-thread tag data. */
+class FMallocProfilerTagsTls
+{
+public:
+	static void Initialize()
+	{
+		TlsSlot = FPlatformTLS::AllocTlsSlot();
+	}
+
+	static void Shutdown()
+	{
+		FPlatformTLS::FreeTlsSlot(TlsSlot);
+		TlsSlot = 0;
+	}
+
+	static const FMallocProfilerTags* GetTagsForCurrentThread()
+	{
+		check(FPlatformTLS::IsValidTlsSlot(TlsSlot));
+
+		return (FMallocProfilerTags*)FPlatformTLS::GetTlsValue(TlsSlot);
+	}
+
+	static FMallocProfilerTags& GetMutableTagsForCurrentThread()
+	{
+		check(FPlatformTLS::IsValidTlsSlot(TlsSlot));
+
+		FMallocProfilerTags* TlsValue = (FMallocProfilerTags*)FPlatformTLS::GetTlsValue(TlsSlot);
+
+		if (!TlsValue)
+		{
+			TlsValue = new FMallocProfilerTags();
+			TlsValue->Register();
+			FPlatformTLS::SetTlsValue(TlsSlot, TlsValue);
+		}
+
+		return *TlsValue;
+	}
+
+private:
+	static uint32 TlsSlot;
+};
+
+uint32 FMallocProfilerTagsTls::TlsSlot = 0;
+
+#endif // MEMORY_PROFILER_INCLUDE_ALLOC_TAGS
+
+
 /*=============================================================================
 	FMallocProfiler implementation.
 =============================================================================*/
@@ -301,6 +443,10 @@ FMallocProfiler::FMallocProfiler(FMalloc* InMalloc)
 	{
 		PanicDump(TYPE_Malloc, nullptr, nullptr);
 	});
+
+#if MEMORY_PROFILER_INCLUDE_ALLOC_TAGS
+	FMallocProfilerTagsTls::Initialize();
+#endif // MEMORY_PROFILER_INCLUDE_ALLOC_TAGS
 }
 
 /**
@@ -320,6 +466,7 @@ void FMallocProfiler::TrackMalloc( void* Ptr, uint32 Size )
 		FProfilerAllocInfo AllocInfo;
 		AllocInfo.Pointer			= (uint64)(UPTRINT) Ptr | TYPE_Malloc;
 		AllocInfo.CallStackIndex	= GetCallStackIndex();
+		AllocInfo.TagsIndex			= GetTagsIndex();
 		AllocInfo.Size				= Size;
 
 		// Serialize to HDD.
@@ -367,6 +514,7 @@ void FMallocProfiler::TrackRealloc( void* OldPtr, void* NewPtr, uint32 NewSize )
 		ReallocInfo.OldPointer		= (uint64)(UPTRINT) OldPtr | TYPE_Realloc;
 		ReallocInfo.NewPointer		= (uint64)(UPTRINT) NewPtr;
 		ReallocInfo.CallStackIndex	= GetCallStackIndex();
+		ReallocInfo.TagsIndex		= GetTagsIndex();
 		ReallocInfo.Size			= NewSize;
 
 		// Serialize to HDD.
@@ -411,6 +559,28 @@ void FMallocProfiler::BeginProfiling()
 	BufferedFileWriter << DummyHeader;
 }
 
+/** 
+ * Adds a tag to the list of tags associated with the calling thread.
+ */
+void FMallocProfiler::AddTag(const FName InTag)
+{
+#if MEMORY_PROFILER_INCLUDE_ALLOC_TAGS
+	FMallocProfilerTags& ActiveTags = FMallocProfilerTagsTls::GetMutableTagsForCurrentThread();
+	ActiveTags.AddTag(InTag);
+#endif // MEMORY_PROFILER_INCLUDE_ALLOC_TAGS
+}
+
+/** 
+ * Removes a tag from the list of tags associated with the calling thread.
+ */
+void FMallocProfiler::RemoveTag(const FName InTag)
+{
+#if MEMORY_PROFILER_INCLUDE_ALLOC_TAGS
+	FMallocProfilerTags& ActiveTags = FMallocProfilerTagsTls::GetMutableTagsForCurrentThread();
+	ActiveTags.RemoveTag(InTag);
+#endif // MEMORY_PROFILER_INCLUDE_ALLOC_TAGS
+}
+
 /** 
 	Added for untracked memory calculation
 	Note that this includes all the memory used by dependent malloc profilers, such as FMallocGcmProfiler,
@@ -422,6 +592,8 @@ int32 FMallocProfiler::CalculateMemoryProfilingOverhead()
 			+ CallStackAddressInfoArray.GetAllocatedSize()
 			+ CRCToCallStackIndexMap.GetAllocatedSize()
 			+ CallStackInfoBuffer.GetAllocatedSize()
+			+ HashToTagTableIndexMap.GetAllocatedSize()
+			+ TagsArray.GetAllocatedSize()
 			+ NameToNameTableIndexMap.GetAllocatedSize()
 			+ NameArray.GetAllocatedSize()
 			+ BufferedFileWriter.GetAllocatedSize());
@@ -566,6 +738,14 @@ void FMallocProfiler::EndProfiling()
 		}
 		CallStackInfoBuffer.Unlock();
 
+		// Write out tags and update the header with offset and count
+		Header.TagsTableOffset = SymbolFileWriter->Tell();
+		Header.TagsTableEntries = TagsArray.Num();
+		for (FString& Tags : TagsArray)
+		{
+			(*SymbolFileWriter) << Tags;
+		}
+
 		Header.ModulesOffset				= SymbolFileWriter->Tell();
 		Header.ModuleEntries				= FPlatformStackWalk::GetProcessModuleCount();
 
@@ -698,6 +878,44 @@ int32 FMallocProfiler::GetCallStackIndex()
 	return Index;
 }	
 
+/** 
+ * Returns index of the tags active on the thread making the allocation.
+ *
+ * @return index into tags array
+ */
+int32 FMallocProfiler::GetTagsIndex()
+{
+	int32 Index = INDEX_NONE;
+
+#if MEMORY_PROFILER_INCLUDE_ALLOC_TAGS
+	const FMallocProfilerTags* ActiveTags = FMallocProfilerTagsTls::GetTagsForCurrentThread();
+	if (ActiveTags)
+	{
+		const uint32 TagsHash = ActiveTags->GetHash();
+		if (TagsHash != 0)
+		{
+			// Use index if found
+			int32* IndexPtr = HashToTagTableIndexMap.Find(TagsHash);
+			if (IndexPtr)
+			{
+				Index = *IndexPtr;
+			}
+			else
+			{
+				// Set mapping for future use.
+				Index = TagsArray.Num();
+				HashToTagTableIndexMap.Add(TagsHash, Index);
+
+				// Add the new tags
+				TagsArray.Add(ActiveTags->AsString());
+			}
+		}
+	}
+#endif // MEMORY_PROFILER_INCLUDE_ALLOC_TAGS
+
+	return Index;
+}
+
 /**
  * Returns index of passed in name into name array. If not found, adds it.
  *
diff --git a/Engine/Source/Runtime/Core/Private/Tests/Internationalization/TextTest.cpp b/Engine/Source/Runtime/Core/Private/Tests/Internationalization/TextTest.cpp
--- a/Engine/Source/Runtime/Core/Private/Tests/Internationalization/TextTest.cpp
+++ b/Engine/Source/Runtime/Core/Private/Tests/Internationalization/TextTest.cpp
@@ -460,7 +460,7 @@ bool FTextTest::RunTest (const FString& Parameters)
 				// When changes are made to FormattedTestLayer2, please pull out the newly translated LEET string and update the below if-statement to keep the test passing!
 				FString LEETTranslatedString = FormattedTestLayer2.ToString();
 
-				FString DesiredOutput = FString(TEXT("\x2021") TEXT("\xAB") TEXT("\xAB") TEXT("L0r3m") TEXT("\xBB") TEXT(" \"L0r3m 1p$um\" ") TEXT("\xAB") TEXT("Ip$um") TEXT("\xBB") TEXT("\xBB") TEXT(" | ") TEXT("\xAB") TEXT("\xAB") TEXT("L0r3m") TEXT("\xBB") TEXT(" \"L0r3m 1p$um\" ") TEXT("\xAB") TEXT("Ip$um") TEXT("\xBB") TEXT("\xBB") TEXT(" | ") TEXT("\xAB") TEXT("5.5421") TEXT("\xBB") TEXT(" | ") TEXT("\xAB") TEXT("5010.89221") TEXT("\xBB") TEXT(" | ") TEXT("\xAB") TEXT("Aug 20, 2080, 9:33:22 AM") TEXT("\xBB") TEXT(" | ") TEXT("\xAB") TEXT("92%") TEXT("\xBB") TEXT(" | ") TEXT("\xAB") TEXT("$") TEXT("\xA0") TEXT("100.25") TEXT("\xBB") TEXT("\x2021"));
+				FString DesiredOutput = FString(TEXT("\x2021") TEXT("\xAB") TEXT("\x2021") TEXT("\xAB") TEXT("\x2021") TEXT("L0r3m") TEXT("\x2021") TEXT("\xBB") TEXT(" \"L0r3m 1p$um\" ") TEXT("\xAB") TEXT("\x2021") TEXT("1p$um") TEXT("\x2021") TEXT("\xBB") TEXT("\x2021") TEXT("\xBB") TEXT(" | ") TEXT("\xAB") TEXT("\x2021") TEXT("\xAB") TEXT("\x2021") TEXT("L0r3m") TEXT("\x2021") TEXT("\xBB") TEXT(" \"L0r3m 1p$um\" ") TEXT("\xAB") TEXT("\x2021") TEXT("1p$um") TEXT("\x2021") TEXT("\xBB") TEXT("\x2021") TEXT("\xBB") TEXT(" | ") TEXT("\xAB") TEXT("5.5421") TEXT("\xBB") TEXT(" | ") TEXT("\xAB") TEXT("5010.89221") TEXT("\xBB") TEXT(" | ") TEXT("\xAB") TEXT("Aug 20, 2080, 9:33:22 AM") TEXT("\xBB") TEXT(" | ") TEXT("\xAB") TEXT("92%") TEXT("\xBB") TEXT(" | ") TEXT("\xAB") TEXT("$") TEXT("\xA0") TEXT("100.25") TEXT("\xBB") TEXT("\x2021"));
 				// Convert the baked string into an FText, which will be leetified, then compare it to the rebuilt FText
 				if(LEETTranslatedString != DesiredOutput)
 				{
diff --git a/Engine/Source/Runtime/Core/Private/Windows/WindowsApplication.cpp b/Engine/Source/Runtime/Core/Private/Windows/WindowsApplication.cpp
--- a/Engine/Source/Runtime/Core/Private/Windows/WindowsApplication.cpp
+++ b/Engine/Source/Runtime/Core/Private/Windows/WindowsApplication.cpp
@@ -807,6 +807,9 @@ int32 FWindowsApplication::ProcessMessage( HWND hwnd, uint32 msg, WPARAM wParam,
 		case WM_NCMOUSEMOVE:
 		case WM_MOUSEMOVE:
 		case WM_MOUSEWHEEL:
+#if WINVER >= 0x0601
+		case WM_TOUCH:
+#endif
 			{
 				DeferMessage( CurrentNativeEventWindowPtr, hwnd, msg, wParam, lParam );
 				// Handled
@@ -1714,6 +1717,66 @@ int32 FWindowsApplication::ProcessDeferredMessage( const FDeferredWindowsMessage
 			}
 			break;
 
+#if WINVER >= 0x0601
+		case WM_TOUCH:
+			{
+				UINT InputCount = LOWORD( wParam );
+				if ( InputCount > 0 )
+				{
+					TUniquePtr<TOUCHINPUT> Inputs( new TOUCHINPUT[InputCount] );
+					if ( GetTouchInputInfo( (HTOUCHINPUT)lParam, InputCount, Inputs.Get(), sizeof(TOUCHINPUT) ) )
+					{
+						for ( uint32 i = 0; i < InputCount; i++ )
+						{
+							TOUCHINPUT Input = Inputs.Get()[i];
+							FVector2D Location( Input.x / 100.0f, Input.y / 100.0f );
+							if ( Input.dwFlags & TOUCHEVENTF_DOWN )
+							{
+								int32 TouchIndex = GetTouchIndexForID( Input.dwID );
+								if (TouchIndex < 0)
+								{
+									TouchIndex = GetFirstFreeTouchIndex();
+									if (TouchIndex >= 0)
+									{
+										TouchIDs[TouchIndex] = TOptional<int32>( Input.dwID );
+										MessageHandler->OnTouchStarted( CurrentNativeEventWindowPtr, Location, TouchIndex + 1, 0 );
+									}
+									else
+									{
+										// TODO: Error handling for more than 10 touches?
+									}
+								}
+							}
+							else if ( Input.dwFlags & TOUCHEVENTF_MOVE )
+							{
+								int32 TouchIndex = GetTouchIndexForID( Input.dwID );
+								if ( TouchIndex >= 0 )
+								{
+									MessageHandler->OnTouchMoved( Location, TouchIndex + 1, 0 );
+								}
+							}
+							else if ( Input.dwFlags & TOUCHEVENTF_UP )
+							{
+								int32 TouchIndex = GetTouchIndexForID( Input.dwID );
+								if ( TouchIndex >= 0 )
+								{
+									TouchIDs[TouchIndex] = TOptional<int32>();
+									MessageHandler->OnTouchEnded( Location, TouchIndex + 1, 0 );
+								}
+								else
+								{
+									// TODO: Error handling.
+								}
+							}
+						}
+						CloseTouchInputHandle( (HTOUCHINPUT)lParam );
+						return 0;
+					}
+				}
+				break;
+			}
+#endif
+
 			// Window focus and activation
 		case WM_MOUSEACTIVATE:
 			{
@@ -2302,6 +2365,32 @@ void FWindowsApplication::QueryConnectedMice()
 	bIsMouseAttached = MouseCount > 0;
 }
 
+#if WINVER >= 0x0601
+uint32 FWindowsApplication::GetTouchIndexForID( int32 TouchID )
+{
+	for ( int i = 0; i < MaxTouches; i++ )
+	{
+		if ( TouchIDs[i].IsSet() && TouchIDs[i].GetValue() == TouchID )
+		{
+			return i;
+		}
+	}
+	return -1;
+}
+
+uint32 FWindowsApplication::GetFirstFreeTouchIndex()
+{
+	for ( int i = 0; i < MaxTouches; i++ )
+	{
+		if ( TouchIDs[i].IsSet() == false )
+		{
+			return i;
+		}
+	}
+	return -1;
+}
+#endif
+
 void FTaskbarList::Initialize()
 {
 	if (FWindowsPlatformMisc::CoInitialize())
diff --git a/Engine/Source/Runtime/Core/Private/Windows/WindowsWindow.cpp b/Engine/Source/Runtime/Core/Private/Windows/WindowsWindow.cpp
--- a/Engine/Source/Runtime/Core/Private/Windows/WindowsWindow.cpp
+++ b/Engine/Source/Runtime/Core/Private/Windows/WindowsWindow.cpp
@@ -168,6 +168,14 @@ void FWindowsWindow::Initialize( FWindowsApplication* const Application, const T
 		( InParent.IsValid() ) ? static_cast<HWND>( InParent->HWnd ) : NULL,
 		NULL, InHInstance, NULL);
 
+#if WINVER >= 0x0601
+	if ( RegisterTouchWindow( HWnd, 0 ) == false )
+	{
+		uint32 Error = GetLastError();
+		UE_LOG(LogWindows, Warning, TEXT("Register touch input failed!"));
+	}
+#endif
+
 	VirtualWidth = ClientWidth;
 	VirtualHeight = ClientHeight;
 
diff --git a/Engine/Source/Runtime/Core/Public/Containers/Algo/Accumulate.h b/Engine/Source/Runtime/Core/Public/Containers/Algo/Accumulate.h
--- a/Engine/Source/Runtime/Core/Public/Containers/Algo/Accumulate.h
+++ b/Engine/Source/Runtime/Core/Public/Containers/Algo/Accumulate.h
@@ -32,7 +32,7 @@ namespace Algo
 	FORCEINLINE T Accumulate(const A& Input, T Init, OpT Op)
 	{
 		T Result = MoveTemp(Init);
-		for (auto&& InputElem : Input)
+		for (const auto& InputElem : Input)
 		{
 			Result = Op(MoveTemp(Result), InputElem);
 		}
@@ -67,7 +67,7 @@ namespace Algo
 	FORCEINLINE T TransformAccumulate(const A& Input, MapT MapOp, T Init, OpT Op)
 	{
 		T Result = MoveTemp(Init);
-		for (auto&& InputElem : Input)
+		for (const auto& InputElem : Input)
 		{
 			Result = Op(MoveTemp(Result), MapOp(InputElem));
 		}
diff --git a/Engine/Source/Runtime/Core/Public/Containers/Algo/Copy.h b/Engine/Source/Runtime/Core/Public/Containers/Algo/Copy.h
--- a/Engine/Source/Runtime/Core/Public/Containers/Algo/Copy.h
+++ b/Engine/Source/Runtime/Core/Public/Containers/Algo/Copy.h
@@ -2,6 +2,8 @@
 
 #pragma once
 
+#include "Templates/Invoke.h"
+
 namespace Algo
 {
 	/**
@@ -14,9 +16,9 @@ namespace Algo
 	template <typename InT, typename OutT, typename PredicateT>
 	FORCEINLINE void CopyIf(const InT& Input, OutT& Output, PredicateT Predicate)
 	{
-		for (auto&& Value : Input)
+		for (const auto& Value : Input)
 		{
-			if (Predicate(Value))
+			if (Invoke(Predicate, Value))
 			{
 				Output.Add(Value);
 			}
@@ -32,7 +34,7 @@ namespace Algo
 	template <typename InT, typename OutT>
 	FORCEINLINE void Copy(const InT& Input, OutT& Output)
 	{
-		for (auto&& Value : Input)
+		for (const auto& Value : Input)
 		{
 			Output.Add(Value);
 		}
diff --git a/Engine/Source/Runtime/Core/Public/Containers/Algo/Transform.h b/Engine/Source/Runtime/Core/Public/Containers/Algo/Transform.h
--- a/Engine/Source/Runtime/Core/Public/Containers/Algo/Transform.h
+++ b/Engine/Source/Runtime/Core/Public/Containers/Algo/Transform.h
@@ -17,7 +17,7 @@ namespace Algo
 	template <typename InT, typename OutT, typename PredicateT, typename TransformT>
 	FORCEINLINE void TransformIf(const InT& Input, OutT& Output, PredicateT Predicate, TransformT Trans)
 	{
-		for (auto&& Value : Input)
+		for (const auto& Value : Input)
 		{
 			if (Invoke(Predicate, Value))
 			{
@@ -36,7 +36,7 @@ namespace Algo
 	template <typename InT, typename OutT, typename TransformT>
 	FORCEINLINE void Transform(const InT& Input, OutT& Output, TransformT Trans)
 	{
-		for (auto&& Value : Input)
+		for (const auto& Value : Input)
 		{
 			Output.Add(Invoke(Trans, Value));
 		}
diff --git a/Engine/Source/Runtime/Core/Public/Misc/AutomationTest.h b/Engine/Source/Runtime/Core/Public/Misc/AutomationTest.h
--- a/Engine/Source/Runtime/Core/Public/Misc/AutomationTest.h
+++ b/Engine/Source/Runtime/Core/Public/Misc/AutomationTest.h
@@ -131,19 +131,6 @@ public:
 		, NumDevicesCurrentlyRunningTest( 0 )
 	{}
 
-
-	/**
-	 * Constructor
-	 *
-	 * @param	InTestInfo - comma separated value string containing the test info
-	 */
-	FAutomationTestInfo( const FString& InTestInfo )
-		: NumParticipantsRequired( 0 )
-		, NumDevicesCurrentlyRunningTest( 0 )
-	{
-		ParseStringInfo( InTestInfo );
-	}
-
 	/**
 	 * Constructor
 	 *
@@ -152,8 +139,9 @@ public:
 	 * @param	InTestFlag - Test flags
 	 * @param	InParameterName - optional parameter. e.g. asset name
 	 */
-	FAutomationTestInfo(const FString& InDisplayName, const FString& InTestName, const uint32 InTestFlags, const int32 InNumParticipantsRequired, const FString& InParameterName = FString(), const FString& InSourceFile = FString(), int32 InSourceFileLine = 0, const FString& InAssetPath = FString(), const FString& InOpenCommand = FString())
+	FAutomationTestInfo(const FString& InDisplayName, const FString& InFullTestPath, const FString& InTestName, const uint32 InTestFlags, const int32 InNumParticipantsRequired, const FString& InParameterName = FString(), const FString& InSourceFile = FString(), int32 InSourceFileLine = 0, const FString& InAssetPath = FString(), const FString& InOpenCommand = FString())
 		: DisplayName( InDisplayName )
+		, FullTestPath( InFullTestPath )
 		, TestName( InTestName )
 		, TestParameter( InParameterName )
 		, SourceFile( InSourceFile )
@@ -177,7 +165,6 @@ public:
 		TestFlags |= InTestFlags;
 	}
 
-
 	/**
 	 * Get the display name of this test.
 	 *
@@ -188,14 +175,15 @@ public:
 		return DisplayName;
 	}
 
-
 	/**
-	 * Get the test as a string so it can be sent over the network.
+	 * Gets the full path for this test if you wanted to run it.
 	 *
-	 * @return The test as a string.
+	 * @return the display name.
 	 */
-	FString GetTestAsString() const;
-
+	const FString& GetFullTestPath() const
+	{
+		return FullTestPath;
+	}
 
 	/**
 	 * Get the test name of this test.
@@ -207,7 +195,6 @@ public:
 		return TestName;
 	}
 
-
 	/**
 	 * Get the type of parameter. This will be the asset name for linked assets.
 	 *
@@ -325,20 +312,12 @@ public:
 		NumParticipantsRequired = NumRequired;
 	}
 
-
-private:
-
-	/**
-	 * Set the test info from a string.
-	 *
-	 * @Param InTestInfo - the test information as a string.
-	 */
-	void ParseStringInfo(const FString& InTestInfo);
-
 private:
 	/** Display name used in the UI */
 	FString DisplayName; 
 
+	FString FullTestPath;
+
 	/** Test name used to run the test */
 	FString TestName;
 
@@ -680,6 +659,9 @@ public:
 		return CurrentTest;
 	}
 
+	bool GetTreatWarningsAsErrors() const;
+	void SetTreatWarningsAsErrors(TOptional<bool> bTreatWarningsAsErrors);
+
 private:
 
 	/** Special feedback context used exclusively while automation testing */
diff --git a/Engine/Source/Runtime/Core/Public/Misc/Build.h b/Engine/Source/Runtime/Core/Public/Misc/Build.h
--- a/Engine/Source/Runtime/Core/Public/Misc/Build.h
+++ b/Engine/Source/Runtime/Core/Public/Misc/Build.h
@@ -186,21 +186,21 @@
 #if UE_BUILD_DEBUG
 	#define DO_GUARD_SLOW									1
 	#define DO_CHECK										1
-	#define STATS											(!UE_BUILD_MINIMAL || !WITH_EDITORONLY_DATA || USE_STATS_WITHOUT_ENGINE)
+	#define STATS											(!UE_BUILD_MINIMAL || !WITH_EDITORONLY_DATA || USE_STATS_WITHOUT_ENGINE || USE_MALLOC_PROFILER)
 	#define ALLOW_DEBUG_FILES								1
 	#define ALLOW_CONSOLE									1
 	#define NO_LOGGING										0
 #elif UE_BUILD_DEVELOPMENT
 	#define DO_GUARD_SLOW									0
 	#define DO_CHECK										1
-	#define STATS											(!UE_BUILD_MINIMAL || !WITH_EDITORONLY_DATA || USE_STATS_WITHOUT_ENGINE)
+	#define STATS											(!UE_BUILD_MINIMAL || !WITH_EDITORONLY_DATA || USE_STATS_WITHOUT_ENGINE || USE_MALLOC_PROFILER)
 	#define ALLOW_DEBUG_FILES								1
 	#define ALLOW_CONSOLE									1
 	#define NO_LOGGING										0
 #elif UE_BUILD_TEST
 	#define DO_GUARD_SLOW									0
 	#define DO_CHECK										USE_CHECKS_IN_SHIPPING
-	#define STATS											0
+	#define STATS											(USE_MALLOC_PROFILER)
 	#define ALLOW_DEBUG_FILES								1
 	#define ALLOW_CONSOLE									1
 	#define NO_LOGGING										!USE_LOGGING_IN_SHIPPING
diff --git a/Engine/Source/Runtime/Core/Public/ProfilingDebugging/MallocProfiler.h b/Engine/Source/Runtime/Core/Public/ProfilingDebugging/MallocProfiler.h
--- a/Engine/Source/Runtime/Core/Public/ProfilingDebugging/MallocProfiler.h
+++ b/Engine/Source/Runtime/Core/Public/ProfilingDebugging/MallocProfiler.h
@@ -219,6 +219,11 @@ protected:
 	/** Buffer of unique call stack infos.															*/
 	FCompressedGrowableBuffer				CallStackInfoBuffer;
 
+	/** Mapping from a hash to an index in the tags array.											*/
+	TMap<uint32, int32>						HashToTagTableIndexMap;
+	/** Array of unique tags.																		*/
+	TArray<FString>							TagsArray;
+
 	/** Mapping from name to index in name array.													*/
 	TMap<FString,int32>						NameToNameTableIndexMap;
 	/** Array of unique names.																		*/
@@ -253,6 +258,13 @@ protected:
 	 */
 	int32 GetCallStackIndex();
 
+	/** 
+	 * Returns index of the tags active on the thread making the allocation.
+	 *
+	 * @return index into tags array
+	 */
+	int32 GetTagsIndex();
+
 	/**
 	 * Returns index of passed in program counter into program counter array. If not found, adds it.
 	 *
@@ -407,6 +419,18 @@ public:
 	 */
 	void BeginProfiling();
 
+	/** 
+	 * Adds a tag to the list of tags associated with the calling thread.
+	 * @note Tags are ref-counted so calls to Add and Remove must be paired correctly.
+	 */
+	void AddTag(const FName InTag);
+
+	/** 
+	 * Removes a tag from the list of tags associated with the calling thread.
+	 * @note Tags are ref-counted so calls to Add and Remove must be paired correctly.
+	 */
+	void RemoveTag(const FName InTag);
+
 	/** 
 	 * Malloc
 	 */
diff --git a/Engine/Source/Runtime/Core/Public/ProfilingDebugging/ResourceSize.h b/Engine/Source/Runtime/Core/Public/ProfilingDebugging/ResourceSize.h
new file mode 100644
--- /dev/null
+++ b/Engine/Source/Runtime/Core/Public/ProfilingDebugging/ResourceSize.h
@@ -0,0 +1,247 @@
+// Copyright 1998-2016 Epic Games, Inc. All Rights Reserved.
+
+#pragma once
+
+/** Indicates which resource size should be returned. */
+namespace EResourceSizeMode
+{
+	enum Type
+	{
+		/** Only exclusive resource size */
+		Exclusive,
+		/** Resource size of the object and all of its references */
+		Inclusive,
+	};
+};
+
+/**
+ * Struct used to count up the amount of memory used by a resource.
+ * This is typically used for assets via UObject::GetResourceSizeEx(...).
+ */
+struct FResourceSizeEx
+{
+public:
+	/**
+	 * Default constructor. 
+	 */
+	explicit FResourceSizeEx()
+		: ResourceSizeMode(EResourceSizeMode::Exclusive)
+		, DedicatedSystemMemoryBytes(0)
+		, SharedSystemMemoryBytes(0)
+		, DedicatedVideoMemoryBytes(0)
+		, SharedVideoMemoryBytes(0)
+		, UnknownMemoryBytes(0)
+	{
+	}
+
+	/**
+	 * Construct using a given mode. 
+	 */
+	explicit FResourceSizeEx(const EResourceSizeMode::Type InResourceSizeMode)
+		: ResourceSizeMode(InResourceSizeMode)
+		, DedicatedSystemMemoryBytes(0)
+		, SharedSystemMemoryBytes(0)
+		, DedicatedVideoMemoryBytes(0)
+		, SharedVideoMemoryBytes(0)
+		, UnknownMemoryBytes(0)
+	{
+	}
+
+	/**
+	 * Construct from known sizes. 
+	 */
+	FResourceSizeEx(const EResourceSizeMode::Type InResourceSizeMode, const SIZE_T InDedicatedSystemMemoryBytes, const SIZE_T InSharedSystemMemoryBytes, const SIZE_T InDedicatedVideoMemoryBytes, const SIZE_T InSharedVideoMemoryBytes)
+		: ResourceSizeMode(InResourceSizeMode)
+		, DedicatedSystemMemoryBytes(InDedicatedSystemMemoryBytes)
+		, SharedSystemMemoryBytes(InSharedSystemMemoryBytes)
+		, DedicatedVideoMemoryBytes(InDedicatedVideoMemoryBytes)
+		, SharedVideoMemoryBytes(InSharedVideoMemoryBytes)
+		, UnknownMemoryBytes(0)
+	{
+	}
+
+	/**
+	 * Construct from legacy unknown size.
+	 * Deliberately explicit to avoid accidental use.
+	 */
+	FResourceSizeEx(const EResourceSizeMode::Type InResourceSizeMode, const SIZE_T InUnknownMemoryBytes)
+		: ResourceSizeMode(InResourceSizeMode)
+		, DedicatedSystemMemoryBytes(0)
+		, SharedSystemMemoryBytes(0)
+		, DedicatedVideoMemoryBytes(0)
+		, SharedVideoMemoryBytes(0)
+		, UnknownMemoryBytes(InUnknownMemoryBytes)
+	{
+	}
+
+	/**
+	 * Get the type of resource size held in this struct.
+	 */
+	EResourceSizeMode::Type GetResourceSizeMode() const
+	{
+		return ResourceSizeMode;
+	}
+
+	/**
+	 * Add the given number of bytes to the dedicated system memory count.
+	 * @see DedicatedSystemMemoryBytes for a description of that memory type.
+	 */
+	FResourceSizeEx& AddDedicatedSystemMemoryBytes(const SIZE_T InMemoryBytes)
+	{
+		DedicatedSystemMemoryBytes += InMemoryBytes;
+		return *this;
+	}
+
+	/**
+	 * Get the number of bytes allocated from dedicated system memory.
+	 * @see DedicatedSystemMemoryBytes for a description of that memory type.
+	 */
+	SIZE_T GetDedicatedSystemMemoryBytes() const
+	{
+		return DedicatedSystemMemoryBytes;
+	}
+
+	/**
+	 * Add the given number of bytes to the shared system memory count.
+	 * @see SharedSystemMemoryBytes for a description of that memory type.
+	 */
+	FResourceSizeEx& AddSharedSystemMemoryBytes(const SIZE_T InMemoryBytes)
+	{
+		SharedSystemMemoryBytes += InMemoryBytes;
+		return *this;
+	}
+
+	/**
+	 * Get the number of bytes allocated from shared system memory.
+	 * @see SharedSystemMemoryBytes for a description of that memory type.
+	 */
+	SIZE_T GetSharedSystemMemoryBytes() const
+	{
+		return SharedSystemMemoryBytes;
+	}
+
+	/**
+	 * Add the given number of bytes to the dedicated video memory count.
+	 * @see DedicatedVideoMemoryBytes for a description of that memory type.
+	 */
+	FResourceSizeEx& AddDedicatedVideoMemoryBytes(const SIZE_T InMemoryBytes)
+	{
+		DedicatedVideoMemoryBytes += InMemoryBytes;
+		return *this;
+	}
+
+	/**
+	 * Get the number of bytes allocated from dedicated video memory.
+	 * @see DedicatedVideoMemoryBytes for a description of that memory type.
+	 */
+	SIZE_T GetDedicatedVideoMemoryBytes() const
+	{
+		return DedicatedVideoMemoryBytes;
+	}
+
+	/**
+	 * Add the given number of bytes to the shared video memory count.
+	 * @see SharedVideoMemoryBytes for a description of that memory type.
+	 */
+	FResourceSizeEx& AddSharedVideoMemoryBytes(const SIZE_T InMemoryBytes)
+	{
+		SharedVideoMemoryBytes += InMemoryBytes;
+		return *this;
+	}
+
+	/**
+	 * Get the number of bytes allocated from shared video memory.
+	 * @see SharedVideoMemoryBytes for a description of that memory type.
+	 */
+	SIZE_T GetSharedVideoMemoryBytes() const
+	{
+		return SharedVideoMemoryBytes;
+	}
+
+	/**
+	 * Add the given number of bytes to the unknown memory count.
+	 * @see UnknownMemoryBytes for a description of that memory type.
+	 */
+	FResourceSizeEx& AddUnknownMemoryBytes(const SIZE_T InMemoryBytes)
+	{
+		UnknownMemoryBytes += InMemoryBytes;
+		return *this;
+	}
+
+	/**
+	 * Get the number of bytes allocated from unknown memory.
+	 * @see UnknownMemoryBytes for a description of that memory type.
+	 */
+	SIZE_T GetUnknownMemoryBytes() const
+	{
+		return UnknownMemoryBytes;
+	}
+
+	/**
+	 * Get the total number of bytes allocated from any memory.
+	 */
+	SIZE_T GetTotalMemoryBytes() const
+	{
+		return DedicatedSystemMemoryBytes + SharedSystemMemoryBytes + DedicatedVideoMemoryBytes + SharedVideoMemoryBytes + UnknownMemoryBytes;
+	}
+
+	/**
+	 * Add another FResourceSizeEx to this one.
+	 */
+	FResourceSizeEx& operator+=(const FResourceSizeEx& InRHS)
+	{
+		ensureAlwaysMsgf(ResourceSizeMode == InRHS.ResourceSizeMode, TEXT("The two resource sizes use different counting modes. The result of adding them together may be incorrect."));
+
+		DedicatedSystemMemoryBytes += InRHS.DedicatedSystemMemoryBytes;
+		SharedSystemMemoryBytes += InRHS.SharedSystemMemoryBytes;
+		DedicatedVideoMemoryBytes += InRHS.DedicatedVideoMemoryBytes;
+		SharedVideoMemoryBytes += InRHS.SharedVideoMemoryBytes;
+		UnknownMemoryBytes += InRHS.UnknownMemoryBytes;
+		return *this;
+	}
+
+	/**
+	 * Add two FResourceSizeEx instances together and return a copy.
+	 */
+	friend FResourceSizeEx operator+(FResourceSizeEx InLHS, const FResourceSizeEx& InRHS)
+	{
+		InLHS += InRHS;
+		return InLHS;
+	}
+
+private:
+	/**
+	 * Type of resource size held in this struct.
+	 */
+	EResourceSizeMode::Type ResourceSizeMode;
+
+	/**
+	 * The number of bytes of memory that this resource is using for CPU resources that have been allocated from dedicated system memory.
+	 * On platforms with unified memory, this typically refers to the things allocated in the preferred memory for CPU use.
+	 */
+	SIZE_T DedicatedSystemMemoryBytes;
+
+	/**
+	 * The number of bytes of memory that this resource is using for CPU resources that have been allocated from non-dedicated system memory (typically zero).
+	 * On platforms with unified memory, this typically refers to the things allocated in the memory outside the preferred memory for CPU use.
+	 */
+	SIZE_T SharedSystemMemoryBytes;
+
+	/**
+	 * The number of bytes of memory that this resource is using for GPU resources that have been allocated from dedicated video memory.
+	 * On platforms with unified memory, this typically refers to the things allocated in the preferred memory for GPU use.
+	 */
+	SIZE_T DedicatedVideoMemoryBytes;
+
+	/**
+	 * The number of bytes of memory that this resource is using for GPU resources that have been allocated from non-dedicated video memory.
+	 * On platforms with unified memory, this typically refers to the things allocated in the memory outside the preferred memory for GPU use.
+	 */
+	SIZE_T SharedVideoMemoryBytes;
+
+	/**
+	 * The number of bytes of memory that this resource is using from an unspecified section of memory.
+	 * This exists so that the legacy GetResourceSize(...) functions can still report back memory usage until they're updated to use FResourceSizeEx, and should not be used in new memory tracking code.
+	 */
+	SIZE_T UnknownMemoryBytes;
+};
diff --git a/Engine/Source/Runtime/Core/Public/Windows/WindowsApplication.h b/Engine/Source/Runtime/Core/Public/Windows/WindowsApplication.h
--- a/Engine/Source/Runtime/Core/Public/Windows/WindowsApplication.h
+++ b/Engine/Source/Runtime/Core/Public/Windows/WindowsApplication.h
@@ -428,6 +428,14 @@ private:
 	/** Queries and caches the number of connected mouse devices. */
 	void QueryConnectedMice();
 
+#if WINVER >= 0x0601
+	/** Gets the touch index for a given windows touch ID. */
+	uint32 GetTouchIndexForID(int32 TouchID);
+
+	/** Searches for a free touch index. */
+	uint32 GetFirstFreeTouchIndex();
+#endif
+
 	/** Helper function to update the cached states of all modifier keys */
 	void UpdateAllModifierKeyStates();
 
@@ -492,6 +500,13 @@ private:
 	STICKYKEYS							StartupStickyKeys;
 	TOGGLEKEYS							StartupToggleKeys;
 	FILTERKEYS							StartupFilterKeys;
+
+#if WINVER >= 0x0601
+	static const int32 MaxTouches = 10;
+
+	/** Maps touch indexes to windows touch IDs. */
+	TOptional<int32> TouchIDs[MaxTouches];
+#endif
 };
 
 
diff --git a/Engine/Source/Runtime/CoreUObject/Private/UObject/Obj.cpp b/Engine/Source/Runtime/CoreUObject/Private/UObject/Obj.cpp
--- a/Engine/Source/Runtime/CoreUObject/Private/UObject/Obj.cpp
+++ b/Engine/Source/Runtime/CoreUObject/Private/UObject/Obj.cpp
@@ -932,7 +932,7 @@ bool UObject::Modify( bool bAlwaysMarkDirty/*=true*/ )
 	{
 		// Do not consider PIE world objects or script packages, as they should never end up in the
 		// transaction buffer and we don't want to mark them dirty here either.
-		if (GetOutermost()->HasAnyPackageFlags(PKG_PlayInEditor | PKG_ContainsScript | PKG_CompiledIn) == false)
+		if (GetOutermost()->HasAnyPackageFlags(PKG_PlayInEditor | PKG_ContainsScript | PKG_CompiledIn) == false || GetClass()->HasAnyClassFlags(CLASS_DefaultConfig | CLASS_Config))
 		{
 			// Attempt to mark the package dirty and save a copy of the object to the transaction
 			// buffer. The save will fail if there isn't a valid transactor, the object isn't
@@ -1367,7 +1367,7 @@ void UObject::FAssetRegistryTag::GetAssetRegistryTagsFromSearchableProperties(co
 void UObject::GetAssetRegistryTags(TArray<FAssetRegistryTag>& OutTags) const
 {
 	// Add ResourceSize if non-zero. GetResourceSize is not const because many override implementations end up calling Serialize on this pointers.
-	SIZE_T ResourceSize = const_cast<UObject*>(this)->GetResourceSize(EResourceSizeMode::Exclusive);
+	const SIZE_T ResourceSize = const_cast<UObject*>(this)->GetResourceSizeBytes(EResourceSizeMode::Exclusive);
 	if ( ResourceSize > 0 || ( !GetClass()->HasAnyClassFlags(CLASS_CompiledFromBlueprint) && HasAnyFlags(RF_ClassDefaultObject) ) )
 	{
 		OutTags.Add( FAssetRegistryTag("ResourceSize", FString::Printf(TEXT("%d"), (ResourceSize + 512) / 1024), FAssetRegistryTag::TT_Numerical) );
@@ -3379,8 +3379,8 @@ bool StaticExec( UWorld* InWorld, const TCHAR* Cmd, FOutputDevice& Ar )
 								FArchiveCountMem Count(Object);
 
 								// Get the 'old-style' resource size and the truer resource size
-								const SIZE_T ResourceSize = It->GetResourceSize(EResourceSizeMode::Inclusive);
-								const SIZE_T TrueResourceSize = It->GetResourceSize(EResourceSizeMode::Exclusive);
+								const SIZE_T ResourceSize = It->GetResourceSizeBytes(EResourceSizeMode::Inclusive);
+								const SIZE_T TrueResourceSize = It->GetResourceSizeBytes(EResourceSizeMode::Exclusive);
 								
 								if (bListClass)
 								{
diff --git a/Engine/Source/Runtime/CoreUObject/Private/UObject/ObjectMemoryAnalyzer.cpp b/Engine/Source/Runtime/CoreUObject/Private/UObject/ObjectMemoryAnalyzer.cpp
--- a/Engine/Source/Runtime/CoreUObject/Private/UObject/ObjectMemoryAnalyzer.cpp
+++ b/Engine/Source/Runtime/CoreUObject/Private/UObject/ObjectMemoryAnalyzer.cpp
@@ -189,8 +189,8 @@ SIZE_T FObjectMemoryAnalyzer::CalculateSizeRecursive(UObject* Object)
 	}
 
 	SIZE_T InclusiveSize = Annotation.ExclusiveMemoryUsage;
-	Annotation.ExclusiveResourceSize = Object->GetResourceSize(EResourceSizeMode::Exclusive);
-	Annotation.InclusiveResourceSize = Object->GetResourceSize(EResourceSizeMode::Inclusive);
+	Annotation.ExclusiveResourceSize = Object->GetResourceSizeBytes(EResourceSizeMode::Exclusive);
+	Annotation.InclusiveResourceSize = Object->GetResourceSizeBytes(EResourceSizeMode::Inclusive);
 
 	TArray<UObject*> ReferencedObjects;
 	GetReferencedObjects(Object, ReferencedObjects);
diff --git a/Engine/Source/Runtime/CoreUObject/Private/UObject/TextProperty.cpp b/Engine/Source/Runtime/CoreUObject/Private/UObject/TextProperty.cpp
--- a/Engine/Source/Runtime/CoreUObject/Private/UObject/TextProperty.cpp
+++ b/Engine/Source/Runtime/CoreUObject/Private/UObject/TextProperty.cpp
@@ -40,29 +40,35 @@ bool UTextProperty::ConvertFromType(const FPropertyTag& Tag, FArchive& Ar, uint8
 
 bool UTextProperty::Identical_Implementation(const FText& ValueA, const FText& ValueB, uint32 PortFlags)
 {
-	if (ValueA.IsCultureInvariant() != ValueB.IsCultureInvariant() || ValueA.IsTransient() != ValueB.IsTransient())
+	// A culture variant text is never equal to a culture invariant text
+	// A transient text is never equal to a non-transient text
+	// An empty text is never equal to a non-empty text
+	if (ValueA.IsCultureInvariant() != ValueB.IsCultureInvariant() || ValueA.IsTransient() != ValueB.IsTransient() || ValueA.IsEmpty() != ValueB.IsEmpty())
 	{
-		//A culture variant text is never equal to a culture invariant text
-		//A transient text is never equal to a non-transient text
 		return false;
 	}
 
-	if (ValueA.IsCultureInvariant() == ValueB.IsCultureInvariant() || ValueA.IsTransient() == ValueB.IsTransient())
+	// If both texts are empty (see the above check), then they must be equal
+	if (ValueA.IsEmpty())
 	{
-		//Culture invariant text don't have a namespace/key so we compare the source string
-		//Transient text don't have a namespace/key or source so we compare the display string
-		return FTextInspector::GetDisplayString(ValueA) == FTextInspector::GetDisplayString(ValueB);
+		return true;
 	}
 
-	if (GIsEditor)
+	// If both texts share the same pointer, then they must be equal
+	if (ValueA.IdenticalTo(ValueB))
 	{
-		return FTextInspector::GetSourceString(ValueA)->Compare(*FTextInspector::GetSourceString(ValueB), ESearchCase::CaseSensitive) == 0;
+		return true;
 	}
-	else
+
+	// We compare the display strings in editor (as we author in the native language)
+	// We compare the display string for culture invariant and transient texts as they don't have an identity
+	if (GIsEditor || ValueA.IsCultureInvariant() || ValueA.IsTransient())
 	{
-		return	FTextInspector::GetNamespace(ValueA) == FTextInspector::GetNamespace(ValueB) &&
-			FTextInspector::GetKey(ValueA) == FTextInspector::GetKey(ValueB);
+		return FTextInspector::GetDisplayString(ValueA).Equals(FTextInspector::GetDisplayString(ValueB), ESearchCase::CaseSensitive);
 	}
+	
+	// If we got this far then the texts don't share the same pointer, which means that they can't share the same identity
+	return false;
 }
 
 bool UTextProperty::Identical( const void* A, const void* B, uint32 PortFlags ) const
diff --git a/Engine/Source/Runtime/CoreUObject/Private/UObject/UObjectBaseUtility.cpp b/Engine/Source/Runtime/CoreUObject/Private/UObject/UObjectBaseUtility.cpp
--- a/Engine/Source/Runtime/CoreUObject/Private/UObject/UObjectBaseUtility.cpp
+++ b/Engine/Source/Runtime/CoreUObject/Private/UObject/UObjectBaseUtility.cpp
@@ -8,6 +8,7 @@
 #include "Interface.h"
 #include "ModuleManager.h"
 #include "FastReferenceCollector.h"
+#include "MallocProfiler.h"
 
 /***********************/
 /******** Names ********/
@@ -417,3 +418,113 @@ bool UObjectBaseUtility::IsDefaultSubobject() const
 	const bool bIsInstanced = GetOuter() && (GetOuter()->HasAnyFlags(RF_ClassDefaultObject) || ((UObject*)this)->GetArchetype() != GetClass()->GetDefaultObject());
 	return bIsInstanced;
 }
+
+#if STATS && USE_MALLOC_PROFILER
+void FScopeCycleCounterUObject::TrackObjectForMallocProfiling(const UObjectBaseUtility *InObject)
+{
+	// Get the package name from the outermost item (if available - can't use GetOutermost here)
+	FName PackageName;
+	if (InObject->GetOuter())
+	{
+		UObjectBaseUtility* Top = InObject->GetOuter();
+		for (;;)
+		{
+			UObjectBaseUtility* CurrentOuter = Top->GetOuter();
+			if (!CurrentOuter)
+			{
+				PackageName = Top->GetFName();
+				break;
+			}
+			Top = CurrentOuter;
+		}
+	}
+
+	// Get the class name (if available)
+	FName ClassName;
+	if (InObject->GetClass())
+	{
+		ClassName = InObject->GetClass()->GetFName();
+	}
+
+	TrackObjectForMallocProfiling(PackageName, ClassName, InObject->GetFName());
+}
+void FScopeCycleCounterUObject::TrackObjectForMallocProfiling(const FName InPackageName, const FName InClassName, const FName InObjectName)
+{
+	static const TCHAR PackageTagCategory[] = TEXT("Package:");
+	static const TCHAR ObjectTagCategory[] = TEXT("Object:");
+	static const TCHAR ClassTagCategory[] = TEXT("Class:");
+
+	// We use an array rather than an FString to try and minimize heap allocations
+	TArray<TCHAR, TInlineAllocator<256>> ScratchSpaceBuffer;
+
+	auto AppendNameToBuffer = [&](const FName InName)
+	{
+		const FNameEntry* NameEntry = InName.GetDisplayNameEntry();
+		if (NameEntry->IsWide())
+		{
+			const WIDECHAR* NameCharPtr = NameEntry->GetWideName();
+			while (*NameCharPtr != 0)
+			{
+				ScratchSpaceBuffer.Add((TCHAR)*NameCharPtr++);
+			}
+		}
+		else
+		{
+			const ANSICHAR* NameCharPtr = NameEntry->GetAnsiName();
+			while (*NameCharPtr != 0)
+			{
+				ScratchSpaceBuffer.Add((TCHAR)*NameCharPtr++);
+			}
+		}
+	};
+
+	if (!InPackageName.IsNone())
+	{
+		// "Package:/Path/To/Package"
+		ScratchSpaceBuffer.Reset();
+		ScratchSpaceBuffer.Append(PackageTagCategory, ARRAY_COUNT(PackageTagCategory) - 1);
+		AppendNameToBuffer(InPackageName);
+		ScratchSpaceBuffer.Add(0);
+		PackageTag = FName(ScratchSpaceBuffer.GetData());
+		GMallocProfiler->AddTag(PackageTag);
+
+		// "Object:/Path/To/Package/ObjectName"
+		ScratchSpaceBuffer.Reset();
+		ScratchSpaceBuffer.Append(ObjectTagCategory, ARRAY_COUNT(ObjectTagCategory) - 1);
+		AppendNameToBuffer(InPackageName);
+		ScratchSpaceBuffer.Add(TEXT('/'));
+		AppendNameToBuffer(InObjectName);
+		ScratchSpaceBuffer.Add(0);
+		ObjectTag = FName(ScratchSpaceBuffer.GetData());
+		GMallocProfiler->AddTag(ObjectTag);
+	}
+
+	if (!InClassName.IsNone())
+	{
+		// "Class:ClassName"
+		ScratchSpaceBuffer.Reset();
+		ScratchSpaceBuffer.Append(ClassTagCategory, ARRAY_COUNT(ClassTagCategory) - 1);
+		AppendNameToBuffer(InClassName);
+		ScratchSpaceBuffer.Add(0);
+		ClassTag = FName(ScratchSpaceBuffer.GetData());
+		GMallocProfiler->AddTag(ClassTag);
+	}
+}
+void FScopeCycleCounterUObject::UntrackObjectForMallocProfiling()
+{
+	if (!PackageTag.IsNone())
+	{
+		GMallocProfiler->RemoveTag(PackageTag);
+	}
+
+	if (!ClassTag.IsNone())
+	{
+		GMallocProfiler->RemoveTag(ClassTag);
+	}
+
+	if (!ObjectTag.IsNone())
+	{
+		GMallocProfiler->RemoveTag(ObjectTag);
+	}
+}
+#endif
diff --git a/Engine/Source/Runtime/CoreUObject/Public/UObject/Object.h b/Engine/Source/Runtime/CoreUObject/Public/UObject/Object.h
--- a/Engine/Source/Runtime/CoreUObject/Public/UObject/Object.h
+++ b/Engine/Source/Runtime/CoreUObject/Public/UObject/Object.h
@@ -3,6 +3,7 @@
 #pragma once
 
 #include "UObjectBaseUtility.h"
+#include "ResourceSize.h"
 
 DECLARE_LOG_CATEGORY_EXTERN(LogObj, Log, All);
 
@@ -16,18 +17,6 @@ namespace ECastCheckedType
 	};
 };
 
-/** Passed to GetResourceSize() to indicate which resource size should be returned.*/
-namespace EResourceSizeMode
-{
-	enum Type
-	{
-		/** Only exclusive resource size */
-		Exclusive,
-		/** Resource size of the object and all of its references */
-		Inclusive,
-	};
-};
-
 //
 // The base class of all objects.
 //
@@ -463,12 +452,37 @@ public:
 	 * default behavior is to return 0 which indicates that the resource shouldn't
 	 * display its size.
 	 *
-	 * @param	Type	Indicates which resource size should be returned
+	 * @param	Mode	Indicates which resource size should be returned
 	 * @return	Size of resource as to be displayed to artists/ LDs in the Editor.
 	 */
+	DEPRECATED(4.14, "GetResourceSize is deprecated. Please use GetResourceSizeEx or GetResourceSizeBytes instead.")
 	virtual SIZE_T GetResourceSize(EResourceSizeMode::Type Mode)
 	{
-		return 0;
+		return GetResourceSizeBytes(Mode);
+	}
+
+	/**
+	 * Get the size of the object/resource for display to artists/LDs in the Editor.
+	 * This is the extended version which separates up the used memory into different memory regions (the actual definition of which may be platform specific).
+	 *
+	 * @param	CumulativeResourceSize	Struct used to count up the cumulative size of the resource as to be displayed to artists/LDs in the Editor.
+	 */
+	virtual void GetResourceSizeEx(FResourceSizeEx& CumulativeResourceSize)
+	{
+	}
+
+	/**
+	 * Get the size of the object/resource for display to artists/LDs in the Editor.
+	 * This is the simple version which just returns the total number of bytes used by this object.
+	 *
+	 * @param	Mode					Indicates which resource size should be returned.
+	 * @return The cumulative size of the resource as to be displayed to artists/LDs in the Editor.
+	 */
+	SIZE_T GetResourceSizeBytes(EResourceSizeMode::Type Mode)
+	{
+		FResourceSizeEx ResSize = FResourceSizeEx(Mode);
+		GetResourceSizeEx(ResSize);
+		return ResSize.GetTotalMemoryBytes();
 	}
 
 	/** 
diff --git a/Engine/Source/Runtime/CoreUObject/Public/UObject/UObjectBaseUtility.h b/Engine/Source/Runtime/CoreUObject/Public/UObject/UObjectBaseUtility.h
--- a/Engine/Source/Runtime/CoreUObject/Public/UObject/UObjectBaseUtility.h
+++ b/Engine/Source/Runtime/CoreUObject/Public/UObject/UObjectBaseUtility.h
@@ -627,6 +627,15 @@ FORCEINLINE FString GetFullNameSafe(const UObjectBaseUtility *Object)
 #if STATS
 struct FScopeCycleCounterUObject : public FCycleCounter
 {
+#if USE_MALLOC_PROFILER
+	/** Package path being tracked */
+	FName PackageTag;
+	/** Class path being tracked */
+	FName ClassTag;
+	/** Object path being tracked */
+	FName ObjectTag;
+#endif
+
 	/**
 	 * Constructor, starts timing
 	 */
@@ -640,6 +649,13 @@ struct FScopeCycleCounterUObject : public FCycleCounter
 				Start(StatId);
 			}
 		}
+
+#if USE_MALLOC_PROFILER
+		if (Object)
+		{
+			TrackObjectForMallocProfiling(Object);
+		}
+#endif
 	}
 	/**
 	 * Constructor, starts timing with an alternate enable stat to use high performance disable for only SOME UObject stats
@@ -654,6 +670,13 @@ struct FScopeCycleCounterUObject : public FCycleCounter
 				Start(StatId);
 			}
 		}
+
+#if USE_MALLOC_PROFILER
+		if (Object)
+		{
+			TrackObjectForMallocProfiling(Object);
+		}
+#endif
 	}
 	/**
 	 * Updates the stat with the time spent
@@ -661,7 +684,17 @@ struct FScopeCycleCounterUObject : public FCycleCounter
 	FORCEINLINE_STATS ~FScopeCycleCounterUObject()
 	{
 		Stop();
+
+#if USE_MALLOC_PROFILER
+		UntrackObjectForMallocProfiling();
+#endif
 	}
+
+#if USE_MALLOC_PROFILER
+	COREUOBJECT_API void TrackObjectForMallocProfiling(const UObjectBaseUtility *InObject);
+	COREUOBJECT_API void TrackObjectForMallocProfiling(const FName InPackageName, const FName InClassName, const FName InObjectName);
+	COREUOBJECT_API void UntrackObjectForMallocProfiling();
+#endif
 };
 #else
 struct FScopeCycleCounterUObject
diff --git a/Engine/Source/Runtime/Engine/Classes/AI/Navigation/NavCollision.h b/Engine/Source/Runtime/Engine/Classes/AI/Navigation/NavCollision.h
--- a/Engine/Source/Runtime/Engine/Classes/AI/Navigation/NavCollision.h
+++ b/Engine/Source/Runtime/Engine/Classes/AI/Navigation/NavCollision.h
@@ -85,7 +85,7 @@ class ENGINE_API UNavCollision : public UObject
 	//~ Begin UObject Interface.
 	virtual void Serialize(FArchive& Ar) override;
 	virtual void PostLoad() override;
-	virtual SIZE_T GetResourceSize(EResourceSizeMode::Type Mode) override;
+	virtual void GetResourceSizeEx(FResourceSizeEx& CumulativeResourceSize) override;
 	//~ End UObject Interface.
 
 	FGuid GetGuid() const;
diff --git a/Engine/Source/Runtime/Engine/Classes/Animation/AnimSequence.h b/Engine/Source/Runtime/Engine/Classes/Animation/AnimSequence.h
--- a/Engine/Source/Runtime/Engine/Classes/Animation/AnimSequence.h
+++ b/Engine/Source/Runtime/Engine/Classes/Animation/AnimSequence.h
@@ -543,7 +543,7 @@ public:
 	virtual void PostDuplicate(bool bDuplicateForPIE) override;
 #endif // WITH_EDITOR
 	virtual void BeginDestroy() override;
-	virtual SIZE_T GetResourceSize(EResourceSizeMode::Type Mode) override;
+	virtual void GetResourceSizeEx(FResourceSizeEx& CumulativeResourceSize) override;
 	virtual void GetAssetRegistryTags(TArray<FAssetRegistryTag>& OutTags) const override;
 	//~ End UObject Interface
 
diff --git a/Engine/Source/Runtime/Engine/Classes/Animation/MorphTarget.h b/Engine/Source/Runtime/Engine/Classes/Animation/MorphTarget.h
--- a/Engine/Source/Runtime/Engine/Classes/Animation/MorphTarget.h
+++ b/Engine/Source/Runtime/Engine/Classes/Animation/MorphTarget.h
@@ -102,7 +102,10 @@ struct FMorphTargetLODModel
 	int32 NumBaseMeshVerts;
 	
 	/** Get Resource Size */
+	DEPRECATED(4.14, "GetResourceSize is deprecated. Please use GetResourceSizeEx or GetResourceSizeBytes instead.")
 	SIZE_T GetResourceSize() const;
+	void GetResourceSizeEx(FResourceSizeEx& CumulativeResourceSize) const;
+	SIZE_T GetResourceSizeBytes() const;
 
 	FMorphTargetLODModel()
 		: NumBaseMeshVerts(0)
@@ -151,7 +154,7 @@ public:
 	//~ UObject interface
 
 	virtual void Serialize(FArchive& Ar) override;
-	virtual SIZE_T GetResourceSize(EResourceSizeMode::Type Mode) override;
+	virtual void GetResourceSizeEx(FResourceSizeEx& CumulativeResourceSize) override;
 
 private:
 
diff --git a/Engine/Source/Runtime/Engine/Classes/Animation/SkeletalMeshActor.h b/Engine/Source/Runtime/Engine/Classes/Animation/SkeletalMeshActor.h
--- a/Engine/Source/Runtime/Engine/Classes/Animation/SkeletalMeshActor.h
+++ b/Engine/Source/Runtime/Engine/Classes/Animation/SkeletalMeshActor.h
@@ -62,6 +62,9 @@ public:
 	//~ Begin UObject Interface
 protected:
 	virtual FString GetDetailedInfoInternal() const override;
+#if WITH_EDITOR
+	virtual void PostEditChangeChainProperty(FPropertyChangedChainEvent& PropertyChangedEvent) override;
+#endif // WITH_EDITOR
 public:
 	//~ End UObject Interface
 
diff --git a/Engine/Source/Runtime/Engine/Classes/Components/ChildActorComponent.h b/Engine/Source/Runtime/Engine/Classes/Components/ChildActorComponent.h
--- a/Engine/Source/Runtime/Engine/Classes/Components/ChildActorComponent.h
+++ b/Engine/Source/Runtime/Engine/Classes/Components/ChildActorComponent.h
@@ -47,7 +47,7 @@ class ENGINE_API UChildActorComponent : public USceneComponent
 
 private:
 	/** The class of Actor to spawn */
-	UPROPERTY(EditAnywhere, BlueprintReadOnly, Category=ChildActorComponent, meta=(OnlyPlaceable, AllowPrivateAccess="true"))
+	UPROPERTY(EditAnywhere, BlueprintReadOnly, Category=ChildActorComponent, meta=(OnlyPlaceable, AllowPrivateAccess="true", ForceRebuildProperty="ChildActorTemplate"))
 	TSubclassOf<AActor>	ChildActorClass;
 
 	/** The actor that we spawned and own */
diff --git a/Engine/Source/Runtime/Engine/Classes/Components/MeshComponent.h b/Engine/Source/Runtime/Engine/Classes/Components/MeshComponent.h
--- a/Engine/Source/Runtime/Engine/Classes/Components/MeshComponent.h
+++ b/Engine/Source/Runtime/Engine/Classes/Components/MeshComponent.h
@@ -38,6 +38,15 @@ class ENGINE_API UMeshComponent : public UPrimitiveComponent
 	/** Returns override Materials count */
 	virtual int32 GetNumOverrideMaterials() const;
 
+#if WITH_EDITOR
+	/*
+	 * Make sure the Override array is using only the space it should use.
+	 * 1. The override array cannot be bigger then the number of mesh material.
+	 * 2. The override array must not end with a nullptr UMaterialInterface.
+	 */
+	void CleanUpOverrideMaterials();
+#endif
+
 	//~ Begin UObject Interface
 	virtual void BeginDestroy() override;
 	//~ End UObject Interface
diff --git a/Engine/Source/Runtime/Engine/Classes/Engine/DataTable.h b/Engine/Source/Runtime/Engine/Classes/Engine/DataTable.h
--- a/Engine/Source/Runtime/Engine/Classes/Engine/DataTable.h
+++ b/Engine/Source/Runtime/Engine/Classes/Engine/DataTable.h
@@ -13,6 +13,9 @@ template <class CharType, class PrintPolicy>
 class TJsonWriter;
 
 
+class UDataTable;
+
+
 /**
  * Base class for all table row structs to inherit from.
  */
@@ -27,9 +30,11 @@ struct FTableRowBase
 	 * Can be overridden by subclasses; Called whenever the owning data table is imported or re-imported.
 	 * Allows for custom fix-ups, parsing, etc. after initial data is read in.
 	 * 
-	 * @param OutCollectedImportProblems	[OUT] List of problems accumulated during import; Can be added to via this method
+	 * @param InDataTable					The data table that owns this row
+	 * @param InRowName						The name of the row we're performing fix-up on
+	 * @param OutCollectedImportProblems	List of problems accumulated during import; Can be added to via this method
 	 */
-	virtual void OnPostDataImport(OUT TArray<FString>& OutCollectedImportProblems) {}
+	virtual void OnPostDataImport(const UDataTable* InDataTable, const FName InRowName, TArray<FString>& OutCollectedImportProblems) {}
 };
 
 
@@ -53,7 +58,7 @@ class UDataTable
 	ENGINE_API virtual void FinishDestroy() override;
 	ENGINE_API virtual void Serialize(FArchive& Ar) override;
 	ENGINE_API static void AddReferencedObjects(UObject* InThis, FReferenceCollector& Collector);
-	ENGINE_API virtual SIZE_T GetResourceSize(EResourceSizeMode::Type Mode) override;
+	ENGINE_API virtual void GetResourceSizeEx(FResourceSizeEx& CumulativeResourceSize) override;
 #if WITH_EDITORONLY_DATA
 	ENGINE_API FName GetRowStructName() const;
 	ENGINE_API virtual void GetAssetRegistryTags(TArray<FAssetRegistryTag>& OutTags) const override;
diff --git a/Engine/Source/Runtime/Engine/Classes/Engine/GameViewportDelegates.h b/Engine/Source/Runtime/Engine/Classes/Engine/GameViewportDelegates.h
--- a/Engine/Source/Runtime/Engine/Classes/Engine/GameViewportDelegates.h
+++ b/Engine/Source/Runtime/Engine/Classes/Engine/GameViewportDelegates.h
@@ -50,4 +50,5 @@ DECLARE_MULTICAST_DELEGATE_OneParam(FOnGameViewportClientPlayerAction, int32);
  * @see UGameViewportClient
  */
 DECLARE_MULTICAST_DELEGATE_OneParam(FOnGameViewportTick, float);
-DECLARE_DELEGATE_RetVal_TwoParams(bool, FOnGameViewportInputKey, FKey, FModifierKeysState);
+
+DECLARE_DELEGATE_RetVal_ThreeParams(bool, FOnGameViewportInputKey, FKey, FModifierKeysState, EInputEvent);
diff --git a/Engine/Source/Runtime/Engine/Classes/Engine/RendererSettings.h b/Engine/Source/Runtime/Engine/Classes/Engine/RendererSettings.h
--- a/Engine/Source/Runtime/Engine/Classes/Engine/RendererSettings.h
+++ b/Engine/Source/Runtime/Engine/Classes/Engine/RendererSettings.h
@@ -398,12 +398,6 @@ public:
 
 private:
 	void SanatizeReflectionCaptureResolution();
-	
-	UPROPERTY(config)
-	EUIScalingRule UIScaleRule_DEPRECATED;
-
-	UPROPERTY(config)
-	FRuntimeFloatCurve UIScaleCurve_DEPRECATED;
 };
 
 UCLASS(config = Engine, globaluserconfig, meta = (DisplayName = "Rendering Overrides (Local)"))
diff --git a/Engine/Source/Runtime/Engine/Classes/Engine/SkeletalMesh.h b/Engine/Source/Runtime/Engine/Classes/Engine/SkeletalMesh.h
--- a/Engine/Source/Runtime/Engine/Classes/Engine/SkeletalMesh.h
+++ b/Engine/Source/Runtime/Engine/Classes/Engine/SkeletalMesh.h
@@ -468,7 +468,10 @@ struct FClothingAssetData
 	friend FArchive& operator<<(FArchive& Ar, FClothingAssetData& A);
 
 	// get resource size
+	DEPRECATED(4.14, "GetResourceSize is deprecated. Please use GetResourceSizeEx or GetResourceSizeBytes instead.")
 	SIZE_T GetResourceSize() const;
+	void GetResourceSizeEx(FResourceSizeEx& CumulativeResourceSize) const;
+	SIZE_T GetResourceSizeBytes() const;
 };
 
 //~ Begin Material Interface for USkeletalMesh - contains a material and a shadow casting flag
@@ -815,7 +818,7 @@ public:
 	virtual void GetAssetRegistryTags(TArray<FAssetRegistryTag>& OutTags) const override;
 	virtual FString GetDesc() override;
 	virtual FString GetDetailedInfoInternal() const override;
-	virtual SIZE_T GetResourceSize(EResourceSizeMode::Type Mode) override;
+	virtual void GetResourceSizeEx(FResourceSizeEx& CumulativeResourceSize) override;
 	static void AddReferencedObjects(UObject* InThis, FReferenceCollector& Collector);
 	//~ End UObject Interface.
 
diff --git a/Engine/Source/Runtime/Engine/Classes/Engine/StaticMesh.h b/Engine/Source/Runtime/Engine/Classes/Engine/StaticMesh.h
--- a/Engine/Source/Runtime/Engine/Classes/Engine/StaticMesh.h
+++ b/Engine/Source/Runtime/Engine/Classes/Engine/StaticMesh.h
@@ -325,15 +325,16 @@ struct FStaticMaterial
 	ENGINE_API friend bool operator==(const UMaterialInterface& LHS, const FStaticMaterial& RHS);
 
 	UPROPERTY(EditAnywhere, BlueprintReadOnly, transient, Category = StaticMesh)
-	class UMaterialInterface *	MaterialInterface;
+	class UMaterialInterface* MaterialInterface;
 
 	/*This name should be use by the gameplay to avoid error if the skeletal mesh Materials array topology change*/
 	UPROPERTY(EditAnywhere, BlueprintReadOnly, Category = StaticMesh)
-		FName						MaterialSlotName;
+	FName MaterialSlotName;
+
 #if WITH_EDITORONLY_DATA
 	/*This name should be use when we re-import a skeletal mesh so we can order the Materials array like it should be*/
 	UPROPERTY(VisibleAnywhere, BlueprintReadOnly, Category = StaticMesh)
-		FName						ImportedMaterialSlotName;
+	FName ImportedMaterialSlotName;
 #endif //WITH_EDITORONLY_DATA
 };
 
@@ -542,7 +543,7 @@ public:
 	ENGINE_API virtual bool IsReadyForFinishDestroy() override;
 	ENGINE_API virtual void GetAssetRegistryTags(TArray<FAssetRegistryTag>& OutTags) const override;
 	ENGINE_API virtual FString GetDesc() override;
-	ENGINE_API virtual SIZE_T GetResourceSize(EResourceSizeMode::Type Mode) override;
+	ENGINE_API virtual void GetResourceSizeEx(FResourceSizeEx& CumulativeResourceSize) override;
 	//~ End UObject Interface.
 
 	/**
diff --git a/Engine/Source/Runtime/Engine/Classes/Engine/TextureRenderTarget2D.h b/Engine/Source/Runtime/Engine/Classes/Engine/TextureRenderTarget2D.h
--- a/Engine/Source/Runtime/Engine/Classes/Engine/TextureRenderTarget2D.h
+++ b/Engine/Source/Runtime/Engine/Classes/Engine/TextureRenderTarget2D.h
@@ -100,7 +100,7 @@ class UTextureRenderTarget2D : public UTextureRenderTarget
 	ENGINE_API virtual void PostEditChangeProperty(FPropertyChangedEvent& PropertyChangedEvent) override;
 #endif // WITH_EDITOR
 	ENGINE_API virtual void PostLoad() override;
-	ENGINE_API virtual SIZE_T GetResourceSize(EResourceSizeMode::Type Mode) override;
+	ENGINE_API virtual void GetResourceSizeEx(FResourceSizeEx& CumulativeResourceSize) override;
 	ENGINE_API virtual FString GetDesc() override;
 	//~ End UObject Interface
 
diff --git a/Engine/Source/Runtime/Engine/Classes/Engine/UserInterfaceSettings.h b/Engine/Source/Runtime/Engine/Classes/Engine/UserInterfaceSettings.h
--- a/Engine/Source/Runtime/Engine/Classes/Engine/UserInterfaceSettings.h
+++ b/Engine/Source/Runtime/Engine/Classes/Engine/UserInterfaceSettings.h
@@ -99,10 +99,6 @@ public:
 	/** Gets the current scale of the UI based on the size of a viewport */
 	float GetDPIScaleBasedOnSize(FIntPoint Size) const;
 
-#if WITH_EDITOR
-	void PostEditChangeProperty(FPropertyChangedEvent& PropertyChangedEvent) override;
-#endif
-
 private:
 	UPROPERTY(Transient)
 	TArray<UObject*> CursorClasses;
diff --git a/Engine/Source/Runtime/Engine/Classes/Materials/Material.h b/Engine/Source/Runtime/Engine/Classes/Materials/Material.h
--- a/Engine/Source/Runtime/Engine/Classes/Materials/Material.h
+++ b/Engine/Source/Runtime/Engine/Classes/Materials/Material.h
@@ -894,7 +894,7 @@ public:
 	ENGINE_API virtual void BeginDestroy() override;
 	ENGINE_API virtual bool IsReadyForFinishDestroy() override;
 	ENGINE_API virtual void FinishDestroy() override;
-	ENGINE_API virtual SIZE_T GetResourceSize(EResourceSizeMode::Type Mode) override;
+	ENGINE_API virtual void GetResourceSizeEx(FResourceSizeEx& CumulativeResourceSize) override;
 	ENGINE_API static void AddReferencedObjects(UObject* InThis, FReferenceCollector& Collector);
 	ENGINE_API virtual bool CanBeClusterRoot() const override;
 	ENGINE_API virtual void GetAssetRegistryTags(TArray<FAssetRegistryTag>& OutTags) const override;
diff --git a/Engine/Source/Runtime/Engine/Classes/Materials/MaterialInstance.h b/Engine/Source/Runtime/Engine/Classes/Materials/MaterialInstance.h
--- a/Engine/Source/Runtime/Engine/Classes/Materials/MaterialInstance.h
+++ b/Engine/Source/Runtime/Engine/Classes/Materials/MaterialInstance.h
@@ -303,7 +303,7 @@ public:
 	//~ End UMaterialInterface Interface.
 
 	//~ Begin UObject Interface.
-	virtual ENGINE_API SIZE_T GetResourceSize(EResourceSizeMode::Type Mode) override;
+	virtual ENGINE_API void GetResourceSizeEx(FResourceSizeEx& CumulativeResourceSize) override;
 	virtual ENGINE_API void PostInitProperties() override;	
 #if WITH_EDITOR
 	virtual ENGINE_API void BeginCacheForCookedPlatformData(const ITargetPlatform *TargetPlatform) override;
diff --git a/Engine/Source/Runtime/Engine/Classes/PhysicsEngine/BodyInstance.h b/Engine/Source/Runtime/Engine/Classes/PhysicsEngine/BodyInstance.h
--- a/Engine/Source/Runtime/Engine/Classes/PhysicsEngine/BodyInstance.h
+++ b/Engine/Source/Runtime/Engine/Classes/PhysicsEngine/BodyInstance.h
@@ -1041,7 +1041,9 @@ public:
 	/** 
 	 * Returns memory used by resources allocated for this body instance ( ex. Physx resources )
 	 **/
+	DEPRECATED(4.14, "GetBodyInstanceResourceSize is deprecated. Please use GetBodyInstanceResourceSizeEx instead.")
 	SIZE_T GetBodyInstanceResourceSize(EResourceSizeMode::Type Mode) const;
+	void GetBodyInstanceResourceSizeEx(FResourceSizeEx& CumulativeResourceSize) const;
 
 	/**
 	 * UObject notification by OwningComponent
diff --git a/Engine/Source/Runtime/Engine/Classes/PhysicsEngine/BodySetup.h b/Engine/Source/Runtime/Engine/Classes/PhysicsEngine/BodySetup.h
--- a/Engine/Source/Runtime/Engine/Classes/PhysicsEngine/BodySetup.h
+++ b/Engine/Source/Runtime/Engine/Classes/PhysicsEngine/BodySetup.h
@@ -32,7 +32,10 @@ struct FBodySetupUVInfo
 	}
 
 	/** Get resource size of UV info */
-	SIZE_T GetResourceSize();
+	DEPRECATED(4.14, "GetResourceSize is deprecated. Please use GetResourceSizeEx or GetResourceSizeBytes instead.")
+	SIZE_T GetResourceSize() const;
+	void GetResourceSizeEx(FResourceSizeEx& CumulativeResourceSize) const;
+	SIZE_T GetResourceSizeBytes() const;
 };
 
 /**
@@ -180,7 +183,7 @@ public:
 #if WITH_EDITOR
 	virtual void PostEditUndo() override;
 #endif // WITH_EDITOR
-	virtual SIZE_T GetResourceSize(EResourceSizeMode::Type Mode) override;
+	virtual void GetResourceSizeEx(FResourceSizeEx& CumulativeResourceSize) override;
 	//~ End UObject Interface.
 
 	//
diff --git a/Engine/Source/Runtime/Engine/Private/AI/Navigation/NavCollision.cpp b/Engine/Source/Runtime/Engine/Private/AI/Navigation/NavCollision.cpp
--- a/Engine/Source/Runtime/Engine/Private/AI/Navigation/NavCollision.cpp
+++ b/Engine/Source/Runtime/Engine/Private/AI/Navigation/NavCollision.cpp
@@ -482,15 +482,13 @@ FByteBulkData* UNavCollision::GetCookedData(FName Format)
 	return Result->GetBulkDataSize() > 0 ? Result : NULL; // we don't return empty bulk data...but we save it to avoid thrashing the DDC
 }
 
-SIZE_T UNavCollision::GetResourceSize( EResourceSizeMode::Type Mode )
+void UNavCollision::GetResourceSizeEx(FResourceSizeEx& CumulativeResourceSize)
 {
-	SIZE_T ResSize = 0;
+	Super::GetResourceSizeEx(CumulativeResourceSize);
 	
 	if (CookedFormatData.Contains(NAVCOLLISION_FORMAT))
 	{
 		const FByteBulkData& FmtData = CookedFormatData.GetFormat(NAVCOLLISION_FORMAT);
-		ResSize += FmtData.GetElementSize() * FmtData.GetElementCount();
+		CumulativeResourceSize.AddDedicatedSystemMemoryBytes(FmtData.GetElementSize() * FmtData.GetElementCount());
 	}
-	
-	return ResSize;
 }
diff --git a/Engine/Source/Runtime/Engine/Private/Actor.cpp b/Engine/Source/Runtime/Engine/Private/Actor.cpp
--- a/Engine/Source/Runtime/Engine/Private/Actor.cpp
+++ b/Engine/Source/Runtime/Engine/Private/Actor.cpp
@@ -4186,7 +4186,7 @@ void AActor::InvalidateLightingCacheDetailed(bool bTranslationOnly)
 	// Validate that we didn't change it during this action
 	TInlineComponentArray<UActorComponent*> NewComponents;
 	GetComponents(NewComponents);
-	check(Components == NewComponents);
+	ensure(Components == NewComponents);
 #endif
 }
 
diff --git a/Engine/Source/Runtime/Engine/Private/Animation/AnimSequence.cpp b/Engine/Source/Runtime/Engine/Private/Animation/AnimSequence.cpp
--- a/Engine/Source/Runtime/Engine/Private/Animation/AnimSequence.cpp
+++ b/Engine/Source/Runtime/Engine/Private/Animation/AnimSequence.cpp
@@ -140,16 +140,17 @@ void UAnimSequence::PostInitProperties()
 	Super::PostInitProperties();
 }
 
-SIZE_T UAnimSequence::GetResourceSize(EResourceSizeMode::Type Mode)
+void UAnimSequence::GetResourceSizeEx(FResourceSizeEx& CumulativeResourceSize)
 {
-	if (Mode == EResourceSizeMode::Exclusive)
+	Super::GetResourceSizeEx(CumulativeResourceSize);
+
+	if (CumulativeResourceSize.GetResourceSizeMode() == EResourceSizeMode::Exclusive)
 	{
 		// All of the sequence data is serialized and will be counted as part of the direct object size rather than as a resource
-		return 0;
 	}
 	else
 	{
-		return (CompressedTrackOffsets.Num() == 0) ? GetApproxRawSize() : GetApproxCompressedSize();
+		CumulativeResourceSize.AddDedicatedSystemMemoryBytes((CompressedTrackOffsets.Num() == 0) ? GetApproxRawSize() : GetApproxCompressedSize());
 	}
 }
 
@@ -446,7 +447,7 @@ void UAnimSequence::PostLoad()
 	if( IsRunningGame() )
 	{
 		// this probably will not show newly created animations in PIE but will show them in the game once they have been saved off
-		INC_DWORD_STAT_BY( STAT_AnimationMemory, GetResourceSize(EResourceSizeMode::Exclusive) );
+		INC_DWORD_STAT_BY( STAT_AnimationMemory, GetResourceSizeBytes(EResourceSizeMode::Exclusive) );
 	}
 
 	{
@@ -605,7 +606,7 @@ void UAnimSequence::BeginDestroy()
 
 	if( IsRunningGame() )
 	{
-		DEC_DWORD_STAT_BY( STAT_AnimationMemory, GetResourceSize(EResourceSizeMode::Exclusive) );
+		DEC_DWORD_STAT_BY( STAT_AnimationMemory, GetResourceSizeBytes(EResourceSizeMode::Exclusive) );
 	}
 }
 
@@ -5077,7 +5078,7 @@ void GatherAnimSequenceStats(FOutputDevice& Ar)
 			NumTransTracksWithOneKey, NumRotTracksWithOneKey, NumScaleTracksWithOneKey,
 			TotalNumTransKeys, TotalNumRotKeys, TotalNumScaleKeys,
 			*FAnimationUtils::GetAnimationKeyFormatString(static_cast<AnimationKeyFormat>(Seq->KeyEncodingFormat)),
-			(int32)Seq->GetResourceSize(EResourceSizeMode::Exclusive) );
+			(int32)Seq->GetResourceSizeBytes(EResourceSizeMode::Exclusive) );
 	}
 	Ar.Logf( TEXT("======================================================================") );
 	Ar.Logf( TEXT("Total Num Tracks: %i trans, %i rot, %i scale, %i trans1, %i rot1, %i scale1"), GlobalNumTransTracks, GlobalNumRotTracks, GlobalNumScaleTracks, GlobalNumTransTracksWithOneKey, GlobalNumRotTracksWithOneKey, GlobalNumScaleTracksWithOneKey  );
diff --git a/Engine/Source/Runtime/Engine/Private/Animation/AnimSet.cpp b/Engine/Source/Runtime/Engine/Private/Animation/AnimSet.cpp
--- a/Engine/Source/Runtime/Engine/Private/Animation/AnimSet.cpp
+++ b/Engine/Source/Runtime/Engine/Private/Animation/AnimSet.cpp
@@ -126,27 +126,26 @@ int32 UAnimSet::GetMeshLinkupIndex(USkeletalMesh* SkelMesh)
 }
 
 
-SIZE_T UAnimSet::GetResourceSize(EResourceSizeMode::Type Mode)
+void UAnimSet::GetResourceSizeEx(FResourceSizeEx& CumulativeResourceSize)
 {
-	if (Mode == EResourceSizeMode::Exclusive)
+	Super::GetResourceSizeEx(CumulativeResourceSize);
+
+	if (CumulativeResourceSize.GetResourceSizeMode() == EResourceSizeMode::Exclusive)
 	{
 		// This object only references others, it doesn't have any real resource bytes
-		return 0;
 	}
 	else
 	{
-		int32 ResourceSize = 0;
 #if WITH_EDITORONLY_DATA
 		for( int32 i=0; i<Sequences.Num(); i++ )
 		{
 			UAnimSequence* AnimSeq = Sequences[i];
 			if( AnimSeq )
 			{
-				ResourceSize += AnimSeq->GetResourceSize(Mode);
+				AnimSeq->GetResourceSizeEx(CumulativeResourceSize);
 			}			
 		}
 #endif	//#if WITH_EDITORONLY_DATA
-		return ResourceSize;
 	}
 }
 
diff --git a/Engine/Source/Runtime/Engine/Private/Animation/AnimationUtils.cpp b/Engine/Source/Runtime/Engine/Private/Animation/AnimationUtils.cpp
--- a/Engine/Source/Runtime/Engine/Private/Animation/AnimationUtils.cpp
+++ b/Engine/Source/Runtime/Engine/Private/Animation/AnimationUtils.cpp
@@ -438,22 +438,22 @@ bool FAnimationUtils::GetForcedRecompressionSetting()
 	#define TRYCOMPRESSION_INNER(compressionname,winningcompressor_count,winningcompressor_error,winningcompressor_margin,compressionalgorithm)				\
 {																																							\
 	/* try the alternative compressor	*/																													\
-	(compressionalgorithm)->Reduce( AnimSeq, CompressContext );																				\
-	AnimSeq->SetUseRawDataOnly(false);																								\
-	const SIZE_T NewSize = AnimSeq->GetResourceSize(EResourceSizeMode::Inclusive);																												\
+	(compressionalgorithm)->Reduce( AnimSeq, CompressContext );																								\
+	AnimSeq->SetUseRawDataOnly(false);																														\
+	const SIZE_T NewSize = AnimSeq->GetResourceSizeBytes(EResourceSizeMode::Inclusive);																		\
 																																							\
 	/* compute the savings and compression error*/																											\
-	const SIZE_T MemorySavingsFromOriginal = OriginalSize - NewSize;																									\
-	const SIZE_T MemorySavingsFromPrevious = CurrentSize - NewSize;																									\
+	const SIZE_T MemorySavingsFromOriginal = OriginalSize - NewSize;																						\
+	const SIZE_T MemorySavingsFromPrevious = CurrentSize - NewSize;																							\
 	PctSaving = 0.f;																																		\
 	/* figure out our new compression error*/																												\
-	FAnimationUtils::ComputeCompressionError(AnimSeq, BoneData, NewErrorStats);														\
+	FAnimationUtils::ComputeCompressionError(AnimSeq, BoneData, NewErrorStats);																				\
 																																							\
 	const bool bLowersError = NewErrorStats.MaxError < WinningCompressorError;																				\
 	const bool bErrorUnderThreshold = NewErrorStats.MaxError <= MasterTolerance;																			\
 																																							\
 	/* keep it if it we want to force the error below the threshold and it reduces error */																	\
-	bKeepNewCompressionMethod = false;																												\
+	bKeepNewCompressionMethod = false;																														\
 	bKeepNewCompressionMethod |= (bLowersError && (WinningCompressorError > MasterTolerance) && bForceBelowThreshold);										\
 	/* or if has an acceptable error and saves space  */																									\
 	bKeepNewCompressionMethod |= bErrorUnderThreshold && (MemorySavingsFromPrevious > 0);																	\
@@ -461,19 +461,19 @@ bool FAnimationUtils::GetForcedRecompressionSetting()
 	bKeepNewCompressionMethod |= bErrorUnderThreshold && bLowersError && (MemorySavingsFromPrevious >= 0);													\
 																																							\
 	if (bKeepNewCompressionMethod)																															\
-			{																																						\
+			{																																				\
 		WinningCompressorMarginalSavings = MemorySavingsFromPrevious;																						\
-		WinningCompressorCounter = &(winningcompressor_count);																							\
-		WinningCompressorErrorSum = &(winningcompressor_error);																							\
-		WinningCompressorMarginalSavingsSum = &(winningcompressor_margin);																				\
+		WinningCompressorCounter = &(winningcompressor_count);																								\
+		WinningCompressorErrorSum = &(winningcompressor_error);																								\
+		WinningCompressorMarginalSavingsSum = &(winningcompressor_margin);																					\
 		WinningCompressorName = compressionname;																											\
 		CurrentSize = NewSize;																																\
 		WinningCompressorSavings = MemorySavingsFromOriginal;																								\
 		WinningCompressorError = NewErrorStats.MaxError;																									\
-			}																																						\
+			}																																				\
 																																							\
 	PctSaving = OriginalSize > 0 ? 100.f - (100.f * float(NewSize) / float(OriginalSize)) : 0.f;															\
-	UE_LOG(LogAnimation, Warning, TEXT("- %s - bytes saved: %i (%3.1f%% saved), maxdiff: %f %s"),																					\
+	UE_LOG(LogAnimation, Warning, TEXT("- %s - bytes saved: %i (%3.1f%% saved), maxdiff: %f %s"),															\
 	compressionname, MemorySavingsFromOriginal, PctSaving, NewErrorStats.MaxError, bKeepNewCompressionMethod ? TEXT("(**Best so far**)") : TEXT(""));		\
 																																							\
 	if( !bKeepNewCompressionMethod )																														\
@@ -492,7 +492,7 @@ bool FAnimationUtils::GetForcedRecompressionSetting()
 		AnimSeq->SetUseRawDataOnly(bSavedUseRawDataOnly);																									\
 		AnimationFormat_SetInterfaceLinks(*AnimSeq);																										\
 																																							\
-		const SIZE_T RestoredSize = AnimSeq->GetResourceSize(EResourceSizeMode::Inclusive);																	\
+		const SIZE_T RestoredSize = AnimSeq->GetResourceSizeBytes(EResourceSizeMode::Inclusive);															\
 		check(RestoredSize == CurrentSize);																													\
 	}																																						\
 	else																																					\
@@ -687,7 +687,7 @@ void FAnimationUtils::CompressAnimSequenceExplicit(
 		bool const bTryAlternateCompressor = MasterTolerance > 0.0f;
 
 		// Get the current size
-		int32 OriginalSize = AnimSeq->GetResourceSize(EResourceSizeMode::Inclusive);
+		int32 OriginalSize = AnimSeq->GetResourceSizeBytes(EResourceSizeMode::Inclusive);
 		TotalSizeBefore += OriginalSize;
 
 		// Estimate total uncompressed
@@ -724,14 +724,14 @@ void FAnimationUtils::CompressAnimSequenceExplicit(
 
 			OriginalCompressionAlgorithm->Reduce( AnimSeq, CompressContext );
 			AnimSeq->SetUseRawDataOnly(false);
-			AfterOriginalRecompression = AnimSeq->GetResourceSize(EResourceSizeMode::Inclusive);
+			AfterOriginalRecompression = AnimSeq->GetResourceSizeBytes(EResourceSizeMode::Inclusive);
 
 			// figure out our current compression error
 			FAnimationUtils::ComputeCompressionError(AnimSeq, BoneData, OriginalErrorStats);
 		}
 		else
 		{
-			AfterOriginalRecompression = AnimSeq->GetResourceSize(EResourceSizeMode::Inclusive);
+			AfterOriginalRecompression = AnimSeq->GetResourceSizeBytes(EResourceSizeMode::Inclusive);
 			OriginalErrorStats = TrueOriginalErrorStats;
 		}
  
@@ -766,7 +766,7 @@ void FAnimationUtils::CompressAnimSequenceExplicit(
 				++TotalRecompressions;
 
 				// Prepare to compress
-				int32 CurrentSize = AnimSeq->GetResourceSize(EResourceSizeMode::Inclusive);
+				int32 CurrentSize = AnimSeq->GetResourceSizeBytes(EResourceSizeMode::Inclusive);
 				int32* WinningCompressorCounter = NULL;
 				float* WinningCompressorErrorSum = NULL;
 				int32* WinningCompressorMarginalSavingsSum = NULL;
@@ -1256,7 +1256,7 @@ void FAnimationUtils::CompressAnimSequenceExplicit(
 				}
 
 				// Make sure we got that right.
-				check(CurrentSize == AnimSeq->GetResourceSize(EResourceSizeMode::Inclusive));
+				check(CurrentSize == AnimSeq->GetResourceSizeBytes(EResourceSizeMode::Inclusive));
 				TotalSizeNow += CurrentSize;
 
 				PctSaving = TotalSizeBefore > 0 ? 100.f - (100.f * float(TotalSizeNow) / float(TotalSizeBefore)) : 0.f;
@@ -1338,7 +1338,7 @@ void FAnimationUtils::CompressAnimSequenceExplicit(
 		// Do not recompress - Still take into account size for stats.
 		else
 		{
-			TotalSizeNow += AnimSeq->GetResourceSize(EResourceSizeMode::Inclusive);
+			TotalSizeNow += AnimSeq->GetResourceSizeBytes(EResourceSizeMode::Inclusive);
 		}
 	}
 	else
diff --git a/Engine/Source/Runtime/Engine/Private/AudioDevice.cpp b/Engine/Source/Runtime/Engine/Private/AudioDevice.cpp
--- a/Engine/Source/Runtime/Engine/Private/AudioDevice.cpp
+++ b/Engine/Source/Runtime/Engine/Private/AudioDevice.cpp
@@ -957,12 +957,12 @@ bool FAudioDevice::HandleAudioMemoryInfo(const TCHAR* Cmd, FOutputDevice& Ar)
 	struct FSoundWaveInfo
 	{
 		USoundWave* SoundWave;
-		SIZE_T ResourceSize;
+		FResourceSizeEx ResourceSize;
 		FString SoundGroupName;
 		float Duration;
 		bool bDecompressed;
 
-		FSoundWaveInfo(USoundWave* InSoundWave, SIZE_T InResourceSize, const FString& InSoundGroupName, float InDuration, bool bInDecompressed)
+		FSoundWaveInfo(USoundWave* InSoundWave, FResourceSizeEx InResourceSize, const FString& InSoundGroupName, float InDuration, bool bInDecompressed)
 			: SoundWave(InSoundWave)
 			, ResourceSize(InResourceSize)
 			, SoundGroupName(InSoundGroupName)
@@ -973,12 +973,12 @@ bool FAudioDevice::HandleAudioMemoryInfo(const TCHAR* Cmd, FOutputDevice& Ar)
 
 	struct FSoundWaveGroupInfo
 	{
-		SIZE_T ResourceSize;
-		SIZE_T CompressedResourceSize;
+		FResourceSizeEx ResourceSize;
+		FResourceSizeEx CompressedResourceSize;
 
 		FSoundWaveGroupInfo()
-			: ResourceSize(0)
-			, CompressedResourceSize(0)
+			: ResourceSize()
+			, CompressedResourceSize()
 		{}
 	};
 
@@ -1023,9 +1023,9 @@ bool FAudioDevice::HandleAudioMemoryInfo(const TCHAR* Cmd, FOutputDevice& Ar)
 		}
 	}
 
-	SIZE_T TotalResourceSize = 0;
-	SIZE_T CompressedResourceSize = 0;
-	SIZE_T DecompressedResourceSize = 0;
+	FResourceSizeEx TotalResourceSize;
+	FResourceSizeEx CompressedResourceSize;
+	FResourceSizeEx DecompressedResourceSize;
 	int32 CompressedResourceCount = 0;
 
 	if (SoundWaveClass != nullptr)
@@ -1039,8 +1039,9 @@ bool FAudioDevice::HandleAudioMemoryInfo(const TCHAR* Cmd, FOutputDevice& Ar)
 			}
 
 			// Get the resource size of the sound wave
-			const SIZE_T TrueResourceSize = It->GetResourceSize(EResourceSizeMode::Exclusive);
-			if (TrueResourceSize == 0)
+			FResourceSizeEx TrueResourceSize = FResourceSizeEx(EResourceSizeMode::Exclusive);
+			It->GetResourceSizeEx(TrueResourceSize);
+			if (TrueResourceSize.GetTotalMemoryBytes() == 0)
 			{
 				continue;
 			}
@@ -1138,9 +1139,9 @@ bool FAudioDevice::HandleAudioMemoryInfo(const TCHAR* Cmd, FOutputDevice& Ar)
 			// Log the sound wave objects
 			
 			ReportAr->Logf(TEXT("Memory (MB),Count"));
-			ReportAr->Logf(TEXT("Total,%.3f,%d"), TotalResourceSize / 1024.f / 1024.f, SoundWaveObjects.Num());
-			ReportAr->Logf(TEXT("Decompressed,%.3f,%d"), DecompressedResourceSize / 1024.f / 1024.f, CompressedResourceCount);
-			ReportAr->Logf(TEXT("Compressed,%.3f,%d"), CompressedResourceSize / 1024.f / 1024.f, SoundWaveObjects.Num() - CompressedResourceCount);
+			ReportAr->Logf(TEXT("Total,%.3f,%d"), TotalResourceSize.GetTotalMemoryBytes() / 1024.f / 1024.f, SoundWaveObjects.Num());
+			ReportAr->Logf(TEXT("Decompressed,%.3f,%d"), DecompressedResourceSize.GetTotalMemoryBytes() / 1024.f / 1024.f, CompressedResourceCount);
+			ReportAr->Logf(TEXT("Compressed,%.3f,%d"), CompressedResourceSize.GetTotalMemoryBytes() / 1024.f / 1024.f, SoundWaveObjects.Num() - CompressedResourceCount);
 
 			if (SoundWaveGroupFolders.Num())
 			{
@@ -1152,7 +1153,7 @@ bool FAudioDevice::HandleAudioMemoryInfo(const TCHAR* Cmd, FOutputDevice& Ar)
 				{
 					FSoundWaveGroupInfo* SubDirSize = SoundWaveGroupSizes.Find(SoundWaveGroupFolder);
 					check(SubDirSize);
-					ReportAr->Logf(TEXT("%s,%10.2f,%10.2f"), *SoundWaveGroupFolder, SubDirSize->ResourceSize / 1024.0f / 1024.0f, SubDirSize->CompressedResourceSize / 1024.0f / 1024.0f);
+					ReportAr->Logf(TEXT("%s,%10.2f,%10.2f"), *SoundWaveGroupFolder, SubDirSize->ResourceSize.GetTotalMemoryBytes() / 1024.0f / 1024.0f, SubDirSize->CompressedResourceSize.GetTotalMemoryBytes() / 1024.0f / 1024.0f);
 				}
 			}
 
@@ -1163,7 +1164,7 @@ bool FAudioDevice::HandleAudioMemoryInfo(const TCHAR* Cmd, FOutputDevice& Ar)
 			ReportAr->Logf(TEXT("%s,%s,%s,%s,%s,%s"), TEXT("SoundWave"), TEXT("KB"), TEXT("MB"), TEXT("SoundGroup"), TEXT("Duration"), TEXT("CompressionState"));
 			for (const FSoundWaveInfo& Info : SoundWaveObjects)
 			{
-				float Kbytes = Info.ResourceSize / 1024.0f;
+				float Kbytes = Info.ResourceSize.GetTotalMemoryBytes() / 1024.0f;
 				ReportAr->Logf(TEXT("%s,%10.2f,%10.2f,%s,%10.2f,%s"), *Info.SoundWave->GetPathName(), Kbytes, Kbytes / 1024.0f, *Info.SoundGroupName, Info.Duration, Info.bDecompressed ? TEXT("Decompressed") : TEXT("Compressed"));
 			}
 		}
@@ -4054,7 +4055,7 @@ void FAudioDevice::Precache(USoundWave* SoundWave, bool bSynchronous, bool bTrac
 
 	if (bTrackMemory)
 	{
-		const int32 ResourceSize = SoundWave->GetResourceSize(EResourceSizeMode::Exclusive);
+		const int32 ResourceSize = SoundWave->GetResourceSizeBytes(EResourceSizeMode::Exclusive);
 		SoundWave->TrackedMemoryUsage += ResourceSize;
 
 		// If we aren't decompressing it above, then count the memory
@@ -4502,7 +4503,7 @@ void FAudioDevice::DumpActiveSounds() const
 			for (const TPair<UPTRINT, FWaveInstance*>& WaveInstancePair : ActiveSound->WaveInstances)
 			{
 				const FWaveInstance* WaveInstance = WaveInstancePair.Value;
-				UE_LOG(LogAudio, Display, TEXT("   %s (%.3g) (%d) - %.3g"), *WaveInstance->GetName(), WaveInstance->WaveData->GetDuration(), WaveInstance->WaveData->GetResourceSize(EResourceSizeMode::Inclusive), WaveInstance->GetActualVolume());
+				UE_LOG(LogAudio, Display, TEXT("   %s (%.3g) (%d) - %.3g"), *WaveInstance->GetName(), WaveInstance->WaveData->GetDuration(), WaveInstance->WaveData->GetResourceSizeBytes(EResourceSizeMode::Inclusive), WaveInstance->GetActualVolume());
 			}
 		}
 	}
