Fix UHT needlessly rewriting configs.

diff --git a/Engine/Source/Runtime/Core/Private/Misc/ConfigCacheIni.cpp b/Engine/Source/Runtime/Core/Private/Misc/ConfigCacheIni.cpp
--- a/Engine/Source/Runtime/Core/Private/Misc/ConfigCacheIni.cpp
+++ b/Engine/Source/Runtime/Core/Private/Misc/ConfigCacheIni.cpp
@@ -333,8 +333,6 @@ FConfigSection* FConfigFile::FindOrAddSection(const FString& SectionName)
 	if (Section == nullptr)
 	{
 		Section = &Add(SectionName, FConfigSection());
-		// mark that we have been modified
-		Dirty = true;
 	}
 	return Section;
 }
