Copying //UE4/Dev-Core to //UE4/Dev-Main (Source: //UE4/Dev-Core @ 3151653)

diff --git a/Engine/Source/Runtime/Core/Private/Misc/OutputDeviceFile.cpp b/Engine/Source/Runtime/Core/Private/Misc/OutputDeviceFile.cpp
--- a/Engine/Source/Runtime/Core/Private/Misc/OutputDeviceFile.cpp
+++ b/Engine/Source/Runtime/Core/Private/Misc/OutputDeviceFile.cpp
@@ -13,7 +13,8 @@ static UTF8BOMType UTF8BOM = { 0xEF, 0xBB, 0xBF };
 
 /**
 * Thread heartbeat check class.
-* Used by crash handling code to check for hangs.
+* Used by crash handling code to check for hangs. 
+* [] tags identify which thread owns a variable or function
 */
 class CORE_API FAsyncWriter : public FRunnable, public FArchive
 {
@@ -31,49 +32,72 @@ class CORE_API FAsyncWriter : public FRunnable, public FArchive
 	FArchive& Ar;
 	/** Data ring buffer */
 	TArray<uint8> Buffer;
-	/** Position where the unserialized data starts in the buffer */
+	/** [WRITER THREAD] Position where the unserialized data starts in the buffer */
 	int32 BufferStartPos;
-	/** Position where the unserialized data ends in the buffer (such as if (BufferEndPos > BufferStartPos) Length = BufferEndPos - BufferStartPos; */
+	/** [CLIENT THREAD] Position where the unserialized data ends in the buffer (such as if (BufferEndPos > BufferStartPos) Length = BufferEndPos - BufferStartPos; */
 	int32 BufferEndPos;
-	/** Sync object for the buffer pos */
+	/** [CLIENT THREAD] Sync object for the buffer pos */
 	FCriticalSection BufferPosCritical;
-	/** Event to let the worker thread know there is data to save to disk. */
-	FEvent* BufferDirtyEvent;
-	/** Event flagged when the worker thread finished processing data */
-	FEvent* BufferEmptyEvent;
+	/** [CLIENT/WRITER THREAD] Outstanding serialize request counter. This is to make sure we flush all requests. */
+	FThreadSafeCounter SerializeRequestCounter;
 
-	/** Serialize the contents of the ring buffer to disk */
+	/** [WRITER THREAD] Last time the archive was flushed. used in threaded situations to flush the underlying archive at a certain maximum rate. */
+	double LastArchiveFlushTime;
+
+	/** [WRITER THREAD] Serialize the contents of the ring buffer to disk */
 	void SerializeBufferToArchive()
 	{
-		// Grab a local copy of the end pos. It's ok if it changes on the client thread later on.
-		// We won't be modifying it anyway and will later serialize new data in the next iteration.
-		// Here we only serialize what we know exists at the beginning of this function.
-		int32 ThisThreadEndPos = BufferEndPos;
-		if (ThisThreadEndPos >= BufferStartPos)
+		while (SerializeRequestCounter.GetValue() > 0)
 		{
-			Ar.Serialize(Buffer.GetData() + BufferStartPos, ThisThreadEndPos - BufferStartPos);
-		}
-		else
-		{
-			// Data is wrapped around the ring buffer
-			Ar.Serialize(Buffer.GetData() + BufferStartPos, Buffer.Num() - BufferStartPos);
-			Ar.Serialize(Buffer.GetData(), BufferEndPos);
+			// Grab a local copy of the end pos. It's ok if it changes on the client thread later on.
+			// We won't be modifying it anyway and will later serialize new data in the next iteration.
+			// Here we only serialize what we know exists at the beginning of this function.
+			int32 ThisThreadEndPos = BufferEndPos;
+
+			if (ThisThreadEndPos >= BufferStartPos)
+			{
+				Ar.Serialize(Buffer.GetData() + BufferStartPos, ThisThreadEndPos - BufferStartPos);
+			}
+			else
+			{
+				// Data is wrapped around the ring buffer
+				Ar.Serialize(Buffer.GetData() + BufferStartPos, Buffer.Num() - BufferStartPos);
+				Ar.Serialize(Buffer.GetData(), BufferEndPos);
+			}
+			// Modify the start pos. Only the worker thread modifies this value so it's ok to not guard it with a critical section.
+			BufferStartPos = ThisThreadEndPos;
+
+			// Flush the archive periodically if running on a separate thread
+			if (Thread)
+			{
+				const double ArchiveFlushIntervalSec = 0.2;
+				if ((FPlatformTime::Seconds() - LastArchiveFlushTime) > ArchiveFlushIntervalSec)
+				{
+					Ar.Flush();
+					LastArchiveFlushTime = FPlatformTime::Seconds();
+				}
+			}
+
+			// Decrement the request counter, we now know we serialized at least one request.
+			// We might have serialized more requests but it's irrelevant, the counter will go down to 0 eventually
+			SerializeRequestCounter.Decrement();
 		}
-		// Modify the start pos. Only the worker thread modifies this value so it's ok to not guard it with a critical section.
-		BufferStartPos = ThisThreadEndPos;
-		// Let the client threads know we're done (this is used by FlushBuffer).
-		BufferEmptyEvent->Trigger();
 	}
 
-	/** Flush the memory buffer (doesn't force the archive to flush) */
+	/** [CLIENT THREAD] Flush the memory buffer (doesn't force the archive to flush). Can only be used from inside of BufferPosCritical lock. */
 	void FlushBuffer()
 	{
-		BufferDirtyEvent->Trigger();
+		SerializeRequestCounter.Increment();
 		if (!Thread)
 		{
 			SerializeBufferToArchive();
 		}
-		BufferEmptyEvent->Wait();
+		while (SerializeRequestCounter.GetValue() != 0)
+		{
+			FPlatformProcess::SleepNoStats(0);
+		}
+		// Make sure there's been no unexpected concurrency
+		check(SerializeRequestCounter.GetValue() == 0);
 	}
 
 public:
@@ -83,11 +107,9 @@ public:
 		, Ar(InAr)
 		, BufferStartPos(0)
 		, BufferEndPos(0)
+		, LastArchiveFlushTime(0.0)
 	{
 		Buffer.AddUninitialized(InitialBufferSize);
-		BufferDirtyEvent = FPlatformProcess::GetSynchEventFromPool();
-		BufferEmptyEvent = FPlatformProcess::GetSynchEventFromPool(true);
-		BufferEmptyEvent->Trigger();
 
 		if (FPlatformProcess::SupportsMultithreading())
 		{
@@ -101,13 +123,9 @@ public:
 		Flush();
 		delete Thread;
 		Thread = nullptr;
-		FPlatformProcess::ReturnSynchEventToPool(BufferDirtyEvent);
-		BufferDirtyEvent = nullptr;
-		FPlatformProcess::ReturnSynchEventToPool(BufferEmptyEvent);
-		BufferEmptyEvent = nullptr;
 	}
 
-	/** Serialize data to buffer that will later be saved to disk by the async thread */
+	/** [CLIENT THREAD] Serialize data to buffer that will later be saved to disk by the async thread */
 	virtual void Serialize(void* InData, int64 Length) override
 	{
 		if (!InData || Length <= 0)
@@ -117,7 +135,6 @@ public:
 
 		const uint8* Data = (uint8*)InData;
 
-		BufferEmptyEvent->Reset();
 		FScopeLock WriteLock(&BufferPosCritical);
 
 		// Store the local copy of the current buffer start pos. It may get moved by the worker thread but we don't
@@ -131,7 +148,6 @@ public:
 			if (BufferFreeSize <= Length)
 			{
 				// Force the async thread to call SerializeBufferToArchive even if it's currently empty
-				BufferDirtyEvent->Trigger();
 				FlushBuffer();
 
 				// Resize the buffer if needed
@@ -160,7 +176,7 @@ public:
 
 		// Update the end position and let the async thread know we need to write to disk
 		BufferEndPos = (BufferEndPos + Length) % Buffer.Num();
-		BufferDirtyEvent->Trigger();
+		SerializeRequestCounter.Increment();
 
 		// No async thread? Serialize now.
 		if (!Thread)
@@ -174,6 +190,8 @@ public:
 	{
 		FScopeLock WriteLock(&BufferPosCritical);
 		FlushBuffer();
+		// At this point the serialize queue should be empty and the writer thread should no longer
+		// be accessing the Archive so we should be safe to flush it from here.
 		Ar.Flush();
 	}
 
@@ -186,10 +204,14 @@ public:
 	{
 		while (StopTaskCounter.GetValue() == 0)
 		{
-			if (BufferDirtyEvent->Wait(10))
+			if (SerializeRequestCounter.GetValue() > 0)
 			{
 				SerializeBufferToArchive();
 			}
+			else
+			{
+				FPlatformProcess::Sleep(0.01f);
+			}
 		}
 		return 0;
 	}
diff --git a/Engine/Source/Runtime/Core/Private/Misc/Parse.cpp b/Engine/Source/Runtime/Core/Private/Misc/Parse.cpp
--- a/Engine/Source/Runtime/Core/Private/Misc/Parse.cpp
+++ b/Engine/Source/Runtime/Core/Private/Misc/Parse.cpp
@@ -254,6 +254,9 @@ bool FParse::Value(
 	}
 	else
 	{
+		// Skip initial whitespace
+		Start += FCString::Strspn(Start, TEXT(" \r\n\t"));
+
 		// Non-quoted string without spaces.
 		FCString::Strncpy( Value, Start, MaxLen );
 		Value[MaxLen-1]=0;
diff --git a/Engine/Source/Runtime/Core/Private/Modules/ModuleManager.cpp b/Engine/Source/Runtime/Core/Private/Modules/ModuleManager.cpp
--- a/Engine/Source/Runtime/Core/Private/Modules/ModuleManager.cpp
+++ b/Engine/Source/Runtime/Core/Private/Modules/ModuleManager.cpp
@@ -48,25 +48,30 @@ FModuleManager& FModuleManager::Get()
 
 	if( ModuleManager == NULL )
 	{
-		// FModuleManager is not thread-safe
-		ensure(IsInGameThread());
+		static FCriticalSection FModuleManagerSingletonConstructor;
+		FScopeLock Guard(&FModuleManagerSingletonConstructor);
+		if (ModuleManager == NULL)
+		{
+			// FModuleManager is not thread-safe
+			ensure(IsInGameThread());
 
-		ModuleManager = new FModuleManager();
+			ModuleManager = new FModuleManager();
 
-		//temp workaround for IPlatformFile being used for FPaths::DirectoryExists before main() sets up the commandline.
+			//temp workaround for IPlatformFile being used for FPaths::DirectoryExists before main() sets up the commandline.
 #if PLATFORM_DESKTOP
 		// Ensure that dependency dlls can be found in restricted sub directories
-		const TCHAR* RestrictedFolderNames[] = { TEXT("NoRedist"), TEXT("NotForLicensees"), TEXT("CarefullyRedist") };
-		FString ModuleDir = FPlatformProcess::GetModulesDirectory();
-		for (const TCHAR* RestrictedFolderName : RestrictedFolderNames)
-		{
-			FString RestrictedFolder = ModuleDir / RestrictedFolderName;
-			if (FPaths::DirectoryExists(RestrictedFolder))
+			const TCHAR* RestrictedFolderNames[] = { TEXT("NoRedist"), TEXT("NotForLicensees"), TEXT("CarefullyRedist") };
+			FString ModuleDir = FPlatformProcess::GetModulesDirectory();
+			for (const TCHAR* RestrictedFolderName : RestrictedFolderNames)
 			{
-				ModuleManager->AddBinariesDirectory(*RestrictedFolder, false);
+				FString RestrictedFolder = ModuleDir / RestrictedFolderName;
+				if (FPaths::DirectoryExists(RestrictedFolder))
+				{
+					ModuleManager->AddBinariesDirectory(*RestrictedFolder, false);
+				}
 			}
-		}
 #endif
+		}
 	}
 
 	return *ModuleManager;
diff --git a/Engine/Source/Runtime/Core/Private/ProfilingDebugging/LoadTimeTracker.cpp b/Engine/Source/Runtime/Core/Private/ProfilingDebugging/LoadTimeTracker.cpp
--- a/Engine/Source/Runtime/Core/Private/ProfilingDebugging/LoadTimeTracker.cpp
+++ b/Engine/Source/Runtime/Core/Private/ProfilingDebugging/LoadTimeTracker.cpp
@@ -96,11 +96,11 @@ void FLoadTimeTracker::DumpRawLoadTimes() const
 	UE_LOG(LogStreaming, Display, TEXT("LinkerLoad_SerializeGatherableTextDataMap: %f"), LinkerLoad_SerializeGatherableTextDataMap);
 	UE_LOG(LogStreaming, Display, TEXT("LinkerLoad_SerializeImportMap: %f"), LinkerLoad_SerializeImportMap);
 	UE_LOG(LogStreaming, Display, TEXT("LinkerLoad_SerializeExportMap: %f"), LinkerLoad_SerializeExportMap);
-	UE_LOG(LogStreaming, Display, TEXT("LinkerLoad_StartTextureAllocation: %f"), LinkerLoad_StartTextureAllocation);
 	UE_LOG(LogStreaming, Display, TEXT("LinkerLoad_FixupImportMap: %f"), LinkerLoad_FixupImportMap);
 	UE_LOG(LogStreaming, Display, TEXT("LinkerLoad_RemapImports: %f"), LinkerLoad_RemapImports);
 	UE_LOG(LogStreaming, Display, TEXT("LinkerLoad_FixupExportMap: %f"), LinkerLoad_FixupExportMap);
 	UE_LOG(LogStreaming, Display, TEXT("LinkerLoad_SerializeDependsMap: %f"), LinkerLoad_SerializeDependsMap);
+	UE_LOG(LogStreaming, Display, TEXT("LinkerLoad_SerializePreloadDependencies: %f"), LinkerLoad_SerializePreloadDependencies);
 	UE_LOG(LogStreaming, Display, TEXT("LinkerLoad_CreateExportHash: %f"), LinkerLoad_CreateExportHash);
 	UE_LOG(LogStreaming, Display, TEXT("LinkerLoad_FindExistingExports: %f"), LinkerLoad_FindExistingExports);
 	UE_LOG(LogStreaming, Display, TEXT("LinkerLoad_FinalizeCreation: %f"), LinkerLoad_FinalizeCreation);
@@ -109,12 +109,12 @@ void FLoadTimeTracker::DumpRawLoadTimes() const
 	UE_LOG(LogStreaming, Display, TEXT("Package_LoadImports: %f"), Package_LoadImports);
 	UE_LOG(LogStreaming, Display, TEXT("Package_CreateImports: %f"), Package_CreateImports);
 	UE_LOG(LogStreaming, Display, TEXT("Package_CreateLinker: %f"), Package_CreateLinker);
-	UE_LOG(LogStreaming, Display, TEXT("Package_FinishTextureAllocations: %f"), Package_FinishTextureAllocations);
 	UE_LOG(LogStreaming, Display, TEXT("Package_CreateExports: %f"), Package_CreateExports);
 	UE_LOG(LogStreaming, Display, TEXT("Package_PreLoadObjects: %f"), Package_PreLoadObjects);
 	UE_LOG(LogStreaming, Display, TEXT("Package_PostLoadObjects: %f"), Package_PostLoadObjects);
 	UE_LOG(LogStreaming, Display, TEXT("Package_Tick: %f"), Package_Tick);
 	UE_LOG(LogStreaming, Display, TEXT("Package_CreateAsyncPackagesFromQueue: %f"), Package_CreateAsyncPackagesFromQueue);
+	UE_LOG(LogStreaming, Display, TEXT("Package_EventIOWait: %f"), Package_EventIOWait);
 
 	UE_LOG(LogStreaming, Display, TEXT("TickAsyncLoading_ProcessLoadedPackages: %f"), TickAsyncLoading_ProcessLoadedPackages);
 
@@ -123,7 +123,13 @@ void FLoadTimeTracker::DumpRawLoadTimes() const
 	UE_LOG(LogStreaming, Display, TEXT("Package_Temp3: %f"), Package_Temp3);
 	UE_LOG(LogStreaming, Display, TEXT("Package_Temp4: %f"), Package_Temp4);
 
-
+	UE_LOG(LogStreaming, Display, TEXT("Graph_AddNode: %f     %u"), Graph_AddNode, Graph_AddNodeCnt);
+	UE_LOG(LogStreaming, Display, TEXT("Graph_AddArc: %f     %u"), Graph_AddArc, Graph_AddArcCnt);
+	UE_LOG(LogStreaming, Display, TEXT("Graph_RemoveNode: %f     %u"), Graph_RemoveNode, Graph_RemoveNodeCnt);
+	UE_LOG(LogStreaming, Display, TEXT("Graph_RemoveNodeFire: %f     %u"), Graph_RemoveNodeFire, Graph_RemoveNodeFireCnt);
+	UE_LOG(LogStreaming, Display, TEXT("Graph_DoneAddingPrerequistesFireIfNone: %f     %u"), Graph_DoneAddingPrerequistesFireIfNone, Graph_DoneAddingPrerequistesFireIfNoneCnt);
+	UE_LOG(LogStreaming, Display, TEXT("Graph_DoneAddingPrerequistesFireIfNoneFire: %f     %u"), Graph_DoneAddingPrerequistesFireIfNoneFire, Graph_DoneAddingPrerequistesFireIfNoneFireCnt);
+	UE_LOG(LogStreaming, Display, TEXT("Graph_Misc: %f     %u"), Graph_Misc, Graph_MiscCnt);
 	UE_LOG(LogStreaming, Display, TEXT("-------------------------------------------------"));
 
 #endif
@@ -151,17 +157,19 @@ void FLoadTimeTracker::ResetRawLoadTimes()
 	MaterialSerializeTime = 0.0;
 	MaterialInstanceSerializeTime = 0.0;
 	AsyncLoadingTime = 0.0;
+	CreateMetaDataTime = 0.0;
+
 	LinkerLoad_CreateLoader = 0.0;
 	LinkerLoad_SerializePackageFileSummary = 0.0;
 	LinkerLoad_SerializeNameMap = 0.0;
 	LinkerLoad_SerializeGatherableTextDataMap = 0.0;
 	LinkerLoad_SerializeImportMap = 0.0;
 	LinkerLoad_SerializeExportMap = 0.0;
-	LinkerLoad_StartTextureAllocation = 0.0;
 	LinkerLoad_FixupImportMap = 0.0;
 	LinkerLoad_RemapImports = 0.0;
 	LinkerLoad_FixupExportMap = 0.0;
 	LinkerLoad_SerializeDependsMap = 0.0;
+	LinkerLoad_SerializePreloadDependencies = 0.0;
 	LinkerLoad_CreateExportHash = 0.0;
 	LinkerLoad_FindExistingExports = 0.0;
 	LinkerLoad_FinalizeCreation = 0.0;
@@ -170,18 +178,40 @@ void FLoadTimeTracker::ResetRawLoadTimes()
 	Package_LoadImports = 0.0;
 	Package_CreateImports = 0.0;
 	Package_CreateLinker = 0.0;
-	Package_FinishTextureAllocations = 0.0;
 	Package_CreateExports = 0.0;
 	Package_PreLoadObjects = 0.0;
 	Package_PostLoadObjects = 0.0;
 	Package_Tick = 0.0;
 	Package_CreateAsyncPackagesFromQueue = 0.0;
+	Package_CreateMetaData = 0.0;
+	Package_EventIOWait = 0.0;
 
 	Package_Temp1 = 0.0;
 	Package_Temp2 = 0.0;
 	Package_Temp3 = 0.0;
 	Package_Temp4 = 0.0;
 
+	Graph_AddNode = 0.0;
+	Graph_AddNodeCnt = 0;
+
+	Graph_AddArc = 0.0;
+	Graph_AddArcCnt = 0;
+
+	Graph_RemoveNode = 0.0;
+	Graph_RemoveNodeCnt = 0;
+
+	Graph_RemoveNodeFire = 0.0;
+	Graph_RemoveNodeFireCnt = 0;
+
+	Graph_DoneAddingPrerequistesFireIfNone = 0.0;
+	Graph_DoneAddingPrerequistesFireIfNoneCnt = 0;
+
+	Graph_DoneAddingPrerequistesFireIfNoneFire = 0.0;
+	Graph_DoneAddingPrerequistesFireIfNoneFireCnt = 0;
+
+	Graph_Misc = 0.0;
+	Graph_MiscCnt = 0;
+
 	TickAsyncLoading_ProcessLoadedPackages = 0.0;
 
 	LinkerLoad_SerializeNameMap_ProcessingEntries = 0.0;
@@ -195,3 +225,8 @@ static FAutoConsoleCommand LoadTimerDumpCmd(
 	TEXT("Dump high level load times being tracked"),
 	FConsoleCommandDelegate::CreateStatic(&FLoadTimeTracker::DumpHighLevelLoadTimesStatic)
 	);
+static FAutoConsoleCommand LoadTimerDumpLowCmd(
+	TEXT("LoadTimes.DumpTrackingLow"),
+	TEXT("Dump low level load times being tracked"),
+	FConsoleCommandDelegate::CreateStatic(&FLoadTimeTracker::DumpRawLoadTimesStatic)
+	);
diff --git a/Engine/Source/Runtime/Core/Private/Serialization/Archive.cpp b/Engine/Source/Runtime/Core/Private/Serialization/Archive.cpp
--- a/Engine/Source/Runtime/Core/Private/Serialization/Archive.cpp
+++ b/Engine/Source/Runtime/Core/Private/Serialization/Archive.cpp
@@ -18,6 +18,9 @@
 
 FArchive::FArchive()
 {
+#if DEVIRTUALIZE_FLinkerLoad_Serialize
+	ActiveFPLB = &InlineFPLB;
+#endif
 	CustomVersionContainer = new FCustomVersionContainer;
 
 #if USE_STABLE_LOCALIZATION_KEYS
@@ -29,6 +32,9 @@ FArchive::FArchive()
 
 FArchive::FArchive(const FArchive& ArchiveToCopy)
 {
+#if DEVIRTUALIZE_FLinkerLoad_Serialize
+	ActiveFPLB = &InlineFPLB;
+#endif
 #if USE_STABLE_LOCALIZATION_KEYS
 	LocalizationNamespacePtr = nullptr;
 #endif // USE_STABLE_LOCALIZATION_KEYS
@@ -43,13 +49,16 @@ FArchive::FArchive(const FArchive& ArchiveToCopy)
 
 FArchive& FArchive::operator=(const FArchive& ArchiveToCopy)
 {
+#if DEVIRTUALIZE_FLinkerLoad_Serialize
+	ActiveFPLB = &InlineFPLB;
+	ActiveFPLB->Reset();
+#endif
 	CopyTrivialFArchiveStatusMembers(ArchiveToCopy);
 
 	// Don't know why this is set to false, but this is what the original copying code did
 	ArIsFilterEditorOnly  = false;
 
 	*CustomVersionContainer = *ArchiveToCopy.CustomVersionContainer;
-
 	return *this;
 }
 
@@ -65,6 +74,9 @@ FArchive::~FArchive()
 // Resets all of the base archive members
 void FArchive::Reset()
 {
+#if DEVIRTUALIZE_FLinkerLoad_Serialize
+	ActiveFPLB->Reset();
+#endif
 	ArUE4Ver							= GPackageFileUE4Version;
 	ArLicenseeUE4Ver					= GPackageFileLicenseeUE4Version;
 	ArEngineVer							= FEngineVersion::Current();
@@ -240,6 +252,13 @@ FArchive& FArchive::operator<<(struct FStringAssetReference& Value)
 	return *this;
 }
 
+FArchive& FArchive::operator<<(struct FWeakObjectPtr& Value)
+{
+	// The base FArchive does not implement this method. Use FArchiveUOBject instead.
+	UE_LOG(LogSerialization, Fatal, TEXT("FArchive does not support FWeakObjectPtr serialization. Use FArchiveUObject instead."));
+	return *this;
+}
+
 const FCustomVersionContainer& FArchive::GetCustomVersions() const
 {
 	if (bCustomVersionsAreReset)
diff --git a/Engine/Source/Runtime/Core/Private/Serialization/AsyncIOSystemBase.cpp b/Engine/Source/Runtime/Core/Private/Serialization/AsyncIOSystemBase.cpp
--- a/Engine/Source/Runtime/Core/Private/Serialization/AsyncIOSystemBase.cpp
+++ b/Engine/Source/Runtime/Core/Private/Serialization/AsyncIOSystemBase.cpp
@@ -1,10 +1,12 @@
 // Copyright 1998-2016 Epic Games, Inc. All Rights Reserved.
 
+#include "CorePrivatePCH.h"
+#if !USE_NEW_ASYNC_IO
+
 /*=============================================================================
 	AsyncIOSystemBase.h: Base implementation of the async IO system
 =============================================================================*/
 
-#include "CorePrivatePCH.h"
 #include "Serialization/AsyncIOSystemBase.h"
 #include "CompressedChunkInfo.h"
 
@@ -30,15 +32,6 @@ DECLARE_FLOAT_COUNTER_STAT( TEXT( "Bandwidth (MByte/ sec)" ), STAT_AsyncIO_Bandw
 
 #define BLOCK_ON_ASYNCIO 0
 
-// Constrain bandwidth if wanted. Value is in MByte/ sec.
-float GAsyncIOBandwidthLimit = 0.0f;
-static FAutoConsoleVariableRef CVarAsyncIOBandwidthLimit(
-	TEXT("s.AsyncIOBandwidthLimit"),
-	GAsyncIOBandwidthLimit,
-	TEXT("Constrain bandwidth if wanted. Value is in MByte/ sec."),
-	ECVF_Default
-	);
-
 CORE_API bool GbLogAsyncLoading = false;
 CORE_API bool GbLogAsyncTiming = false;
 
@@ -925,3 +918,5 @@ bool FIOSystem::HasShutdown()
 {
 	return AsyncIOThread == nullptr || AsyncIOThread == (FRunnableThread*)-1;
 }
+
+#endif
\ No newline at end of file
diff --git a/Engine/Source/Runtime/Core/Private/Serialization/AsyncIOSystemBase.h b/Engine/Source/Runtime/Core/Private/Serialization/AsyncIOSystemBase.h
--- a/Engine/Source/Runtime/Core/Private/Serialization/AsyncIOSystemBase.h
+++ b/Engine/Source/Runtime/Core/Private/Serialization/AsyncIOSystemBase.h
@@ -5,8 +5,10 @@
 =============================================================================*/
 
 #pragma once
+#if !USE_NEW_ASYNC_IO
 #include "HAL/IOBase.h"
 
+
 /**
  * Base implementation of an async IO system allowing most of the code to be shared across platforms.
  */
@@ -438,3 +440,5 @@ protected:
 	/** Low level file system that we use for our requests.											*/
 	IPlatformFile&					LowLevel;
 };
+
+#endif
\ No newline at end of file
diff --git a/Engine/Source/Runtime/Core/Private/Windows/WindowsPlatformMemory.cpp b/Engine/Source/Runtime/Core/Private/Windows/WindowsPlatformMemory.cpp
--- a/Engine/Source/Runtime/Core/Private/Windows/WindowsPlatformMemory.cpp
+++ b/Engine/Source/Runtime/Core/Private/Windows/WindowsPlatformMemory.cpp
@@ -7,15 +7,13 @@
 #include "MallocStomp.h"
 #include "GenericPlatformMemoryPoolStats.h"
 #include "MemoryMisc.h"
+#include "MallocBinned.h"
+#include "MallocBinned2.h"
 
 #if ENABLE_WIN_ALLOC_TRACKING
 #include <crtdbg.h>
 #endif // ENABLE_WIN_ALLOC_TRACKING
 
-#if !FORCE_ANSI_ALLOCATOR
-#include "MallocBinnedRedirect.h"
-#endif
-
 #include "AllowWindowsPlatformTypes.h"
 #include <Psapi.h>
 #pragma comment(lib, "psapi.lib")
@@ -62,6 +60,9 @@ void FWindowsPlatformMemory::Init()
 	DumpStats( *GLog );
 }
 
+// Set rather to use BinnedMalloc2 for binned malloc, can be overridden below
+#define USE_MALLOC_BINNED2 (1)
+
 FMalloc* FWindowsPlatformMemory::BaseAllocator()
 {
 #if ENABLE_WIN_ALLOC_TRACKING
@@ -71,25 +72,66 @@ FMalloc* FWindowsPlatformMemory::BaseAllocator()
 	_CrtSetAllocHook(WindowsAllocHook);
 #endif // ENABLE_WIN_ALLOC_TRACKING
 
-#if FORCE_ANSI_ALLOCATOR
-	return new FMallocAnsi();
-#elif USE_MALLOC_STOMP
-	// Allocator that helps track down memory stomps.
-	return new FMallocStomp();
-#elif (WITH_EDITORONLY_DATA || IS_PROGRAM) && TBB_ALLOCATOR_ALLOWED
-	// -run is used for commandlets.  Cooker and other commandlets run faster and with less memory using malloc binned.
-	if ((FCString::Stristr(::GetCommandLine(), TEXT("-mallocbinned")) || FCString::Stristr(::GetCommandLine(), TEXT("-run="))) &&
-		(!FCString::Stristr(::GetCommandLine(), TEXT("-malloctbb"))))
+	if (FORCE_ANSI_ALLOCATOR)
+	{
+		AllocatorToUse = EMemoryAllocatorToUse::Ansi;
+	}
+	else if (USE_MALLOC_STOMP)
+	{
+		AllocatorToUse = EMemoryAllocatorToUse::Stomp;
+	}
+	else if ((WITH_EDITORONLY_DATA || IS_PROGRAM) && TBB_ALLOCATOR_ALLOWED)
 	{
-		return new FMallocBinnedRedirect((uint32)(GetConstants().PageSize&MAX_uint32), (uint64)MAX_uint32 + 1);
+		AllocatorToUse = EMemoryAllocatorToUse::TBB;
+	}
+	else if (USE_MALLOC_BINNED2)
+	{
+		AllocatorToUse = EMemoryAllocatorToUse::Binned2;
 	}
 	else
 	{
-		return new FMallocTBB();
-}
-#else
-	return new FMallocBinnedRedirect((uint32)(GetConstants().PageSize&MAX_uint32), (uint64)MAX_uint32 + 1);
+		AllocatorToUse = EMemoryAllocatorToUse::Binned;
+	}
+	
+#if !UE_BUILD_SHIPPING
+	// If not shipping, allow overriding with command line options, this happens very early so we need to use windows functions
+	const TCHAR* CommandLine = ::GetCommandLine();
+
+	if (FCString::Stristr(CommandLine, TEXT("-ansimalloc")))
+	{
+		AllocatorToUse = EMemoryAllocatorToUse::Ansi;
+	}
+	else if (FCString::Stristr(CommandLine, TEXT("-tbbmalloc")))
+	{
+		AllocatorToUse = EMemoryAllocatorToUse::TBB;
+	}
+	else if (FCString::Stristr(CommandLine, TEXT("-binnedmalloc2")))
+	{
+		AllocatorToUse = EMemoryAllocatorToUse::Binned2;
+	}
+	else if (FCString::Stristr(CommandLine, TEXT("-binnedmalloc")))
+	{
+		AllocatorToUse = EMemoryAllocatorToUse::Binned;
+	}
 #endif
+
+	switch (AllocatorToUse)
+	{
+	case EMemoryAllocatorToUse::Ansi:
+		return new FMallocAnsi();
+#if USE_MALLOC_STOMP
+	case EMemoryAllocatorToUse::Stomp:
+		return new FMallocStomp();
+#endif
+	case EMemoryAllocatorToUse::TBB:
+		return new FMallocTBB();
+	case EMemoryAllocatorToUse::Binned2:
+		return new FMallocBinned2();
+		
+	default:	// intentional fall-through
+	case EMemoryAllocatorToUse::Binned:
+		return new FMallocBinned((uint32)(GetConstants().PageSize&MAX_uint32), (uint64)MAX_uint32 + 1);
+	}
 }
 
 FPlatformMemoryStats FWindowsPlatformMemory::GetStats()
diff --git a/Engine/Source/Runtime/Core/Public/Async/AsyncFileHandle.h b/Engine/Source/Runtime/Core/Public/Async/AsyncFileHandle.h
--- a/Engine/Source/Runtime/Core/Public/Async/AsyncFileHandle.h
+++ b/Engine/Source/Runtime/Core/Public/Async/AsyncFileHandle.h
@@ -10,6 +10,10 @@
 
 #include "Function.h"
 
+DECLARE_MEMORY_STAT_EXTERN(TEXT("Async File Handle Memory"), STAT_AsyncFileMemory, STATGROUP_Memory, CORE_API);
+DECLARE_DWORD_ACCUMULATOR_STAT_EXTERN(TEXT("Num Async File Handles"), STAT_AsyncFileHandles, STATGROUP_Memory, CORE_API);
+DECLARE_DWORD_ACCUMULATOR_STAT_EXTERN(TEXT("Num Async File Requests"), STAT_AsyncFileRequests, STATGROUP_Memory, CORE_API);
+
 // Note on threading. Like the rest of the filesystem platform abstraction, these methods are threadsafe, but it is expected you are not concurrently _using_ these data structures. 
 
 class IAsyncReadRequest;
@@ -29,34 +33,36 @@ protected:
 	bool bCompleteSync;
 	bool bCanceled;
 	const bool bSizeRequest;
+	const bool bUserSuppliedMemory;
 public:
 
-	FORCEINLINE IAsyncReadRequest(FAsyncFileCallBack* InCallback, bool bInSizeRequest)
+	FORCEINLINE IAsyncReadRequest(FAsyncFileCallBack* InCallback, bool bInSizeRequest, uint8* UserSuppliedMemory)
 		: Callback(InCallback)
 		, bDataIsReady(false)
 		, bCompleteAndCallbackCalled(false)
 		, bCompleteSync(false)
 		, bCanceled(false)
 		, bSizeRequest(bInSizeRequest)
+		, bUserSuppliedMemory(!!UserSuppliedMemory)
 	{
 		if (bSizeRequest)
 		{
 			Size = -1;
+			check(!UserSuppliedMemory && !bUserSuppliedMemory); // size requests don't do memory
 		}
 		else
 		{
-			Memory = nullptr;
+			Memory = UserSuppliedMemory;
+			check((!!Memory) == bUserSuppliedMemory);
 		}
+		INC_DWORD_STAT(STAT_AsyncFileRequests);
 	}
 
 	/* Not legal to destroy the request until it is complete. */
 	virtual ~IAsyncReadRequest()
 	{
-		check(bCompleteAndCallbackCalled);
-		if (!bSizeRequest && Memory)
-		{
-			FMemory::Free(Memory);
-		}
+		check(bCompleteAndCallbackCalled && (bSizeRequest || !Memory)); // must be complete, and if it was a read request, the memory should be gone
+		DEC_DWORD_STAT(STAT_AsyncFileRequests);
 	}
 
 	/**
@@ -107,7 +113,7 @@ public:
 
 	/**
 	* Return the bytes of a completed read request. Not legal to call unless the request is complete.
-	* @return Returned memory block which if non-null contains the bytes read. Caller owns the memory block and must call FMemory::Free on it when done. Can be null if the file was not found or could not be read or the request was cancelled.
+	* @return Returned memory block which if non-null contains the bytes read. Caller owns the memory block and must call FMemory::Free on it when done. Can be null if the file was not found or could not be read or the request was cancelled, or the request priority was AIOP_Precache.
 	**/
 	FORCEINLINE uint8* GetReadResults()
 	{
@@ -115,10 +121,12 @@ public:
 		uint8* Result = Memory;
 		if (bCanceled && Result)
 		{
-			FMemory::Free(Result);
 			Result = nullptr;
 		}
-		Memory = nullptr;
+		else
+		{
+			Memory = nullptr;
+		}
 		return Result;
 	}
 
@@ -150,9 +158,14 @@ protected:
 class CORE_API IAsyncReadFileHandle
 {
 public:
-	/** Destructor, also the only way to close the file handle. It is not legal to delete an async file with outstanding requests. **/
+	IAsyncReadFileHandle()
+	{
+		INC_DWORD_STAT(STAT_AsyncFileHandles);
+	}
+	/** Destructor, also the only way to close the file handle. It is not legal to delete an async file with outstanding requests. You must always call WaitCompletion before deleting a request. **/
 	virtual ~IAsyncReadFileHandle()
 	{
+		DEC_DWORD_STAT(STAT_AsyncFileHandles);
 	}
 
 	/**
@@ -165,11 +178,12 @@ public:
 	/**
 	* Submit an async request and/or wait for an async request
 	* @param Offset					Offset into the file to start reading.
-	* @param BytesToRead			number of bytes to read.
+	* @param BytesToRead			number of bytes to read. If this request is AIOP_Preache, the size can be anything, even MAX_int64, otherwise the size and offset must be fully contained in the file.
+	* @PAram Priority				Priority of the request. If this is AIOP_Precache, then memory will never be returned. The request should always be canceled and waited for, even for a precache request.
 	* @param CompleteCallback		Called from an arbitrary thread when the request is complete. Can be nullptr, if non-null, must remain valid until it is called. It will always be called.
 	* @return A request for the read. This is owned by the caller and must be deleted by the caller.
 	**/
-	virtual IAsyncReadRequest* ReadRequest(int64 Offset, int64 BytesToRead, EAsyncIOPriority Priority = AIOP_Normal, FAsyncFileCallBack* CompleteCallback = nullptr) = 0;
+	virtual IAsyncReadRequest* ReadRequest(int64 Offset, int64 BytesToRead, EAsyncIOPriority Priority = AIOP_Normal, FAsyncFileCallBack* CompleteCallback = nullptr, uint8* UserSuppliedMemory = nullptr) = 0;
 };
 
 
diff --git a/Engine/Source/Runtime/Core/Public/Async/TaskGraphInterfaces.h b/Engine/Source/Runtime/Core/Public/Async/TaskGraphInterfaces.h
--- a/Engine/Source/Runtime/Core/Public/Async/TaskGraphInterfaces.h
+++ b/Engine/Source/Runtime/Core/Public/Async/TaskGraphInterfaces.h
@@ -339,7 +339,7 @@ public:
 	*	A (slow) function to call a function on every known thread, both named and workers
 	*	@param	Callback - function to call prior to shutting down the taskgraph
 	**/
-	static void BroadcastSlow_OnlyUseForSpecialPurposes(bool bDoBackgroundThreads, TFunction<void(ENamedThreads::Type CurrentThread)>& Callback);
+	static void BroadcastSlow_OnlyUseForSpecialPurposes(bool bDoTaskThreads, bool bDoBackgroundThreads, TFunction<void(ENamedThreads::Type CurrentThread)>& Callback);
 };
 
 /** 
diff --git a/Engine/Source/Runtime/Core/Public/CoreGlobals.h b/Engine/Source/Runtime/Core/Public/CoreGlobals.h
--- a/Engine/Source/Runtime/Core/Public/CoreGlobals.h
+++ b/Engine/Source/Runtime/Core/Public/CoreGlobals.h
@@ -323,9 +323,6 @@ extern CORE_API float GHitchThresholdMS;
 /** Size to break up data into when saving compressed data */
 extern CORE_API int32 GSavingCompressionChunkSize;
 
-/** Whether we are using the seek-free/ cooked loading code path. */
-extern CORE_API bool GUseSeekFreeLoading;
-
 /** Thread ID of the main/game thread */
 extern CORE_API uint32 GGameThreadId;
 
diff --git a/Engine/Source/Runtime/Core/Public/GenericPlatform/GenericPlatformFile.h b/Engine/Source/Runtime/Core/Public/GenericPlatform/GenericPlatformFile.h
--- a/Engine/Source/Runtime/Core/Public/GenericPlatform/GenericPlatformFile.h
+++ b/Engine/Source/Runtime/Core/Public/GenericPlatform/GenericPlatformFile.h
@@ -13,7 +13,15 @@ class FArchive;
 class FString;
 struct FDateTime;
 
-#define USE_NEW_ASYNC_IO (0)
+#if !defined(WITH_EDITORONLY_DATA)
+#error "WITH_EDITORONLY_DATA must be defined"
+#endif
+
+#if !defined(USE_NEW_ASYNC_IO)
+#define USE_NEW_ASYNC_IO 0
+#elif USE_NEW_ASYNC_IO && WITH_EDITORONLY_DATA
+#error USE_NEW_ASYNC_IO can not be used with WITH_EDITORONLY_DATA
+#endif
 
 #if USE_NEW_ASYNC_IO
 class IAsyncReadFileHandle;
diff --git a/Engine/Source/Runtime/Core/Public/GenericPlatform/GenericPlatformMemory.h b/Engine/Source/Runtime/Core/Public/GenericPlatform/GenericPlatformMemory.h
--- a/Engine/Source/Runtime/Core/Public/GenericPlatform/GenericPlatformMemory.h
+++ b/Engine/Source/Runtime/Core/Public/GenericPlatform/GenericPlatformMemory.h
@@ -150,6 +150,21 @@ struct CORE_API FGenericPlatformMemory
 		MCR_MAX
 	};
 
+	/** Which allocator is being used */
+	enum EMemoryAllocatorToUse
+	{
+		Ansi, // Default C allocator
+		Stomp, // Allocator to check for memory stomping
+		TBB, // Thread Building Blocks malloc
+		Jemalloc, // Linux/FreeBSD malloc
+		Binned, // Older binned malloc
+		Binned2, // Newer binned malloc
+		Platform, // Custom platform specific allocator
+	};
+
+	/** Current allocator */
+	static EMemoryAllocatorToUse AllocatorToUse;
+
 	/**
 	 * Flags used for shared memory creation/open
 	 */
diff --git a/Engine/Source/Runtime/Core/Public/HAL/IOBase.h b/Engine/Source/Runtime/Core/Public/HAL/IOBase.h
--- a/Engine/Source/Runtime/Core/Public/HAL/IOBase.h
+++ b/Engine/Source/Runtime/Core/Public/HAL/IOBase.h
@@ -1,8 +1,7 @@
 // Copyright 1998-2016 Epic Games, Inc. All Rights Reserved.
 
 #pragma once
-
-
+#if !USE_NEW_ASYNC_IO
 /*-----------------------------------------------------------------------------
 	FIOSystem.
 -----------------------------------------------------------------------------*/
@@ -133,3 +132,4 @@ struct CORE_API FIOSystem
 	virtual void FlushLog() = 0;
 };
 
+#endif
\ No newline at end of file
diff --git a/Engine/Source/Runtime/Core/Public/HAL/Platform.h b/Engine/Source/Runtime/Core/Public/HAL/Platform.h
--- a/Engine/Source/Runtime/Core/Public/HAL/Platform.h
+++ b/Engine/Source/Runtime/Core/Public/HAL/Platform.h
@@ -130,6 +130,9 @@
 #ifndef PLATFORM_LITTLE_ENDIAN
 	#define PLATFORM_LITTLE_ENDIAN				0
 #endif
+#ifndef PLATFORM_SUPPORTS_UNALIGNED_INT_LOADS
+	#define PLATFORM_SUPPORTS_UNALIGNED_INT_LOADS	0
+#endif
 #ifndef PLATFORM_EXCEPTIONS_DISABLED
 	#define PLATFORM_EXCEPTIONS_DISABLED		!PLATFORM_DESKTOP
 #endif
@@ -508,57 +511,6 @@
 #define private_subobject public
 #endif
 
-// explicit bool support
-namespace FHasOperatorImpl
-{
-	struct FNotSpecified {};
-
-	template <typename T>
-	struct FReturnValueCheck
-	{
-		static char (&Func())[2];
-	};
-
-	template <>
-	struct FReturnValueCheck<FNotSpecified>
-	{
-		static char (&Func())[1];
-	};
-
-	template <typename T>
-	FNotSpecified operator==(const T&, const T&);
-
-	template <typename T>
-	FNotSpecified operator!=(const T&, const T&);
-
-	template <typename T>
-	const T& Make();
-
-	template <typename T>
-	struct Equals
-	{
-		enum { Value = sizeof(FReturnValueCheck<decltype(Make<T>() == Make<T>())>::Func()) == sizeof(char[2]) };
-	};
-
-	template <typename T>
-	struct NotEquals
-	{
-		enum { Value = sizeof(FReturnValueCheck<decltype(Make<T>() != Make<T>())>::Func()) == sizeof(char[2]) };
-	};
-}
-
-template <typename T>
-struct THasOperatorEquals
-{
-	enum { Value = FHasOperatorImpl::Equals<T>::Value };
-};
-
-template <typename T>
-struct THasOperatorNotEquals
-{
-	enum { Value = FHasOperatorImpl::NotEquals<T>::Value };
-};
-
 // Console ANSICHAR/TCHAR command line handling
 #if PLATFORM_COMPILER_HAS_TCHAR_WMAIN
 #define INT32_MAIN_INT32_ARGC_TCHAR_ARGV() int32 wmain(int32 ArgC, TCHAR* ArgV[])
diff --git a/Engine/Source/Runtime/Core/Public/Linux/LinuxPlatform.h b/Engine/Source/Runtime/Core/Public/Linux/LinuxPlatform.h
--- a/Engine/Source/Runtime/Core/Public/Linux/LinuxPlatform.h
+++ b/Engine/Source/Runtime/Core/Public/Linux/LinuxPlatform.h
@@ -34,6 +34,7 @@ typedef FLinuxPlatformTypes FPlatformTypes;
 // Base defines, defaults are commented out
 
 #define PLATFORM_LITTLE_ENDIAN						1
+#define PLATFORM_SUPPORTS_UNALIGNED_INT_LOADS		1
 #define PLATFORM_COMPILER_DISTINGUISHES_INT_AND_LONG 1
 #define PLATFORM_SUPPORTS_PRAGMA_PACK				1
 #define PLATFORM_USE_LS_SPEC_FOR_WIDECHAR			1
diff --git a/Engine/Source/Runtime/Core/Public/Misc/CString.h b/Engine/Source/Runtime/Core/Public/Misc/CString.h
--- a/Engine/Source/Runtime/Core/Public/Misc/CString.h
+++ b/Engine/Source/Runtime/Core/Public/Misc/CString.h
@@ -651,12 +651,15 @@ int32 TCString<T>::Strspn( const CharType* String, const CharType* Mask )
 	{
 		for (const TCHAR* MaskIt = Mask; *MaskIt; ++MaskIt)
 		{
-			if (*StringIt != *MaskIt)
+			if (*StringIt == *MaskIt)
 			{
-				return StringIt - String;
+				goto NextChar;
 			}
 		}
 
+		return StringIt - String;
+
+	NextChar:
 		++StringIt;
 	}
 
diff --git a/Engine/Source/Runtime/Core/Public/ProfilingDebugging/LoadTimeTracker.h b/Engine/Source/Runtime/Core/Public/ProfilingDebugging/LoadTimeTracker.h
--- a/Engine/Source/Runtime/Core/Public/ProfilingDebugging/LoadTimeTracker.h
+++ b/Engine/Source/Runtime/Core/Public/ProfilingDebugging/LoadTimeTracker.h
@@ -5,7 +5,7 @@
 */
 
 #pragma once
-#define ENABLE_LOADTIME_TRACKING 1
+#define ENABLE_LOADTIME_TRACKING 0
 #define ENABLE_LOADTIME_RAW_TIMINGS 0
 
 #include "ScopedTimers.h"
@@ -41,6 +41,11 @@ public:
 	/** Prints out raw load times for individual timers */
 	void DumpRawLoadTimes() const;
 
+	static void DumpRawLoadTimesStatic()
+	{
+		Get().DumpRawLoadTimes();
+	}
+
 	void ResetRawLoadTimes();
 
 #if ENABLE_LOADTIME_RAW_TIMINGS
@@ -64,6 +69,7 @@ public:
 	double MaterialSerializeTime;
 	double MaterialInstanceSerializeTime;
 	double AsyncLoadingTime;
+	double CreateMetaDataTime;
 
 	double LinkerLoad_CreateLoader;
 	double LinkerLoad_SerializePackageFileSummary;
@@ -71,11 +77,11 @@ public:
 	double LinkerLoad_SerializeGatherableTextDataMap;
 	double LinkerLoad_SerializeImportMap;
 	double LinkerLoad_SerializeExportMap;
-	double LinkerLoad_StartTextureAllocation;
 	double LinkerLoad_FixupImportMap;
 	double LinkerLoad_RemapImports;
 	double LinkerLoad_FixupExportMap;
 	double LinkerLoad_SerializeDependsMap;
+	double LinkerLoad_SerializePreloadDependencies;
 	double LinkerLoad_CreateExportHash;
 	double LinkerLoad_FindExistingExports;
 	double LinkerLoad_FinalizeCreation;
@@ -84,18 +90,40 @@ public:
 	double Package_LoadImports;
 	double Package_CreateImports;
 	double Package_CreateLinker;
-	double Package_FinishTextureAllocations;
 	double Package_CreateExports;
 	double Package_PreLoadObjects;
 	double Package_PostLoadObjects;
 	double Package_Tick;
 	double Package_CreateAsyncPackagesFromQueue;
+	double Package_CreateMetaData;
+	double Package_EventIOWait;
 
 	double Package_Temp1;
 	double Package_Temp2;
 	double Package_Temp3;
 	double Package_Temp4;
 
+	double Graph_AddNode;
+	uint32 Graph_AddNodeCnt;
+
+	double Graph_AddArc;
+	uint32 Graph_AddArcCnt;
+
+	double Graph_RemoveNode;
+	uint32 Graph_RemoveNodeCnt;
+
+	double Graph_RemoveNodeFire;
+	uint32 Graph_RemoveNodeFireCnt;
+
+	double Graph_DoneAddingPrerequistesFireIfNone;
+	uint32 Graph_DoneAddingPrerequistesFireIfNoneCnt;
+
+	double Graph_DoneAddingPrerequistesFireIfNoneFire;
+	uint32 Graph_DoneAddingPrerequistesFireIfNoneFireCnt;
+
+	double Graph_Misc;
+	uint32 Graph_MiscCnt;
+
 	double TickAsyncLoading_ProcessLoadedPackages;
 
 
@@ -116,6 +144,8 @@ private:
 
 #if ENABLE_LOADTIME_RAW_TIMINGS
 #define SCOPED_LOADTIMER(TimerName) FScopedDurationTimer DurationTimer_##TimerName(FLoadTimeTracker::Get().TimerName);
+#define SCOPED_LOADTIMER_CNT(TimerName) FScopedDurationTimer DurationTimer_##TimerName(FLoadTimeTracker::Get().TimerName); FLoadTimeTracker::Get().TimerName##Cnt++;
 #else
 #define SCOPED_LOADTIMER(TimerName)
+#define SCOPED_LOADTIMER_CNT(TimerName)
 #endif
\ No newline at end of file
diff --git a/Engine/Source/Runtime/Core/Public/Serialization/Archive.h b/Engine/Source/Runtime/Core/Public/Serialization/Archive.h
--- a/Engine/Source/Runtime/Core/Public/Serialization/Archive.h
+++ b/Engine/Source/Runtime/Core/Public/Serialization/Archive.h
@@ -15,6 +15,21 @@ class FCustomVersionContainer;
 class FLazyObjectPtr;
 struct FStringAssetReference;
 struct FUntypedBulkData;
+struct FWeakObjectPtr;
+
+// this is the master switch
+//@todoio if this is off, then we should leave the package file format completely unchanged....!!!!! this is really important to fix before we merge to main
+//#define COOK_FOR_EVENT_DRIVEN_LOAD (1)
+
+#define SPLIT_COOKED_FILES (1)
+
+#if !defined(USE_NEW_ASYNC_IO)
+#error "USE_NEW_ASYNC_IO must be defined"
+#endif
+
+// We only use event driven loading if we have cooked for it and it is a cooked platform
+#define USE_EVENT_DRIVEN_ASYNC_LOAD (!WITH_EDITORONLY_DATA && USE_NEW_ASYNC_IO)
+#define DEVIRTUALIZE_FLinkerLoad_Serialize (USE_EVENT_DRIVEN_ASYNC_LOAD)
 
 /**
  * Base class for archives that can be used for loading, saving, and garbage
@@ -104,6 +119,16 @@ public:
 	 */
 	virtual FArchive& operator<<(struct FStringAssetReference& Value);
 
+	/**
+	* Serializes FWeakObjectPtr value from or into this archive.
+	*
+	* This operator can be implemented by sub-classes that wish to serialize FWeakObjectPtr instances.
+	*
+	* @param Value The value to serialize.
+	* @return This instance.
+	*/
+	virtual FArchive& operator<<(struct FWeakObjectPtr& Value);
+
 	/** 
 	 * Inform the archive that a blueprint would like to force finalization, normally
 	 * this is triggered by CDO load, but if there's no CDO we force finalization.
@@ -119,8 +144,12 @@ public:
 	 */
 	FORCEINLINE friend FArchive& operator<<(FArchive& Ar, ANSICHAR& Value)
 	{
-		Ar.Serialize(&Value, 1);
-
+#if DEVIRTUALIZE_FLinkerLoad_Serialize
+		if (!Ar.FastPathLoad<sizeof(Value)>(&Value))
+#endif
+		{
+			Ar.Serialize(&Value, 1);
+		}
 		return Ar;
 	}
 
@@ -132,8 +161,12 @@ public:
 	 */
 	FORCEINLINE friend FArchive& operator<<(FArchive& Ar, WIDECHAR& Value)
 	{
-		Ar.ByteOrderSerialize(&Value, sizeof(Value));
-
+#if DEVIRTUALIZE_FLinkerLoad_Serialize
+		if (!Ar.FastPathLoad<sizeof(Value)>(&Value))
+#endif
+		{
+			Ar.ByteOrderSerialize(&Value, sizeof(Value));
+		}
 		return Ar;
 	}
 
@@ -145,8 +178,12 @@ public:
 	 */
 	FORCEINLINE friend FArchive& operator<<(FArchive& Ar, uint8& Value)
 	{
-		Ar.Serialize(&Value, 1);
-
+#if DEVIRTUALIZE_FLinkerLoad_Serialize
+		if (!Ar.FastPathLoad<sizeof(Value)>(&Value))
+#endif
+		{
+			Ar.Serialize(&Value, 1);
+		}
 		return Ar;
 	}
 
@@ -159,8 +196,12 @@ public:
 	template<class TEnum>
 	FORCEINLINE friend FArchive& operator<<(FArchive& Ar, TEnumAsByte<TEnum>& Value)
 	{
-		Ar.Serialize(&Value, 1);
-
+#if DEVIRTUALIZE_FLinkerLoad_Serialize
+		if (!Ar.FastPathLoad<sizeof(Value)>(&Value))
+#endif
+		{
+			Ar.Serialize(&Value, 1);
+		}
 		return Ar;
 	}
 
@@ -172,8 +213,12 @@ public:
 	 */
 	FORCEINLINE friend FArchive& operator<<(FArchive& Ar, int8& Value)
 	{
-		Ar.Serialize(&Value, 1);
-
+#if DEVIRTUALIZE_FLinkerLoad_Serialize
+		if (!Ar.FastPathLoad<sizeof(Value)>(&Value))
+#endif
+		{
+			Ar.Serialize(&Value, 1);
+		}
 		return Ar;
 	}
 
@@ -185,8 +230,12 @@ public:
 	 */
 	FORCEINLINE friend FArchive& operator<<(FArchive& Ar, uint16& Value)
 	{
-		Ar.ByteOrderSerialize(&Value, sizeof(Value));
-
+#if DEVIRTUALIZE_FLinkerLoad_Serialize
+		if (!Ar.FastPathLoad<sizeof(Value)>(&Value))
+#endif
+		{
+			Ar.ByteOrderSerialize(&Value, sizeof(Value));
+		}
 		return Ar;
 	}
 
@@ -198,8 +247,12 @@ public:
 	 */
 	FORCEINLINE friend FArchive& operator<<(FArchive& Ar, int16& Value)
 	{
-		Ar.ByteOrderSerialize(&Value, sizeof(Value));
-
+#if DEVIRTUALIZE_FLinkerLoad_Serialize
+		if (!Ar.FastPathLoad<sizeof(Value)>(&Value))
+#endif
+		{
+			Ar.ByteOrderSerialize(&Value, sizeof(Value));
+		}
 		return Ar;
 	}
 
@@ -211,8 +264,12 @@ public:
 	 */
 	FORCEINLINE friend FArchive& operator<<(FArchive& Ar, uint32& Value)
 	{
-		Ar.ByteOrderSerialize(&Value, sizeof(Value));
-
+#if DEVIRTUALIZE_FLinkerLoad_Serialize
+		if (!Ar.FastPathLoad<sizeof(Value)>(&Value))
+#endif
+		{
+			Ar.ByteOrderSerialize(&Value, sizeof(Value));
+		}
 		return Ar;
 	}
 
@@ -225,9 +282,26 @@ public:
 	FORCEINLINE friend FArchive& operator<<( FArchive& Ar, bool& D )
 	{
 		// Serialize bool as if it were UBOOL (legacy, 32 bit int).
-		uint32 OldUBoolValue = D ? 1 : 0;
-		Ar.ByteOrderSerialize( &OldUBoolValue, sizeof(OldUBoolValue) );
-		D = !!OldUBoolValue;
+
+#if DEVIRTUALIZE_FLinkerLoad_Serialize
+		const uint8 * RESTRICT Src = Ar.ActiveFPLB->StartFastPathLoadBuffer;
+		if (Src + sizeof(uint32) <= Ar.ActiveFPLB->EndFastPathLoadBuffer)
+		{
+#if PLATFORM_SUPPORTS_UNALIGNED_INT_LOADS
+			D = !!*(uint32* RESTRICT)Src;
+#else
+			static_assert(sizeof(uint32) == 4, "assuming sizeof(uint32) == 4");
+			D = !!(Src[0] | Src[1] | Src[2] | Src[3]);
+#endif
+			Ar.ActiveFPLB->StartFastPathLoadBuffer += 4;
+		}
+		else
+#endif
+		{
+			uint32 OldUBoolValue = D ? 1 : 0;
+			Ar.Serialize(&OldUBoolValue, sizeof(OldUBoolValue));
+			D = !!OldUBoolValue;
+		}
 		return Ar;
 	}
 
@@ -239,8 +313,12 @@ public:
 	 */
 	FORCEINLINE friend FArchive& operator<<(FArchive& Ar, int32& Value)
 	{
-		Ar.ByteOrderSerialize(&Value, sizeof(Value));
-
+#if DEVIRTUALIZE_FLinkerLoad_Serialize
+		if (!Ar.FastPathLoad<sizeof(Value)>(&Value))
+#endif
+		{
+			Ar.ByteOrderSerialize(&Value, sizeof(Value));
+		}
 		return Ar;
 	}
 
@@ -253,8 +331,12 @@ public:
 	 */
 	FORCEINLINE friend FArchive& operator<<(FArchive& Ar, long& Value)
 	{
-		Ar.ByteOrderSerialize(&Value, sizeof(Value));
-
+#if DEVIRTUALIZE_FLinkerLoad_Serialize
+		if (!Ar.FastPathLoad<sizeof(Value)>(&Value))
+#endif
+		{
+			Ar.ByteOrderSerialize(&Value, sizeof(Value));
+		}
 		return Ar;
 	}	
 #endif
@@ -267,8 +349,12 @@ public:
 	 */
 	FORCEINLINE friend FArchive& operator<<( FArchive& Ar, float& Value)
 	{
-		Ar.ByteOrderSerialize(&Value, sizeof(Value));
-
+#if DEVIRTUALIZE_FLinkerLoad_Serialize
+		if (!Ar.FastPathLoad<sizeof(Value)>(&Value))
+#endif
+		{
+			Ar.ByteOrderSerialize(&Value, sizeof(Value));
+		}
 		return Ar;
 	}
 
@@ -280,8 +366,12 @@ public:
 	 */
 	FORCEINLINE friend FArchive& operator<<(FArchive& Ar, double& Value)
 	{
-		Ar.ByteOrderSerialize(&Value, sizeof(Value));
-
+#if DEVIRTUALIZE_FLinkerLoad_Serialize
+		if (!Ar.FastPathLoad<sizeof(Value)>(&Value))
+#endif
+		{
+			Ar.ByteOrderSerialize(&Value, sizeof(Value));
+		}
 		return Ar;
 	}
 
@@ -293,8 +383,12 @@ public:
 	 */
 	FORCEINLINE friend FArchive& operator<<(FArchive &Ar, uint64& Value)
 	{
-		Ar.ByteOrderSerialize(&Value, sizeof(Value));
-
+#if DEVIRTUALIZE_FLinkerLoad_Serialize
+		if (!Ar.FastPathLoad<sizeof(Value)>(&Value))
+#endif
+		{
+			Ar.ByteOrderSerialize(&Value, sizeof(Value));
+		}
 		return Ar;
 	}
 
@@ -306,8 +400,12 @@ public:
 	 */
 	/*FORCEINLINE*/friend FArchive& operator<<(FArchive& Ar, int64& Value)
 	{
-		Ar.ByteOrderSerialize(&Value, sizeof(Value));
-
+#if DEVIRTUALIZE_FLinkerLoad_Serialize
+		if (!Ar.FastPathLoad<sizeof(Value)>(&Value))
+#endif
+		{
+			Ar.ByteOrderSerialize(&Value, sizeof(Value));
+		}
 		return Ar;
 	}
 
@@ -556,6 +654,14 @@ public:
 	 */
 	virtual void MarkScriptSerializationEnd(const UObject* Obj) { }
 
+	/**
+	* Called to retrieve the archetype from the event driven loader. If this returns null, then call GetArchetype yourself.
+	*/
+	virtual UObject* GetArchetypeFromLoader(const UObject* Obj)
+	{
+		return nullptr;
+	}
+
 	virtual void IndicateSerializationMismatch() { }
 
 	// Logf implementation for convenience.
@@ -1022,110 +1128,158 @@ protected:
 
 	/** Resets all of the base archive members. */
 	void Reset();
+#if DEVIRTUALIZE_FLinkerLoad_Serialize
+	template<SIZE_T Size>
+	FORCEINLINE bool FastPathLoad(void* InDest)
+	{
+		const uint8* RESTRICT Src = ActiveFPLB->StartFastPathLoadBuffer;
+		if (Src + Size <= ActiveFPLB->EndFastPathLoadBuffer)
+		{
+#if PLATFORM_SUPPORTS_UNALIGNED_INT_LOADS
+			if (Size == 2)
+			{
+				uint16 * RESTRICT Dest = (uint16 * RESTRICT)InDest;
+				*Dest = *(uint16 * RESTRICT)Src;
+			}
+			else if (Size == 4)
+			{
+				uint32 * RESTRICT Dest = (uint32 * RESTRICT)InDest;
+				*Dest = *(uint32 * RESTRICT)Src;
+			}
+			else if (Size == 8)
+			{
+				uint64 * RESTRICT Dest = (uint64 * RESTRICT)InDest;
+				*Dest = *(uint64 * RESTRICT)Src;
+			}
+			else
+#endif
+			{
+				uint8 * RESTRICT Dest = (uint8 * RESTRICT)InDest;
+				for (SIZE_T Index = 0; Index < Size; Index++)
+				{
+					Dest[Index] = Src[Index];
+				}
+			}
+			ActiveFPLB->StartFastPathLoadBuffer += Size;
+			return true;
+		}
+		return false;
+	}
+public:
+	/* These are used for fastpath inline serializers  */
+	struct FFastPathLoadBuffer
+	{
+		const uint8* StartFastPathLoadBuffer;
+		const uint8* EndFastPathLoadBuffer;
+		const uint8* OriginalFastPathLoadBuffer;
+		FORCEINLINE FFastPathLoadBuffer()
+		{
+			Reset();
+		}
+		FORCEINLINE void Reset()
+		{
+			StartFastPathLoadBuffer = nullptr;
+			EndFastPathLoadBuffer = nullptr;
+			OriginalFastPathLoadBuffer = nullptr;
+		}
+	};
+	//@todoio FArchive is really a horrible class and the way it is proxied by FLinkerLoad is double terrible. It makes the fast path really hacky and slower than it would need to be.
+	FFastPathLoadBuffer* ActiveFPLB;
+	FFastPathLoadBuffer InlineFPLB;
 
-private:
-
-	/** Copies all of the members except CustomVersionContainer */
-	void CopyTrivialFArchiveStatusMembers(const FArchive& ArchiveStatusToCopy);
-
-protected:
 
-	/** Holds the archive version. */
-	int32 ArUE4Ver;
-	
-	/** Holds the archive version for licensees. */
-	int32 ArLicenseeUE4Ver;
 
-	/** Holds the engine version. */
-	FEngineVersionBase ArEngineVer;
 
-	/** Holds the engine network protocol version. */
-	uint32 ArEngineNetVer;
+#else
+	template<SIZE_T Size>
+	FORCEINLINE bool FastPathLoad(void* InDest)
+	{
+		return false;
+	}
+#endif
 
-	/** Holds the game network protocol version. */
-	uint32 ArGameNetVer;
 
 private:
 
-	/**
-	 * All the custom versions stored in the archive.
-	 * Stored as a pointer to a heap-allocated object because of a 3-way dependency between TArray, FCustomVersionContainer and FArchive, which is too much work to change right now.
-	 */
-	FCustomVersionContainer* CustomVersionContainer;
+	/** Copies all of the members except CustomVersionContainer */
+	void CopyTrivialFArchiveStatusMembers(const FArchive& ArchiveStatusToCopy);
 
 public:
 
 	/** Whether this archive is for loading data. */
-	bool ArIsLoading;
-	
+	uint8 ArIsLoading : 1;
+
 	/** Whether this archive is for saving data. */
-	bool ArIsSaving;
-	
+	uint8 ArIsSaving : 1;
+
 	/** Whether archive is transacting. */
-	bool ArIsTransacting;
-	
+	uint8 ArIsTransacting : 1;
+
 	/** Whether this archive wants properties to be serialized in binary form instead of tagged. */
-	bool ArWantBinaryPropertySerialization;
-	
+	uint8 ArWantBinaryPropertySerialization : 1;
+
 	/** Whether this archive wants to always save strings in unicode format */
-	bool ArForceUnicode;
-	
+	uint8 ArForceUnicode : 1;
+
 	/** Whether this archive saves to persistent storage. */
-	bool ArIsPersistent;
-			
+	uint8 ArIsPersistent : 1;
+
 	/** Whether this archive contains errors. */
-	bool ArIsError;
+	uint8 ArIsError : 1;
 
 	/** Whether this archive contains critical errors. */
-	bool ArIsCriticalError;
-	
+	uint8 ArIsCriticalError : 1;
+
 	/** Quickly tell if an archive contains script code. */
-	bool ArContainsCode;
-	
+	uint8 ArContainsCode : 1;
+
 	/** Used to determine whether FArchive contains a level or world. */
-	bool ArContainsMap;
+	uint8 ArContainsMap : 1;
 
 	/** Used to determine whether FArchive contains data required to be gathered for localization. */
-	bool ArRequiresLocalizationGather;
-	
+	uint8 ArRequiresLocalizationGather : 1;
+
 	/** Whether we should forcefully swap bytes. */
-	bool ArForceByteSwapping;
-	
+	uint8 ArForceByteSwapping : 1;
+
 	/** If true, we will not serialize the ObjectArchetype reference in UObject. */
-	bool ArIgnoreArchetypeRef;
-	
+	uint8 ArIgnoreArchetypeRef : 1;
+
 	/** If true, we will not serialize the ObjectArchetype reference in UObject. */
-	bool ArNoDelta;
+	uint8 ArNoDelta : 1;
 
 	/** If true, we will not serialize the Outer reference in UObject. */
-	bool ArIgnoreOuterRef;
+	uint8 ArIgnoreOuterRef : 1;
 
 	/** If true, we will not serialize ClassGeneratedBy reference in UClass. */
-	bool ArIgnoreClassGeneratedByRef;
+	uint8 ArIgnoreClassGeneratedByRef : 1;
 
 	/** If true, UObject::Serialize will skip serialization of the Class property. */
-	bool ArIgnoreClassRef;
-	
+	uint8 ArIgnoreClassRef : 1;
+
 	/** Whether to allow lazy loading. */
-	bool ArAllowLazyLoading;
-	
+	uint8 ArAllowLazyLoading : 1;
+
 	/** Whether this archive only cares about serializing object references. */
-	bool ArIsObjectReferenceCollector;
-	
+	uint8 ArIsObjectReferenceCollector : 1;
+
 	/** Whether a reference collector is modifying the references and wants both weak and strong ones */
-	bool ArIsModifyingWeakAndStrongReferences;
-	
+	uint8 ArIsModifyingWeakAndStrongReferences : 1;
+
 	/** Whether this archive is counting memory and therefore wants e.g. TMaps to be serialized. */
-	bool ArIsCountingMemory;
-	
+	uint8 ArIsCountingMemory : 1;
+
 	/** Whether bulk data serialization should be skipped or not. */
-	bool ArShouldSkipBulkData;
+	uint8 ArShouldSkipBulkData : 1;
 
 	/** Whether editor only properties are being filtered from the archive (or has been filtered). */
-	bool ArIsFilterEditorOnly;
+	uint8 ArIsFilterEditorOnly : 1;
 
 	/** Whether this archive is saving/loading game state */
-	bool ArIsSaveGame;
+	uint8 ArIsSaveGame : 1;
+
+	/** Set TRUE to use the custom property list attribute for serialization. */
+	uint8 ArUseCustomPropertyList : 1;
 
 	/** Whether we are currently serializing defaults. > 0 means yes, <= 0 means no. */
 	int32 ArSerializingDefaults;
@@ -1136,12 +1290,36 @@ public:
 	/** Max size of data that this archive is allowed to serialize. */
 	int64 ArMaxSerializeSize;
 
+protected:
+
+	/** Holds the archive version. */
+	int32 ArUE4Ver;
+
+	/** Holds the archive version for licensees. */
+	int32 ArLicenseeUE4Ver;
+
+	/** Holds the engine version. */
+	FEngineVersionBase ArEngineVer;
+
+	/** Holds the engine network protocol version. */
+	uint32 ArEngineNetVer;
+
+	/** Holds the game network protocol version. */
+	uint32 ArGameNetVer;
+
+private:
+
+	/**
+	* All the custom versions stored in the archive.
+	* Stored as a pointer to a heap-allocated object because of a 3-way dependency between TArray, FCustomVersionContainer and FArchive, which is too much work to change right now.
+	*/
+	FCustomVersionContainer* CustomVersionContainer;
+
+public:
+
 	/** Custom property list attribute. If the flag below is set, only these properties will be iterated during serialization. If NULL, then no properties will be iterated. */
 	const struct FCustomPropertyListNode* ArCustomPropertyList;
 
-	/** Set TRUE to use the custom property list attribute for serialization. */
-	bool ArUseCustomPropertyList;
-
 	class FScopeSetDebugSerializationFlags
 	{
 	private:
diff --git a/Engine/Source/Runtime/Core/Public/Templates/Function.h b/Engine/Source/Runtime/Core/Public/Templates/Function.h
--- a/Engine/Source/Runtime/Core/Public/Templates/Function.h
+++ b/Engine/Source/Runtime/Core/Public/Templates/Function.h
@@ -9,6 +9,7 @@
 #include "UnrealMathUtility.h"
 #include "UnrealMemory.h"
 #include "Templates/Decay.h"
+#include "Templates/Invoke.h"
 #include "Templates/RemoveReference.h"
 
 // Disable visualization hack for shipping or test builds.
@@ -225,7 +226,7 @@ namespace UE4Function_Private
 	{
 		static Ret Call(void* Obj, ParamTypes&... Params)
 		{
-			return (*(Functor*)Obj)(Forward<ParamTypes>(Params)...);
+			return Invoke(*(Functor*)Obj, Forward<ParamTypes>(Params)...);
 		}
 	};
 
@@ -234,7 +235,7 @@ namespace UE4Function_Private
 	{
 		static void Call(void* Obj, ParamTypes&... Params)
 		{
-			(*(Functor*)Obj)(Forward<ParamTypes>(Params)...);
+			Invoke(*(Functor*)Obj, Forward<ParamTypes>(Params)...);
 		}
 	};
 
diff --git a/Engine/Source/Runtime/Core/Public/Templates/IsArray.h b/Engine/Source/Runtime/Core/Public/Templates/IsArray.h
new file mode 100644
--- /dev/null
+++ b/Engine/Source/Runtime/Core/Public/Templates/IsArray.h
@@ -0,0 +1,22 @@
+// Copyright 1998-2016 Epic Games, Inc. All Rights Reserved.
+
+#pragma once
+
+/**
+ * Traits class which tests if a type is a C++ array.
+ */
+template <typename T>           struct TIsArray       { enum { Value = false }; };
+template <typename T>           struct TIsArray<T[]>  { enum { Value = true  }; };
+template <typename T, uint32 N> struct TIsArray<T[N]> { enum { Value = true  }; };
+
+/**
+ * Traits class which tests if a type is a bounded C++ array.
+ */
+template <typename T>           struct TIsBoundedArray       { enum { Value = false }; };
+template <typename T, uint32 N> struct TIsBoundedArray<T[N]> { enum { Value = true  }; };
+
+/**
+ * Traits class which tests if a type is an unbounded C++ array.
+ */
+template <typename T> struct TIsUnboundedArray      { enum { Value = false }; };
+template <typename T> struct TIsUnboundedArray<T[]> { enum { Value = true  }; };
diff --git a/Engine/Source/Runtime/Core/Public/Templates/SharedPointer.h b/Engine/Source/Runtime/Core/Public/Templates/SharedPointer.h
--- a/Engine/Source/Runtime/Core/Public/Templates/SharedPointer.h
+++ b/Engine/Source/Runtime/Core/Public/Templates/SharedPointer.h
@@ -470,6 +470,16 @@ private:
 	/** Interface to the reference counter for this object.  Note that the actual reference
 		controller object is shared by all shared and weak pointers that refer to the object */
 	SharedPointerInternals::FSharedReferencer< Mode > SharedReferenceCount;
+
+	template <typename InObjectType, ESPMode InMode, typename... InArgTypes>
+	friend TSharedRef<InObjectType, InMode> MakeShared(InArgTypes&&... Args);
+
+	FORCEINLINE explicit TSharedRef(ObjectType* InObject, SharedPointerInternals::FReferenceControllerBase* InSharedReferenceCount)
+		: Object(InObject)
+		, SharedReferenceCount(InSharedReferenceCount)
+	{
+		Init(InObject);
+	}
 };
 
 
@@ -494,7 +504,7 @@ struct FMakeReferenceTo<void>
 
 
 template <typename InObjectType, ESPMode InMode = ESPMode::Fast, typename... InArgTypes>
-TSharedPtr<InObjectType, InMode> MakeShared(InArgTypes&&... Args);
+TSharedRef<InObjectType, InMode> MakeShared(InArgTypes&&... Args);
 
 
 /**
@@ -872,18 +882,6 @@ private:
 	/** Interface to the reference counter for this object.  Note that the actual reference
 		controller object is shared by all shared and weak pointers that refer to the object */
 	SharedPointerInternals::FSharedReferencer< Mode > SharedReferenceCount;
-
-	template <typename InObjectType, ESPMode InMode, typename... InArgTypes>
-	friend TSharedPtr<InObjectType, InMode> MakeShared(InArgTypes&&... Args);
-
-	FORCEINLINE explicit TSharedPtr(ObjectType* InObject, SharedPointerInternals::FReferenceControllerBase* InSharedReferenceCount)
-		: Object(InObject)
-		, SharedReferenceCount(InSharedReferenceCount)
-	{
-		// If the object happens to be derived from TSharedFromThis, the following method
-		// will prime the object with a weak pointer to itself.
-		SharedPointerInternals::EnableSharedFromThis(this, InObject, InObject);
-	}
 };
 
 
@@ -1631,10 +1629,10 @@ FORCEINLINE SharedPointerInternals::FRawPtrProxy< ObjectType > MakeShareable( Ob
  * Equivalent to std::make_shared.
  */
 template <typename InObjectType, ESPMode InMode, typename... InArgTypes>
-FORCEINLINE TSharedPtr<InObjectType, InMode> MakeShared(InArgTypes&&... Args)
+FORCEINLINE TSharedRef<InObjectType, InMode> MakeShared(InArgTypes&&... Args)
 {
 	SharedPointerInternals::TIntrusiveReferenceController<InObjectType>* Controller = SharedPointerInternals::NewIntrusiveReferenceController<InObjectType>(Forward<InArgTypes>(Args)...);
-	return TSharedPtr<InObjectType, InMode>(Controller->GetObjectPtr(), (SharedPointerInternals::FReferenceControllerBase*)Controller);
+	return TSharedRef<InObjectType, InMode>(Controller->GetObjectPtr(), (SharedPointerInternals::FReferenceControllerBase*)Controller);
 }
 
 
diff --git a/Engine/Source/Runtime/Core/Public/Templates/UniquePtr.h b/Engine/Source/Runtime/Core/Public/Templates/UniquePtr.h
--- a/Engine/Source/Runtime/Core/Public/Templates/UniquePtr.h
+++ b/Engine/Source/Runtime/Core/Public/Templates/UniquePtr.h
@@ -4,6 +4,9 @@
 
 #include "UnrealTemplate.h"
 #include "Templates/EnableIf.h"
+#include "Templates/IsArray.h"
+#include "Templates/PointerIsConvertibleFromTo.h"
+#include "Templates/RemoveExtent.h"
 
 // Single-ownership smart pointer in the vein of std::unique_ptr.
 // Use this when you need an object's lifetime to be strictly bound to the lifetime of a single smart pointer.
@@ -21,13 +24,19 @@ struct TDefaultDelete
 	{
 	}
 
-	template <typename U, typename = typename TEnableIf<TPointerIsConvertibleFromTo<U, T>::Value>::Type>
-	TDefaultDelete(TDefaultDelete<U>&&)
+	template <
+		typename U,
+		typename = typename TEnableIf<TPointerIsConvertibleFromTo<U, T>::Value>::Type
+	>
+	TDefaultDelete(const TDefaultDelete<U>&)
 	{
 	}
 
-	template <typename U, typename = typename TEnableIf<TPointerIsConvertibleFromTo<U, T>::Value>::Type>
-	TDefaultDelete& operator=(TDefaultDelete<U>&&)
+	template <
+		typename U,
+		typename = typename TEnableIf<TPointerIsConvertibleFromTo<U, T>::Value>::Type
+	>
+	TDefaultDelete& operator=(const TDefaultDelete<U>&)
 	{
 		return *this;
 	}
@@ -42,6 +51,41 @@ struct TDefaultDelete
 	}
 };
 
+template <typename T>
+struct TDefaultDelete<T[]>
+{
+	TDefaultDelete()
+	{
+	}
+
+	template <
+		typename U,
+		typename = typename TEnableIf<TPointerIsConvertibleFromTo<U[], T[]>::Value>::Type
+	>
+	TDefaultDelete(const TDefaultDelete<U[]>&)
+	{
+	}
+
+	template <
+		typename U,
+		typename = typename TEnableIf<TPointerIsConvertibleFromTo<U[], T[]>::Value>::Type
+	>
+	TDefaultDelete& operator=(const TDefaultDelete<U[]>&)
+	{
+		return *this;
+	}
+
+	~TDefaultDelete()
+	{
+	}
+
+	template <typename U>
+	typename TEnableIf<TPointerIsConvertibleFromTo<U[], T[]>::Value>::Type operator()(U* Ptr) const
+	{
+		delete [] Ptr;
+	}
+};
+
 template <typename T, typename Deleter = TDefaultDelete<T>>
 class TUniquePtr : private Deleter
 {
@@ -88,7 +132,11 @@ public:
 	/**
 	 * Constructor from rvalues of other (usually derived) types
 	 */
-	template <typename OtherT, typename OtherDeleter>
+	template <
+		typename OtherT,
+		typename OtherDeleter,
+		typename = typename TEnableIf<!TIsArray<OtherT>::Value>::Type
+	>
 	FORCEINLINE TUniquePtr(TUniquePtr<OtherT, OtherDeleter>&& Other)
 		: Deleter(MoveTemp(Other.GetDeleter()))
 		, Ptr    (Other.Ptr)
@@ -119,7 +167,7 @@ public:
 	 * Assignment operator for rvalues of other (usually derived) types
 	 */
 	template <typename OtherT, typename OtherDeleter>
-	FORCEINLINE TUniquePtr& operator=(TUniquePtr<OtherT>&& Other)
+	FORCEINLINE typename TEnableIf<!TIsArray<OtherT>::Value, TUniquePtr&>::Type operator=(TUniquePtr<OtherT, OtherDeleter>&& Other)
 	{
 		// We delete last, because we don't want odd side effects if the destructor of T relies on the state of this or Other
 		T* OldPtr = Ptr;
@@ -266,6 +314,237 @@ private:
 	T* Ptr;
 };
 
+template <typename T, typename Deleter>
+class TUniquePtr<T[], Deleter> : private Deleter
+{
+	template <typename OtherT, typename OtherDeleter>
+	friend class TUniquePtr;
+
+public:
+	/**
+	 * Default constructor - initializes the TUniquePtr to null.
+	 */
+	FORCEINLINE TUniquePtr()
+		: Ptr(nullptr)
+	{
+	}
+
+	/**
+	 * Pointer constructor - takes ownership of the pointed-to array
+	 *
+	 * @param InPtr The pointed-to array to take ownership of.
+	 */
+	template <
+		typename U,
+		typename = typename TEnableIf<TPointerIsConvertibleFromTo<U[], T[]>::Value>::Type
+	>
+	explicit FORCEINLINE TUniquePtr(U* InPtr)
+		: Ptr(InPtr)
+	{
+	}
+
+	/**
+	 * nullptr constructor - initializes the TUniquePtr to null.
+	 */
+	FORCEINLINE TUniquePtr(TYPE_OF_NULLPTR)
+		: Ptr(nullptr)
+	{
+	}
+
+	/**
+	 * Move constructor
+	 */
+	FORCEINLINE TUniquePtr(TUniquePtr&& Other)
+		: Deleter(MoveTemp(Other.GetDeleter()))
+		, Ptr    (Other.Ptr)
+	{
+		Other.Ptr = nullptr;
+	}
+
+	/**
+	 * Constructor from rvalues of other (usually less qualified) types
+	 */
+	template <
+		typename OtherT,
+		typename OtherDeleter,
+		typename = typename TEnableIf<TPointerIsConvertibleFromTo<OtherT[], T[]>::Value>::Type
+	>
+	FORCEINLINE TUniquePtr(TUniquePtr<OtherT, OtherDeleter>&& Other)
+		: Deleter(MoveTemp(Other.GetDeleter()))
+		, Ptr    (Other.Ptr)
+	{
+		Other.Ptr = nullptr;
+	}
+
+	/**
+	 * Move assignment operator
+	 */
+	FORCEINLINE TUniquePtr& operator=(TUniquePtr&& Other)
+	{
+		if (this != &Other)
+		{
+			// We delete last, because we don't want odd side effects if the destructor of T relies on the state of this or Other
+			T* OldPtr = Ptr;
+			Ptr = Other.Ptr;
+			Other.Ptr = nullptr;
+			GetDeleter()(OldPtr);
+		}
+
+		GetDeleter() = MoveTemp(Other.GetDeleter());
+
+		return *this;
+	}
+
+	/**
+	 * Assignment operator for rvalues of other (usually less qualified) types
+	 */
+	template <typename OtherT, typename OtherDeleter>
+	FORCEINLINE typename TEnableIf<TPointerIsConvertibleFromTo<OtherT[], T[]>::Value, TUniquePtr&>::Type operator=(TUniquePtr<OtherT>&& Other)
+	{
+		// We delete last, because we don't want odd side effects if the destructor of T relies on the state of this or Other
+		T* OldPtr = Ptr;
+		Ptr = Other.Ptr;
+		Other.Ptr = nullptr;
+		GetDeleter()(OldPtr);
+
+		GetDeleter() = MoveTemp(Other.GetDeleter());
+
+		return *this;
+	}
+
+	/**
+	 * Nullptr assignment operator
+	 */
+	FORCEINLINE TUniquePtr& operator=(TYPE_OF_NULLPTR)
+	{
+		// We delete last, because we don't want odd side effects if the destructor of T relies on the state of this
+		T* OldPtr = Ptr;
+		Ptr = nullptr;
+		GetDeleter()(OldPtr);
+
+		return *this;
+	}
+
+	/**
+	 * Destructor
+	 */
+	FORCEINLINE ~TUniquePtr()
+	{
+		GetDeleter()(Ptr);
+	}
+
+	/**
+	 * Tests if the TUniquePtr currently owns an array.
+	 *
+	 * @return true if the TUniquePtr currently owns an array, false otherwise.
+	 */
+	bool IsValid() const
+	{
+		return Ptr != nullptr;
+	}
+
+	/**
+	 * operator bool
+	 *
+	 * @return true if the TUniquePtr currently owns an array, false otherwise.
+	 */
+	FORCEINLINE explicit operator bool() const
+	{
+		return IsValid();
+	}
+
+	/**
+	 * Logical not operator
+	 *
+	 * @return false if the TUniquePtr currently owns an array, true otherwise.
+	 */
+	FORCEINLINE bool operator!() const
+	{
+		return !IsValid();
+	}
+
+	/**
+	 * Indexing operator
+	 *
+	 * @return A reference to the object at the specified index by the TUniquePtr.
+	 */
+	FORCEINLINE T& operator[](SIZE_T Index) const
+	{
+		return Ptr[Index];
+	}
+
+	/**
+	 * Returns a pointer to the owned array without relinquishing ownership.
+	 *
+	 * @return A copy of the pointer to the array owned by the TUniquePtr, or nullptr if no array is being owned.
+	 */
+	FORCEINLINE T* Get() const
+	{
+		return Ptr;
+	}
+
+	/**
+	 * Relinquishes control of the owned array to the caller and nulls the TUniquePtr.
+	 *
+	 * @return The pointer to the array that was owned by the TUniquePtr, or nullptr if no array was being owned.
+	 */
+	FORCEINLINE T* Release()
+	{
+		T* Result = Ptr;
+		Ptr = nullptr;
+		return Result;
+	}
+
+	/**
+	 * Gives the TUniquePtr a new array to own, destroying any previously-owned array.
+	 *
+	 * @param InPtr A pointer to the array to take ownership of.
+	 */
+	template <typename U>
+	FORCEINLINE typename TEnableIf<TPointerIsConvertibleFromTo<U[], T[]>::Value>::Type Reset(U* InPtr)
+	{
+		// We delete last, because we don't want odd side effects if the destructor of T relies on the state of this
+		T* OldPtr = Ptr;
+		Ptr = InPtr;
+		GetDeleter()(OldPtr);
+	}
+
+	FORCEINLINE void Reset(TYPE_OF_NULLPTR InPtr = nullptr)
+	{
+		// We delete last, because we don't want odd side effects if the destructor of T relies on the state of this
+		T* OldPtr = Ptr;
+		Ptr = InPtr;
+		GetDeleter()(OldPtr);
+	}
+
+	/**
+	 * Returns a reference to the deleter subobject.
+	 *
+	 * @return A reference to the deleter.
+	 */
+	FORCEINLINE Deleter& GetDeleter()
+	{
+		return static_cast<Deleter&>(*this);
+	}
+
+	/**
+	 * Returns a reference to the deleter subobject.
+	 *
+	 * @return A reference to the deleter.
+	 */
+	FORCEINLINE const Deleter& GetDeleter() const
+	{
+		return static_cast<const Deleter&>(*this);
+	}
+
+private:
+	// Non-copyable
+	TUniquePtr(const TUniquePtr&);
+	TUniquePtr& operator=(const TUniquePtr&);
+
+	T* Ptr;
+};
+
 /**
  * Equality comparison operator
  *
@@ -355,14 +634,34 @@ struct TIsBitwiseConstructible<TUniquePtr<T>, T*>
 };
 
 /**
- * Constructs a new object with the given arguments and returns it as a TUniquePtr.
+ * Allocates a new object of type T with the given arguments and returns it as a TUniquePtr.  Disabled for array-type TUniquePtrs.
  *
  * @param Args The arguments to pass to the constructor of T.
  *
  * @return A TUniquePtr which points to a newly-constructed T with the specified Args.
  */
 template <typename T, typename... TArgs>
-FORCEINLINE TUniquePtr<T> MakeUnique(TArgs&&... Args)
+FORCEINLINE typename TEnableIf<!TIsArray<T>::Value, TUniquePtr<T>>::Type MakeUnique(TArgs&&... Args)
 {
 	return TUniquePtr<T>(new T(Forward<TArgs>(Args)...));
 }
+
+/**
+ * Allocates a new array of type T with the given size and returns it as a TUniquePtr.  Only enabled for array-type TUniquePtrs.
+ *
+ * @param Size The size of the array to allocate.
+ *
+ * @return A TUniquePtr which points to a newly-constructed T array of the specified Size.
+ */
+template <typename T>
+FORCEINLINE typename TEnableIf<TIsUnboundedArray<T>::Value, TUniquePtr<T>>::Type MakeUnique(SIZE_T Size)
+{
+	typedef typename TRemoveExtent<T>::Type ElementType;
+	return TUniquePtr<T>(new ElementType[Size]);
+}
+
+/**
+ * Overload to cause a compile error when MakeUnique<T[N]> is attempted.  Use MakeUnique<T>(N) instead.
+ */
+template <typename T, typename... TArgs>
+typename TEnableIf<TIsBoundedArray<T>::Value, TUniquePtr<T>>::Type MakeUnique(TArgs&&... Args) = delete;
diff --git a/Engine/Source/Runtime/Core/Public/Templates/UnrealTemplate.h b/Engine/Source/Runtime/Core/Public/Templates/UnrealTemplate.h
--- a/Engine/Source/Runtime/Core/Public/Templates/UnrealTemplate.h
+++ b/Engine/Source/Runtime/Core/Public/Templates/UnrealTemplate.h
@@ -430,19 +430,6 @@ FORCEINLINE T StaticCast(ArgType&& Arg)
 template <typename T> struct TRValueToLValueReference      { typedef T  Type; };
 template <typename T> struct TRValueToLValueReference<T&&> { typedef T& Type; };
 
-/**
- * A traits class which tests if a type is a C++ array.
- */
-template <typename T>           struct TIsCPPArray       { enum { Value = false }; };
-template <typename T, uint32 N> struct TIsCPPArray<T[N]> { enum { Value = true  }; };
-
-/**
- * Removes one dimension of extents from an array type.
- */
-template <typename T>           struct TRemoveExtent       { typedef T Type; };
-template <typename T>           struct TRemoveExtent<T[]>  { typedef T Type; };
-template <typename T, uint32 N> struct TRemoveExtent<T[N]> { typedef T Type; };
-
 /**
  * Reverses the order of the bits of a value.
  * This is an TEnableIf'd template to ensure that no undesirable conversions occur.  Overloads for other types can be added in the same way.
diff --git a/Engine/Source/Runtime/Core/Public/UObject/ObjectVersion.h b/Engine/Source/Runtime/Core/Public/UObject/ObjectVersion.h
--- a/Engine/Source/Runtime/Core/Public/UObject/ObjectVersion.h
+++ b/Engine/Source/Runtime/Core/Public/UObject/ObjectVersion.h
@@ -608,6 +608,10 @@ enum EUnrealEngineObjectUE4Version
 	VER_UE4_INSTANCED_STEREO_UNIFORM_REFACTOR,
 	// Added compression to the shader resource for memory savings
 	VER_UE4_COMPRESSED_SHADER_RESOURCES,
+	// Cooked files contain the dependency graph for the event driven loader (the serialization is largely independent of the use of the new loader)
+	VER_UE4_PRELOAD_DEPENDENCIES_IN_COOKED_EXPORTS,
+	// Cooked files contain the TemplateIndex used by the event driven loader (the serialization is largely independent of the use of the new loader, i.e. this will be null if cooking for the old loader)
+	VER_UE4_TemplateIndex_IN_COOKED_EXPORTS,
 	
 	// -----<new versions can be added before this line>-------------------------------------------------
 	// - this needs to be the last line (see note below)
diff --git a/Engine/Source/Runtime/Core/Public/Windows/WIndowsPlatform.h b/Engine/Source/Runtime/Core/Public/Windows/WIndowsPlatform.h
--- a/Engine/Source/Runtime/Core/Public/Windows/WIndowsPlatform.h
+++ b/Engine/Source/Runtime/Core/Public/Windows/WIndowsPlatform.h
@@ -36,6 +36,7 @@ typedef FWindowsPlatformTypes FPlatformTypes;
 // Base defines, defaults are commented out
 
 #define PLATFORM_LITTLE_ENDIAN								1
+#define PLATFORM_SUPPORTS_UNALIGNED_INT_LOADS				1
 #if defined(__clang__)
 	// @todo clang: Clang compiler on Windows doesn't support SEH exception handling yet (__try/__except)
 	#define PLATFORM_SEH_EXCEPTIONS_DISABLED				1
diff --git a/Engine/Source/Runtime/CoreUObject/Private/Misc/GCObjectReferencer.cpp b/Engine/Source/Runtime/CoreUObject/Private/Misc/GCObjectReferencer.cpp
--- a/Engine/Source/Runtime/CoreUObject/Private/Misc/GCObjectReferencer.cpp
+++ b/Engine/Source/Runtime/CoreUObject/Private/Misc/GCObjectReferencer.cpp
@@ -6,13 +6,18 @@
 
 #include "CoreUObjectPrivate.h"
 
+// Global GC state flags
+extern bool GObjIncrementalPurgeIsInProgress;
+extern bool GObjUnhashUnreachableIsInProgress;
+
 void UGCObjectReferencer::AddReferencedObjects(UObject* InThis, FReferenceCollector& Collector)
 {	
 	UGCObjectReferencer* This = CastChecked<UGCObjectReferencer>(InThis);
+	// Note we're not locking ReferencedObjectsCritical here because we guard
+	// against adding new references during GC in AddObject and RemoveObject.
 	// Let each registered object handle its AddReferencedObjects call
-	for (int32 i = 0; i < This->ReferencedObjects.Num(); i++)
+	for (FGCObject* Object : This->ReferencedObjects)
 	{
-		FGCObject* Object = This->ReferencedObjects[i];
 		check(Object);
 		Object->AddReferencedObjects(Collector);
 	}
@@ -22,6 +27,8 @@ void UGCObjectReferencer::AddReferencedObjects(UObject* InThis, FReferenceCollec
 void UGCObjectReferencer::AddObject(FGCObject* Object)
 {
 	check(Object);
+	check(GObjUnhashUnreachableIsInProgress || GObjIncrementalPurgeIsInProgress || !IsGarbageCollecting());
+	FScopeLock ReferencedObjectsLock(&ReferencedObjectsCritical);
 	// Make sure there are no duplicates. Should be impossible...
 	ReferencedObjects.AddUnique(Object);
 }
@@ -29,6 +36,9 @@ void UGCObjectReferencer::AddObject(FGCObject* Object)
 void UGCObjectReferencer::RemoveObject(FGCObject* Object)
 {
 	check(Object);
+	// Only allow to remove FGCObjects during GC when purging objects.
+	check(GObjUnhashUnreachableIsInProgress || GObjIncrementalPurgeIsInProgress || !IsGarbageCollecting());
+	FScopeLock ReferencedObjectsLock(&ReferencedObjectsCritical);
 	ReferencedObjects.Remove(Object);
 }
 
diff --git a/Engine/Source/Runtime/CoreUObject/Private/Misc/StartupPackages.cpp b/Engine/Source/Runtime/CoreUObject/Private/Misc/StartupPackages.cpp
--- a/Engine/Source/Runtime/CoreUObject/Private/Misc/StartupPackages.cpp
+++ b/Engine/Source/Runtime/CoreUObject/Private/Misc/StartupPackages.cpp
@@ -30,12 +30,7 @@ void FStartupPackages::GetStartupPackageNames(TArray<FString>& PackageNames, con
 	}
 }
 
-/**
- * Kicks off a list of packages to be read in asynchronously in the background by the
- * async file manager. The package will be serialized from RAM later.
- * 
- * @param PackageNames The list of package names to async preload
- */
+#if !USE_NEW_ASYNC_IO
 static void AsyncPreloadPackageList(const TArray<FString>& PackageNames)
 {
 	// Iterate over all native script packages and preload them.
@@ -45,6 +40,7 @@ static void AsyncPreloadPackageList(const TArray<FString>& PackageNames)
 		FLinkerLoad::AsyncPreloadPackage(*PackageNames[PackageIndex]);
 	}
 }
+#endif
 
 void FStartupPackages::LoadPackageList(const TArray<FString>& PackageNames)
 {
@@ -61,10 +57,6 @@ bool FStartupPackages::LoadAll()
 
 	DECLARE_SCOPE_CYCLE_COUNTER(TEXT("Loading Startup Packages"), STAT_StartupPackages, STATGROUP_LoadTime);
 
-	// should startup packages load from memory?
-	bool bSerializeStartupPackagesFromMemory = false;
-	GConfig->GetBool(TEXT("Engine.StartupPackages"), TEXT("bSerializeStartupPackagesFromMemory"), bSerializeStartupPackagesFromMemory, GEngineIni);
-
 	// Get list of startup packages.
 	TArray<FString> StartupPackages;
 	
@@ -74,14 +66,21 @@ bool FStartupPackages::LoadAll()
 		FStartupPackages::GetStartupPackageNames(StartupPackages);
 	}
 
-	if( bSerializeStartupPackagesFromMemory )
+//@todoio this doesn't seem to work with !USE_NEW_ASYNC_IO anymore...and honestly I can't see how it ever worked well in UE4
+#if 0 && !USE_NEW_ASYNC_IO
+	// should startup packages load from memory?
+	bool bSerializeStartupPackagesFromMemory = false;
+	GConfig->GetBool(TEXT("Engine.StartupPackages"), TEXT("bSerializeStartupPackagesFromMemory"), bSerializeStartupPackagesFromMemory, GEngineIni);
+
+	if (bSerializeStartupPackagesFromMemory)
 	{
-		if( GUseSeekFreeLoading )
+		// @todo: this needs proper support in the cooker
+		if (FPlatformProperties::RequiresCookedData())
 		{
 			// kick them off to be preloaded
 			AsyncPreloadPackageList(StartupPackages);
 		}
 	}
-
+#endif
 	return bReturn;
 }
diff --git a/Engine/Source/Runtime/CoreUObject/Private/Serialization/ArchiveUObject.cpp b/Engine/Source/Runtime/CoreUObject/Private/Serialization/ArchiveUObject.cpp
--- a/Engine/Source/Runtime/CoreUObject/Private/Serialization/ArchiveUObject.cpp
+++ b/Engine/Source/Runtime/CoreUObject/Private/Serialization/ArchiveUObject.cpp
@@ -93,6 +93,12 @@ FArchive& FArchiveUObject::operator<<(struct FStringAssetReference& Value)
 	return *this;
 }
 
+FArchive& FArchiveUObject::operator<<(FWeakObjectPtr& WeakObjectPtr)
+{
+	WeakObjectPtr.Serialize(*this);
+	return *this;
+}
+
 /*----------------------------------------------------------------------------
 	FObjectAndNameAsStringProxyArchive.
 ----------------------------------------------------------------------------*/
