Copying //UE4/Orion-Staging to //UE4/Main (Source: //Orion/Dev-General @ 3196521)

diff --git a/Engine/Source/Runtime/AnimGraphRuntime/Private/AnimNodes/AnimNode_TwoWayBlend.cpp b/Engine/Source/Runtime/AnimGraphRuntime/Private/AnimNodes/AnimNode_TwoWayBlend.cpp
--- a/Engine/Source/Runtime/AnimGraphRuntime/Private/AnimNodes/AnimNode_TwoWayBlend.cpp
+++ b/Engine/Source/Runtime/AnimGraphRuntime/Private/AnimNodes/AnimNode_TwoWayBlend.cpp
@@ -12,6 +12,9 @@ void FAnimationNode_TwoWayBlend::Initialize(const FAnimationInitializeContext& C
 
 	A.Initialize(Context);
 	B.Initialize(Context);
+
+	bAIsRelevant = false;
+	bBIsRelevant = false;
 }
 
 void FAnimationNode_TwoWayBlend::CacheBones(const FAnimationCacheBonesContext& Context) 
@@ -25,14 +28,40 @@ void FAnimationNode_TwoWayBlend::Update(const FAnimationUpdateContext& Context)
 	QUICK_SCOPE_CYCLE_COUNTER(STAT_FAnimationNode_TwoWayBlend_Update);
 	EvaluateGraphExposedInputs.Execute(Context);
 
-	const float ActualAlpha = AlphaScaleBias.ApplyTo(Alpha);
-	if (ActualAlpha > ZERO_ANIMWEIGHT_THRESH)
+	InternalBlendAlpha = AlphaScaleBias.ApplyTo(Alpha);
+	const bool bNewAIsRelevant = !FAnimWeight::IsFullWeight(InternalBlendAlpha);
+	const bool bNewBIsRelevant = FAnimWeight::IsRelevant(InternalBlendAlpha);
+
+	// when this flag is true, we'll reinitialize the children
+	if (bResetChildOnActivation)
+	{
+		if (bNewAIsRelevant && !bAIsRelevant)
+		{
+			FAnimationInitializeContext ReinitializeContext(Context.AnimInstanceProxy);
+
+			// reinitialize
+			A.Initialize(ReinitializeContext);
+		}
+
+		if (bNewBIsRelevant && !bBIsRelevant)
+		{
+			FAnimationInitializeContext ReinitializeContext(Context.AnimInstanceProxy);
+
+			// reinitialize
+			B.Initialize(ReinitializeContext);
+		}
+	}
+
+	bAIsRelevant = bNewAIsRelevant;
+	bBIsRelevant = bNewBIsRelevant;
+
+	if (bBIsRelevant)
 	{
-		if (ActualAlpha < 1.0f - ZERO_ANIMWEIGHT_THRESH)
+		if (bAIsRelevant)
 		{
 			// Blend A and B together
-			A.Update(Context.FractionalWeight(1.0f - ActualAlpha));
-			B.Update(Context.FractionalWeight(ActualAlpha));
+			A.Update(Context.FractionalWeight(1.0f - InternalBlendAlpha));
+			B.Update(Context.FractionalWeight(InternalBlendAlpha));
 		}
 		else
 		{
@@ -49,10 +78,9 @@ void FAnimationNode_TwoWayBlend::Update(const FAnimationUpdateContext& Context)
 
 void FAnimationNode_TwoWayBlend::Evaluate(FPoseContext& Output)
 {
-	const float ActualAlpha = AlphaScaleBias.ApplyTo(Alpha);
-	if (ActualAlpha > ZERO_ANIMWEIGHT_THRESH)
+	if (bBIsRelevant)
 	{
-		if (ActualAlpha < 1.0f - ZERO_ANIMWEIGHT_THRESH)
+		if (bAIsRelevant)
 		{
 			FPoseContext Pose1(Output);
 			FPoseContext Pose2(Output);
@@ -60,7 +88,7 @@ void FAnimationNode_TwoWayBlend::Evaluate(FPoseContext& Output)
 			A.Evaluate(Pose1);
 			B.Evaluate(Pose2);
 
-			FAnimationRuntime::BlendTwoPosesTogether(Pose1.Pose, Pose2.Pose, Pose1.Curve, Pose2.Curve, (1.0f - ActualAlpha), Output.Pose, Output.Curve);
+			FAnimationRuntime::BlendTwoPosesTogether(Pose1.Pose, Pose2.Pose, Pose1.Curve, Pose2.Curve, (1.0f - InternalBlendAlpha), Output.Pose, Output.Curve);
 		}
 		else
 		{
@@ -76,12 +104,10 @@ void FAnimationNode_TwoWayBlend::Evaluate(FPoseContext& Output)
 
 void FAnimationNode_TwoWayBlend::GatherDebugData(FNodeDebugData& DebugData)
 {
-	const float ActualAlpha = AlphaScaleBias.ApplyTo(Alpha);
-
 	FString DebugLine = DebugData.GetNodeName(this);
-	DebugLine += FString::Printf(TEXT("(Alpha: %.1f%%)"), ActualAlpha*100);
+	DebugLine += FString::Printf(TEXT("(Alpha: %.1f%%)"), InternalBlendAlpha *100);
 	DebugData.AddDebugItem(DebugLine);
 
-	A.GatherDebugData(DebugData.BranchFlow(1.f - ActualAlpha));
-	B.GatherDebugData(DebugData.BranchFlow(ActualAlpha));
+	A.GatherDebugData(DebugData.BranchFlow(1.f - InternalBlendAlpha));
+	B.GatherDebugData(DebugData.BranchFlow(InternalBlendAlpha));
 }
\ No newline at end of file
diff --git a/Engine/Source/Runtime/AnimGraphRuntime/Private/BoneControllers/AnimNode_AnimDynamics.cpp b/Engine/Source/Runtime/AnimGraphRuntime/Private/BoneControllers/AnimNode_AnimDynamics.cpp
--- a/Engine/Source/Runtime/AnimGraphRuntime/Private/BoneControllers/AnimNode_AnimDynamics.cpp
+++ b/Engine/Source/Runtime/AnimGraphRuntime/Private/BoneControllers/AnimNode_AnimDynamics.cpp
@@ -135,47 +135,6 @@ void FAnimNode_AnimDynamics::EvaluateBoneTransforms(USkeletalMeshComponent* Skel
 
 		if (bDoUpdate && NextTimeStep > 0.0f)
 		{
-			// Wind / Force update
-			if(CVarEnableWind.GetValueOnAnyThread() == 1 && bEnableWind)
-			{
-				SCOPE_CYCLE_COUNTER(STAT_AnimDynamicsWindData);
-
-				for(FAnimPhysRigidBody* Body : BaseBodyPtrs)
-				{
-					if(SkelComp && SkelComp->GetWorld())
-					{
-						Body->bWindEnabled = bEnableWind;
-
-						if(Body->bWindEnabled)
-						{
-							UWorld* World = SkelComp->GetWorld();
-							FSceneInterface* Scene = World->Scene;
-
-							// Unused by our simulation but needed for the call to GetWindParameters below
-							float WindMinGust;
-							float WindMaxGust;
-
-							// Setup wind data
-							Body->bWindEnabled = true;
-							Scene->GetWindParameters(SkelComp->ComponentToWorld.TransformPosition(Body->Pose.Position), Body->WindData.WindDirection, Body->WindData.WindSpeed, WindMinGust, WindMaxGust);
-
-							Body->WindData.WindDirection = SkelComp->ComponentToWorld.Inverse().TransformVector(Body->WindData.WindDirection);
-							Body->WindData.WindAdaption = FMath::FRandRange(0.0f, 2.0f);
-							Body->WindData.BodyWindScale = WindScale;
-						}
-					}
-				}
-			}
-			else
-			{
-				SCOPE_CYCLE_COUNTER(STAT_AnimDynamicsWindData);
-				// Disable wind.
-				for(FAnimPhysRigidBody* Body : BaseBodyPtrs)
-				{
-					Body->bWindEnabled = false;
-				}
-			}
-
 			FVector OrientedExternalForce = ExternalForce;
 			if(!OrientedExternalForce.IsNearlyZero())
 			{
@@ -515,6 +474,7 @@ void FAnimNode_AnimDynamics::InitPhysics(USkeletalMeshComponent* Component, FCSP
 		}
 
 		PhysicsBody.GravityScale = GravityScale;
+		PhysicsBody.bWindEnabled = bWindWasEnabled;
 
 		// Link to parent
 		if (Bodies.Num() > 0)
@@ -730,6 +690,45 @@ void FAnimNode_AnimDynamics::PreUpdate(const UAnimInstance* InAnimInstance)
 	const UWorld* World = SkelComp->GetWorld();
 	check(World->GetWorldSettings());
 	CurrentTimeDilation = World->GetWorldSettings()->GetEffectiveTimeDilation();
+	if(CVarEnableWind.GetValueOnAnyThread() == 1 && bEnableWind)
+	{
+		SCOPE_CYCLE_COUNTER(STAT_AnimDynamicsWindData);
+
+		for(FAnimPhysRigidBody* Body : BaseBodyPtrs)
+		{
+			if(SkelComp && SkelComp->GetWorld())
+			{
+				Body->bWindEnabled = bEnableWind;
+
+				if(Body->bWindEnabled)
+				{
+					FSceneInterface* Scene = World->Scene;
+
+					// Unused by our simulation but needed for the call to GetWindParameters below
+					float WindMinGust;
+					float WindMaxGust;
+
+					// Setup wind data
+					Body->bWindEnabled = true;
+					Scene->GetWindParameters_GameThread(SkelComp->ComponentToWorld.TransformPosition(Body->Pose.Position), Body->WindData.WindDirection, Body->WindData.WindSpeed, WindMinGust, WindMaxGust);
+
+					Body->WindData.WindDirection = SkelComp->ComponentToWorld.Inverse().TransformVector(Body->WindData.WindDirection);
+					Body->WindData.WindAdaption = FMath::FRandRange(0.0f, 2.0f);
+					Body->WindData.BodyWindScale = WindScale;
+				}
+			}
+		}
+	}
+	else if (bWindWasEnabled)
+	{
+		SCOPE_CYCLE_COUNTER(STAT_AnimDynamicsWindData);
+	
+		bWindWasEnabled = false;
+		for(FAnimPhysRigidBody* Body : BaseBodyPtrs)
+		{
+			Body->bWindEnabled = false;
+		}
+	}
 }
 
 FTransform FAnimNode_AnimDynamics::GetBoneTransformInSimSpace(USkeletalMeshComponent* SkelComp, FCSPose<FCompactPose>& MeshBases, const FCompactPoseBoneIndex& BoneIndex)
diff --git a/Engine/Source/Runtime/AnimGraphRuntime/Public/AnimNodes/AnimNode_TwoWayBlend.h b/Engine/Source/Runtime/AnimGraphRuntime/Public/AnimNodes/AnimNode_TwoWayBlend.h
--- a/Engine/Source/Runtime/AnimGraphRuntime/Public/AnimNodes/AnimNode_TwoWayBlend.h
+++ b/Engine/Source/Runtime/AnimGraphRuntime/Public/AnimNodes/AnimNode_TwoWayBlend.h
@@ -25,9 +25,24 @@ public:
 	UPROPERTY(EditAnywhere, BlueprintReadWrite, Category=Settings)
 	FInputScaleBias AlphaScaleBias;
 
+protected:
+	UPROPERTY(Transient)
+	mutable float InternalBlendAlpha;
+
+	UPROPERTY(Transient)
+	mutable bool bAIsRelevant;
+
+	UPROPERTY(Transient)
+	mutable bool bBIsRelevant;
+
+	/** This reinitializes child pose when re-activated. For example, when active child changes */
+	UPROPERTY(EditAnywhere, Category = Option)
+	bool bResetChildOnActivation;
+
 public:
 	FAnimationNode_TwoWayBlend()
 		: Alpha(0.0f)
+		, bResetChildOnActivation(false)
 	{
 	}
 
diff --git a/Engine/Source/Runtime/AnimGraphRuntime/Public/BoneControllers/AnimNode_AnimDynamics.h b/Engine/Source/Runtime/AnimGraphRuntime/Public/BoneControllers/AnimNode_AnimDynamics.h
--- a/Engine/Source/Runtime/AnimGraphRuntime/Public/BoneControllers/AnimNode_AnimDynamics.h
+++ b/Engine/Source/Runtime/AnimGraphRuntime/Public/BoneControllers/AnimNode_AnimDynamics.h
@@ -255,6 +255,9 @@ struct ANIMGRAPHRUNTIME_API FAnimNode_AnimDynamics : public FAnimNode_SkeletalCo
 	UPROPERTY(EditAnywhere, BlueprintReadWrite, Category = Wind)
 	bool bEnableWind;
 
+	UPROPERTY(Transient)
+	bool bWindWasEnabled;
+
 	/** Scale to apply to calculated wind velocities in the solver */
 	UPROPERTY(EditAnywhere, BlueprintReadWrite, Category = Wind)
 	float WindScale;
diff --git a/Engine/Source/Runtime/CinematicCamera/Public/CineCameraComponent.h b/Engine/Source/Runtime/CinematicCamera/Public/CineCameraComponent.h
--- a/Engine/Source/Runtime/CinematicCamera/Public/CineCameraComponent.h
+++ b/Engine/Source/Runtime/CinematicCamera/Public/CineCameraComponent.h
@@ -222,13 +222,16 @@ public:
 
 #if WITH_EDITORONLY_DATA
 	/** Read-only. Control this value with CurrentFocalLength (and filmback settings). */
-	UPROPERTY(VisibleAnywhere, BlueprintReadOnly, Category = "Current Camera Settings")
+	UPROPERTY(VisibleAnywhere, Category = "Current Camera Settings")
 	float CurrentHorizontalFOV;
 #endif
 	
 	/** Returns the horizonal FOV of the camera with current settings. */
+	UFUNCTION(BlueprintCallable, Category = "Cine Camera")
 	float GetHorizontalFieldOfView() const;
+	
 	/** Returns the vertical FOV of the camera with current settings. */
+	UFUNCTION(BlueprintCallable, Category = "Cine Camera")
 	float GetVerticalFieldOfView() const;
 
 	/** Returns a list of available filmback presets. */
diff --git a/Engine/Source/Runtime/Core/Private/HAL/MallocLeakDetection.cpp b/Engine/Source/Runtime/Core/Private/HAL/MallocLeakDetection.cpp
--- a/Engine/Source/Runtime/Core/Private/HAL/MallocLeakDetection.cpp
+++ b/Engine/Source/Runtime/Core/Private/HAL/MallocLeakDetection.cpp
@@ -11,9 +11,10 @@
 #if MALLOC_LEAKDETECTION
 
 FMallocLeakDetection::FMallocLeakDetection()
-	: bCaptureAllocs(false)
+	: bRecursive(false)
+	, bCaptureAllocs(false)
 	, MinAllocationSize(0)
-	, bDumpOustandingAllocs(false)
+	, TotalTracked(0)
 {
 	Contexts.Empty(16);
 }
@@ -28,14 +29,36 @@ FMallocLeakDetection::~FMallocLeakDetection()
 {	
 }
 
+static uint32 ContextsTLSID = FPlatformTLS::AllocTlsSlot();
+
 void FMallocLeakDetection::PushContext(const FString& Context)
 {
-	Contexts.Push(Context);
+	FMallocLeakDetectionProxy::Get().Lock();
+
+	FScopeLock Lock(&AllocatedPointersCritical);
+
+	TArray<ContextString>* TLContexts = (TArray<ContextString>*)FPlatformTLS::GetTlsValue(ContextsTLSID);
+	if (!TLContexts)
+	{
+		TLContexts = new TArray<ContextString>();
+		FPlatformTLS::SetTlsValue(ContextsTLSID, TLContexts);
+	}
+
+	bRecursive = true;
+
+	ContextString Str;
+	FCString::Strncpy(Str.Buffer, *Context, 128);
+	TLContexts->Push(Str);
+	bRecursive = false;
+
+	FMallocLeakDetectionProxy::Get().Unlock();
 }
 
 void FMallocLeakDetection::PopContext()
 {
-	Contexts.Pop(false);
+	TArray<ContextString>* TLContexts = (TArray<ContextString>*)FPlatformTLS::GetTlsValue(ContextsTLSID);
+	check(TLContexts);
+	TLContexts->Pop(false);
 }
 
 bool FMallocLeakDetection::Exec(UWorld* InWorld, const TCHAR* Cmd, FOutputDevice& Ar)
@@ -61,7 +84,10 @@ bool FMallocLeakDetection::Exec(UWorld* InWorld, const TCHAR* Cmd, FOutputDevice
 			SetAllocationCollection(false);
 
 			UE_LOG(LogConsoleResponse, Display, TEXT("Dumping unique calltacks with %i bytes or more oustanding."), FilterSize);
-			DumpOpenCallstacks(FilterSize);
+
+			FMallocLeakReportOptions Options;
+			Options.FilterSize = FilterSize;
+			DumpOpenCallstacks(Options);
 
 			ClearData();
 		}
@@ -72,7 +98,7 @@ bool FMallocLeakDetection::Exec(UWorld* InWorld, const TCHAR* Cmd, FOutputDevice
 		}
 		else if (FParse::Command(&Cmd, TEXT("Dump")))
 		{
-			uint32 FilterSize = 64;
+			uint32 FilterSize = 128;
 			FParse::Value(Cmd, TEXT("filtersize="), FilterSize);
 			FilterSize *= 1024;
 
@@ -80,14 +106,41 @@ bool FMallocLeakDetection::Exec(UWorld* InWorld, const TCHAR* Cmd, FOutputDevice
 			FParse::Value(Cmd, TEXT("name="), FileName);
 
 			UE_LOG(LogConsoleResponse, Display, TEXT("Dumping unique calltacks with %i KB or more oustanding."), FilterSize/1024);
-			DumpOpenCallstacks(FilterSize, *FileName);
+
+			FMallocLeakReportOptions Options;
+			Options.FilterSize = FilterSize;
+			Options.OutputFile = *FileName;
+
+			DumpOpenCallstacks(Options);
+		}
+		else if (FParse::Command(&Cmd, TEXT("dumpleaks")))
+		{
+			uint32 FilterSize = 0;
+			FParse::Value(Cmd, TEXT("filtersize="), FilterSize);
+			FilterSize *= 1024;
+
+			FString FileName;
+			FParse::Value(Cmd, TEXT("name="), FileName);
+
+			FMallocLeakReportOptions Options;
+			Options.FilterSize = FilterSize;
+			Options.OutputFile = *FileName;
+
+			UE_LOG(LogConsoleResponse, Display, TEXT("Dumping potential leakers with %i KB or more oustanding."), FilterSize / 1024);
+			DumpPotentialLeakers(Options);
+		}
+		else if (FParse::Command(&Cmd, TEXT("Checkpoint")))
+		{
+			CheckpointLinearFit();
 		}
 		else
 		{
 			UE_LOG(LogConsoleResponse, Display, TEXT("'MallocLeak Start [filtersize=Size]'  Begins accumulating unique callstacks of allocations > Size KBytes\n")
-				TEXT("'MallocLeak Stop filtersize=[Size]'  Dumps outstanding unique callstacks and stops accumulation (with an optional filter size in KBytes, if not specified 128 KB is assumed).  Also clears data.\n")
-				TEXT("'MallocLeak Clear'  Clears all accumulated data.  Does not change start/stop state.\n")
-				TEXT("'MallocLeak Dump filtersize=[Size] [save]' Dumps oustanding unique callstacks with optional size filter in KBytes (if size is not specified it is assumed to be 128 KB).\n"));
+				TEXT("'MallocLeak stop filtersize=[Size]'  Dumps outstanding unique callstacks and stops accumulation (with an optional filter size in KBytes, if not specified 128 KB is assumed).  Also clears data.\n")
+				TEXT("'MallocLeak clear'  Clears all accumulated data.  Does not change start/stop state.\n")
+				TEXT("'MallocLeak dump filtersize=[Size] name=[file]' Dumps oustanding unique callstacks with optional size filter in KBytes (if size is not specified it is assumed to be 128 KB).\n")
+				TEXT("'MallocLeak dumpleaks filtersize=[Size] name=[file]' Dumps oustanding unique callstacks with optional size filter in KBytes (if size is not specified it is assumed to be 0).\n")
+			);
 		}
 
 		return true;
@@ -95,19 +148,20 @@ bool FMallocLeakDetection::Exec(UWorld* InWorld, const TCHAR* Cmd, FOutputDevice
 	return false;
 }
 
+
 void FMallocLeakDetection::AddCallstack(FCallstackTrack& Callstack)
 {
 	FScopeLock Lock(&AllocatedPointersCritical);
 	uint32 CallstackHash = Callstack.GetHash();
 	FCallstackTrack& UniqueCallstack = UniqueCallstacks.FindOrAdd(CallstackHash);
 	//if we had a hash collision bail and lose the data rather than corrupting existing data.
-	if (UniqueCallstack.Count > 0 && UniqueCallstack != Callstack)
+	if ((UniqueCallstack.Count > 0 || UniqueCallstack.NumCheckPoints > 0) && UniqueCallstack != Callstack)
 	{
 		ensureMsgf(false, TEXT("Callstack hash collision.  Throwing away new stack."));
 		return;
 	}
 
-	if (UniqueCallstack.Count == 0)
+	if (UniqueCallstack.Count == 0 && UniqueCallstack.NumCheckPoints == 0)
 	{
 		UniqueCallstack = Callstack;
 	}
@@ -117,6 +171,7 @@ void FMallocLeakDetection::AddCallstack(FCallstackTrack& Callstack)
 		UniqueCallstack.LastFrame = Callstack.LastFrame;
 	}
 	UniqueCallstack.Count++;	
+	TotalTracked += Callstack.Size;
 }
 
 void FMallocLeakDetection::RemoveCallstack(FCallstackTrack& Callstack)
@@ -133,6 +188,7 @@ void FMallocLeakDetection::RemoveCallstack(FCallstackTrack& Callstack)
 			UniqueCallstack = nullptr;
 			UniqueCallstacks.Remove(CallstackHash);
 		}
+		TotalTracked -= Callstack.Size;
 	}
 }
 
@@ -147,61 +203,86 @@ void FMallocLeakDetection::SetAllocationCollection(bool bEnabled, int32 Size)
 	}
 }
 
-void FMallocLeakDetection::DumpOpenCallstacks(uint32 FilterSize, const TCHAR* FileName /*= nullptr*/)
+void FMallocLeakDetection::GetOpenCallstacks(TArray<uint32>& OutCallstacks, const FMallocLeakReportOptions& Options /*= FMallocLeakReportOptions()*/)
 {
-	const int32 MaxCallstackLineChars = 2048;
-	char CallstackString[MaxCallstackLineChars];
-	TCHAR CallstackStringWide[MaxCallstackLineChars];
-	FMemory::Memzero(CallstackString);
-	TArray<uint32> SortedKeys;
-	SIZE_T TotalSize = 0;
-	SIZE_T ReportedSize = 0;
-
-	// Make sure we preallocate memory so GetKeys() doesn't allocate!
-	SortedKeys.Empty(UniqueCallstacks.Num() + 32);
+	// Make sure we preallocate memory to avoid
+	OutCallstacks.Empty(UniqueCallstacks.Num() + 32);
 
 	{
 		FScopeLock Lock(&AllocatedPointersCritical);
 
 		for (const auto& Pair : UniqueCallstacks)
 		{
-			if (Pair.Value.Size >= FilterSize)
+
+			// Filter on type?
+			if (Options.OnlyNonDeleters && KnownDeleters.Contains(Pair.Value.GetHash()))
 			{
-				SortedKeys.Add(Pair.Key);
+				continue;
+			}
 
-				ReportedSize += Pair.Value.Size;
+			// frame number?
+			if (Options.FrameStart > Pair.Value.LastFrame)
+			{
+				continue;
 			}
 
-			TotalSize += Pair.Value.Size;
-		}
+			if (Options.FrameEnd && Options.FrameEnd < Pair.Value.LastFrame)
+			{
+				continue;
+			}
 
-		SortedKeys.Sort([this](uint32 lhs, uint32 rhs) {
+			// Filter on size?
+			if (Pair.Value.Size < Options.FilterSize)
+			{
+				continue;
+			}
 
-			FCallstackTrack& Left = UniqueCallstacks[lhs];
-			FCallstackTrack& Right = UniqueCallstacks[rhs];
+			OutCallstacks.Add(Pair.Key);
+		}
+	}
+}
 
-			return Right.Size < Left.Size;
-		});
-	}	
+void FMallocLeakDetection::CheckpointLinearFit()
+{
+	FScopeLock Lock(&AllocatedPointersCritical);
 
+	float FrameNum = float(GFrameCounter);
+	float FrameNum2 = FrameNum * FrameNum;
 
-	FOutputDevice* ReportAr = nullptr;
-	FArchive* FileAr = nullptr;
-	FOutputDeviceArchiveWrapper* FileArWrapper = nullptr;
+	for (auto& Pair : UniqueCallstacks)
+	{
+		FCallstackTrack& Callstack = Pair.Value;
+
+		Callstack.NumCheckPoints++;
+		Callstack.SumOfFramesNumbers += FrameNum;
+		Callstack.SumOfFramesNumbersSquared += FrameNum2;
+		Callstack.SumOfMemory += float(Callstack.Size);
+		Callstack.SumOfMemoryTimesFrameNumber += float(Callstack.Size) * FrameNum;
+		Callstack.GetLinearFit();
+	}
+}
 
-	if (FileName)
+void FMallocLeakDetection::FCallstackTrack::GetLinearFit()
+{
+	Baseline = 0.0f;
+	BytesPerFrame = 0.0f;
+	float DetInv = (float(NumCheckPoints) * SumOfFramesNumbersSquared) - (2.0f * SumOfFramesNumbers);
+	if (!NumCheckPoints || DetInv == 0.0f)
 	{
-		const FString PathName = *(FPaths::ProfilingDir() + TEXT("MallocLeak/"));
-		IFileManager::Get().MakeDirectory(*PathName);
-		
-		FString FilePath = PathName + CreateProfileFilename(FileName, TEXT(".report"), true);
+		return;
+	}
+	float Det = 1.0f / DetInv;
 
-		FileAr = IFileManager::Get().CreateDebugFileWriter(*FilePath);
-		FileArWrapper = new FOutputDeviceArchiveWrapper(FileAr);
-		ReportAr = FileArWrapper;
+	float AtAInv[2][2];
 
-		//UE_LOG(LogConsoleResponse, Log, TEXT("MemReportDeferred: saving to %s"), *FilePath);
-	}
+	AtAInv[0][0] = Det * SumOfFramesNumbersSquared;
+	AtAInv[0][1] = -Det * SumOfFramesNumbers;
+	AtAInv[1][0] = -Det * SumOfFramesNumbers;
+	AtAInv[1][1] = Det * float(NumCheckPoints);
+
+	Baseline = AtAInv[0][0] * SumOfMemory + AtAInv[0][1] * SumOfMemoryTimesFrameNumber;
+	BytesPerFrame = AtAInv[1][0] * SumOfMemory + AtAInv[1][1] * SumOfMemoryTimesFrameNumber;
+}
 
 #define LOG_OUTPUT(Format, ...) \
 	do {\
@@ -214,34 +295,132 @@ void FMallocLeakDetection::DumpOpenCallstacks(uint32 FilterSize, const TCHAR* Fi
 		}\
 	} while(0);\
 
-	uint32 EffectiveFilter = FilterSize ? FilterSize : MinAllocationSize;
+int32 FMallocLeakDetection::DumpPotentialLeakers(const FMallocLeakReportOptions& Options /*= FMallocLeakReportOptions()*/)
+{
+	const int32 MaxCallstackLineChars = 2048;
+	char CallstackString[MaxCallstackLineChars];
+	TCHAR CallstackStringWide[MaxCallstackLineChars];
+	FMemory::Memzero(CallstackString);
 
-	const float InvToMb = 1.0 / (1024 * 1024);
-	FPlatformMemoryStats MemoryStats = FPlatformMemory::GetStats();
+	TArray<uint32> CallstackHashes;
 
-	LOG_OUTPUT(TEXT("Current Memory: %.02fMB (Peak: %.02fMB)."),
-		MemoryStats.UsedPhysical * InvToMb,
-		MemoryStats.PeakUsedPhysical * InvToMb);
+	GetOpenCallstacks(CallstackHashes, Options);
 
-	LOG_OUTPUT(TEXT("Leak detection allocation filter: %dKB"), MinAllocationSize / 1024);
-	LOG_OUTPUT(TEXT("Leak detection report filter: %dKB"), EffectiveFilter / 1024);
+	CheckpointLinearFit();
 
-	LOG_OUTPUT(TEXT("Found %i open callstacks that hold %uMB of memory on frame: %i"), UniqueCallstacks.Num(), TotalSize / (1024 * 1024), (int32)GFrameCounter);
-	LOG_OUTPUT(TEXT("Dumping %d callstacks that hold more than %uKBs and account for %uMB"), SortedKeys.Num(), EffectiveFilter / 1024, ReportedSize / (1024 * 1024));
-		
-	for (const auto& Key : SortedKeys)
+	const int kMaxDump = 30;
+
+	TMap<uint32, float> HashesToLeakRate;
+	float TotalLeakRate = 0.0f;
+	SIZE_T TotalSize = 0;
+
+
+	for (uint32 Hash : CallstackHashes)
+	{
+		FCallstackTrack Callstack;
+
+		{
+			FScopeLock Lock(&AllocatedPointersCritical);
+			if (UniqueCallstacks.Contains(Hash) == false)
+				continue;
+			
+			Callstack = UniqueCallstacks[Hash];
+		}
+				
+		if (Callstack.NumCheckPoints > 5 && Callstack.BytesPerFrame > .01f && Callstack.Count > 5)
+		{		
+			HashesToLeakRate.Add(Hash) = Callstack.BytesPerFrame;
+
+			TotalLeakRate += Callstack.BytesPerFrame;
+			TotalSize += Callstack.Size;
+
+			if (HashesToLeakRate.Num() > kMaxDump)
+			{
+				float SmallestValue = MAX_FLT;
+				uint32 SmallestKey = 0;
+				for (const auto& KV : HashesToLeakRate)
+				{
+					if (KV.Value < SmallestValue)
+					{
+						SmallestValue = Callstack.BytesPerFrame;
+						SmallestKey = KV.Key;
+					}
+				}
+
+				check(SmallestKey);
+				HashesToLeakRate.Remove(SmallestKey);
+				TotalLeakRate -= SmallestValue;
+				TotalSize -= UniqueCallstacks[SmallestKey].Size;
+			}
+		}
+	}
+
+	if (HashesToLeakRate.Num() > 0)
 	{
-		const FCallstackTrack& Callstack = UniqueCallstacks[Key];
-		if (Callstack.Size >= FilterSize)
+		FOutputDevice* ReportAr = nullptr;
+		FArchive* FileAr = nullptr;
+		FOutputDeviceArchiveWrapper* FileArWrapper = nullptr;
+
+		if (Options.OutputFile && FCString::Strlen(Options.OutputFile))
 		{
+			const FString PathName = *(FPaths::ProfilingDir() + TEXT("memreports/"));
+			IFileManager::Get().MakeDirectory(*PathName);
+
+			FString FilePath = PathName + CreateProfileFilename(Options.OutputFile, TEXT(".leaks"), true);
+
+			FileAr = IFileManager::Get().CreateDebugFileWriter(*FilePath);
+			FileArWrapper = new FOutputDeviceArchiveWrapper(FileAr);
+			ReportAr = FileArWrapper;
+
+			//UE_LOG(LogConsoleResponse, Log, TEXT("MemReportDeferred: saving to %s"), *FilePath);
+		}
+
+		const float InvToMb = 1.0 / (1024 * 1024);
+		FPlatformMemoryStats MemoryStats = FPlatformMemory::GetStats();
+
+		LOG_OUTPUT(TEXT("Current Time: %s, Current Frame %d"), *FDateTime::Now().ToString(TEXT("%m.%d-%H.%M.%S")), GFrameCounter);
+
+		LOG_OUTPUT(TEXT("Current Memory: %.02fMB (Peak: %.02fMB)."),
+			MemoryStats.UsedPhysical * InvToMb,
+			MemoryStats.PeakUsedPhysical * InvToMb);
+
+		LOG_OUTPUT(TEXT("Tracking %d callstacks that hold %.02fMB"), UniqueCallstacks.Num(), TotalTracked * InvToMb);
+
+		LOG_OUTPUT(TEXT("Leak detection allocation filter: %dKB"), MinAllocationSize / 1024);
+		LOG_OUTPUT(TEXT("Leak detection report filter: %dKB"), Options.FilterSize / 1024);
+		LOG_OUTPUT(TEXT("Dumping %d callstacks that hold %.02fMB and may be leaking memory."), HashesToLeakRate.Num(), TotalSize * InvToMb);
+		LOG_OUTPUT(TEXT("Current estimated leak rate: %.02fB/frame\n.") , TotalLeakRate);
+
+		TArray<uint32> SortedKeys;
+		HashesToLeakRate.GetKeys(SortedKeys);
+
+		SortedKeys.Sort([&](uint32 lhs, uint32 rhs) {
+			//return HashesToLeakRate[rhs] < HashesToLeakRate[lhs];
+			return lhs < rhs;
+		});
+
+		for (uint32 Key : SortedKeys)
+		{
+			FCallstackTrack Callstack;
+
+			{
+				FScopeLock Lock(&AllocatedPointersCritical);
+				if (UniqueCallstacks.Contains(Key) == false)
+					continue;
+
+				Callstack = UniqueCallstacks[Key];
+			}
+
 			bool KnownDeleter = KnownDeleters.Contains(Callstack.GetHash());
 
-			LOG_OUTPUT(TEXT("\nAllocSize: %i KB, Num: %i, FirstFrame %i, LastFrame %i, KnownDeleter: %d")
-				, Callstack.Size / 1024
+			LOG_OUTPUT(TEXT("AllocSize: %i KB, Num: %i, FirstFrame %i, LastFrame %i, KnownDeleter: %d, Baseline memory %.3fMB  Leak Rate %.2fB/frame"),
+				Callstack.Size / 1024
 				, Callstack.Count
 				, Callstack.FirstFame
 				, Callstack.LastFrame
-				, KnownDeleter);
+				, KnownDeleter
+				, Callstack.Baseline / 1024.0f / 1024.0f
+				, Callstack.BytesPerFrame);
 
 			for (int32 i = 0; i < FCallstackTrack::Depth; ++i)
 			{
@@ -258,38 +437,187 @@ void FMallocLeakDetection::DumpOpenCallstacks(uint32 FilterSize, const TCHAR* Fi
 					LOG_OUTPUT(TEXT("%s"), CallstackStringWide);
 				}
 				FMemory::Memzero(CallstackString);
-			}			
-
-			TArray<FString> SortedContexts;
+			}
 
-			for (const auto& Pair : OpenPointers)
 			{
-				if (Pair.Value.GetHash() == Key)
+				FScopeLock Lock(&AllocatedPointersCritical);
+
+				TArray<FString> SortedContexts;
+
+				for (const auto& Pair : OpenPointers)
 				{
-					if (PointerContexts.Contains(Pair.Key))
+					if (Pair.Value.GetHash() == Key)
 					{
-						SortedContexts.Add(*PointerContexts.Find(Pair.Key));
+						if (PointerContexts.Contains(Pair.Key))
+						{
+							SortedContexts.Add(*PointerContexts.Find(Pair.Key));
+						}
+					}
+				}
+
+				if (SortedContexts.Num())
+				{
+					LOG_OUTPUT(TEXT("%d contexts:"), SortedContexts.Num(), CallstackStringWide);
+
+					SortedContexts.Sort();
+
+					for (const auto& Str : SortedContexts)
+					{
+						LOG_OUTPUT(TEXT("\t%s"), *Str);
 					}
 				}
 			}
 
-			if (SortedContexts.Num())
+			LOG_OUTPUT(TEXT("\n"));
+		}
+
+		if (FileArWrapper != nullptr)
+		{
+			FileArWrapper->TearDown();
+			delete FileArWrapper;
+			delete FileAr;
+		}
+	}
+
+	return HashesToLeakRate.Num();
+}
+
+int32 FMallocLeakDetection::DumpOpenCallstacks(const FMallocLeakReportOptions& Options /*= FMallocLeakReportOptions()*/)
+{
+	const int32 MaxCallstackLineChars = 2048;
+	char CallstackString[MaxCallstackLineChars];
+	TCHAR CallstackStringWide[MaxCallstackLineChars];
+	FMemory::Memzero(CallstackString);
+
+	SIZE_T ReportedSize = 0;
+	SIZE_T TotalSize = 0;
+
+	TArray<uint32> SortedKeys;
+
+	GetOpenCallstacks(SortedKeys, Options);
+
+	{
+		FScopeLock Lock(&AllocatedPointersCritical);
+
+		SortedKeys.Sort([this](uint32 lhs, uint32 rhs) {
+
+			FCallstackTrack& Left = UniqueCallstacks[lhs];
+			FCallstackTrack& Right = UniqueCallstacks[rhs];
+
+			//return Right.Size < Left.Size;
+			return Left.GetHash() < Right.GetHash();
+		});
+
+		for (uint32 ID : SortedKeys)
+		{
+			ReportedSize += UniqueCallstacks[ID].Size;
+		}
+	}
+	
+	FOutputDevice* ReportAr = nullptr;
+	FArchive* FileAr = nullptr;
+	FOutputDeviceArchiveWrapper* FileArWrapper = nullptr;
+
+	if (Options.OutputFile && FCString::Strlen(Options.OutputFile))
+	{
+		const FString PathName = *(FPaths::ProfilingDir() + TEXT("memreports/"));
+		IFileManager::Get().MakeDirectory(*PathName);
+		
+		FString FilePath = PathName + CreateProfileFilename(Options.OutputFile, TEXT(".allocs"), true);
+
+		FileAr = IFileManager::Get().CreateDebugFileWriter(*FilePath);
+		FileArWrapper = new FOutputDeviceArchiveWrapper(FileAr);
+		ReportAr = FileArWrapper;
+
+		//UE_LOG(LogConsoleResponse, Log, TEXT("MemReportDeferred: saving to %s"), *FilePath);
+	}
+
+	const float InvToMb = 1.0 / (1024 * 1024);
+	FPlatformMemoryStats MemoryStats = FPlatformMemory::GetStats();
+
+	LOG_OUTPUT(TEXT("Current Time: %s, Current Frame %d"), *FDateTime::Now().ToString(TEXT("%m.%d-%H.%M.%S")), GFrameCounter);
+
+	LOG_OUTPUT(TEXT("Current Memory: %.02fMB (Peak: %.02fMB)."),
+		MemoryStats.UsedPhysical * InvToMb,
+		MemoryStats.PeakUsedPhysical * InvToMb);
+
+	LOG_OUTPUT(TEXT("Tracking %d callstacks that hold %.02fMB"), UniqueCallstacks.Num(), TotalTracked * InvToMb);
+
+	LOG_OUTPUT(TEXT("Leak detection allocation filter: %dKB"), MinAllocationSize / 1024);
+	LOG_OUTPUT(TEXT("Leak detection report filter: %dKB"), Options.FilterSize / 1024);
+
+	LOG_OUTPUT(TEXT("Have  %i open callstacks at frame: %i"), UniqueCallstacks.Num(),(int32)GFrameCounter);
+	
+	//for (const auto& Key : SortedKeys)
+
+	LOG_OUTPUT(TEXT("Dumping %d callstacks that hold more than %uKBs and total %uKBs"),
+		SortedKeys.Num(), Options.FilterSize / 1024, ReportedSize / 1024);
+
+	for (const auto& Key : SortedKeys)
+	{		
+		FCallstackTrack Callstack;
+
+		{
+			FScopeLock Lock(&AllocatedPointersCritical);
+			if (UniqueCallstacks.Contains(Key) == false)
+				continue;
+
+			Callstack = UniqueCallstacks[Key];
+		}
+
+		bool KnownDeleter = KnownDeleters.Contains(Callstack.GetHash());
+
+		LOG_OUTPUT(TEXT("\nAllocSize: %i KB, Num: %i, FirstFrame %i, LastFrame %i, KnownDeleter: %d")
+			, Callstack.Size / 1024
+			, Callstack.Count
+			, Callstack.FirstFame
+			, Callstack.LastFrame
+			, KnownDeleter);
+
+		for (int32 i = 0; i < FCallstackTrack::Depth; ++i)
+		{
+			FPlatformStackWalk::ProgramCounterToHumanReadableString(i, Callstack.CallStack[i], CallstackString, MaxCallstackLineChars);
+
+			if (Callstack.CallStack[i] && FCStringAnsi::Strlen(CallstackString) > 0)
 			{
-				LOG_OUTPUT(TEXT("%d contexts:"), SortedContexts.Num(), CallstackStringWide);
+				//convert ansi -> tchar without mallocs in case we are in OOM handler.
+				for (int32 CurrChar = 0; CurrChar < MaxCallstackLineChars; ++CurrChar)
+				{
+					CallstackStringWide[CurrChar] = CallstackString[CurrChar];
+				}
 
-				SortedContexts.Sort();
+				LOG_OUTPUT(TEXT("%s"), CallstackStringWide);
+			}
+			FMemory::Memzero(CallstackString);
+		}
 
-				for (const auto& Str : SortedContexts)
+		TArray<FString> SortedContexts;
+
+		for (const auto& Pair : OpenPointers)
+		{
+			if (Pair.Value.GetHash() == Key)
+			{
+				if (PointerContexts.Contains(Pair.Key))
 				{
-					LOG_OUTPUT(TEXT("\t%s"), *Str);
+					SortedContexts.Add(*PointerContexts.Find(Pair.Key));
 				}
 			}
+		}
 
-			LOG_OUTPUT(TEXT("\n"), CallstackStringWide);
+		if (SortedContexts.Num())
+		{
+			LOG_OUTPUT(TEXT("%d contexts:"), SortedContexts.Num(), CallstackStringWide);
+
+			SortedContexts.Sort();
+
+			for (const auto& Str : SortedContexts)
+			{
+				LOG_OUTPUT(TEXT("\t%s"), *Str);
+			}
 		}
-	}
 
-#undef LOG_OUTPUT
+		LOG_OUTPUT(TEXT("\n"));
+	}
 
 	if (FileArWrapper != nullptr)
 	{
@@ -297,14 +625,28 @@ void FMallocLeakDetection::DumpOpenCallstacks(uint32 FilterSize, const TCHAR* Fi
 		delete FileArWrapper;
 		delete FileAr;
 	}
+
+	return SortedKeys.Num();
 }
 
+#undef LOG_OUTPUT
+
 void FMallocLeakDetection::ClearData()
 {
-	FScopeLock Lock(&AllocatedPointersCritical);
+	bool OldCapture = bCaptureAllocs;
+	{
+		FScopeLock Lock(&AllocatedPointersCritical);
+		bCaptureAllocs = false;
+	}
+
 	OpenPointers.Empty();
 	UniqueCallstacks.Empty();
 	PointerContexts.Empty();
+
+	{
+		FScopeLock Lock(&AllocatedPointersCritical);
+		bCaptureAllocs = OldCapture;
+	}
 }
 
 void FMallocLeakDetection::Malloc(void* Ptr, SIZE_T Size)
@@ -312,8 +654,9 @@ void FMallocLeakDetection::Malloc(void* Ptr, SIZE_T Size)
 	if (Ptr)
 	{
 		if (bCaptureAllocs && (MinAllocationSize == 0 || Size >= MinAllocationSize))
-		{
+		{			
 			FScopeLock Lock(&AllocatedPointersCritical);
+
 			if (!bRecursive)
 			{
 				bRecursive = true;
@@ -325,9 +668,11 @@ void FMallocLeakDetection::Malloc(void* Ptr, SIZE_T Size)
 				AddCallstack(Callstack);
 				OpenPointers.Add(Ptr, Callstack);
 
-				if (Contexts.Num())
+				TArray<ContextString>* TLContexts = (TArray<ContextString>*)FPlatformTLS::GetTlsValue(ContextsTLSID);
+
+				if (TLContexts && TLContexts->Num())
 				{
-					PointerContexts.Add(Ptr, Contexts.Top());
+					PointerContexts.Add(Ptr, TLContexts->Top().Buffer);
 				}
 
 				bRecursive = false;
@@ -338,11 +683,63 @@ void FMallocLeakDetection::Malloc(void* Ptr, SIZE_T Size)
 
 void FMallocLeakDetection::Realloc(void* OldPtr, void* NewPtr, SIZE_T Size)
 {	
-	if (bCaptureAllocs)
+	if (bRecursive == false && (bCaptureAllocs || OpenPointers.Num()))
 	{
-		FScopeLock Lock(&AllocatedPointersCritical);		
-		Free(OldPtr);
-		Malloc(NewPtr, Size);
+		FScopeLock Lock(&AllocatedPointersCritical);
+
+		// realloc may return the same pointer, if so skip this
+	
+		if (OldPtr != NewPtr)
+		{
+			FCallstackTrack* Callstack = OpenPointers.Find(OldPtr);
+			uint32 OldHash = 0;
+			bool WasKnownDeleter = false;
+			if (Callstack)
+			{
+				OldHash = Callstack->GetHash();
+				WasKnownDeleter = KnownDeleters.Contains(OldHash);
+			}
+
+			// Note - malloc and then free so linearfit checkpoints in the callstack are
+			// preserved
+			Malloc(NewPtr, Size);
+
+			// See if we should/need to copy the context across
+			if (OldPtr && NewPtr)
+			{
+				FString* OldContext = PointerContexts.Find(OldPtr);
+				FString* NewContext = PointerContexts.Find(NewPtr);
+
+				if (OldContext && !NewContext)
+				{
+					FString Tmp = *OldContext;
+					PointerContexts.Add(NewPtr) = Tmp;
+				}
+			}
+
+			Free(OldPtr);
+
+			// if we had a tracked callstack, and it was not a deleter, remark it as such because this 'free'
+			// was not actually deallocating memory.
+			if (Size > 0 && OldHash > 0 && WasKnownDeleter == false)
+			{
+				KnownDeleters.Remove(OldHash);
+			}
+		}
+		else
+		{
+			// realloc returned the same pointer, if there was a context when the call happened then
+			// update it.
+			TArray<ContextString>* TLContexts = (TArray<ContextString>*)FPlatformTLS::GetTlsValue(ContextsTLSID);
+
+			if (TLContexts && TLContexts->Num())
+			{
+				bRecursive = true;
+				PointerContexts.Remove(OldPtr);
+				PointerContexts.FindOrAdd(NewPtr) = TLContexts->Top().Buffer;
+				bRecursive = false;
+			}
+		}
 	}
 }
 
@@ -350,7 +747,7 @@ void FMallocLeakDetection::Free(void* Ptr)
 {
 	if (Ptr)
 	{
-		if (bCaptureAllocs)
+		if (bCaptureAllocs || OpenPointers.Num())
 		{
 			FScopeLock Lock(&AllocatedPointersCritical);			
 			if (!bRecursive)
@@ -363,18 +760,31 @@ void FMallocLeakDetection::Free(void* Ptr)
 					KnownDeleters.Add(Callstack->GetHash());
 				}
 				OpenPointers.Remove(Ptr);
-				PointerContexts.Remove(Ptr);
 				bRecursive = false;
 			}
+
+			if (PointerContexts.Contains(Ptr))
+			{
+				PointerContexts.Remove(Ptr);
+			}
 		}
 	}
 }
 
+static FMallocLeakDetectionProxy* ProxySingleton = nullptr;
+
 FMallocLeakDetectionProxy::FMallocLeakDetectionProxy(FMalloc* InMalloc)
 	: UsedMalloc(InMalloc)
 	, Verify(FMallocLeakDetection::Get())
 {
-	
+	check(ProxySingleton == nullptr);
+	ProxySingleton = this;
+}
+
+FMallocLeakDetectionProxy& FMallocLeakDetectionProxy::Get()
+{
+	check(ProxySingleton);
+	return *ProxySingleton;
 }
 
 #endif // MALLOC_LEAKDETECTION
\ No newline at end of file
diff --git a/Engine/Source/Runtime/Core/Private/HAL/MallocLeakDetection.h b/Engine/Source/Runtime/Core/Private/HAL/MallocLeakDetection.h
--- a/Engine/Source/Runtime/Core/Private/HAL/MallocLeakDetection.h
+++ b/Engine/Source/Runtime/Core/Private/HAL/MallocLeakDetection.h
@@ -14,11 +14,27 @@
 
 #if MALLOC_LEAKDETECTION
 
+
+struct FMallocLeakReportOptions
+{
+	FMallocLeakReportOptions()
+	{
+		FMemory::Memzero(this, sizeof(FMallocLeakReportOptions));
+	}
+	uint32			FilterSize;
+	float			LeakRate;
+	bool			OnlyNonDeleters;
+	uint32			FrameStart;
+	uint32			FrameEnd;
+	const TCHAR*	OutputFile;
+};
+
 /**
  * Maintains a list of all pointers to currently allocated memory.
  */
 class FMallocLeakDetection
 {
+
 	struct FCallstackTrack
 	{
 		FCallstackTrack()
@@ -29,9 +45,21 @@ class FMallocLeakDetection
 		uint64 CallStack[Depth];
 		uint32 FirstFame;
 		uint32 LastFrame;
-		uint32 Size;
+		uint64 Size;
 		uint32 Count;
 
+		// least square line fit stuff
+		uint32 NumCheckPoints;
+		float SumOfFramesNumbers;
+		float SumOfFramesNumbersSquared;
+		float SumOfMemory;
+		float SumOfMemoryTimesFrameNumber;
+
+		// least square line results
+		float Baseline;
+		float BytesPerFrame;
+
+
 		bool operator==(const FCallstackTrack& Other) const
 		{
 			bool bEqual = true;
@@ -51,15 +79,22 @@ class FMallocLeakDetection
 			return !(*this == Other);
 		}
 
+		void GetLinearFit();
+		
 		uint32 GetHash() const 
 		{
 			return FCrc::MemCrc32(CallStack, sizeof(CallStack), 0);
 		}
 	};
 
+private:
+
 	FMallocLeakDetection();
 	~FMallocLeakDetection();
 
+	/** Track a callstack */
+
+	/** Stop tracking a callstack */
 	void AddCallstack(FCallstackTrack& Callstack);
 	void RemoveCallstack(FCallstackTrack& Callstack);
 
@@ -75,24 +110,49 @@ class FMallocLeakDetection
 	/** Contexts that are associated with allocations */
 	TMap<void*, FString>		PointerContexts;
 
-	/** Stack of contects */
-	TArray<FString>				Contexts;
+	/** Stack of contexts */
+	struct ContextString { TCHAR Buffer[128]; };
+	TArray<ContextString>	Contexts;
+		
+	/** Critical section for mutating internal data */
+	FCriticalSection AllocatedPointersCritical;	
 
+	/** Set during mutating operationms to prevent internal allocations from recursing */
 	bool	bRecursive;
+
+	/** Is allocation capture enabled? */
 	bool	bCaptureAllocs;
+
+	/** Minimal size to capture? */
 	int32	MinAllocationSize;
-	bool	bDumpOustandingAllocs;	
-		
-	FCriticalSection AllocatedPointersCritical;	
+
+	SIZE_T	TotalTracked;
 
 public:	
 
 	static FMallocLeakDetection& Get();
+	static void HandleMallocLeakCommand(const TArray< FString >& Args);
+
+	/** Enable/disable collection of allocation with an optional filter on allocation size */
+	void SetAllocationCollection(bool bEnabled, int32 Size = 0);
+
+	/** Returns state of allocation collection */
+	bool IsAllocationCollectionEnabled(void) const { return bCaptureAllocs; }
 
-	void SetAllocationCollection(bool bEnabled, int32 Size=0);
-	void DumpOpenCallstacks(uint32 FilterSize = 0, const TCHAR* FileName=nullptr);
+	/** Clear currently accumulated data */
 	void ClearData();
 
+	/** Dumps callstacks that appear to be leaks based on allocation trends */
+	int32 DumpPotentialLeakers(const FMallocLeakReportOptions& Options = FMallocLeakReportOptions());
+
+	/** Dumps currently open callstacks */
+	int32 DumpOpenCallstacks(const FMallocLeakReportOptions& Options = FMallocLeakReportOptions());
+
+	/** Perform a linear fit checkpoint of all open callstacks */
+	void CheckpointLinearFit();
+
+
+	/** Cmd handler */
 	bool Exec(UWorld* InWorld, const TCHAR* Cmd, FOutputDevice& Ar);
 
 	/** Handles new allocated pointer */
@@ -104,11 +164,13 @@ public:
 	/** Removes allocated pointer from list */
 	void Free(void* Ptr);	
 
-	/** Push/Pop contects that will be associated with allocations. This is very rough.
-	 not thread-safe (let alone TLS), and should probably be removed. However if you really
-	 need to associate state with allocations it's very useful.... Cavet Emptor! */
+	/** Push/Pop a context that will be associated with allocations. All open contexts will be displayed alongside
+	callstacks in a report.  */
 	void PushContext(const FString& Context);
 	void PopContext();
+
+	/** Returns */
+	void GetOpenCallstacks(TArray<uint32>& OutCallstacks, const FMallocLeakReportOptions& Options = FMallocLeakReportOptions());
 };
 
 /**
@@ -129,6 +191,8 @@ private:
 public:
 	explicit FMallocLeakDetectionProxy(FMalloc* InMalloc);	
 
+	static FMallocLeakDetectionProxy& Get();
+
 	virtual void* Malloc(SIZE_T Size, uint32 Alignment) override
 	{
 		FScopeLock SafeLock(&AllocatedPointersCritical);
@@ -210,11 +274,20 @@ public:
 		return UsedMalloc->ClearAndDisableTLSCachesOnCurrentThread();
 	}
 
-
 	virtual const TCHAR* GetDescriptiveName() override
 	{ 
 		return UsedMalloc->GetDescriptiveName();
 	}
+
+	void Lock()
+	{
+		AllocatedPointersCritical.Lock();
+	}
+
+	void Unlock()
+	{
+		AllocatedPointersCritical.Unlock();
+	}
 };
 
 #endif // MALLOC_LEAKDETECTION
diff --git a/Engine/Source/Runtime/Core/Private/HAL/ThreadHeartBeat.cpp b/Engine/Source/Runtime/Core/Private/HAL/ThreadHeartBeat.cpp
--- a/Engine/Source/Runtime/Core/Private/HAL/ThreadHeartBeat.cpp
+++ b/Engine/Source/Runtime/Core/Private/HAL/ThreadHeartBeat.cpp
@@ -11,6 +11,8 @@ FThreadHeartBeat::FThreadHeartBeat()
 	: Thread(nullptr)
 	, bReadyToCheckHeartbeat(false)
 	, HangDuration(25.0)
+	, LastHangCallstackCRC(0)
+	, LastHungThreadId(0)
 {
 	if (GConfig)
 	{
@@ -89,55 +91,76 @@ bool FThreadHeartBeat::Init()
 
 uint32 FThreadHeartBeat::Run()
 {
+	bool InHungState = false;
+
 	while (StopTaskCounter.GetValue() == 0)
 	{
 		uint32 ThreadThatHung = CheckHeartBeat();
-		if (ThreadThatHung != FThreadHeartBeat::InvalidThreadId)
+
+		if (ThreadThatHung == FThreadHeartBeat::InvalidThreadId)
+		{
+			InHungState = false;
+		}
+		else if (InHungState == false)
 		{
+			// Only want to call this once per hang (particularly if we're just ensuring).
+			InHungState = true;
+
 			const SIZE_T StackTraceSize = 65535;
 			ANSICHAR* StackTrace = (ANSICHAR*)GMalloc->Malloc(StackTraceSize);
 			StackTrace[0] = 0;
 			// Walk the stack and dump it to the allocated memory. This process usually allocates a lot of memory.
 			FPlatformStackWalk::ThreadStackWalkAndDump(StackTrace, StackTraceSize, 0, ThreadThatHung);
-			FString StackTraceText(StackTrace);
-			TArray<FString> StackLines;
-			StackTraceText.ParseIntoArrayLines(StackLines);
 
-			// Dump the callstack and the thread name to log
-			FString ThreadName(ThreadThatHung == GGameThreadId ? TEXT("GameThread") : FThreadManager::Get().GetThreadName(ThreadThatHung));
-			if (ThreadName.IsEmpty())
-			{
-				ThreadName = FString::Printf(TEXT("unknown thread (%u)"), ThreadThatHung);
-			}
-			UE_LOG(LogCore, Error, TEXT("Hang detected on %s (thread hasn't sent a heartbeat for %.2llf seconds):"), *ThreadName, HangDuration);
-			for (FString& StackLine : StackLines)
-			{
-				UE_LOG(LogCore, Error, TEXT("  %s"), *StackLine);
-			}
-			
-			// Assert (on the current thread unfortunately) with a trimmed stack.
-			FString StackTrimmed;
-			for (int32 LineIndex = 0; LineIndex < StackLines.Num() && StackTrimmed.Len() < 512; ++LineIndex)
+			// First verify we're not reporting the same hang over and over again
+			uint32 CallstackCRC = FCrc::StrCrc32(StackTrace);
+			if (CallstackCRC != LastHangCallstackCRC && ThreadThatHung != LastHungThreadId)
 			{
-				StackTrimmed += TEXT("  ");
-				StackTrimmed += StackLines[LineIndex];
-				StackTrimmed += LINE_TERMINATOR;
-			}
+				LastHangCallstackCRC = CallstackCRC;
+				LastHungThreadId = ThreadThatHung;
+
+				FString StackTraceText(StackTrace);
+				TArray<FString> StackLines;
+				StackTraceText.ParseIntoArrayLines(StackLines);
+
+				// Dump the callstack and the thread name to log
+				FString ThreadName(ThreadThatHung == GGameThreadId ? TEXT("GameThread") : FThreadManager::Get().GetThreadName(ThreadThatHung));
+				if (ThreadName.IsEmpty())
+				{
+					ThreadName = FString::Printf(TEXT("unknown thread (%u)"), ThreadThatHung);
+				}
+				UE_LOG(LogCore, Error, TEXT("Hang detected on %s (thread hasn't sent a heartbeat for %.2llf seconds):"), *ThreadName, HangDuration);
+				for (FString& StackLine : StackLines)
+				{
+					UE_LOG(LogCore, Error, TEXT("  %s"), *StackLine);
+				}
 
-			const FString ErrorMessage = FString::Printf(TEXT("Hang detected on %s:%s%s%sCheck log for full callstack."), *ThreadName, LINE_TERMINATOR, *StackTrimmed, LINE_TERMINATOR);
+				// Assert (on the current thread unfortunately) with a trimmed stack.
+				FString StackTrimmed;
+				for (int32 LineIndex = 0; LineIndex < StackLines.Num() && StackTrimmed.Len() < 512; ++LineIndex)
+				{
+					StackTrimmed += TEXT("  ");
+					StackTrimmed += StackLines[LineIndex];
+					StackTrimmed += LINE_TERMINATOR;
+				}
+
+				const FString ErrorMessage = FString::Printf(TEXT("Hang detected on %s:%s%s%sCheck log for full callstack."), *ThreadName, LINE_TERMINATOR, *StackTrimmed, LINE_TERMINATOR);
 #if UE_ASSERT_ON_HANG
-			UE_LOG(LogCore, Fatal, TEXT("%s"), *ErrorMessage);
+				UE_LOG(LogCore, Fatal, TEXT("%s"), *ErrorMessage);
 #else
-			UE_LOG(LogCore, Error, TEXT("%s"), *ErrorMessage);
+				UE_LOG(LogCore, Error, TEXT("%s"), *ErrorMessage);
 #if PLATFORM_DESKTOP
-			GLog->PanicFlushThreadedLogs();
-			// GErrorMessage here is very unfortunate but it's used internally by the crash context code.
-			FCString::Strcpy(GErrorMessage, ARRAY_COUNT(GErrorMessage), *ErrorMessage);
-			// Skip macros and FDebug, we always want this to fire
-			NewReportEnsure(*ErrorMessage);
-			GErrorMessage[0] = '\0';
+				GLog->PanicFlushThreadedLogs();
+				// GErrorMessage here is very unfortunate but it's used internally by the crash context code.
+				FCString::Strcpy(GErrorMessage, ARRAY_COUNT(GErrorMessage), *ErrorMessage);
+				// Skip macros and FDebug, we always want this to fire
+				NewReportEnsure(*ErrorMessage);
+				GErrorMessage[0] = '\0';
 #endif
 #endif
+			}
+
+			GMalloc->Free(StackTrace);
 		}
 		FPlatformProcess::SleepNoStats(0.5f);
 	}
@@ -218,3 +241,16 @@ void FThreadHeartBeat::ResumeHeartBeat()
 		}
 	}
 }
+
+bool FThreadHeartBeat::IsBeating()
+{
+	uint32 ThreadId = FPlatformTLS::GetCurrentThreadId();
+	FScopeLock HeartBeatLock(&HeartBeatCritical);
+	FHeartBeatInfo* HeartBeatInfo = ThreadHeartBeat.Find(ThreadId);
+	if (HeartBeatInfo && HeartBeatInfo->SuspendedCount == 0)
+	{
+		return true;
+	}
+
+	return false;
+}
diff --git a/Engine/Source/Runtime/Core/Private/Internationalization/Text.cpp b/Engine/Source/Runtime/Core/Private/Internationalization/Text.cpp
--- a/Engine/Source/Runtime/Core/Private/Internationalization/Text.cpp
+++ b/Engine/Source/Runtime/Core/Private/Internationalization/Text.cpp
@@ -76,6 +76,16 @@ uint32 FTextInspector::GetFlags(const FText& Text)
 	return Text.Flags;
 }
 
+void FTextInspector::GetHistoricFormatData(const FText& Text, TArray<FHistoricTextFormatData>& OutHistoricFormatData)
+{
+	Text.GetHistoricFormatData(OutHistoricFormatData);
+}
+
+bool FTextInspector::GetHistoricNumericData(const FText& Text, FHistoricTextNumericData& OutHistoricNumericData)
+{
+	return Text.GetHistoricNumericData(OutHistoricNumericData);
+}
+
 // These default values have been duplicated to the KismetTextLibrary functions for Blueprints. Please replicate any changes there!
 FNumberFormattingOptions::FNumberFormattingOptions()
 	: UseGrouping(true)
@@ -915,6 +925,16 @@ const FString& FText::GetSourceString() const
 	return TextData->GetDisplayString();
 }
 
+void FText::GetHistoricFormatData(TArray<FHistoricTextFormatData>& OutHistoricFormatData) const
+{
+	TextData->GetTextHistory().GetHistoricFormatData(*this, OutHistoricFormatData);
+}
+
+bool FText::GetHistoricNumericData(FHistoricTextNumericData& OutHistoricNumericData) const
+{
+	return TextData->GetTextHistory().GetHistoricNumericData(*this, OutHistoricNumericData);
+}
+
 bool FText::IdenticalTo( const FText& Other ) const
 {
 	// If both instances point to the same data or localized string, then both instances are considered identical.
@@ -923,11 +943,6 @@ bool FText::IdenticalTo( const FText& Other ) const
 	return TextData == Other.TextData || TextData->GetLocalizedString() == Other.TextData->GetLocalizedString();
 }
 
-void FText::GetSourceTextsFromFormatHistory(TArray<FText>& OutSourceTexts) const
-{
-	TextData->GetTextHistory().GetSourceTextsFromFormatHistory(*this, OutSourceTexts);
-}
-
 FArchive& operator<<(FArchive& Ar, FFormatArgumentValue& Value)
 {
 	int8 TypeAsInt8 = Value.Type;
diff --git a/Engine/Source/Runtime/Core/Private/Internationalization/TextHistory.cpp b/Engine/Source/Runtime/Core/Private/Internationalization/TextHistory.cpp
--- a/Engine/Source/Runtime/Core/Private/Internationalization/TextHistory.cpp
+++ b/Engine/Source/Runtime/Core/Private/Internationalization/TextHistory.cpp
@@ -41,10 +41,13 @@ const FString* FTextHistory::GetSourceString() const
 	return NULL;
 }
 
-void FTextHistory::GetSourceTextsFromFormatHistory(FText Text, TArray<FText>& OutSourceTexts) const
+void FTextHistory::GetHistoricFormatData(const FText& InText, TArray<FHistoricTextFormatData>& OutHistoricFormatData) const
 {
-	// If we get here, we have no more source so this must be the base text
-	OutSourceTexts.Add(Text);
+}
+
+bool FTextHistory::GetHistoricNumericData(const FText& InText, FHistoricTextNumericData& OutHistoricNumericData) const
+{
+	return false;
 }
 
 void FTextHistory::SerializeForDisplayString(FArchive& Ar, FTextDisplayStringPtr& InOutDisplayString)
@@ -272,21 +275,23 @@ void FTextHistory_NamedFormat::Serialize( FArchive& Ar )
 	Ar << Arguments;
 }
 
-void FTextHistory_NamedFormat::GetSourceTextsFromFormatHistory(FText, TArray<FText>& OutSourceTexts) const
+void FTextHistory_NamedFormat::GetHistoricFormatData(const FText& InText, TArray<FHistoricTextFormatData>& OutHistoricFormatData) const
 {
-	// Search the formatting text itself for source text
-	SourceFmt.GetSourceText().GetSourceTextsFromFormatHistory(OutSourceTexts);
+	// Process the formatting text in-case it's a recursive format
+	FTextInspector::GetHistoricFormatData(SourceFmt.GetSourceText(), OutHistoricFormatData);
 
 	for (auto It = Arguments.CreateConstIterator(); It; ++It)
 	{
 		const FFormatArgumentValue& ArgumentValue = It.Value();
 		if (ArgumentValue.GetType() == EFormatArgumentType::Text)
 		{
-			// Search any text arguments for source text
-			const FText& TextValue = ArgumentValue.GetTextValue();
-			TextValue.TextData->GetTextHistory().GetSourceTextsFromFormatHistory(TextValue, OutSourceTexts);
+			// Process the text argument in-case it's a recursive format
+			FTextInspector::GetHistoricFormatData(ArgumentValue.GetTextValue(), OutHistoricFormatData);
 		}
 	}
+
+	// Add ourself now that we've processed any format dependencies
+	OutHistoricFormatData.Emplace(InText, FTextFormat(SourceFmt), FFormatNamedArguments(Arguments));
 }
 
 ///////////////////////////////////////
@@ -344,21 +349,30 @@ void FTextHistory_OrderedFormat::Serialize( FArchive& Ar )
 	Ar << Arguments;
 }
 
-void FTextHistory_OrderedFormat::GetSourceTextsFromFormatHistory(FText, TArray<FText>& OutSourceTexts) const
+void FTextHistory_OrderedFormat::GetHistoricFormatData(const FText& InText, TArray<FHistoricTextFormatData>& OutHistoricFormatData) const
 {
-	// Search the formatting text itself for source text
-	SourceFmt.GetSourceText().GetSourceTextsFromFormatHistory(OutSourceTexts);
+	// Process the formatting text in-case it's a recursive format
+	FTextInspector::GetHistoricFormatData(SourceFmt.GetSourceText(), OutHistoricFormatData);
 
 	for (auto It = Arguments.CreateConstIterator(); It; ++It)
 	{
 		const FFormatArgumentValue& ArgumentValue = *It;
 		if (ArgumentValue.GetType() == EFormatArgumentType::Text)
 		{
-			// Search any text arguments for source text
-			const FText& TextValue = ArgumentValue.GetTextValue();
-			TextValue.TextData->GetTextHistory().GetSourceTextsFromFormatHistory(TextValue, OutSourceTexts);
+			// Process the text argument in-case it's a recursive format
+			FTextInspector::GetHistoricFormatData(ArgumentValue.GetTextValue(), OutHistoricFormatData);
 		}
 	}
+
+	// Add ourself now that we've processed any format dependencies
+	FFormatNamedArguments NamedArgs;
+	NamedArgs.Reserve(Arguments.Num());
+	for (int32 ArgIndex = 0; ArgIndex < Arguments.Num(); ++ArgIndex)
+	{
+		const FFormatArgumentValue& ArgumentValue = Arguments[ArgIndex];
+		NamedArgs.Emplace(FString::FromInt(ArgIndex), ArgumentValue);
+	}
+	OutHistoricFormatData.Emplace(InText, FTextFormat(SourceFmt), MoveTemp(NamedArgs));
 }
 
 ///////////////////////////////////////
@@ -416,17 +430,45 @@ void FTextHistory_ArgumentDataFormat::Serialize( FArchive& Ar )
 	Ar << Arguments;
 }
 
-void FTextHistory_ArgumentDataFormat::GetSourceTextsFromFormatHistory(FText, TArray<FText>& OutBaseTexts) const
+void FTextHistory_ArgumentDataFormat::GetHistoricFormatData(const FText& InText, TArray<FHistoricTextFormatData>& OutHistoricFormatData) const
 {
-	// Search the formatting text itself for source text
-	SourceFmt.GetSourceText().GetSourceTextsFromFormatHistory(OutBaseTexts);
+	// Process the formatting text in-case it's a recursive format
+	FTextInspector::GetHistoricFormatData(SourceFmt.GetSourceText(), OutHistoricFormatData);
 
-	for (int32 x = 0; x < Arguments.Num(); ++x)
+	for (const FFormatArgumentData& ArgumentData : Arguments)
 	{
-		// Search any text arguments for source text
-		const FText& TextValue = Arguments[x].ArgumentValue;
-		TextValue.TextData->GetTextHistory().GetSourceTextsFromFormatHistory(TextValue, OutBaseTexts);
+		if (ArgumentData.ArgumentValueType == EFormatArgumentType::Text)
+		{
+			// Process the text argument in-case it's a recursive format
+			FTextInspector::GetHistoricFormatData(ArgumentData.ArgumentValue, OutHistoricFormatData);
+		}
 	}
+
+	// Add ourself now that we've processed any format dependencies
+	FFormatNamedArguments NamedArgs;
+	NamedArgs.Reserve(Arguments.Num());
+	for (const FFormatArgumentData& ArgumentData : Arguments)
+	{
+		FFormatArgumentValue ArgumentValue;
+		switch (ArgumentData.ArgumentValueType)
+		{
+		case EFormatArgumentType::Int:
+			ArgumentValue = FFormatArgumentValue(ArgumentData.ArgumentValueInt);
+			break;
+		case EFormatArgumentType::Float:
+			ArgumentValue = FFormatArgumentValue(ArgumentData.ArgumentValueFloat);
+			break;
+		case EFormatArgumentType::Gender:
+			ArgumentValue = FFormatArgumentValue(ArgumentData.ArgumentValueGender);
+			break;
+		default:
+			ArgumentValue = FFormatArgumentValue(ArgumentData.ArgumentValue);
+			break;
+		}
+
+		NamedArgs.Emplace(ArgumentData.ArgumentName, MoveTemp(ArgumentValue));
+	}
+	OutHistoricFormatData.Emplace(InText, FTextFormat(SourceFmt), MoveTemp(NamedArgs));
 }
 
 ///////////////////////////////////////
@@ -566,6 +608,12 @@ void FTextHistory_AsNumber::Serialize( FArchive& Ar )
 	FTextHistory_FormatNumber::Serialize(Ar);
 }
 
+bool FTextHistory_AsNumber::GetHistoricNumericData(const FText& InText, FHistoricTextNumericData& OutHistoricNumericData) const
+{
+	OutHistoricNumericData = FHistoricTextNumericData(FHistoricTextNumericData::EType::AsNumber, SourceValue, FormatOptions);
+	return true;
+}
+
 ///////////////////////////////////////
 // FTextHistory_AsPercent
 
@@ -620,6 +668,12 @@ void FTextHistory_AsPercent::Serialize( FArchive& Ar )
 	FTextHistory_FormatNumber::Serialize(Ar);
 }
 
+bool FTextHistory_AsPercent::GetHistoricNumericData(const FText& InText, FHistoricTextNumericData& OutHistoricNumericData) const
+{
+	OutHistoricNumericData = FHistoricTextNumericData(FHistoricTextNumericData::EType::AsPercent, SourceValue, FormatOptions);
+	return true;
+}
+
 ///////////////////////////////////////
 // FTextHistory_AsCurrency
 
diff --git a/Engine/Source/Runtime/Core/Private/Internationalization/TextHistory.h b/Engine/Source/Runtime/Core/Private/Internationalization/TextHistory.h
--- a/Engine/Source/Runtime/Core/Private/Internationalization/TextHistory.h
+++ b/Engine/Source/Runtime/Core/Private/Internationalization/TextHistory.h
@@ -44,8 +44,11 @@ public:
 	/** Returns the source string managed by the history (if any). */
 	virtual const FString* GetSourceString() const;
 
-	/** Trace the history of this Text until we find the base Texts it was comprised from */
-	virtual void GetSourceTextsFromFormatHistory(FText Text, TArray<FText>& OutSourceTexts) const;
+	/** Get any historic text format data from this history */
+	virtual void GetHistoricFormatData(const FText& InText, TArray<FHistoricTextFormatData>& OutHistoricFormatData) const;
+
+	/** Get any historic numeric format data from this history */
+	virtual bool GetHistoricNumericData(const FText& InText, FHistoricTextNumericData& OutHistoricNumericData) const;
 
 	/** Will rebuild the display string if out of date. */
 	void Rebuild(TSharedRef< FString, ESPMode::ThreadSafe > InDisplayString);
@@ -112,7 +115,7 @@ public:
 	//~ Begin FTextHistory Interface
 	virtual FText ToText(bool bInAsSource) const override;
 	virtual void Serialize(FArchive& Ar) override;
-	virtual void GetSourceTextsFromFormatHistory(FText, TArray<FText>& OutSourceTexts) const override;
+	virtual void GetHistoricFormatData(const FText& InText, TArray<FHistoricTextFormatData>& OutHistoricFormatData) const override;
 	//~ End FTextHistory Interface
 
 private:
@@ -140,7 +143,7 @@ public:
 	//~ Begin FTextHistory Interface
 	virtual FText ToText(bool bInAsSource) const override;
 	virtual void Serialize(FArchive& Ar) override;
-	virtual void GetSourceTextsFromFormatHistory(FText, TArray<FText>& OutSourceTexts) const override;
+	virtual void GetHistoricFormatData(const FText& InText, TArray<FHistoricTextFormatData>& OutHistoricFormatData) const override;
 	//~ End FTextHistory Interface
 
 private:
@@ -168,7 +171,7 @@ public:
 	//~ Begin FTextHistory Interface
 	virtual FText ToText(bool bInAsSource) const override;
 	virtual void Serialize(FArchive& Ar) override;
-	virtual void GetSourceTextsFromFormatHistory(FText, TArray<FText>& OutSourceTexts) const override;
+	virtual void GetHistoricFormatData(const FText& InText, TArray<FHistoricTextFormatData>& OutHistoricFormatData) const override;
 	//~ End FTextHistory Interface
 
 private:
@@ -225,6 +228,7 @@ public:
 	//~ Begin FTextHistory Interface
 	virtual FText ToText(bool bInAsSource) const override;
 	virtual void Serialize(FArchive& Ar) override;
+	virtual bool GetHistoricNumericData(const FText& InText, FHistoricTextNumericData& OutHistoricNumericData) const override;
 	//~ End FTextHistory interface
 
 private:
@@ -247,6 +251,7 @@ public:
 	//~ Begin FTextHistory Interface
 	virtual FText ToText(bool bInAsSource) const override;
 	virtual void Serialize(FArchive& Ar) override;
+	virtual bool GetHistoricNumericData(const FText& InText, FHistoricTextNumericData& OutHistoricNumericData) const override;
 	//~ End FTextHistory interface
 
 private:
diff --git a/Engine/Source/Runtime/Core/Private/Linux/LinuxPlatformStackWalk.cpp b/Engine/Source/Runtime/Core/Private/Linux/LinuxPlatformStackWalk.cpp
--- a/Engine/Source/Runtime/Core/Private/Linux/LinuxPlatformStackWalk.cpp
+++ b/Engine/Source/Runtime/Core/Private/Linux/LinuxPlatformStackWalk.cpp
@@ -7,6 +7,7 @@
 #include "CorePrivatePCH.h"
 #include "Misc/App.h"
 #include "LinuxPlatformCrashContext.h"
+#include "ExceptionHandling.h"
 #include <execinfo.h>
 #include <dlfcn.h>
 #include <cxxabi.h>
diff --git a/Engine/Source/Runtime/Core/Private/Misc/EventPool.h b/Engine/Source/Runtime/Core/Private/Misc/EventPool.h
--- a/Engine/Source/Runtime/Core/Private/Misc/EventPool.h
+++ b/Engine/Source/Runtime/Core/Private/Misc/EventPool.h
@@ -15,6 +15,48 @@ enum class EEventPoolTypes
 	ManualReset
 };
 
+class FSafeRecyclableEvent  final: public FEvent
+{
+public:
+	FEvent* InnerEvent;
+
+	FSafeRecyclableEvent(FEvent* InInnerEvent)
+		: InnerEvent(InInnerEvent)
+	{
+	}
+
+	~FSafeRecyclableEvent()
+	{
+		InnerEvent = nullptr;
+	}
+
+	virtual bool Create(bool bIsManualReset = false)
+	{
+		return InnerEvent->Create(bIsManualReset);
+	}
+
+	virtual bool IsManualReset()
+	{
+		return InnerEvent->IsManualReset();
+	}
+
+	virtual void Trigger()
+	{
+		InnerEvent->Trigger();
+	}
+
+	virtual void Reset()
+	{
+		InnerEvent->Reset();
+	}
+
+	virtual bool Wait(uint32 WaitTime, const bool bIgnoreThreadIdleStats = false)
+	{
+		return InnerEvent->Wait(WaitTime, bIgnoreThreadIdleStats);
+	}
+
+};
+
 /**
  * Template class for event pools.
  *
@@ -59,11 +101,10 @@ public:
 			Result = FPlatformProcess::CreateSynchEvent((PoolType == EEventPoolTypes::ManualReset));
 			PRAGMA_ENABLE_DEPRECATION_WARNINGS
 		}
-
 		Result->AdvanceStats();
 
 		check(Result);
-		return Result;
+		return new FSafeRecyclableEvent(Result);
 	}
 
 	/**
@@ -76,10 +117,12 @@ public:
 	{
 		check(Event);
 		check(Event->IsManualReset() == (PoolType == EEventPoolTypes::ManualReset));
-
-		Event->Reset();
-
-		Pool.Push(Event);
+		FSafeRecyclableEvent* SafeEvent = (FSafeRecyclableEvent*)Event;
+		FEvent* Result = SafeEvent->InnerEvent;
+		delete SafeEvent;
+		check(Result);
+		Result->Reset();
+		Pool.Push(Result);
 	}
 
 private:
diff --git a/Engine/Source/Runtime/Core/Private/Misc/NetworkVersion.cpp b/Engine/Source/Runtime/Core/Private/Misc/NetworkVersion.cpp
--- a/Engine/Source/Runtime/Core/Private/Misc/NetworkVersion.cpp
+++ b/Engine/Source/Runtime/Core/Private/Misc/NetworkVersion.cpp
@@ -24,6 +24,9 @@ uint32 FNetworkVersion::CachedNetworkChecksum			= 0;
 uint32 FNetworkVersion::EngineNetworkProtocolVersion	= HISTORY_REPLAY_BACKWARDS_COMPAT;
 uint32 FNetworkVersion::GameNetworkProtocolVersion		= 0;
 
+uint32 FNetworkVersion::EngineCompatibleNetworkProtocolVersion		= HISTORY_REPLAY_BACKWARDS_COMPAT;
+uint32 FNetworkVersion::GameCompatibleNetworkProtocolVersion		= 0;
+
 uint32 FNetworkVersion::GetNetworkCompatibleChangelist()
 {
 	//return FEngineVersion::Current().GetChangelist();
@@ -40,11 +43,21 @@ uint32 FNetworkVersion::GetEngineNetworkProtocolVersion()
 	return EngineNetworkProtocolVersion;
 }
 
+uint32 FNetworkVersion::GetEngineCompatibleNetworkProtocolVersion()
+{
+	return EngineCompatibleNetworkProtocolVersion;
+}
+
 uint32 FNetworkVersion::GetGameNetworkProtocolVersion()
 {
 	return GameNetworkProtocolVersion;
 }
 
+uint32 FNetworkVersion::GetGameCompatibleNetworkProtocolVersion()
+{
+	return GameCompatibleNetworkProtocolVersion;
+}
+
 uint32 FNetworkVersion::GetLocalNetworkVersion( bool AllowOverrideDelegate /*=true*/ )
 {
 	if ( bHasCachedNetworkChecksum )
@@ -98,7 +111,7 @@ bool FNetworkVersion::IsNetworkCompatible( const uint32 LocalNetworkVersion, con
 
 FNetworkReplayVersion FNetworkVersion::GetReplayVersion()
 {
-	const uint32 ReplayVersion = ( GameNetworkProtocolVersion << 16 ) | EngineNetworkProtocolVersion;
+	const uint32 ReplayVersion = ( GameCompatibleNetworkProtocolVersion << 16 ) | EngineCompatibleNetworkProtocolVersion;
 
 	return FNetworkReplayVersion( FApp::GetGameName(), ReplayVersion, GetReplayCompatibleChangelist() );
 }
diff --git a/Engine/Source/Runtime/Core/Private/Misc/OutputDeviceFile.cpp b/Engine/Source/Runtime/Core/Private/Misc/OutputDeviceFile.cpp
--- a/Engine/Source/Runtime/Core/Private/Misc/OutputDeviceFile.cpp
+++ b/Engine/Source/Runtime/Core/Private/Misc/OutputDeviceFile.cpp
@@ -245,6 +245,15 @@ FOutputDeviceFile::FOutputDeviceFile( const TCHAR* InFilename, bool bInDisableBa
 	}
 }
 
+/**
+* Destructor to perform teardown
+*
+*/
+FOutputDeviceFile::~FOutputDeviceFile()
+{
+	TearDown();
+}
+
 void FOutputDeviceFile::SetFilename(const TCHAR* InFilename)
 {
 	// Close any existing file.
diff --git a/Engine/Source/Runtime/Core/Private/Serialization/BitReader.cpp b/Engine/Source/Runtime/Core/Private/Serialization/BitReader.cpp
--- a/Engine/Source/Runtime/Core/Private/Serialization/BitReader.cpp
+++ b/Engine/Source/Runtime/Core/Private/Serialization/BitReader.cpp
@@ -179,13 +179,10 @@ void FBitReader::AppendTo( TArray<uint8> &DestBuffer )
 
 void FBitReader::SetOverflowed(int32 LengthBits)
 {
-	if (!ensure(LengthBits == 0))
-	{
-		UE_LOG(LogNetSerialization, Error, TEXT("FBitReader::SetOverflowed() called! (ReadLen: %i, Remaining: %i, Max: %i)"),
-					LengthBits, (Num-Pos), Num);
+	UE_LOG(LogNetSerialization, Error, TEXT("FBitReader::SetOverflowed() called! (ReadLen: %i, Remaining: %i, Max: %i)"),
+			LengthBits, (Num - Pos), Num);
 
-		ArIsError = 1;
-	}
+	ArIsError = 1;
 }
 
 /*-----------------------------------------------------------------------------
diff --git a/Engine/Source/Runtime/Core/Private/Windows/WindowsPlatformProcess.cpp b/Engine/Source/Runtime/Core/Private/Windows/WindowsPlatformProcess.cpp
--- a/Engine/Source/Runtime/Core/Private/Windows/WindowsPlatformProcess.cpp
+++ b/Engine/Source/Runtime/Core/Private/Windows/WindowsPlatformProcess.cpp
@@ -711,7 +711,16 @@ const TCHAR* FWindowsPlatformProcess::BaseDir()
 		// workingdir in your debugger to the <path>/Project/Binaries/Win64 of your cooked
 		// data
 		// Too early to use the FCommand line interface
-		if (FCString::Stristr(::GetCommandLine(), TEXT("-BaseFromWorkingDir")))
+		FString BaseArg;
+		FParse::Value(::GetCommandLine(), TEXT("-basedir="), BaseArg);
+
+		if (BaseArg.Len())
+		{
+			BaseArg = BaseArg.Replace(TEXT("\\"), TEXT("/"));
+			BaseArg += TEXT('/');
+			FCString::Strcpy(Result, *BaseArg);
+		}
+		else if (FCString::Stristr(::GetCommandLine(), TEXT("-BaseFromWorkingDir")))
 		{
 			::GetCurrentDirectory(512, Result);
 
diff --git a/Engine/Source/Runtime/Core/Private/Windows/WindowsWindow.cpp b/Engine/Source/Runtime/Core/Private/Windows/WindowsWindow.cpp
--- a/Engine/Source/Runtime/Core/Private/Windows/WindowsWindow.cpp
+++ b/Engine/Source/Runtime/Core/Private/Windows/WindowsWindow.cpp
@@ -269,8 +269,8 @@ FWindowsWindow::FWindowsWindow()
 	, bIsVisible( false )
 	, DPIScaleFactor(1.0f)
 {
+	// PreFullscreenWindowPlacement.length will be set when we save the window placement and then used to check if the structure is valid
 	FMemory::Memzero(PreFullscreenWindowPlacement);
-	PreFullscreenWindowPlacement.length = sizeof(WINDOWPLACEMENT);
 
 	FMemory::Memzero(PreParentMinimizedWindowPlacement);
 	PreParentMinimizedWindowPlacement.length = sizeof(WINDOWPLACEMENT);
@@ -306,7 +306,7 @@ void FWindowsWindow::OnTransparencySupportChanged(EWindowTransparency NewTranspa
 #endif
 }
 
-HRGN FWindowsWindow::MakeWindowRegionObject() const
+HRGN FWindowsWindow::MakeWindowRegionObject(bool bIncludeBorderWhenMaximized) const
 {
 	HRGN Region;
 	if ( RegionWidth != INDEX_NONE && RegionHeight != INDEX_NONE )
@@ -324,11 +324,12 @@ HRGN FWindowsWindow::MakeWindowRegionObject() const
 				WindowInfo.cbSize = sizeof(WindowInfo);
 				::GetWindowInfo(HWnd, &WindowInfo);
 
-				Region = CreateRectRgn(WindowInfo.cxWindowBorders, WindowInfo.cxWindowBorders, RegionWidth + WindowInfo.cxWindowBorders, RegionHeight + WindowInfo.cxWindowBorders);
+				const int32 WindowBorderSize = bIncludeBorderWhenMaximized ? WindowInfo.cxWindowBorders : 0;
+				Region = CreateRectRgn(WindowBorderSize, WindowBorderSize, RegionWidth + WindowBorderSize, RegionHeight + WindowBorderSize);
 			}
 			else
 			{
-				int32 WindowBorderSize = GetWindowBorderSize();
+				const int32 WindowBorderSize = bIncludeBorderWhenMaximized ? GetWindowBorderSize() : 0;
 				Region = CreateRectRgn(WindowBorderSize, WindowBorderSize, RegionWidth - WindowBorderSize, RegionHeight - WindowBorderSize);
 			}
 		}
@@ -371,7 +372,7 @@ void FWindowsWindow::AdjustWindowRegion( int32 Width, int32 Height )
 	RegionWidth = Width;
 	RegionHeight = Height;
 
-	HRGN Region = MakeWindowRegionObject();
+	HRGN Region = MakeWindowRegionObject(true);
 
 	// NOTE: We explicitly don't delete the Region object, because the OS deletes the handle when it no longer needed after
 	// a call to SetWindowRgn.
@@ -597,6 +598,7 @@ void FWindowsWindow::SetWindowMode( EWindowMode::Type NewWindowMode )
 {
 	if (NewWindowMode != WindowMode)
 	{
+		EWindowMode::Type PreviousWindowMode = WindowMode;
 		WindowMode = NewWindowMode;
 
 		const bool bTrueFullscreen = NewWindowMode == EWindowMode::Fullscreen;
@@ -635,7 +637,11 @@ void FWindowsWindow::SetWindowMode( EWindowMode::Type NewWindowMode )
 		// If we're not in fullscreen, make it so
 		if (NewWindowMode == EWindowMode::WindowedFullscreen || NewWindowMode == EWindowMode::Fullscreen)
 		{
-			::GetWindowPlacement(HWnd, &PreFullscreenWindowPlacement);
+			if (PreviousWindowMode == EWindowMode::Windowed)
+			{
+				PreFullscreenWindowPlacement.length = sizeof(WINDOWPLACEMENT);
+				::GetWindowPlacement(HWnd, &PreFullscreenWindowPlacement);
+			}
 
 			// Setup Win32 flags for fullscreen window
 			WindowStyle &= ~WindowedModeStyle;
@@ -692,7 +698,10 @@ void FWindowsWindow::SetWindowMode( EWindowMode::Type NewWindowMode )
 			SetWindowLong(HWnd, GWL_STYLE, WindowStyle);
 			::SetWindowPos(HWnd, nullptr, 0, 0, 0, 0, SWP_NOMOVE | SWP_NOSIZE | SWP_NOZORDER | SWP_FRAMECHANGED);
 
-			::SetWindowPlacement(HWnd, &PreFullscreenWindowPlacement);
+			if (PreFullscreenWindowPlacement.length) // Was PreFullscreenWindowPlacement initialized?
+			{
+				::SetWindowPlacement(HWnd, &PreFullscreenWindowPlacement);
+			}
 		}
 	}
 }
@@ -808,7 +817,7 @@ bool FWindowsWindow::IsPointInWindow( int32 X, int32 Y ) const
 {
 	bool Result = false;
 
-	HRGN Region = MakeWindowRegionObject();
+	HRGN Region = MakeWindowRegionObject(false);
 	Result = !!PtInRegion( Region, X, Y );
 	DeleteObject( Region );
 
diff --git a/Engine/Source/Runtime/Core/Public/HAL/ThreadHeartBeat.h b/Engine/Source/Runtime/Core/Public/HAL/ThreadHeartBeat.h
--- a/Engine/Source/Runtime/Core/Public/HAL/ThreadHeartBeat.h
+++ b/Engine/Source/Runtime/Core/Public/HAL/ThreadHeartBeat.h
@@ -35,6 +35,11 @@ class CORE_API FThreadHeartBeat : public FRunnable
 	/** Max time the thread is allowed to not send the heartbeat*/
 	double HangDuration;
 
+	/** CRC of the last hang's callstack */
+	uint32 LastHangCallstackCRC;
+	/** Id of the last thread that hung */
+	uint32 LastHungThreadId;
+
 	FThreadHeartBeat();
 	virtual ~FThreadHeartBeat();
 
@@ -67,6 +72,11 @@ public:
 	 */
 	void ResumeHeartBeat();
 
+	/**
+	* Returns true/false if this thread is currently performing heartbeat monitoring
+	*/
+	bool IsBeating();
+
 	//~ Begin FRunnable Interface.
 	virtual bool Init();
 	virtual uint32 Run();
diff --git a/Engine/Source/Runtime/Core/Public/Internationalization/Text.h b/Engine/Source/Runtime/Core/Public/Internationalization/Text.h
--- a/Engine/Source/Runtime/Core/Public/Internationalization/Text.h
+++ b/Engine/Source/Runtime/Core/Public/Internationalization/Text.h
@@ -15,6 +15,8 @@ struct FDateTime;
 class FText;
 class FTextHistory;
 class FTextFormatData;
+class FHistoricTextFormatData;
+class FHistoricTextNumericData;
 
 #define ENABLE_TEXT_ERROR_CHECKING_RESULTS (UE_BUILD_DEBUG | UE_BUILD_DEVELOPMENT | UE_BUILD_TEST )
 
@@ -375,9 +377,6 @@ public:
 	 */
 	bool IdenticalTo( const FText& Other ) const;
 
-	/** Trace the history of this Text until we find the base Texts it was comprised from */
-	void GetSourceTextsFromFormatHistory(TArray<FText>& OutSourceTexts) const;
-
 	class CORE_API FSortPredicate
 	{
 	public:
@@ -498,6 +497,12 @@ private:
 	/** Returns the source string of the FText */
 	const FString& GetSourceString() const;
 
+	/** Get any historic text format data from the history used by this FText */
+	void GetHistoricFormatData(TArray<FHistoricTextFormatData>& OutHistoricFormatData) const;
+
+	/** Get any historic numeric format data from the history used by this FText */
+	bool GetHistoricNumericData(FHistoricTextNumericData& OutHistoricNumericData) const;
+
 	/** Rebuilds the FText under the current culture if needed */
 	void Rebuild() const;
 
@@ -730,6 +735,63 @@ FText FText::FormatOrdered( FTextFormat Fmt, TArguments&&... Args )
 	return FormatOrderedImpl( MoveTemp( Fmt ), MoveTemp( FormatArguments ) );
 }
 
+/** Used to gather information about a historic text format operation */
+class CORE_API FHistoricTextFormatData
+{
+public:
+	FHistoricTextFormatData()
+	{
+	}
+
+	FHistoricTextFormatData(FText InFormattedText, FTextFormat&& InSourceFmt, FFormatNamedArguments&& InArguments)
+		: FormattedText(MoveTemp(InFormattedText))
+		, SourceFmt(MoveTemp(InSourceFmt))
+		, Arguments(MoveTemp(InArguments))
+	{
+	}
+
+	/** The final formatted text this data is for */
+	FText FormattedText;
+
+	/** The pattern used to format the text */
+	FTextFormat SourceFmt;
+
+	/** Arguments to replace in the pattern string */
+	FFormatNamedArguments Arguments;
+};
+
+/** Used to gather information about a historic numeric format operation */
+class CORE_API FHistoricTextNumericData
+{
+public:
+	enum class EType : uint8
+	{
+		AsNumber,
+		AsPercent,
+	};
+
+	FHistoricTextNumericData()
+		: FormatType(EType::AsNumber)
+	{
+	}
+
+	FHistoricTextNumericData(const EType InFormatType, const FFormatArgumentValue& InSourceValue, const TOptional<FNumberFormattingOptions>& InFormatOptions)
+		: FormatType(InFormatType)
+		, SourceValue(InSourceValue)
+		, FormatOptions(InFormatOptions)
+	{
+	}
+
+	/** Type of numeric format that was performed */
+	EType FormatType;
+
+	/** The source number to format */
+	FFormatArgumentValue SourceValue;
+
+	/** Custom formatting options used when formatting this number (if any) */
+	TOptional<FNumberFormattingOptions> FormatOptions;
+};
+
 /** A snapshot of an FText at a point in time that can be used to detect changes in the FText, including live-culture changes */
 class CORE_API FTextSnapshot
 {
@@ -772,6 +834,8 @@ public:
 	static const FString& GetDisplayString(const FText& Text);
 	static const FTextDisplayStringRef GetSharedDisplayString(const FText& Text);
 	static uint32 GetFlags(const FText& Text);
+	static void GetHistoricFormatData(const FText& Text, TArray<FHistoricTextFormatData>& OutHistoricFormatData);
+	static bool GetHistoricNumericData(const FText& Text, FHistoricTextNumericData& OutHistoricNumericData);
 };
 
 class CORE_API FTextStringHelper
diff --git a/Engine/Source/Runtime/Core/Public/Misc/NetworkVersion.h b/Engine/Source/Runtime/Core/Public/Misc/NetworkVersion.h
--- a/Engine/Source/Runtime/Core/Public/Misc/NetworkVersion.h
+++ b/Engine/Source/Runtime/Core/Public/Misc/NetworkVersion.h
@@ -33,6 +33,8 @@ struct CORE_API FNetworkVersion
 	static uint32 GetReplayCompatibleChangelist();
 	static uint32 GetEngineNetworkProtocolVersion();
 	static uint32 GetGameNetworkProtocolVersion();
+	static uint32 GetEngineCompatibleNetworkProtocolVersion();
+	static uint32 GetGameCompatibleNetworkProtocolVersion();
 
 	/**
 	* Generates a version number, that by default, is based on a checksum of the engine version + project name + project version string
@@ -67,4 +69,7 @@ struct CORE_API FNetworkVersion
 
 	static uint32	EngineNetworkProtocolVersion;
 	static uint32	GameNetworkProtocolVersion;
+
+	static uint32	EngineCompatibleNetworkProtocolVersion;
+	static uint32	GameCompatibleNetworkProtocolVersion;
 };
diff --git a/Engine/Source/Runtime/Core/Public/Misc/OutputDeviceFile.h b/Engine/Source/Runtime/Core/Public/Misc/OutputDeviceFile.h
--- a/Engine/Source/Runtime/Core/Public/Misc/OutputDeviceFile.h
+++ b/Engine/Source/Runtime/Core/Public/Misc/OutputDeviceFile.h
@@ -30,6 +30,12 @@ public:
 	*/
 	FOutputDeviceFile(const TCHAR* InFilename = nullptr, bool bDisableBackup = false);
 
+	/**
+	* Destructor to perform teardown
+	*
+	*/
+	~FOutputDeviceFile();
+
 	/** Sets the filename that the output device writes to.  If the output device was already writing to a file, closes that file. */
 	void SetFilename(const TCHAR* InFilename);
 
diff --git a/Engine/Source/Runtime/Core/Public/Windows/WindowsWindow.h b/Engine/Source/Runtime/Core/Public/Windows/WindowsWindow.h
--- a/Engine/Source/Runtime/Core/Public/Windows/WindowsWindow.h
+++ b/Engine/Source/Runtime/Core/Public/Windows/WindowsWindow.h
@@ -148,7 +148,7 @@ private:
 
 	/** Creates an HRGN for the window's current region.  Remember to delete this when you're done with it using
 	   ::DeleteObject, unless you're passing it to SetWindowRgn(), which will absorb the reference itself. */
-	HRGN MakeWindowRegionObject() const;
+	HRGN MakeWindowRegionObject(bool bIncludeBorderWhenMaximized) const;
 
 private:
 
diff --git a/Engine/Source/Runtime/CoreUObject/Private/UObject/LinkerLoad.cpp b/Engine/Source/Runtime/CoreUObject/Private/UObject/LinkerLoad.cpp
--- a/Engine/Source/Runtime/CoreUObject/Private/UObject/LinkerLoad.cpp
+++ b/Engine/Source/Runtime/CoreUObject/Private/UObject/LinkerLoad.cpp
@@ -2236,7 +2236,7 @@ void FLinkerLoad::GatherImportDependencies(int32 ImportIndex, TSet<FDependencyRe
 		return;
 	}
 
-	BeginLoad();
+	BeginLoad(TEXT("GatherImportDependencies"));
 
 	// load the linker and find export in sourcelinker
 	if (Import.SourceLinker == NULL || Import.SourceIndex == INDEX_NONE)
@@ -3708,7 +3708,10 @@ UObject* FLinkerLoad::CreateExport( int32 Index )
 				}
 				else
 				{
-					UE_LOG(LogLinker, Warning, TEXT("CreateExport: Failed to load Parent for %s"), *GetExportFullName(Index));
+					if (!FLinkerLoad::IsKnownMissingPackage(*GetExportFullName(Index)))
+					{
+						UE_LOG(LogLinker, Warning, TEXT("CreateExport: Failed to load Parent for %s"), *GetExportFullName(Index));
+					}
 					return NULL;
 				}
 			}
diff --git a/Engine/Source/Runtime/CoreUObject/Private/UObject/UObjectGlobals.cpp b/Engine/Source/Runtime/CoreUObject/Private/UObject/UObjectGlobals.cpp
--- a/Engine/Source/Runtime/CoreUObject/Private/UObject/UObjectGlobals.cpp
+++ b/Engine/Source/Runtime/CoreUObject/Private/UObject/UObjectGlobals.cpp
@@ -872,13 +872,16 @@ UObject* StaticLoadObject(UClass* ObjectClass, UObject* InOuter, const TCHAR* In
 		FString ObjectName = InName;
 		ResolveName(InOuter, ObjectName, true, true, LoadFlags & LOAD_EditorOnly);
 
-		// we haven't created or found the object, error
-		FFormatNamedArguments Arguments;
-		Arguments.Add(TEXT("ClassName"), FText::FromString(ObjectClass->GetName()));
-		Arguments.Add(TEXT("OuterName"), InOuter ? FText::FromString(InOuter->GetPathName()) : NSLOCTEXT("Core", "None", "None"));
-		Arguments.Add(TEXT("ObjectName"), FText::FromString(ObjectName));
-		const FString Error = FText::Format(NSLOCTEXT("Core", "ObjectNotFound", "Failed to find object '{ClassName} {OuterName}.{ObjectName}'"), Arguments).ToString();
-		SafeLoadError(InOuter, LoadFlags, *Error);
+		if (InOuter == nullptr || FLinkerLoad::IsKnownMissingPackage(FName(*InOuter->GetPathName())) == false)
+		{
+			// we haven't created or found the object, error
+			FFormatNamedArguments Arguments;
+			Arguments.Add(TEXT("ClassName"), FText::FromString(ObjectClass->GetName()));
+			Arguments.Add(TEXT("OuterName"), InOuter ? FText::FromString(InOuter->GetPathName()) : NSLOCTEXT("Core", "None", "None"));
+			Arguments.Add(TEXT("ObjectName"), FText::FromString(ObjectName));
+			const FString Error = FText::Format(NSLOCTEXT("Core", "ObjectNotFound", "Failed to find object '{ClassName} {OuterName}.{ObjectName}'"), Arguments).ToString();
+			SafeLoadError(InOuter, LoadFlags, *Error);
+		}
 	}
 	return Result;
 }
@@ -1167,7 +1170,7 @@ static UPackage* LoadPackageInternalInner(UPackage* InOuter, const TCHAR* InLong
 	SlowTask.EnterProgressFrame(10);
 
 	// Try to load.
-	BeginLoad();
+	BeginLoad(InLongPackageNameOrFilename);
 
 	bool bFullyLoadSkipped = false;
 
@@ -1452,11 +1455,16 @@ bool IsLoading()
 // Begin loading packages.
 //warning: Objects may not be destroyed between BeginLoad/EndLoad calls.
 //
-void BeginLoad()
+void BeginLoad(const TCHAR* DebugContext)
 {
 	auto& ThreadContext = FUObjectThreadContext::Get();
 	if (ThreadContext.ObjBeginLoadCount == 0 && !IsInAsyncLoadingThread())
 	{
+		if (IsAsyncLoading() && (DebugContext != nullptr))
+		{
+			UE_LOG(LogUObjectGlobals, Log, TEXT("BeginLoad(%s) is flushing async loading"), DebugContext);
+		}
+
 		// Make sure we're finishing up all pending async loads, and trigger texture streaming next tick if necessary.
 		FlushAsyncLoading();
 
diff --git a/Engine/Source/Runtime/CoreUObject/Public/UObject/CoreNetTypes.h b/Engine/Source/Runtime/CoreUObject/Public/UObject/CoreNetTypes.h
--- a/Engine/Source/Runtime/CoreUObject/Public/UObject/CoreNetTypes.h
+++ b/Engine/Source/Runtime/CoreUObject/Public/UObject/CoreNetTypes.h
@@ -17,7 +17,8 @@ enum ELifetimeCondition
 	COND_ReplayOnly = 10					UMETA(DisplayName = "Replay Only"),						// This property will only send to the replay connection
 	COND_SimulatedOnlyNoReplay = 11			UMETA(DisplayName = "Simulated Only No Replay"),		// This property will send to actors only, but not to replay connections
 	COND_SimulatedOrPhysicsNoReplay = 12	UMETA(DisplayName = "Simulated Or Physics No Replay"),	// This property will send to simulated Or bRepPhysics actors, but not to replay connections
-	COND_Max = 13							UMETA(Hidden)
+	COND_SkipReplay = 13					UMETA(DisplayName = "Skip Replay"),	// This property will not send to the replay connection
+	COND_Max = 14							UMETA(Hidden)
 };
 
 
diff --git a/Engine/Source/Runtime/CoreUObject/Public/UObject/UObjectGlobals.h b/Engine/Source/Runtime/CoreUObject/Public/UObject/UObjectGlobals.h
--- a/Engine/Source/Runtime/CoreUObject/Public/UObject/UObjectGlobals.h
+++ b/Engine/Source/Runtime/CoreUObject/Public/UObject/UObjectGlobals.h
@@ -531,7 +531,7 @@ namespace EAsyncPackageState
  * @return The minimum state of any of the queued packages.
  */
 COREUOBJECT_API EAsyncPackageState::Type ProcessAsyncLoading( bool bUseTimeLimit, bool bUseFullTimeLimit, float TimeLimit);
-COREUOBJECT_API void BeginLoad();
+COREUOBJECT_API void BeginLoad(const TCHAR* DebugContext = nullptr);
 COREUOBJECT_API void EndLoad();
 
 /**
diff --git a/Engine/Source/Runtime/Engine/Classes/AI/Navigation/RecastNavMesh.h b/Engine/Source/Runtime/Engine/Classes/AI/Navigation/RecastNavMesh.h
--- a/Engine/Source/Runtime/Engine/Classes/AI/Navigation/RecastNavMesh.h
+++ b/Engine/Source/Runtime/Engine/Classes/AI/Navigation/RecastNavMesh.h
@@ -951,6 +951,9 @@ public:
 	/** Get all polys that overlap the specified box */
 	bool GetPolysInBox(const FBox& Box, TArray<FNavPoly>& Polys, FSharedConstNavQueryFilter Filter = nullptr, const UObject* Owner = nullptr) const;
 
+	/** Get all polys from tile */
+	bool GetNavLinksInTile(const int32 TileIndex, TArray<FNavPoly>& Polys, const bool bIncludeLinksFromNeighborTiles) const;
+
 	/** Projects point on navmesh, returning all hits along vertical line defined by min-max Z params */
 	bool ProjectPointMulti(const FVector& Point, TArray<FNavLocation>& OutLocations, const FVector& Extent,
 		float MinZ, float MaxZ, FSharedConstNavQueryFilter Filter = NULL, const UObject* Querier = NULL) const;
diff --git a/Engine/Source/Runtime/Engine/Classes/Animation/AnimInstance.h b/Engine/Source/Runtime/Engine/Classes/Animation/AnimInstance.h
--- a/Engine/Source/Runtime/Engine/Classes/Animation/AnimInstance.h
+++ b/Engine/Source/Runtime/Engine/Classes/Animation/AnimInstance.h
@@ -65,6 +65,9 @@ DECLARE_DYNAMIC_MULTICAST_DELEGATE_OneParam(FOnMontageStartedMCDelegate, UAnimMo
 */
 DECLARE_DYNAMIC_MULTICAST_DELEGATE_TwoParams(FOnMontageEndedMCDelegate, UAnimMontage*, Montage, bool, bInterrupted);
 
+/** Delegate for when all montage instances have ended. */
+DECLARE_DYNAMIC_MULTICAST_DELEGATE(FOnAllMontageInstancesEndedMCDelegate);
+
 /**
 * Delegate for when Montage started to blend out, whether interrupted or finished
 * DesiredWeight of this montage becomes 0.f, but this still contributes to the output pose
@@ -411,14 +414,24 @@ public:
 
 	/** Get global weight in AnimGraph for this slot node.
 	* Note: this is the weight of the node, not the weight of any potential montage it is playing. */
-	float GetSlotNodeGlobalWeight(FName SlotNodeName) const;
+	float GetSlotNodeGlobalWeight(const FName& SlotNodeName) const;
 
 	// Should Extract Root Motion or not. Return true if we do. 
 	bool ShouldExtractRootMotion() const { return RootMotionMode == ERootMotionMode::RootMotionFromEverything || RootMotionMode == ERootMotionMode::IgnoreRootMotion; }
 
 	/** Get Global weight of any montages this slot node is playing.
 	* If this slot is not currently playing a montage, it will return 0. */
-	float GetSlotMontageGlobalWeight(FName SlotNodeName) const;
+	float GetSlotMontageGlobalWeight(const FName& SlotNodeName) const;
+
+	/** Get local weight of any montages this slot node is playing.
+	* If this slot is not currently playing a montage, it will return 0.
+	* This is double buffered, will return last frame data if called from Update or Evaluate. */
+	float GetSlotMontageLocalWeight(const FName& SlotNodeName) const;
+
+	/** Get local weight of any montages this slot is playing.
+	* If this slot is not current playing a montage, it will return 0.
+	* This will return up to date data if called during Update or Evaluate. */
+	float CalcSlotMontageLocalWeight(const FName& SlotNodeName) const;
 
 	// kismet event functions
 
@@ -603,6 +616,10 @@ public:
 	UPROPERTY(BlueprintAssignable)
 	FOnMontageEndedMCDelegate OnMontageEnded;
 
+	/** Called when all Montage instances have ended. */
+	UPROPERTY(BlueprintAssignable)
+	FOnAllMontageInstancesEndedMCDelegate OnAllMontageInstancesEnded;
+
 	/*********************************************************************************************
 	* AnimMontage native C++ interface
 	********************************************************************************************* */
@@ -831,6 +848,9 @@ public:
 	UFUNCTION(BlueprintPure, Category="Animation")
 	float GetCurveValue(FName CurveName);
 
+	/** Returns value of named curved in OutValue, returns whether the curve was actually found or not. */
+	bool GetCurveValue(FName CurveName, float& OutValue);
+
 	/** Returns the length (in seconds) of an animation AnimAsset. */
 	DEPRECATED(4.9, "GetAnimAssetPlayerLength is deprecated, use GetInstanceAssetPlayerLength instead")
 	UFUNCTION(BlueprintPure, Category="Animation", meta=(BlueprintInternalUseOnly = "true"))
diff --git a/Engine/Source/Runtime/Engine/Classes/Animation/AnimNotifies/AnimNotifyState_Trail.h b/Engine/Source/Runtime/Engine/Classes/Animation/AnimNotifies/AnimNotifyState_Trail.h
--- a/Engine/Source/Runtime/Engine/Classes/Animation/AnimNotifies/AnimNotifyState_Trail.h
+++ b/Engine/Source/Runtime/Engine/Classes/Animation/AnimNotifies/AnimNotifyState_Trail.h
@@ -6,6 +6,9 @@
 
 DECLARE_LOG_CATEGORY_EXTERN(LogAnimTrails, Log, All);
 
+class USkeletalMeshComponent;
+class UAnimSequenceBase;
+
 UCLASS(editinlinenew, Blueprintable, const, hidecategories = Object, collapsecategories, meta = (ShowWorldContextPin, DisplayName = "Trail"))
 class ENGINE_API UAnimNotifyState_Trail : public UAnimNotifyState
 {
@@ -16,9 +19,11 @@ class ENGINE_API UAnimNotifyState_Trail : public UAnimNotifyState
 	UParticleSystem* PSTemplate;
 
 	UFUNCTION(BlueprintImplementableEvent)
-	UParticleSystem* OverridePSTemplate(class USkeletalMeshComponent* MeshComp, class UAnimSequenceBase* Animation) const;
+	UParticleSystem* OverridePSTemplate(USkeletalMeshComponent* MeshComp, class UAnimSequenceBase* Animation) const;
+
+	virtual UParticleSystem* GetOverridenPSTemplate(USkeletalMeshComponent* MeshComp, UAnimSequenceBase* Animation) const;
 
-	virtual UParticleSystem* GetOverridenPSTemplate(class USkeletalMeshComponent* MeshComp, class UAnimSequenceBase* Animation) const;
+	virtual float GetCurveWidth(USkeletalMeshComponent* MeshComp) const;
 
 	/** Name of the first socket defining this trail. */
 	UPROPERTY(EditAnywhere, BlueprintReadWrite, Category = Trail)
@@ -62,11 +67,14 @@ class ENGINE_API UAnimNotifyState_Trail : public UAnimNotifyState
 	uint32 bRenderTessellation : 1;
 #endif // WITH_EDITORONLY_DATA
 
-	virtual void NotifyBegin(class USkeletalMeshComponent * MeshComp, class UAnimSequenceBase * Animation, float TotalDuration) override;
-	virtual void NotifyTick(class USkeletalMeshComponent * MeshComp, class UAnimSequenceBase * Animation, float FrameDeltaTime) override;
-	virtual void NotifyEnd(class USkeletalMeshComponent * MeshComp, class UAnimSequenceBase * Animation) override;
+	/** Helper function for outside code to get PSC that we are using */
+	UParticleSystemComponent* GetParticleSystemComponent(USkeletalMeshComponent* MeshComp) const;
+
+	virtual void NotifyBegin(USkeletalMeshComponent* MeshComp, UAnimSequenceBase* Animation, float TotalDuration) override;
+	virtual void NotifyTick(USkeletalMeshComponent* MeshComp, UAnimSequenceBase* Animation, float FrameDeltaTime) override;
+	virtual void NotifyEnd(USkeletalMeshComponent* MeshComp, UAnimSequenceBase* Animation) override;
 
-	bool ValidateInput(class USkeletalMeshComponent * MeshComp, bool bReportErrors = false);
+	bool ValidateInput(USkeletalMeshComponent* MeshComp, bool bReportErrors = false);
 };
 
 
diff --git a/Engine/Source/Runtime/Engine/Classes/Animation/AnimNotifies/AnimNotify_ResetClothingSimulation.h b/Engine/Source/Runtime/Engine/Classes/Animation/AnimNotifies/AnimNotify_ResetClothingSimulation.h
new file mode 100644
--- /dev/null
+++ b/Engine/Source/Runtime/Engine/Classes/Animation/AnimNotifies/AnimNotify_ResetClothingSimulation.h
@@ -0,0 +1,26 @@
+// Copyright 1998-2016 Epic Games, Inc. All Rights Reserved.
+
+#pragma once
+
+#include "AnimNotify.h"
+
+#include "AnimNotify_ResetClothingSimulation.generated.h"
+
+UCLASS(const, hidecategories=Object, collapsecategories, meta=(DisplayName="Reset Clothing Simulation"))
+class ENGINE_API UAnimNotify_ResetClothingSimulation : public UAnimNotify
+{
+	GENERATED_BODY()
+
+public:
+
+	UAnimNotify_ResetClothingSimulation();
+
+	// Begin UAnimNotify interface
+	virtual FString GetNotifyName_Implementation() const override;
+	virtual void Notify(USkeletalMeshComponent* MeshComp, UAnimSequenceBase* Animation) override;
+	// End UAnimNotify interface
+
+};
+
+
+
diff --git a/Engine/Source/Runtime/Engine/Classes/Components/ActorComponent.h b/Engine/Source/Runtime/Engine/Classes/Components/ActorComponent.h
--- a/Engine/Source/Runtime/Engine/Classes/Components/ActorComponent.h
+++ b/Engine/Source/Runtime/Engine/Classes/Components/ActorComponent.h
@@ -224,6 +224,9 @@ private:
 	friend class FActorComponentInstanceData;
 	friend class FActorComponentDetails;
 
+	/** True if this component was owned by a net startup actor during level load. */
+	bool bIsNetStartupComponent;
+
 public:
 	UPROPERTY()
 	EComponentCreationMethod CreationMethod;
@@ -397,6 +400,12 @@ public:
 	*/
 	bool IsNetMode(ENetMode Mode) const;
 
+	/** Returns true if this component was owned by a net startup actor during level load. */
+	bool IsNetStartupComponent() const { return bIsNetStartupComponent; }
+
+	/** This should only be called by the engine in ULevel::InitializeNetworkActors to initialize bIsNetStartupComponent. */
+	void SetIsNetStartupComponent(const bool bInIsNetStartupComponent) { bIsNetStartupComponent = bInIsNetStartupComponent; }
+
 private:
 
 	/** 
diff --git a/Engine/Source/Runtime/Engine/Classes/Components/SkeletalMeshComponent.h b/Engine/Source/Runtime/Engine/Classes/Components/SkeletalMeshComponent.h
--- a/Engine/Source/Runtime/Engine/Classes/Components/SkeletalMeshComponent.h
+++ b/Engine/Source/Runtime/Engine/Classes/Components/SkeletalMeshComponent.h
@@ -391,6 +391,9 @@ public:
 	/** Temporary array of local-space (relative to parent bone) rotation/translation for each bone. */
 	TArray<FTransform> BoneSpaceTransforms;
 	
+	/** Temporary storage for curves */
+	FBlendedHeapCurve AnimCurves;
+
 	// Update Rate
 
 	/** Cached BoneSpaceTransforms for Update Rate optimization. */
@@ -562,7 +565,11 @@ public:
 	/** If true TickPose() will not be called from the Component's TickComponent function.
 	* It will instead be called from Autonomous networking updates. See ACharacter. */
 	UPROPERTY(Transient)
-	uint32 bAutonomousTickPose : 1;
+	uint32 bOnlyAllowAutonomousTickPose : 1;
+
+	/** True if calling TickPose() from Autonomous networking updates. See ACharacter. */
+	UPROPERTY(Transient)
+	uint32 bIsAutonomousTickPose : 1;
 
 	/** If true, force the mesh into the reference pose - is an optimization. */
 	UPROPERTY()
diff --git a/Engine/Source/Runtime/Engine/Classes/Components/TimelineComponent.h b/Engine/Source/Runtime/Engine/Classes/Components/TimelineComponent.h
--- a/Engine/Source/Runtime/Engine/Classes/Components/TimelineComponent.h
+++ b/Engine/Source/Runtime/Engine/Classes/Components/TimelineComponent.h
@@ -528,6 +528,8 @@ public:
 
 	/** Set the delegate to call when timeline is finished */
 	ENGINE_API void SetTimelineFinishedFunc(FOnTimelineEvent NewTimelineFinishedFunc);
+	/** Set the static delegate to call when timeline is finished */
+	ENGINE_API void SetTimelineFinishedFunc(FOnTimelineEventStatic NewTimelineFinishedFunc);
 
 	/** Set the delegate to call when timeline is finished */
 	ENGINE_API void SetDirectionPropertyName(FName DirectionPropertyName);
diff --git a/Engine/Source/Runtime/Engine/Classes/Components/WindDirectionalSourceComponent.h b/Engine/Source/Runtime/Engine/Classes/Components/WindDirectionalSourceComponent.h
--- a/Engine/Source/Runtime/Engine/Classes/Components/WindDirectionalSourceComponent.h
+++ b/Engine/Source/Runtime/Engine/Classes/Components/WindDirectionalSourceComponent.h
@@ -12,6 +12,9 @@ enum class EWindSourceType : uint8
 	Point,
 };
 
+class FWindData;
+class FWindSourceSceneProxy;
+
 /** Component that provides a directional wind source. Only affects SpeedTree assets. */
 UCLASS(collapsecategories, hidecategories=(Object, Mobility), editinlinenew)
 class UWindDirectionalSourceComponent : public USceneComponent
@@ -37,7 +40,7 @@ class UWindDirectionalSourceComponent : public USceneComponent
 	uint32 bPointWind : 1;
 
 public:
-	class FWindSourceSceneProxy* SceneProxy;
+	FWindSourceSceneProxy* SceneProxy;
 
 	/**
 	 * Because the actual data used to query wind is stored on the render thread in
@@ -70,6 +73,9 @@ public:
 	UFUNCTION(BlueprintCallable, Category = Wind)
 	void SetWindType(EWindSourceType InNewType);
 
+	/** Calculate wind parameters from the data on this component, safe to call on game thread */
+	ENGINE_API bool GetWindParameters(const FVector& EvaluatePosition, FWindData& OutData, float& Weight) const;
+
 protected:
 	//~ Begin UActorComponent Interface.
 	virtual void Activate(bool bReset) override;
@@ -88,8 +94,5 @@ public:
 	 * Creates a proxy to represent the primitive to the scene manager in the rendering thread.
 	 * @return The proxy object.
 	 */
-	virtual class FWindSourceSceneProxy* CreateSceneProxy() const;
+	virtual FWindSourceSceneProxy* CreateSceneProxy() const;
 };
-
-
-
diff --git a/Engine/Source/Runtime/Engine/Classes/Engine/Canvas.h b/Engine/Source/Runtime/Engine/Classes/Engine/Canvas.h
--- a/Engine/Source/Runtime/Engine/Classes/Engine/Canvas.h
+++ b/Engine/Source/Runtime/Engine/Classes/Engine/Canvas.h
@@ -208,7 +208,6 @@ class ENGINE_API UCanvas
 
 	FDisplayDebugManager DisplayDebugManager;
 public:
-
 	FCanvas* Canvas;
 	FSceneView* SceneView;
 	FMatrix	ViewProjectionMatrix;
@@ -219,6 +218,8 @@ public:
 	/** Initializes the canvas. */
 	void Init(int32 InSizeX, int32 InSizeY, FSceneView* InSceneView);
 
+	virtual void BeginDestroy() override;
+
 	/** Changes the view for the canvas. */
 	void SetView(FSceneView* InView);
 
