Redoing (edigrate) the Switch IWYU changes into Main so that it compiles again. This is part 1/2.

diff --git a/Engine/Source/Runtime/Core/Public/GenericPlatform/StandardPlatformString.h b/Engine/Source/Runtime/Core/Public/GenericPlatform/StandardPlatformString.h
--- a/Engine/Source/Runtime/Core/Public/GenericPlatform/StandardPlatformString.h
+++ b/Engine/Source/Runtime/Core/Public/GenericPlatform/StandardPlatformString.h
@@ -8,7 +8,7 @@
 #include "GenericPlatform/GenericPlatformString.h"
 #include <wchar.h>
 
-#if PLATFORM_APPLE || PLATFORM_LINUX || PLATFORM_HTML5 || PLATFORM_PS4
+#if PLATFORM_APPLE || PLATFORM_LINUX || PLATFORM_HTML5 || PLATFORM_PS4 || PLATFORM_SWITCH
 
 /**
 * Standard implementation
diff --git a/Engine/Source/Runtime/Core/Public/Misc/GenericGrowableAllocator.h b/Engine/Source/Runtime/Core/Public/Misc/GenericGrowableAllocator.h
--- a/Engine/Source/Runtime/Core/Public/Misc/GenericGrowableAllocator.h
+++ b/Engine/Source/Runtime/Core/Public/Misc/GenericGrowableAllocator.h
@@ -6,7 +6,8 @@
 
 #pragma once
 
-#include "Core.h"
+#include "CoreMinimal.h"
+#include "ScopeLock.h"
 
 
 
