Copying //UE4/Portal-Staging to //UE4/Dev-Main (Source: //Portal/Main @ 3216504)

==========================
MAJOR FEATURES + CHANGES
==========================

Change 3216141 on 2016/11/30 by [REDACTED_PERSON]

	Completed first ready to use pass of the new AutomationDriver module and new Spec test type.

Change 3213288 on 2016/11/29 by [REDACTED_PERSON]

	Removing deprecation of IWebBrowserSingleton::SetJSBindingToLoweringEnabled for now.

Change 3212796 on 2016/11/29 by [REDACTED_PERSON]

	Added SetJSBindingToLoweringEnabled to IWebBrowserSingleton so that the to-lowering of binding names can be disabled.
	Deprecated SetJSBindingToLoweringEnabled since 4.15. In future the to-lowering will always occurr.
	Adding GetBindingName helper to FWebJSScripting, which returns a to-lowered name for a UField, unless disabled.
	Updated all current binding code to use GetBindingName when building from UObjects/UStructs. This affects Windows, Mac, Linux, and Android.
	Portal currently disables to-lowering unless a commandline -LowercaseJS is provided.

Change 3200370 on 2016/11/16 by [REDACTED_PERSON]

	Ensure we always get the latest version of the user content catalog when promoting marketplace items.

Change 3192974 on 2016/11/10 by [REDACTED_PERSON]

	This is because of the OnPreExit core delegate being used to null out the Data uobject member on a manifest, also being the only sensible way to ensure threads complete in a safe and clean manner.
	Refactoring BuildPatchServices manifest class to not permanently hold any UObject and simply just use one while serialising.
	This removes the reliance on the OnPreExit delegate from manifest class, making it generally safer behaviour for shutdown.

Change 3187028 on 2016/11/04 by [REDACTED_PERSON]

	PortalPublishingTool: Adding UE_Main app to UnrealEngine project

Change 3186788 on 2016/11/04 by [REDACTED_PERSON]

	Change C# wrapper for BuildPatchTool patch generation to prevent clobbering manifest files by default, unless we specifically pass in an optional flag to allow this.

Change 3186779 on 2016/11/04 by [REDACTED_PERSON]

	Add support to automation tool testing framework for the following assertions:
	    Assert.AreNotEqual(a, b, optionalFailureMessage)
	    Assert.ThrowsError(actionToCarryOut, expectedExceptionType, optionalExceptionMessageContainsString)

	Moved attribute-based expected exception declarations to their own attribute, TestThrowsExceptionAttribute, which can now accept an optional parameter for a string which should be contained within the exception message.

	Fixed a bug where a test method with an attribute-based expected exception would not count towards the success total if the exception was encountered as expected.

	Fixed a bug where NOT throwing an exception when we were expecting one would count as a success.

	Added an internal property bDoNotLogTestFailsAsError which we can set to true to suppress logging of UAT errors when a test fails (but still count them in our failure results), to allow us to deliberately cause test failures to test the test framework!

	Added a suite of unit tests for the test framework itself, in TestRunner.Automation.Tests.cs.

Change 3185411 on 2016/11/03 by [REDACTED_PERSON]

	Allow Rocket_PromoteBuild changelist to be overridden by a changelist read from a file.

Change 3184843 on 2016/11/03 by [REDACTED_PERSON]

	Ensure catalog file synced during user content generation is always the latest one.

Change 3184752 on 2016/11/03 by [REDACTED_PERSON]

	Ensure we log reading changelist from specified file.

Change 3184744 on 2016/11/03 by [REDACTED_PERSON]

	Ensure directory is created for Changelist file if it doesn't already exist.

Change 3184738 on 2016/11/03 by [REDACTED_PERSON]

	Ensure we use latest CL from all of Perforce when generating build versions for user content

	Because of the nature of the build farm, where separate parts of the job are executed on different build agents at different times, this changelist is serialized to the filesystem during execution of a node dedicated to this task, and then made available to all future nodes, so that they're working with a consistent build version.  In the case of an execution where we're updating Perforce with new content, this calculation of the changelist occurs AFTER we've updated Perforce with the new content.

	Have also optimized the build graph scripts to enable Mac and Windows user generated content to execute simultaneously.

Change 3181456 on 2016/11/01 by [REDACTED_PERSON]

	SExpandable area has been modified as the Portal settings mocks weren't able to be achieved with default functionality.
	Added BodyBorderImage arguement and BodyBorderBackgroundColor attribute so we can specify a different brush/color to use for the expanded area compared with the title area. Additional care was made to ensure that rounded corners still appear correctly if the developer doesn't want to specify a different look to the body.
	Added AreaTitlePadding attribute, to be able to specify padding between the expand/collapse icon and the header content.
	Added MinWidth arguement, to ensure that the areas meet a minimum width requirement.

Change 3181285 on 2016/11/01 by [REDACTED_PERSON]

	Ensure user content generated using latest changelist submitted to Perforce, rather than using portal's latest changelist

Change 3177758 on 2016/10/28 by [REDACTED_PERSON]

	Copying //Portal/Dev-Social to Dev-Main (//Portal/Dev-Main)

Change 3175889 on 2016/10/26 by [REDACTED_PERSON]

	Web browser tooltips will no longer continue to appear when the mouse leaves the browser window.

Change 3171388 on 2016/10/22 by [REDACTED_PERSON]

	BPS: FBuildPatchAppManifest needs to listen for FCoreDelegates::OnPreExit in order to clean up references to it's UObject which is about to be destroyed.

Change 3170373 on 2016/10/21 by [REDACTED_PERSON]

	Reducing platform regex to only match pre-defined possibilities.

diff --git a/Engine/Plugins/Online/OnlineSubsystem/Source/Private/OnlineError.cpp b/Engine/Plugins/Online/OnlineSubsystem/Source/Private/OnlineError.cpp
--- a/Engine/Plugins/Online/OnlineSubsystem/Source/Private/OnlineError.cpp
+++ b/Engine/Plugins/Online/OnlineSubsystem/Source/Private/OnlineError.cpp
@@ -37,7 +37,7 @@ void FOnlineError::SetFromErrorCode(const FString& ErrorCodeIn)
 FOnlineError::FOnlineError(const FText& ErrorMessageIn)
 	: bSucceeded(false)
 	, HttpResult(0)
-	, NumericErrorCode(0)
+	, NumericErrorCode(INDEX_NONE)
 {
 	SetFromErrorMessage(ErrorMessageIn);
 }
diff --git a/Engine/Plugins/Online/OnlineSubsystem/Source/Public/OnlineSubsystemTypes.h b/Engine/Plugins/Online/OnlineSubsystem/Source/Public/OnlineSubsystemTypes.h
--- a/Engine/Plugins/Online/OnlineSubsystem/Source/Public/OnlineSubsystemTypes.h
+++ b/Engine/Plugins/Online/OnlineSubsystem/Source/Public/OnlineSubsystemTypes.h
@@ -129,7 +129,7 @@ namespace EOnlineServerConnectionStatus
 	enum Type : unsigned int
 	{
 		/** System normal (used for default state) */
-		Normal,
+		Normal = 0,
 		/** Gracefully disconnected from the online servers */
 		NotConnected,
 		/** Connected to the online servers just fine */
diff --git a/Engine/Source/Developer/AutomationDriver/AutomationDriver.Build.cs b/Engine/Source/Developer/AutomationDriver/AutomationDriver.Build.cs
new file mode 100644
--- /dev/null
+++ b/Engine/Source/Developer/AutomationDriver/AutomationDriver.Build.cs
@@ -0,0 +1,35 @@
+// Copyright 1998-2016 Epic Games, Inc. All Rights Reserved.
+
+using UnrealBuildTool;
+
+public class AutomationDriver : ModuleRules
+{
+    public AutomationDriver(TargetInfo Target)
+    {
+        PublicIncludePaths.AddRange(
+            new string[] {
+				"Developer/AutomationDriver/Public",
+            }
+        );
+
+        PrivateIncludePaths.AddRange(
+            new string[] {
+                "Developer/AutomationDriver/Private",
+                "Developer/AutomationDriver/Private/Locators",
+                "Developer/AutomationDriver/Private/Specs",
+                "Developer/AutomationDriver/Private/MetaData",
+            }
+        );
+
+        PrivateDependencyModuleNames.AddRange(
+            new string[] {
+                "Core",
+                "CoreUObject",
+                "InputCore",
+                "Json",
+                "Slate",
+                "SlateCore",
+            }
+        );
+    }
+}
diff --git a/Engine/Source/Developer/AutomationDriver/Private/AutomationDriverLogging.cpp b/Engine/Source/Developer/AutomationDriver/Private/AutomationDriverLogging.cpp
new file mode 100644
--- /dev/null
+++ b/Engine/Source/Developer/AutomationDriver/Private/AutomationDriverLogging.cpp
@@ -0,0 +1,85 @@
+// Copyright 1998-2016 Epic Games, Inc. All Rights Reserved.
+
+#include "AutomationDriverLogging.h"
+#include "AutomationDriverCommon.h"
+#include "IApplicationElement.h"
+
+DEFINE_LOG_CATEGORY(LogAutomationDriver);
+
+void FAutomationDriverLogging::TooManyElementsFound(const TArray<TSharedRef<IApplicationElement>>& Elements)
+{
+	UE_LOG(LogAutomationDriver, Error, TEXT("Multiple elements found when 1 was expected\nExpected 1\n   Found %d"), Elements.Num());
+
+	for (int32 Index = 0; Index < Elements.Num(); Index++)
+	{
+		const TSharedRef<IApplicationElement>& Element = Elements[Index];
+		UE_LOG(LogAutomationDriver, Error, TEXT("    [%d] -> %s"), Index, *Element->ToDebugString());
+	}
+}
+
+void FAutomationDriverLogging::CannotFindElement(const TSharedPtr<IElementLocator, ESPMode::ThreadSafe>& ElementLocator)
+{
+	UE_LOG(LogAutomationDriver, Error, TEXT("Failed to locate element"));
+
+	if (ElementLocator.IsValid())
+	{
+		UE_LOG(LogAutomationDriver, Error, TEXT("    %s"), *ElementLocator->ToDebugString());
+	}
+}
+
+void FAutomationDriverLogging::ElementNotVisible(const TSharedPtr<IElementLocator, ESPMode::ThreadSafe>& ElementLocator)
+{
+	UE_LOG(LogAutomationDriver, Error, TEXT("Failed to locate visible element"));
+
+	if (ElementLocator.IsValid())
+	{
+		UE_LOG(LogAutomationDriver, Error, TEXT("    Element found but not visible: %s"), *ElementLocator->ToDebugString());
+	}
+	else
+	{
+		UE_LOG(LogAutomationDriver, Error, TEXT("    Element found but not visible"));
+	}
+}
+
+void FAutomationDriverLogging::ElementNotInteractable(const TSharedPtr<IElementLocator, ESPMode::ThreadSafe>& ElementLocator)
+{
+	UE_LOG(LogAutomationDriver, Error, TEXT("Failed to locate interactable element"));
+
+	if (ElementLocator.IsValid())
+	{
+		UE_LOG(LogAutomationDriver, Error, TEXT("    Element found but not interactable: %s"), *ElementLocator->ToDebugString());
+	}
+	else
+	{
+		UE_LOG(LogAutomationDriver, Error, TEXT("    Element found but not interactable"));
+	}
+}
+
+void FAutomationDriverLogging::ElementHasNoWindow(const TSharedPtr<IElementLocator, ESPMode::ThreadSafe>& ElementLocator)
+{
+	UE_LOG(LogAutomationDriver, Error, TEXT("Failed to locate window hosting element"));
+
+	if (ElementLocator.IsValid())
+	{
+		UE_LOG(LogAutomationDriver, Error, TEXT("    Element found but no window is associated with it: %s"), *ElementLocator->ToDebugString());
+	}
+	else
+	{
+		UE_LOG(LogAutomationDriver, Error, TEXT("    Element found but no window is associated with it"));
+	}
+}
+
+void FAutomationDriverLogging::CannotClickUnhoveredElement(const TSharedPtr<IElementLocator, ESPMode::ThreadSafe>& ElementLocator)
+{
+	UE_LOG(LogAutomationDriver, Error, TEXT("Failed to click element"));
+
+	if (ElementLocator.IsValid())
+	{
+		UE_LOG(LogAutomationDriver, Error, TEXT("    Element found but not located under the cursor: %s"), *ElementLocator->ToDebugString());
+	}
+	else
+	{
+		UE_LOG(LogAutomationDriver, Error, TEXT("    Element found but not located under the cursor"));
+	}
+}
+
diff --git a/Engine/Source/Developer/AutomationDriver/Private/IStepExecutor.h b/Engine/Source/Developer/AutomationDriver/Private/IStepExecutor.h
new file mode 100644
--- /dev/null
+++ b/Engine/Source/Developer/AutomationDriver/Private/IStepExecutor.h
@@ -0,0 +1,46 @@
+// Copyright 1998-2016 Epic Games, Inc. All Rights Reserved.
+#pragma once
+
+#include "AsyncResult.h"
+#include "Misc/Timespan.h"
+
+struct FStepResult
+{
+public:
+
+	enum class EState : uint8
+	{
+		DONE,
+		FAILED,
+		REPEAT,
+	};
+
+	FStepResult(EState InState, FTimespan InNextWait)
+		: NextWait(MoveTemp(InNextWait))
+		, State(InState)
+	{ }
+
+	// How long the executor should wait before either executing this same step, the next step or before declaring all steps complete
+	FTimespan NextWait;
+
+	// Whether the step that just completed is completely finished or should be rescheduled again for execution
+	EState State;
+};
+
+class IStepExecutor
+{
+public:
+
+	DECLARE_DELEGATE_RetVal_OneParam(FStepResult, FExecuteStepDelegate, const FTimespan& /*TotalProcessTime*/);
+	virtual void Add(const FExecuteStepDelegate& Step) = 0;
+
+	virtual void Add(const TFunction<FStepResult(const FTimespan&)>& Step) = 0;
+
+	virtual void InsertNext(const FExecuteStepDelegate& Step) = 0;
+
+	virtual void InsertNext(const TFunction<FStepResult(const FTimespan&)>& Step) = 0;
+
+	virtual TAsyncResult<bool> Execute() = 0;
+
+	virtual bool IsExecuting() const = 0;
+};
diff --git a/Engine/Source/Developer/AutomationDriver/Private/LocateBy.cpp b/Engine/Source/Developer/AutomationDriver/Private/LocateBy.cpp
new file mode 100644
--- /dev/null
+++ b/Engine/Source/Developer/AutomationDriver/Private/LocateBy.cpp
@@ -0,0 +1,79 @@
+// Copyright 1998-2016 Epic Games, Inc. All Rights Reserved.
+
+#include "LocateBy.h"
+#include "SlateWidgetLocatorByDelegate.h"
+#include "SlateWidgetLocatorByPath.h"
+
+TSharedRef<IElementLocator, ESPMode::ThreadSafe> By::Delegate(const FLocateSlateWidgetElementDelegate& Value)
+{
+	return FSlateWidgetLocatorByDelegateFactory::Create(Value);
+}
+
+TSharedRef<IElementLocator, ESPMode::ThreadSafe> By::Delegate(const FLocateSlateWidgetPathElementDelegate& Value)
+{
+	return FSlateWidgetLocatorByDelegateFactory::Create(Value);
+}
+
+TSharedRef<IElementLocator, ESPMode::ThreadSafe> By::WidgetLambda(const TFunction<void(TArray<TSharedRef<SWidget>>&)>& Value)
+{
+	return FSlateWidgetLocatorByDelegateFactory::Create(FLocateSlateWidgetElementDelegate::CreateLambda(Value));
+}
+
+TSharedRef<IElementLocator, ESPMode::ThreadSafe> By::WidgetPathLambda(const TFunction<void(TArray<FWidgetPath>&)>& Value)
+{
+	return FSlateWidgetLocatorByDelegateFactory::Create(FLocateSlateWidgetPathElementDelegate::CreateLambda(Value));
+}
+
+TSharedRef<IElementLocator, ESPMode::ThreadSafe> By::Id(const FString& Value)
+{
+	return FSlateWidgetLocatorByPathFactory::Create(TEXT("#") + Value);
+}
+
+TSharedRef<IElementLocator, ESPMode::ThreadSafe> By::Id(const FName& Value)
+{
+	return FSlateWidgetLocatorByPathFactory::Create(TEXT("#") + Value.ToString());
+}
+
+TSharedRef<IElementLocator, ESPMode::ThreadSafe> By::Id(const TCHAR* Value)
+{
+	return FSlateWidgetLocatorByPathFactory::Create(FString(TEXT("#")) + Value);
+}
+
+TSharedRef<IElementLocator, ESPMode::ThreadSafe> By::Id(const char* Value)
+{
+	return FSlateWidgetLocatorByPathFactory::Create(FString(TEXT("#")) + Value);
+}
+
+TSharedRef<IElementLocator, ESPMode::ThreadSafe> By::Path(const FString& Value)
+{
+	return FSlateWidgetLocatorByPathFactory::Create(Value);
+}
+
+TSharedRef<IElementLocator, ESPMode::ThreadSafe> By::Path(const FName& Value)
+{
+	return FSlateWidgetLocatorByPathFactory::Create(Value.ToString());
+}
+
+TSharedRef<IElementLocator, ESPMode::ThreadSafe> By::Path(const TCHAR* Value)
+{
+	return FSlateWidgetLocatorByPathFactory::Create(Value);
+}
+
+TSharedRef<IElementLocator, ESPMode::ThreadSafe> By::Path(const char* Value)
+{
+	return FSlateWidgetLocatorByPathFactory::Create(Value);
+}
+
+TSharedRef<IElementLocator, ESPMode::ThreadSafe> By::Cursor()
+{
+	return By::WidgetPathLambda([](TArray<FWidgetPath>& OutWidgetPaths) -> void {
+		TArray<TSharedRef<SWindow>> Windows;
+		FSlateApplication::Get().GetAllVisibleWindowsOrdered(Windows);
+		FWidgetPath WidgetPath = FSlateApplication::Get().LocateWindowUnderMouse(FSlateApplication::Get().GetCursorPos(), Windows);
+
+		if (WidgetPath.IsValid())
+		{
+			OutWidgetPaths.Add(WidgetPath);
+		}
+	});
+}
diff --git a/Engine/Source/Developer/AutomationDriver/Private/Locators/SlateWidgetLocatorByDelegate.h b/Engine/Source/Developer/AutomationDriver/Private/Locators/SlateWidgetLocatorByDelegate.h
new file mode 100644
--- /dev/null
+++ b/Engine/Source/Developer/AutomationDriver/Private/Locators/SlateWidgetLocatorByDelegate.h
@@ -0,0 +1,17 @@
+// Copyright 1998-2016 Epic Games, Inc. All Rights Reserved.
+#pragma once
+
+#include "LocateBy.h"
+
+class IElementLocator;
+
+class FSlateWidgetLocatorByDelegateFactory
+{
+public:
+
+	static TSharedRef<IElementLocator, ESPMode::ThreadSafe> Create(
+		const FLocateSlateWidgetElementDelegate& Delegate);
+
+	static TSharedRef<IElementLocator, ESPMode::ThreadSafe> Create(
+		const FLocateSlateWidgetPathElementDelegate& Delegate);
+};
diff --git a/Engine/Source/Developer/AutomationDriver/Private/PassThroughMessageHandler .cpp b/Engine/Source/Developer/AutomationDriver/Private/PassThroughMessageHandler .cpp
new file mode 100644
--- /dev/null
+++ b/Engine/Source/Developer/AutomationDriver/Private/PassThroughMessageHandler .cpp	
@@ -0,0 +1,469 @@
+// Copyright 1998-2016 Epic Games, Inc. All Rights Reserved.
+
+#include "PassThroughMessageHandler.h"
+
+class FPassThroughMessageHandlerFactory;
+
+class FPassThroughMessageHandlerImpl
+	: public FPassThroughMessageHandler
+{
+public:
+
+	virtual bool IsHandlingMessages() const override
+	{
+		return bAllowMessageHandling;
+	}
+
+	virtual void SetAllowMessageHandling(bool bValue) override
+	{
+		bAllowMessageHandling = bValue;
+	}
+
+public:
+
+	virtual bool ShouldProcessUserInputMessages(const TSharedPtr< FGenericWindow >& PlatformWindow) const override
+	{
+		return RealMessageHandler->ShouldProcessUserInputMessages(PlatformWindow);
+	}
+
+	virtual bool OnKeyChar(const TCHAR Character, const bool IsRepeat) override
+	{
+		if (!bAllowMessageHandling)
+		{
+			return false;
+		}
+
+		return RealMessageHandler->OnKeyChar(Character, IsRepeat);
+	}
+
+	virtual bool OnKeyDown(const int32 KeyCode, const uint32 CharacterCode, const bool IsRepeat) override
+	{
+		if (!bAllowMessageHandling)
+		{
+			return false;
+		}
+
+		return RealMessageHandler->OnKeyDown(KeyCode, CharacterCode, IsRepeat);
+	}
+
+	virtual bool OnKeyUp(const int32 KeyCode, const uint32 CharacterCode, const bool IsRepeat) override
+	{
+		const FKey Key = FInputKeyManager::Get().GetKeyFromCodes(KeyCode, CharacterCode);
+
+		if (Key == EKeys::ScrollLock)
+		{
+			// Allow scroll lock to toggle whether platform input can be processed by the application
+			bAllowMessageHandling = !bAllowMessageHandling;
+		}
+
+		if (!bAllowMessageHandling)
+		{
+			return false;
+		}
+
+		return RealMessageHandler->OnKeyUp(KeyCode, CharacterCode, IsRepeat);
+	}
+
+	virtual bool OnMouseDown(const TSharedPtr< FGenericWindow >& Window, const EMouseButtons::Type Button) override
+	{
+		if (!bAllowMessageHandling)
+		{
+			return false;
+		}
+
+		return RealMessageHandler->OnMouseDown(Window, Button);
+	}
+
+	virtual bool OnMouseDown(const TSharedPtr< FGenericWindow >& Window, const EMouseButtons::Type Button, const FVector2D CursorPos)
+	{
+		if (!bAllowMessageHandling)
+		{
+			return false;
+		}
+
+		return RealMessageHandler->OnMouseDown(Window, Button, CursorPos);
+	}
+
+	virtual bool OnMouseUp(const EMouseButtons::Type Button) override
+	{
+		if (!bAllowMessageHandling)
+		{
+			return false;
+		}
+
+		return RealMessageHandler->OnMouseUp(Button);
+	}
+
+	virtual bool OnMouseUp(const EMouseButtons::Type Button, const FVector2D CursorPos)
+	{
+		if (!bAllowMessageHandling)
+		{
+			return false;
+		}
+
+		return RealMessageHandler->OnMouseUp(Button, CursorPos);
+	}
+
+	virtual bool OnMouseDoubleClick(const TSharedPtr< FGenericWindow >& Window, const EMouseButtons::Type Button) override
+	{
+		if (!bAllowMessageHandling)
+		{
+			return false;
+		}
+
+		return RealMessageHandler->OnMouseDoubleClick(Window, Button);
+	}
+
+	virtual bool OnMouseDoubleClick(const TSharedPtr< FGenericWindow >& Window, const EMouseButtons::Type Button, const FVector2D CursorPos)
+	{
+		if (!bAllowMessageHandling)
+		{
+			return false;
+		}
+
+		return RealMessageHandler->OnMouseDoubleClick(Window, Button, CursorPos);
+	}
+
+	virtual bool OnMouseWheel(const float Delta) override
+	{
+		if (!bAllowMessageHandling)
+		{
+			return false;
+		}
+
+		return RealMessageHandler->OnMouseWheel(Delta);
+	}
+
+	virtual bool OnMouseWheel(const float Delta, const FVector2D CursorPos)
+	{
+		if (!bAllowMessageHandling)
+		{
+			return false;
+		}
+
+		return RealMessageHandler->OnMouseWheel(Delta, CursorPos);
+	}
+
+	virtual bool OnMouseMove() override
+	{
+		if (!bAllowMessageHandling)
+		{
+			return false;
+		}
+
+		return RealMessageHandler->OnMouseMove();
+	}
+
+	virtual bool OnRawMouseMove(const int32 X, const int32 Y) override
+	{
+		if (!bAllowMessageHandling)
+		{
+			return false;
+		}
+
+		return RealMessageHandler->OnRawMouseMove(X, Y);
+	}
+
+	virtual bool OnCursorSet() override
+	{
+		if (!bAllowMessageHandling)
+		{
+			return false;
+		}
+
+		return RealMessageHandler->OnCursorSet();
+	}
+
+	virtual bool OnControllerAnalog(FGamepadKeyNames::Type KeyName, int32 ControllerId, float AnalogValue) override
+	{
+		if (!bAllowMessageHandling)
+		{
+			return false;
+		}
+
+		return RealMessageHandler->OnControllerAnalog(KeyName, ControllerId, AnalogValue);
+	}
+
+	virtual bool OnControllerButtonPressed(FGamepadKeyNames::Type KeyName, int32 ControllerId, bool IsRepeat) override
+	{
+		if (!bAllowMessageHandling)
+		{
+			return false;
+		}
+
+		return RealMessageHandler->OnControllerButtonPressed(KeyName, ControllerId, IsRepeat);
+	}
+
+	virtual bool OnControllerButtonReleased(FGamepadKeyNames::Type KeyName, int32 ControllerId, bool IsRepeat) override
+	{
+		if (!bAllowMessageHandling)
+		{
+			return false;
+		}
+
+		return RealMessageHandler->OnControllerButtonReleased(KeyName, ControllerId, IsRepeat);
+	}
+
+	virtual void OnBeginGesture() override
+	{
+		if (!bAllowMessageHandling)
+		{
+			return;
+		}
+
+		return RealMessageHandler->OnBeginGesture();
+	}
+
+	virtual bool OnTouchGesture(EGestureEvent::Type GestureType, const FVector2D& Delta, float WheelDelta, bool bIsDirectionInvertedFromDevice) override
+	{
+		if (!bAllowMessageHandling)
+		{
+			return false;
+		}
+
+		return RealMessageHandler->OnTouchGesture(GestureType, Delta, WheelDelta, bIsDirectionInvertedFromDevice);
+	}
+
+	virtual void OnEndGesture() override
+	{
+		if (!bAllowMessageHandling)
+		{
+			return;
+		}
+
+		RealMessageHandler->OnEndGesture();
+	}
+
+	virtual bool OnTouchStarted(const TSharedPtr< FGenericWindow >& Window, const FVector2D& Location, int32 TouchIndex, int32 ControllerId) override
+	{
+		if (!bAllowMessageHandling)
+		{
+			return false;
+		}
+
+		return RealMessageHandler->OnTouchStarted(Window, Location, TouchIndex, ControllerId);
+	}
+
+	virtual bool OnTouchMoved(const FVector2D& Location, int32 TouchIndex, int32 ControllerId) override
+	{
+		if (!bAllowMessageHandling)
+		{
+			return false;
+		}
+
+		return RealMessageHandler->OnTouchMoved(Location, TouchIndex, ControllerId);
+	}
+
+	virtual bool OnTouchEnded(const FVector2D& Location, int32 TouchIndex, int32 ControllerId) override
+	{
+		if (!bAllowMessageHandling)
+		{
+			return false;
+		}
+
+		return RealMessageHandler->OnTouchEnded(Location, TouchIndex, ControllerId);
+	}
+
+	virtual bool OnMotionDetected(const FVector& Tilt, const FVector& RotationRate, const FVector& Gravity, const FVector& Acceleration, int32 ControllerId) override
+	{
+		if (!bAllowMessageHandling)
+		{
+			return false;
+		}
+
+		return RealMessageHandler->OnMotionDetected(Tilt, RotationRate, Gravity, Acceleration, ControllerId);
+	}
+
+	virtual bool OnSizeChanged(const TSharedRef< FGenericWindow >& Window, const int32 Width, const int32 Height, bool bWasMinimized = false) override
+	{
+		if (!bAllowMessageHandling)
+		{
+			return false;
+		}
+
+		return RealMessageHandler->OnSizeChanged(Window, Width, Height, bWasMinimized);
+	}
+
+	virtual void OnOSPaint(const TSharedRef<FGenericWindow>& Window) override
+	{
+		RealMessageHandler->OnOSPaint(Window);
+	}
+
+	virtual FWindowSizeLimits GetSizeLimitsForWindow(const TSharedRef<FGenericWindow>& Window) const override
+	{
+		return RealMessageHandler->GetSizeLimitsForWindow(Window);
+	}
+
+	virtual void OnResizingWindow(const TSharedRef<FGenericWindow>& Window) override
+	{
+		if (!bAllowMessageHandling)
+		{
+			return;
+		}
+
+		RealMessageHandler->OnResizingWindow(Window);
+	}
+
+	virtual bool BeginReshapingWindow(const TSharedRef< FGenericWindow >& Window) override
+	{
+		if (!bAllowMessageHandling)
+		{
+			return false;
+		}
+
+		return RealMessageHandler->BeginReshapingWindow(Window);
+	}
+
+	virtual void FinishedReshapingWindow(const TSharedRef< FGenericWindow >& Window) override
+	{
+		if (!bAllowMessageHandling)
+		{
+			return;
+		}
+
+		RealMessageHandler->FinishedReshapingWindow(Window);
+	}
+
+	virtual void OnMovedWindow(const TSharedRef< FGenericWindow >& Window, const int32 X, const int32 Y) override
+	{
+		if (!bAllowMessageHandling)
+		{
+			return;
+		}
+
+		RealMessageHandler->OnMovedWindow(Window, X, Y);
+	}
+
+	virtual bool OnWindowActivationChanged(const TSharedRef< FGenericWindow >& Window, const EWindowActivation::Type ActivationType) override
+	{
+		return RealMessageHandler->OnWindowActivationChanged(Window, ActivationType);
+	}
+
+	virtual bool OnApplicationActivationChanged(const bool IsActive) override
+	{
+		return RealMessageHandler->OnApplicationActivationChanged(IsActive);
+	}
+
+	virtual bool OnConvertibleLaptopModeChanged() override
+	{
+		if (!bAllowMessageHandling)
+		{
+			return true;
+		}
+
+		return RealMessageHandler->OnConvertibleLaptopModeChanged();
+	}
+
+	virtual EWindowZone::Type GetWindowZoneForPoint(const TSharedRef< FGenericWindow >& Window, const int32 X, const int32 Y) override
+	{
+		if (!bAllowMessageHandling)
+		{
+			return EWindowZone::NotInWindow;
+		}
+
+		return RealMessageHandler->GetWindowZoneForPoint(Window, X, Y);
+	}
+
+	virtual void OnWindowClose(const TSharedRef< FGenericWindow >& Window) override
+	{
+		if (!bAllowMessageHandling)
+		{
+			return;
+		}
+
+		RealMessageHandler->OnWindowClose(Window);
+	}
+
+	virtual EDropEffect::Type OnDragEnterText(const TSharedRef< FGenericWindow >& Window, const FString& Text) override
+	{
+		if (!bAllowMessageHandling)
+		{
+			return EDropEffect::None;
+		}
+
+		return RealMessageHandler->OnDragEnterText(Window, Text);
+	}
+
+	virtual EDropEffect::Type OnDragEnterFiles(const TSharedRef< FGenericWindow >& Window, const TArray< FString >& Files) override
+	{
+		if (!bAllowMessageHandling)
+		{
+			return EDropEffect::None;
+		}
+
+		return RealMessageHandler->OnDragEnterFiles(Window, Files);
+	}
+
+	virtual EDropEffect::Type OnDragOver(const TSharedPtr< FGenericWindow >& Window) override
+	{
+		if (!bAllowMessageHandling)
+		{
+			return EDropEffect::None;
+		}
+
+		return RealMessageHandler->OnDragOver(Window);
+	}
+
+	virtual void OnDragLeave(const TSharedPtr< FGenericWindow >& Window) override
+	{
+		if (!bAllowMessageHandling)
+		{
+			return;
+		}
+
+		RealMessageHandler->OnDragLeave(Window);
+	}
+
+	virtual EDropEffect::Type OnDragDrop(const TSharedPtr< FGenericWindow >& Window) override
+	{
+		if (!bAllowMessageHandling)
+		{
+			return EDropEffect::None;
+		}
+
+		return RealMessageHandler->OnDragDrop(Window);
+	}
+
+	virtual bool OnWindowAction(const TSharedRef< FGenericWindow >& Window, const EWindowAction::Type InActionType) override
+	{
+		if (!bAllowMessageHandling)
+		{
+			return true;
+		}
+
+		return RealMessageHandler->OnWindowAction(Window, InActionType);
+	}
+
+private:
+
+	FPassThroughMessageHandlerImpl(
+		const TSharedRef<FGenericApplicationMessageHandler>& InMessageHandler)
+		: RealMessageHandler(InMessageHandler)
+		, bAllowMessageHandling(false)
+	{ }
+
+private:
+
+	const TSharedRef<FGenericApplicationMessageHandler> RealMessageHandler;
+
+	bool bAllowMessageHandling;
+
+	friend FPassThroughMessageHandlerFactory;
+};
+
+class FPassThroughMessageHandlerFactory
+	: public IPassThroughMessageHandlerFactory
+{
+public:
+
+	virtual TSharedRef<FPassThroughMessageHandler> Create(
+		const TSharedRef<FGenericApplicationMessageHandler>& MessageHandler) const override
+	{
+		return MakeShareable(new FPassThroughMessageHandlerImpl(MessageHandler));
+	}
+};
+
+TSharedRef<IPassThroughMessageHandlerFactory> FPassThroughMessageHandlerFactoryFactory::Create()
+{
+	return MakeShareable(new FPassThroughMessageHandlerFactory());
+}
\ No newline at end of file
diff --git a/Engine/Source/Developer/AutomationDriver/Private/StepExecutor.cpp b/Engine/Source/Developer/AutomationDriver/Private/StepExecutor.cpp
new file mode 100644
--- /dev/null
+++ b/Engine/Source/Developer/AutomationDriver/Private/StepExecutor.cpp
@@ -0,0 +1,170 @@
+// Copyright 1998-2016 Epic Games, Inc. All Rights Reserved.
+
+#include "StepExecutor.h"
+#include "IStepExecutor.h"
+#include "IAutomationDriver.h"
+
+#include "DriverConfiguration.h"
+
+#include "ScopeLock.h"
+#include "Ticker.h"
+#include "Async.h"
+
+
+class FStepExecutor
+	: public IStepExecutor
+	, public TSharedFromThis<FStepExecutor, ESPMode::ThreadSafe>
+{
+public:
+
+	virtual ~FStepExecutor()
+	{
+		if (Promise.IsValid())
+		{
+			Promise->SetValue(false);
+		}
+	}
+
+	virtual void Add(const FExecuteStepDelegate& Step) override
+	{
+		check(!Promise.IsValid());
+		FScopeLock StateLock(&StepsCS);
+		Steps.Add(Step);
+	}
+
+	virtual void Add(const TFunction<FStepResult(const FTimespan&)>& Step) override
+	{
+		check(!Promise.IsValid());
+		FScopeLock StateLock(&StepsCS);
+		Steps.Add(FExecuteStepDelegate::CreateLambda(Step));
+	}
+
+	virtual void InsertNext(const FExecuteStepDelegate& Step) override
+	{
+		check(Promise.IsValid());
+		FScopeLock StateLock(&StepsCS);
+		Steps.Insert(Step, CurrentStepIndex + 1);
+	}
+
+	virtual void InsertNext(const TFunction<FStepResult(const FTimespan&)>& Step) override
+	{
+		check(Promise.IsValid());
+		FScopeLock StateLock(&StepsCS);
+		Steps.Insert(FExecuteStepDelegate::CreateLambda(Step), CurrentStepIndex + 1);
+	}
+
+	virtual TAsyncResult<bool> Execute() override
+	{
+		check(!Promise.IsValid());
+		CurrentStepIndex = 0;
+
+		// Queue result back on main thread
+		TWeakPtr<FStepExecutor, ESPMode::ThreadSafe> LocalWeakThis(SharedThis(this));
+		AsyncTask(
+			ENamedThreads::GameThread,
+			[LocalWeakThis]()
+			{
+				TSharedPtr<FStepExecutor, ESPMode::ThreadSafe> Executor = LocalWeakThis.Pin();
+
+				if (Executor.IsValid())
+				{
+					const int32 StepIndex = 0;
+					FTicker::GetCoreTicker().AddTicker(FTickerDelegate::CreateThreadSafeSP(Executor.ToSharedRef(), &FStepExecutor::ExecuteStep, StepIndex), 0);
+				}
+			}
+		);
+
+		Promise = MakeShareable(new TPromise<bool>());
+		return TAsyncResult<bool>(Promise->GetFuture(), nullptr, nullptr);
+	}
+
+	virtual bool IsExecuting() const override
+	{
+		return Promise.IsValid();
+	}
+
+private:
+
+	FStepExecutor(
+		const TSharedRef<FDriverConfiguration, ESPMode::ThreadSafe>& InConfiguration)
+		: Configuration(InConfiguration)
+		, Steps()
+		, CurrentStepIndex(0)
+		, Promise()
+		, StepTotalProcessTime(FTimespan::Zero())
+		, LastDelay(0)
+	{ }
+
+	bool ExecuteStep(float Delta, int32 StepIndex)
+	{
+		check(IsInGameThread());
+
+		FStepResult Result = FStepResult(FStepResult::EState::FAILED, 0);
+		{
+			FScopeLock StateLock(&StepsCS);
+
+			// If we've encountered an invalid step that's greater then zero then we were just waiting
+			// a little bit after the last step completed before signaling completion.
+			if (StepIndex > 0 && !Steps.IsValidIndex(StepIndex))
+			{
+				Promise->SetValue(true);
+				Promise.Reset();
+				StepTotalProcessTime = FTimespan::Zero();
+				return false;
+			}
+
+			check(Steps.IsValidIndex(StepIndex));
+
+			Result = Steps[StepIndex].Execute(StepTotalProcessTime);
+		}
+
+		if (Result.State == FStepResult::EState::FAILED)
+		{
+			Promise->SetValue(false);
+			Promise.Reset();
+			StepTotalProcessTime = FTimespan::Zero();
+			return false;
+		}
+
+		if (Result.State == FStepResult::EState::DONE)
+		{
+			StepTotalProcessTime = FTimespan::Zero();
+			++StepIndex;
+		}
+
+		CurrentStepIndex = StepIndex;
+		float Milliseconds = Result.NextWait.GetTicks() / ETimespan::TicksPerMillisecond;
+		float Delay = FMath::Max(SMALL_NUMBER, (Milliseconds / 1000) * Configuration->ExecutionSpeedMultiplier);
+		
+		if (LastDelay < KINDA_SMALL_NUMBER)
+		{
+			StepTotalProcessTime += FTimespan::FromSeconds(Delta);
+		}
+
+		StepTotalProcessTime += FTimespan::FromSeconds(Delay);
+		LastDelay = Delay;
+		FTicker::GetCoreTicker().AddTicker(FTickerDelegate::CreateThreadSafeSP(this, &FStepExecutor::ExecuteStep, StepIndex), Delay);
+
+		return false;
+	}
+
+private:
+
+	const TSharedRef<FDriverConfiguration, ESPMode::ThreadSafe> Configuration;
+
+	TArray<FExecuteStepDelegate> Steps;
+	int32 CurrentStepIndex;
+	TSharedPtr<TPromise<bool>> Promise;
+	FTimespan StepTotalProcessTime;
+	float LastDelay;
+
+	mutable FCriticalSection StepsCS;
+
+	friend FStepExecutorFactory;
+};
+
+TSharedRef<IStepExecutor, ESPMode::ThreadSafe> FStepExecutorFactory::Create(
+	const TSharedRef<FDriverConfiguration, ESPMode::ThreadSafe>& Configuration)
+{
+	return MakeShareable(new FStepExecutor(Configuration));
+}
diff --git a/Engine/Source/Developer/AutomationDriver/Private/StepExecutor.h b/Engine/Source/Developer/AutomationDriver/Private/StepExecutor.h
new file mode 100644
--- /dev/null
+++ b/Engine/Source/Developer/AutomationDriver/Private/StepExecutor.h
@@ -0,0 +1,14 @@
+// Copyright 1998-2016 Epic Games, Inc. All Rights Reserved.
+#pragma once
+
+class IStepExecutor;
+class FDriverConfiguration;
+
+class FStepExecutorFactory
+{
+public:
+
+	static TSharedRef<IStepExecutor, ESPMode::ThreadSafe> Create(
+		const TSharedRef<FDriverConfiguration, ESPMode::ThreadSafe>& Configuration);
+
+};
diff --git a/Engine/Source/Developer/AutomationDriver/Public/WaitUntil.h b/Engine/Source/Developer/AutomationDriver/Public/WaitUntil.h
new file mode 100644
--- /dev/null
+++ b/Engine/Source/Developer/AutomationDriver/Public/WaitUntil.h
@@ -0,0 +1,178 @@
+// Copyright 1998-2016 Epic Games, Inc. All Rights Reserved.
+#pragma once
+
+class IElementLocator;
+
+/**
+ * Represents the state of an active wait action for the driver
+ */
+struct FDriverWaitResponse
+{
+public:
+
+	enum class EState : uint8
+	{
+		PASSED,
+		WAIT,
+		FAILED,
+	};
+
+	/**
+	 * @return a FDriverWaitResponse with a state of PASSED and a wait of zero
+	 */
+	static FDriverWaitResponse Passed();
+
+	/**
+	 * @return a FDriverWaitResponse with a state of WAIT and a wait of 0.5 seconds
+	 */
+	static FDriverWaitResponse Wait();
+
+	/**
+	 * @return a FDriverWaitResponse with a state of WAIT and a wait of the specified timespan
+	 */
+	static FDriverWaitResponse Wait(FTimespan Timespan);
+
+	/**
+	 * @return a FDriverWaitResponse with a state of FAILED and a wait of zero
+	 */
+	static FDriverWaitResponse Failed();
+
+	FDriverWaitResponse(EState InState, FTimespan InNextWait);
+
+	// How long the driver should wait before re-evaluating the wait condition again
+	const FTimespan NextWait;
+
+	// Whether the wait condition is completely finished or should be rescheduled again for execution
+	const EState State;
+};
+
+DECLARE_DELEGATE_RetVal_OneParam(FDriverWaitResponse, FDriverWaitDelegate, const FTimespan& /*TotalWaitTime*/);
+DECLARE_DELEGATE_RetVal(bool, FDriverWaitConditionDelegate);
+
+/**
+ * A fluent wrapper around timespan to enforce obvious differences between specified Timeout and Interval values for waits
+ */
+class FWaitTimeout
+{
+public:
+
+	static FWaitTimeout InMilliseconds(double Value);
+	static FWaitTimeout InSeconds(double Value);
+	static FWaitTimeout InMinutes(double Value);
+	static FWaitTimeout InHours(double Value);
+
+	explicit FWaitTimeout(FTimespan InTimespan);
+
+	const FTimespan Timespan;
+};
+
+/**
+ * A fluent wrapper around timespan to enforce obvious differences between specified Timeout and Interval values for waits
+ */
+class FWaitInterval
+{
+public:
+
+	static FWaitInterval InMilliseconds(double Value);
+	static FWaitInterval InSeconds(double Value);
+	static FWaitInterval InMinutes(double Value);
+	static FWaitInterval InHours(double Value);
+
+	explicit FWaitInterval(FTimespan InTimespan);
+
+	const FTimespan Timespan;
+};
+
+/**
+ * Represents a collection of fluent helper functions designed to make accessing and creating driver wait delegates easier
+ */
+class Until
+{
+public:
+
+	/**
+	 * Creates a new wait delegate which completes it's wait only if the specified element locator discovers elements or if the specified timeout timespan elapses
+	 */
+	static FDriverWaitDelegate ElementExists(const TSharedRef<IElementLocator, ESPMode::ThreadSafe>& ElementLocator, FWaitTimeout Timeout);
+
+	/**
+	 * Creates a new wait delegate which completes it's wait only if the specified element locator discovers elements or if the specified timeout timespan elapses;
+	 * The element locator is only re-evaluated at the specified wait interval
+	 */
+	static FDriverWaitDelegate ElementExists(const TSharedRef<IElementLocator, ESPMode::ThreadSafe>& ElementLocator, FWaitInterval Interval, FWaitTimeout Timeout);
+
+	/**
+	 * Creates a new wait delegate which completes it's wait only if the specified element locator discovers visible elements or if the specified timeout timespan elapses
+	 */
+	static FDriverWaitDelegate ElementIsVisible(const TSharedRef<IElementLocator, ESPMode::ThreadSafe>& ElementLocator, FWaitTimeout Timeout);
+
+	/**
+	 * Creates a new wait delegate which completes it's wait only if the specified element locator discovers visible elements or if the specified timeout timespan elapses;
+	 * The element locator is only re-evaluated at the specified wait interval
+	 */
+	static FDriverWaitDelegate ElementIsVisible(const TSharedRef<IElementLocator, ESPMode::ThreadSafe>& ElementLocator, FWaitInterval Interval, FWaitTimeout Timeout);
+
+	/**
+	 * Creates a new wait delegate which completes it's wait only if the specified element locator discovers interactable elements or if the specified timeout timespan elapses
+	 */
+	static FDriverWaitDelegate ElementIsInteractable(const TSharedRef<IElementLocator, ESPMode::ThreadSafe>& ElementLocator, FWaitTimeout Timeout);
+
+	/**
+	 * Creates a new wait delegate which completes it's wait only if the specified element locator discovers interactable elements or if the specified timeout timespan elapses;
+	 * The element locator is only re-evaluated at the specified wait interval
+	 */
+	static FDriverWaitDelegate ElementIsInteractable(const TSharedRef<IElementLocator, ESPMode::ThreadSafe>& ElementLocator, FWaitInterval Interval, FWaitTimeout Timeout);
+
+	/**
+	 * Creates a new wait delegate which completes it's wait only if the specified element locator discovers a scrollable element whose scroll position is at the beginning
+	 * or if the specified timeout timespan elapses
+	 */
+	static FDriverWaitDelegate ElementIsScrolledToBeginning(const TSharedRef<IElementLocator, ESPMode::ThreadSafe>& ElementLocator, FWaitTimeout Timeout);
+
+	/**
+	 * Creates a new wait delegate which completes it's wait only if the specified element locator discovers a scrollable element whose scroll position is at the beginning
+	 * or if the specified timeout timespan elapses
+	 * The element locator is only re-evaluated at the specified wait interval
+	 */
+	static FDriverWaitDelegate ElementIsScrolledToBeginning(const TSharedRef<IElementLocator, ESPMode::ThreadSafe>& ElementLocator, FWaitInterval Interval, FWaitTimeout Timeout);
+
+	/**
+	 * Creates a new wait delegate which completes it's wait only if the specified element locator discovers a scrollable element whose scroll position is at the end
+	 * or if the specified timeout timespan elapses
+	 */
+	static FDriverWaitDelegate ElementIsScrolledToEnd(const TSharedRef<IElementLocator, ESPMode::ThreadSafe>& ElementLocator, FWaitTimeout Timeout);
+	
+	/**
+	 * Creates a new wait delegate which completes it's wait only if the specified element locator discovers a scrollable element whose scroll position is at the end
+	 * or if the specified timeout timespan elapses
+	 * The element locator is only re-evaluated at the specified wait interval
+	 */
+	static FDriverWaitDelegate ElementIsScrolledToEnd(const TSharedRef<IElementLocator, ESPMode::ThreadSafe>& ElementLocator, FWaitInterval Interval, FWaitTimeout Timeout);
+
+	/**
+	 * Creates a new wait delegate which completes it's wait only if the specified condition returns true or if the specified timeout timespan elapses
+	 */
+	static FDriverWaitDelegate Condition(const TFunction<bool()>& Function, FWaitTimeout Timeout);
+
+	/**
+	 * Creates a new wait delegate which completes it's wait only if the specified condition returns true or if the specified timeout timespan elapses
+	 * The lambda is only re-evaluated at the specified wait interval
+	 */
+	static FDriverWaitDelegate Condition(const TFunction<bool()>& Function, FWaitInterval Interval, FWaitTimeout Timeout);
+
+	/**
+	 * Creates a new wait delegate which completes it's wait only if the specified condition returns true or if the specified timeout timespan elapses
+	 */
+	static FDriverWaitDelegate Condition(const FDriverWaitConditionDelegate& Delegate, FWaitTimeout Timeout);
+
+	/**
+	 * Creates a new wait delegate which completes it's wait only if the specified condition returns true or if the specified timeout timespan elapses
+	 * The delegate is only re-evaluated at the specified wait interval
+	 */
+	static FDriverWaitDelegate Condition(const FDriverWaitConditionDelegate& Delegate, FWaitInterval Interval, FWaitTimeout Timeout);
+
+	/**
+	 * Creates a new wait delegate which drives its state off the result of the specified lambda
+	 */
+	static FDriverWaitDelegate Lambda(const TFunction<FDriverWaitResponse(const FTimespan&)>& Value);
+};
diff --git a/Engine/Source/Developer/StandaloneRenderer/Private/Windows/D3D/SlateD3DShaders.cpp b/Engine/Source/Developer/StandaloneRenderer/Private/Windows/D3D/SlateD3DShaders.cpp
--- a/Engine/Source/Developer/StandaloneRenderer/Private/Windows/D3D/SlateD3DShaders.cpp
+++ b/Engine/Source/Developer/StandaloneRenderer/Private/Windows/D3D/SlateD3DShaders.cpp
@@ -5,6 +5,66 @@
 #include "Windows/D3D/SlateD3DRenderingPolicy.h"
 #include "Misc/Paths.h"
 
+// @return pointer to the D3DCompile function
+static pD3DCompile GetD3DCompileFunc()
+{
+	// Override default compiler path to newer dll
+	FString CompilerPath = FPaths::EngineDir();
+#if !PLATFORM_64BITS
+	CompilerPath.Append(TEXT("Binaries/ThirdParty/Windows/DirectX/x86/d3dcompiler_47.dll"));
+#else
+	CompilerPath.Append(TEXT("Binaries/ThirdParty/Windows/DirectX/x64/d3dcompiler_47.dll"));
+#endif
+	static bool bHasCompiler = false;
+	static HMODULE CompilerDLL = 0;
+
+	if (bHasCompiler == false)
+	{
+		CompilerDLL = LoadLibrary(*CompilerPath);
+	}
+
+	if (CompilerDLL)
+	{
+		return (pD3DCompile)(void*)GetProcAddress(CompilerDLL, "D3DCompile");
+	}
+
+	// D3D SDK we compiled with (usually D3DCompiler_43.dll from windows folder)
+	return &D3DCompile;
+}
+
+class StandaloneD3DIncluder: public ID3DInclude
+{
+	public:
+		STDMETHOD(Open)(D3D_INCLUDE_TYPE IncludeType, LPCSTR pFileName, LPCVOID pParentData, LPCVOID* ppData, uint32* pBytes) override
+		{
+			FString FileName(ANSI_TO_TCHAR(pFileName));
+			FString IncludePath = FPaths::EngineDir();
+			IncludePath.Append(TEXT("Shaders/StandaloneRenderer/D3D/"));
+			IncludePath.Append(FileName);
+
+			TArray<uint8>* ShaderFilePtr = new TArray<uint8>();
+			FFileHelper::LoadFileToArray(*ShaderFilePtr, *IncludePath);
+
+			*ppData = ShaderFilePtr->GetData();
+			*pBytes = ShaderFilePtr->Num();
+
+			IncludeMap.Add(*ppData, ShaderFilePtr);
+
+			return S_OK;
+		}
+
+		STDMETHOD(Close)(LPCVOID pData) override
+		{
+			TArray<uint8>* ShaderFilePtr = IncludeMap.FindChecked(pData);
+			IncludeMap.Remove(pData);
+			delete ShaderFilePtr;
+
+			return S_OK;
+		}
+
+		TMap<LPCVOID, TArray<uint8>*>   IncludeMap;
+};
+
 static void CompileShader( const FString& Filename, const FString& EntryPoint, const FString& ShaderModel, TRefCountPtr<ID3DBlob>& OutBlob )
 {
 	uint32 ShaderFlags = D3DCOMPILE_ENABLE_STRICTNESS;
@@ -14,8 +74,15 @@ static void CompileShader( const FString& Filename, const FString& EntryPoint, c
 	ShaderFlags |= D3DCOMPILE_OPTIMIZATION_LEVEL3;
 #endif
 
+	pD3DCompile D3DCompilerFunc = GetD3DCompileFunc();
+
+	StandaloneD3DIncluder Includer;
+
+	TArray<uint8> ShaderFile;
+	FFileHelper::LoadFileToArray(ShaderFile, *Filename);
+
 	TRefCountPtr<ID3DBlob> ErrorBlob;
-	HRESULT Hr = D3DX11CompileFromFile( *Filename, NULL, NULL, TCHAR_TO_ANSI(*EntryPoint), TCHAR_TO_ANSI(*ShaderModel), ShaderFlags, 0, NULL, OutBlob.GetInitReference(), ErrorBlob.GetInitReference(), NULL );
+	HRESULT Hr = D3DCompilerFunc(ShaderFile.GetData(), ShaderFile.Num(), NULL, NULL, &Includer, TCHAR_TO_ANSI(*EntryPoint), TCHAR_TO_ANSI(*ShaderModel), ShaderFlags, 0, OutBlob.GetInitReference(), ErrorBlob.GetInitReference());
 	if( FAILED(Hr) )
 	{
 		if( ErrorBlob.GetReference() )
diff --git a/Engine/Source/Programs/AutomationTool/AutomationUtils/MCPPublic.cs b/Engine/Source/Programs/AutomationTool/AutomationUtils/MCPPublic.cs
--- a/Engine/Source/Programs/AutomationTool/AutomationUtils/MCPPublic.cs
+++ b/Engine/Source/Programs/AutomationTool/AutomationUtils/MCPPublic.cs
@@ -10,12 +10,30 @@ using System.Net;
 using System.Reflection;
 using System.Text.RegularExpressions;
 using UnrealBuildTool;
+using EpicGames.MCP.Automation;
 
 namespace EpicGames.MCP.Automation
 {
 	using EpicGames.MCP.Config;
 	using System.Threading.Tasks;
 
+	public static class Extensions
+	{
+		public static Type[] SafeGetLoadedTypes(this Assembly Dll)
+		{
+			Type[] AllTypes;
+			try
+			{
+				AllTypes = Dll.GetTypes();
+			}
+			catch (ReflectionTypeLoadException e)
+			{
+				AllTypes = e.Types.Where(x => x != null).ToArray();
+			}
+			return AllTypes;
+		}
+	}
+
     /// <summary>
     /// Utility class to provide commit/rollback functionality via an RAII-like functionality.
     /// Usage is to provide a rollback action that will be called on Dispose if the Commit() method is not called.
@@ -628,7 +646,7 @@ namespace EpicGames.MCP.Automation
 				Assembly[] LoadedAssemblies = AppDomain.CurrentDomain.GetAssemblies();
 				foreach (var Dll in LoadedAssemblies)
 				{
-					Type[] AllTypes = Dll.GetTypes();
+					Type[] AllTypes = Dll.SafeGetLoadedTypes();
 					foreach (var PotentialConfigType in AllTypes)
 					{
 						if (PotentialConfigType != typeof(BuildPatchToolBase) && typeof(BuildPatchToolBase).IsAssignableFrom(PotentialConfigType))
@@ -651,7 +669,8 @@ namespace EpicGames.MCP.Automation
 		/// </summary>
 		/// <param name="Opts">Parameters which will be passed to the patch tool generation process.</param>
 		/// <param name="Version">Which version of BuildPatchTool is desired.</param>
-		public abstract void Execute(PatchGenerationOptions Opts, ToolVersion Version = ToolVersion.Live);
+		/// <param name="bAllowManifestClobbering">If set to true, will allow an existing manifest file to be overwritten with this execution. Default is false.</param>
+		public abstract void Execute(PatchGenerationOptions Opts, ToolVersion Version = ToolVersion.Live, bool bAllowManifestClobbering = false);
 
 		/// <summary>
 		/// Runs the Build Patch Tool executable to compactify a cloud directory using the supplied parameters.
@@ -698,7 +717,7 @@ namespace EpicGames.MCP.Automation
                 Assembly[] LoadedAssemblies = AppDomain.CurrentDomain.GetAssemblies();
                 foreach (var Dll in LoadedAssemblies)
                 {
-                    Type[] AllTypes = Dll.GetTypes();
+                    Type[] AllTypes = Dll.SafeGetLoadedTypes();
                     foreach (var PotentialConfigType in AllTypes)
                     {
                         if (PotentialConfigType != typeof(BuildInfoPublisherBase) && typeof(BuildInfoPublisherBase).IsAssignableFrom(PotentialConfigType))
@@ -808,7 +827,7 @@ namespace EpicGames.MCP.Automation
                 Assembly[] LoadedAssemblies = AppDomain.CurrentDomain.GetAssemblies();
                 foreach (var Dll in LoadedAssemblies)
                 {
-                    Type[] AllTypes = Dll.GetTypes();
+                    Type[] AllTypes = Dll.SafeGetLoadedTypes();
                     foreach (var PotentialConfigType in AllTypes)
                     {
                         if (PotentialConfigType != typeof(McpAccountServiceBase) && typeof(McpAccountServiceBase).IsAssignableFrom(PotentialConfigType))
@@ -908,7 +927,7 @@ namespace EpicGames.MCP.Automation
 						Assembly[] LoadedAssemblies = AppDomain.CurrentDomain.GetAssemblies();
 						foreach (var Dll in LoadedAssemblies)
 						{
-							Type[] AllTypes = Dll.GetTypes();
+							Type[] AllTypes = Dll.SafeGetLoadedTypes();
 							foreach (var PotentialConfigType in AllTypes)
 							{
 								if (PotentialConfigType != typeof(CloudStorageBase) && typeof(CloudStorageBase).IsAssignableFrom(PotentialConfigType))
@@ -1301,7 +1320,7 @@ namespace EpicGames.MCP.Config
                 Assembly[] LoadedAssemblies = AppDomain.CurrentDomain.GetAssemblies();
                 foreach (var Dll in LoadedAssemblies)
                 {
-                    Type[] AllTypes = Dll.GetTypes();
+                    Type[] AllTypes = Dll.SafeGetLoadedTypes();
                     foreach (var PotentialConfigType in AllTypes)
                     {
                         if (PotentialConfigType != typeof(McpConfigData) && typeof(McpConfigData).IsAssignableFrom(PotentialConfigType))
diff --git a/Engine/Source/Runtime/Core/Private/HTML5/HTML5Cursor.cpp b/Engine/Source/Runtime/Core/Private/HTML5/HTML5Cursor.cpp
--- a/Engine/Source/Runtime/Core/Private/HTML5/HTML5Cursor.cpp
+++ b/Engine/Source/Runtime/Core/Private/HTML5/HTML5Cursor.cpp
@@ -60,3 +60,8 @@ void FHTML5Cursor::Lock( const RECT* const Bounds )
 		LockStatus = true; 
 	}
 }
+
+void FHTML5Cursor::SetCustomShape(void* CursorHandle)
+{
+	
+}
\ No newline at end of file
diff --git a/Engine/Source/Runtime/Core/Private/Linux/LinuxCursor.cpp b/Engine/Source/Runtime/Core/Private/Linux/LinuxCursor.cpp
--- a/Engine/Source/Runtime/Core/Private/Linux/LinuxCursor.cpp
+++ b/Engine/Source/Runtime/Core/Private/Linux/LinuxCursor.cpp
@@ -332,3 +332,8 @@ bool FLinuxCursor::IsHidden()
 {
 	return bHidden;
 }
+
+void FLinuxCursor::SetCustomShape(void* CursorHandle)
+{
+	
+}
\ No newline at end of file
diff --git a/Engine/Source/Runtime/Core/Private/Mac/MacCursor.cpp b/Engine/Source/Runtime/Core/Private/Mac/MacCursor.cpp
--- a/Engine/Source/Runtime/Core/Private/Mac/MacCursor.cpp
+++ b/Engine/Source/Runtime/Core/Private/Mac/MacCursor.cpp
@@ -208,17 +208,6 @@ FMacCursor::~FMacCursor()
 	}
 }
 
-void FMacCursor::SetCustomShape(NSCursor* CursorHandle)
-{
-	SCOPED_AUTORELEASE_POOL;
-	[CursorHandle retain];
-	if (CursorHandles[EMouseCursor::Custom] != NULL)
-	{
-		[CursorHandles[EMouseCursor::Custom] release];
-	}
-	CursorHandles[EMouseCursor::Custom] = CursorHandle;
-}
-
 FVector2D FMacCursor::GetPosition() const
 {
 	FVector2D CurrentPos = CurrentPosition;
@@ -544,3 +533,16 @@ const FVector2D& FMacCursor::GetMouseScaling() const
 		return FVector2D::UnitVector;
 	}
 }
+
+void FMacCursor::SetCustomShape(void* InCursorHandle)
+{
+	NSCursor* CursorHandle = (NSCursor*)InCursorHandle;
+
+	SCOPED_AUTORELEASE_POOL;
+	[CursorHandle retain];
+	if (CursorHandles[EMouseCursor::Custom] != NULL)
+	{
+		[CursorHandles[EMouseCursor::Custom] release];
+	}
+	CursorHandles[EMouseCursor::Custom] = CursorHandle;
+}
\ No newline at end of file
diff --git a/Engine/Source/Runtime/Core/Private/Misc/AutomationTest.cpp b/Engine/Source/Runtime/Core/Private/Misc/AutomationTest.cpp
--- a/Engine/Source/Runtime/Core/Private/Misc/AutomationTest.cpp
+++ b/Engine/Source/Runtime/Core/Private/Misc/AutomationTest.cpp
@@ -728,6 +728,14 @@ void FAutomationTestBase::AddError(const FString& InError, int32 StackOffset)
 	}
 }
 
+void FAutomationTestBase::AddError(const FString& InError, const FString& InFilename, int32 InLineNumber)
+{
+	if (!bSuppressLogs)
+	{
+		ExecutionInfo.Errors.Add(FAutomationEvent(InError, ExecutionInfo.Context, InFilename, InLineNumber));
+	}
+}
+
 void FAutomationTestBase::AddWarning( const FString& InWarning )
 {
 	if( !bSuppressLogs )
@@ -791,8 +799,8 @@ void FAutomationTestBase::GenerateTestNames(TArray<FAutomationTestInfo>& TestInf
 			GetTestFlags(),
 			GetRequiredDeviceNum(),
 			ParameterNames[ParameterIndex],
-			GetTestSourceFileName(),
-			GetTestSourceFileLine(),
+			GetTestSourceFileName(CompleteTestName),
+			GetTestSourceFileLine(CompleteTestName),
 			GetTestAssetPath(ParameterNames[ParameterIndex]),
 			GetTestOpenCommand(ParameterNames[ParameterIndex])
 		);
diff --git a/Engine/Source/Runtime/Core/Private/Windows/WindowsCursor.cpp b/Engine/Source/Runtime/Core/Private/Windows/WindowsCursor.cpp
--- a/Engine/Source/Runtime/Core/Private/Windows/WindowsCursor.cpp
+++ b/Engine/Source/Runtime/Core/Private/Windows/WindowsCursor.cpp
@@ -143,11 +143,6 @@ FWindowsCursor::~FWindowsCursor()
 	}
 }
 
-void FWindowsCursor::SetCustomShape( HCURSOR CursorHandle )
-{
-	CursorHandles[EMouseCursor::Custom] = CursorHandle;
-}
-
 FVector2D FWindowsCursor::GetPosition() const
 {
 	POINT CursorPos;
@@ -193,5 +188,11 @@ void FWindowsCursor::Lock( const RECT* const Bounds )
 {
 	// Lock/Unlock the cursor
 	::ClipCursor(Bounds);
-		// If the cursor is not visible and we're running game, assume we're in a mode where the mouse is controlling the camera and lock it to the center of the widget.
+	// If the cursor is not visible and we're running game, assume we're in a mode where the mouse is controlling the camera and lock it to the center of the widget.
+}
+
+void FWindowsCursor::SetCustomShape(void* InCursorHandle)
+{
+	HCURSOR CursorHandle = (HCURSOR)InCursorHandle;
+	CursorHandles[EMouseCursor::Custom] = CursorHandle;
 }
diff --git a/Engine/Source/Runtime/Core/Private/Windows/WindowsPlatformProcess.cpp b/Engine/Source/Runtime/Core/Private/Windows/WindowsPlatformProcess.cpp
--- a/Engine/Source/Runtime/Core/Private/Windows/WindowsPlatformProcess.cpp
+++ b/Engine/Source/Runtime/Core/Private/Windows/WindowsPlatformProcess.cpp
@@ -75,8 +75,14 @@ void* FWindowsPlatformProcess::GetDllHandle( const TCHAR* FileName )
 		SearchPaths.Add(DllDirectories[Idx]);
 	}
 
-	::SetErrorMode(SEM_NOOPENFILEERRORBOX);
-	return LoadLibraryWithSearchPaths(FileName, SearchPaths);
+	// Load the DLL, avoiding windows dialog boxes if missing
+	int32 PrevErrorMode = ::SetErrorMode(SEM_NOOPENFILEERRORBOX);
+
+	void* Handle = LoadLibraryWithSearchPaths(FileName, SearchPaths);
+
+	::SetErrorMode(PrevErrorMode);
+
+	return Handle;
 }
 
 void FWindowsPlatformProcess::FreeDllHandle( void* DllHandle )
diff --git a/Engine/Source/Runtime/Core/Public/GenericPlatform/ICursor.h b/Engine/Source/Runtime/Core/Public/GenericPlatform/ICursor.h
--- a/Engine/Source/Runtime/Core/Public/GenericPlatform/ICursor.h
+++ b/Engine/Source/Runtime/Core/Public/GenericPlatform/ICursor.h
@@ -99,4 +99,6 @@ public:
 	 * @param Bounds	The bounds to lock the cursor to.  Pass NULL to unlock.
 	 */
 	virtual void Lock( const RECT* const Bounds ) = 0;
+
+	virtual void SetCustomShape(void* CursorHandle) = 0;
 };
diff --git a/Engine/Source/Runtime/Core/Public/HTML5/HTML5Cursor.h b/Engine/Source/Runtime/Core/Public/HTML5/HTML5Cursor.h
--- a/Engine/Source/Runtime/Core/Public/HTML5/HTML5Cursor.h
+++ b/Engine/Source/Runtime/Core/Public/HTML5/HTML5Cursor.h
@@ -30,6 +30,7 @@ public:
 
 	virtual void Lock( const RECT* const Bounds ) override;
 
+	virtual void SetCustomShape(void* CursorHandle) override;
 
 private:
 	EMouseCursor::Type CurrentType;
diff --git a/Engine/Source/Runtime/Core/Public/Linux/LinuxCursor.h b/Engine/Source/Runtime/Core/Public/Linux/LinuxCursor.h
--- a/Engine/Source/Runtime/Core/Public/Linux/LinuxCursor.h
+++ b/Engine/Source/Runtime/Core/Public/Linux/LinuxCursor.h
@@ -34,6 +34,8 @@ public:
 
 	virtual void Lock( const RECT* const Bounds ) override;
 
+	virtual void SetCustomShape(void* CursorHandle) override;
+
 public:
 
 	/**
diff --git a/Engine/Source/Runtime/Core/Public/Mac/MacCursor.h b/Engine/Source/Runtime/Core/Public/Mac/MacCursor.h
--- a/Engine/Source/Runtime/Core/Public/Mac/MacCursor.h
+++ b/Engine/Source/Runtime/Core/Public/Mac/MacCursor.h
@@ -29,16 +29,10 @@ public:
 
 	virtual void Lock(const RECT* const Bounds) override;
 
+	virtual void SetCustomShape(void* CursorHandle) override;
 
 public:
 
-	/**
-	 * Defines a custom cursor shape for the EMouseCursor::Custom type.
-	 *
-	 * @param CursorHandle	A native NSCursor instance to show when EMouseCursor::Custom is selected.
-	 */
-	virtual void SetCustomShape(NSCursor* CursorHandle);
-
 	bool UpdateCursorClipping(FVector2D& CursorPosition);
 
 	void WarpCursor(const int32 X, const int32 Y);
diff --git a/Engine/Source/Runtime/Core/Public/Misc/AutomationTest.h b/Engine/Source/Runtime/Core/Public/Misc/AutomationTest.h
--- a/Engine/Source/Runtime/Core/Public/Misc/AutomationTest.h
+++ b/Engine/Source/Runtime/Core/Public/Misc/AutomationTest.h
@@ -14,8 +14,13 @@
 #include "Delegates/Delegate.h"
 #include "Misc/Optional.h"
 #include "HAL/PlatformTime.h"
+#include "HAL/ThreadSafeBool.h" 
+#include "HAL/PlatformStackWalk.h"
+#include "GenericPlatformStackWalk.h" 
 #include "Containers/Queue.h"
 #include "Misc/FeedbackContext.h"
+#include "Future.h"
+#include "Async.h"
 
 /** Flags for specifying automation test requirements/behavior */
 namespace EAutomationTestFlags
@@ -402,6 +407,39 @@ protected:
 	friend class FAutomationTestFramework;
 };
 
+/**
+ * A simple latent command that runs the provided function on another thread
+ */
+class FThreadedAutomationLatentCommand : public IAutomationLatentCommand
+{
+public:
+
+	virtual ~FThreadedAutomationLatentCommand() {};
+
+	virtual bool Update() override
+	{
+		if (!Future.IsValid())
+		{
+			Future = Async(EAsyncExecution::Thread, Function);
+		}
+
+		return Future.IsReady();
+	}
+
+	FThreadedAutomationLatentCommand(TFunction<void()> InFunction)
+		: Function(InFunction)
+	{ }
+
+protected:
+
+	const TFunction<void()> Function;
+
+	TFuture<void> Future;
+
+	friend class FAutomationTestFramework;
+};
+
+
 /**
  * Simple abstract base class for networked, multi-participant tests
  */
@@ -879,6 +917,16 @@ public:
 	 */
 	virtual void AddError( const FString& InError, int32 StackOffset = 0 );
 
+
+	/**
+	 * Adds an error message to this test
+	 *
+	 * @param	InError	Error message to add to this test
+	 * @param	InFilename	The filename the error originated in
+	 * @param	InLineNumber	The line number in the file this error originated in
+	 */
+	virtual void AddError(const FString& InError, const FString& InFilename, int32 InLineNumber);
+
 	/**
 	 * Adds a warning to this test
 	 *
@@ -937,6 +985,11 @@ public:
 		return bComplexTask;
 	}
 
+	const bool IsRanOnSeparateThread() const
+	{
+		return bRunOnSeparateThread;
+	}
+
 	/**
 	 * Used to suppress / unsuppress logs.
 	 *
@@ -971,6 +1024,12 @@ public:
 	/** Gets the line number where this test was defined. */
 	virtual int32 GetTestSourceFileLine() const { return 0; }
 
+	/** Gets the filename where this test was defined. */
+	virtual FString GetTestSourceFileName(const FString& InTestName) const { return GetTestSourceFileName(); }
+
+	/** Gets the line number where this test was defined. */
+	virtual int32 GetTestSourceFileLine(const FString& InTestName) const { return GetTestSourceFileLine(); }
+
 	/** Allows navigation to the asset associated with the test if there is one. */
 	virtual FString GetTestAssetPath(const FString& Parameter) const { return TEXT(""); }
 
@@ -1004,7 +1063,8 @@ public:
 	 *
 	 * @see TestNotEqual
 	 */
-	template<typename ValueType> void TestEqual(const FString& Description, const ValueType& A, const ValueType& B)
+	template<typename ValueType> 
+	void TestEqual(const FString& Description, const ValueType& A, const ValueType& B)
 	{
 		if (A != B)
 		{
@@ -1012,6 +1072,30 @@ public:
 		}
 	}
 
+	void TestEqual(const FString& Description, const FString& A, const TCHAR* B)
+	{
+		if (A != B)
+		{
+			AddError(FString::Printf(TEXT("%s: The two values are not equal.\n  Actual: '%s'\nExpected: '%s'"), *Description, *A, B), 1);
+		}
+	}
+
+	void TestEqual(const FString& Description, const TCHAR* A, const TCHAR* B)
+	{
+		if (A != B)
+		{
+			AddError(FString::Printf(TEXT("%s: The two values are not equal.\n  Actual: '%s'\nExpected: '%s'"), *Description, A, B), 1);
+		}
+	}
+
+	void TestEqual(const FString& Description, const TCHAR* A, const FString& B)
+	{
+		if (A != B)
+		{
+			AddError(FString::Printf(TEXT("%s: The two values are not equal.\n  Actual: '%s'\nExpected: '%s'"), *Description, A, *B), 1);
+		}
+	}
+
 	/**
 	 * Logs an error if the specified Boolean value is not false.
 	 *
@@ -1188,6 +1272,9 @@ protected:
 	//Flag to indicate if this is a complex task
 	bool bComplexTask;
 
+	//Flag to indicate if this test should be ran on it's own thread
+	bool bRunOnSeparateThread;
+
 	/** Flag to suppress logs */
 	bool bSuppressLogs;
 
@@ -1221,6 +1308,9 @@ public:
 
 	virtual void GetTests(TArray<FString>& OutBeautifiedNames, TArray <FString>& OutTestCommands) const override
 	{
+		const_cast<FBDDAutomationTestBase*>(this)->BeautifiedNames.Empty();
+		const_cast<FBDDAutomationTestBase*>(this)->TestCommands.Empty();
+
 		bIsDiscoveryMode = true;
 		const_cast<FBDDAutomationTestBase*>(this)->RunTest(FString());
 		bIsDiscoveryMode = false;
@@ -1230,6 +1320,11 @@ public:
 		bBaseRunTestRan = false;
 	}
 
+	bool IsDiscoveryMode() const
+	{
+		return bIsDiscoveryMode;
+	}
+
 	void xDescribe(const FString& InDescription, TFunction<void()> DoWork)
 	{
 		check(bBaseRunTestRan);
@@ -1376,6 +1471,669 @@ private:
 	mutable bool bBaseRunTestRan;
 };
 
+DECLARE_DELEGATE(FDoneDelegate);
+
+class CORE_API FAutomationSpecBase : public FAutomationTestBase
+{
+private:
+
+	class FSingleExecuteLatentCommand : public IAutomationLatentCommand
+	{
+	public:
+		FSingleExecuteLatentCommand(TFunction<void()> InPredicate)
+			: Predicate(MoveTemp(InPredicate))
+		{ }
+
+		virtual ~FSingleExecuteLatentCommand()
+		{ }
+
+		virtual bool Update() override
+		{
+			Predicate();
+			return true;
+		}
+
+	private:
+
+		const TFunction<void()> Predicate;
+	};
+
+	class FUntilDoneLatentCommand : public IAutomationLatentCommand
+	{
+	public:
+
+		FUntilDoneLatentCommand(TFunction<void(const FDoneDelegate&)> InPredicate)
+			: Predicate(MoveTemp(InPredicate))
+			, bIsRunning(false)
+			, bDone(false)
+		{ }
+
+		virtual ~FUntilDoneLatentCommand()
+		{ }
+
+		virtual bool Update() override
+		{
+			if (!bIsRunning)
+			{
+				Predicate(FDoneDelegate::CreateSP(this, &FUntilDoneLatentCommand::Done));
+				bIsRunning = true;
+			}
+
+			if (bDone)
+			{
+				// Reset the done for the next potential run of this command
+				bDone = false;
+				bIsRunning = false;
+				return true;
+			}
+
+			return false;
+		}
+
+	private:
+
+		void Done()
+		{
+			bDone = true;
+		}
+
+	private:
+
+		const TFunction<void(const FDoneDelegate&)> Predicate;
+
+		bool bIsRunning;
+		FThreadSafeBool bDone;
+	};
+
+	class FAsyncUntilDoneLatentCommand : public IAutomationLatentCommand
+	{
+	public:
+
+		FAsyncUntilDoneLatentCommand(EAsyncExecution InExecution, TFunction<void(const FDoneDelegate&)> InPredicate)
+			: Execution(InExecution)
+			, Predicate(MoveTemp(InPredicate))
+			, bDone(false)
+		{ }
+
+		virtual ~FAsyncUntilDoneLatentCommand()
+		{ }
+
+		virtual bool Update() override
+		{
+			if (!Future.IsValid())
+			{
+				Future = Async<void>(Execution, [this]() {
+					Predicate(FDoneDelegate::CreateRaw(this, &FAsyncUntilDoneLatentCommand::Done));
+				});
+			}
+
+			if (bDone)
+			{
+				// Reset the done for the next potential run of this command
+				bDone = false;
+				Future = TFuture<void>();
+				return true;
+			}
+
+			return false;
+		}
+
+	private:
+
+		void Done()
+		{
+			bDone = true;
+		}
+
+	private:
+
+		const EAsyncExecution Execution;
+		const TFunction<void(const FDoneDelegate&)> Predicate;
+
+		TFuture<void> Future;
+		FThreadSafeBool bDone;
+	};
+
+	class FAsyncLatentCommand : public IAutomationLatentCommand
+	{
+	public:
+
+		FAsyncLatentCommand(EAsyncExecution InExecution, TFunction<void()> InPredicate)
+			: Execution(InExecution)
+			, Predicate(MoveTemp(InPredicate))
+			, bDone(false)
+		{ }
+
+		virtual ~FAsyncLatentCommand()
+		{ }
+
+		virtual bool Update() override
+		{
+			if (!Future.IsValid())
+			{
+				Future = Async<void>(Execution, [this]() {
+					Predicate();
+					bDone = true;
+				});
+			}
+
+			if (bDone)
+			{
+				// Reset the done for the next potential run of this command
+				bDone = false;
+				Future = TFuture<void>();
+				return true;
+			}
+
+			return false;
+		}
+
+	private:
+
+		void Done()
+		{
+			bDone = true;
+		}
+
+	private:
+
+		const EAsyncExecution Execution;
+		const TFunction<void()> Predicate;
+
+		TFuture<void> Future;
+		FThreadSafeBool bDone;
+	};
+
+	struct FSpecIt
+	{
+	public:
+
+		FString Description;
+		FString Id;
+		FString Filename;
+		int32 LineNumber;
+		TSharedRef<IAutomationLatentCommand> Command;
+
+		FSpecIt(FString InDescription, FString InId, FString InFilename, int32 InLineNumber, TSharedRef<IAutomationLatentCommand> InCommand)
+			: Description(MoveTemp(InDescription))
+			, Id(MoveTemp(InId))
+			, Filename(InFilename)
+			, LineNumber(MoveTemp(InLineNumber))
+			, Command(MoveTemp(InCommand))
+		{ }
+	};
+
+	struct FSpecDefinitionScope
+	{
+	public:
+
+		FString Description;
+
+		TArray<TSharedRef<IAutomationLatentCommand>> BeforeEach;
+		TArray<TSharedRef<FSpecIt>> It;
+		TArray<TSharedRef<IAutomationLatentCommand>> AfterEach;
+
+		TArray<TSharedRef<FSpecDefinitionScope>> Children;
+	};
+
+	struct FSpec
+	{
+	public:
+
+		FString Id;
+		FString Description;
+		FString Filename;
+		int32 LineNumber;
+		TArray<TSharedRef<IAutomationLatentCommand>> Commands;
+	};
+
+public:
+
+	FAutomationSpecBase(const FString& InName, const bool bInComplexTask)
+		: FAutomationTestBase(InName, bInComplexTask)
+		, RootDefinitionScope(MakeShareable(new FSpecDefinitionScope()))
+	{
+		DefinitionScopeStack.Push(RootDefinitionScope.ToSharedRef());
+	}
+
+	virtual bool RunTest(const FString& InParameters) override
+	{
+		EnsureDefinitions();
+
+		if (!InParameters.IsEmpty())
+		{
+			const TSharedRef<FSpec>* SpecToRun = IdToSpecMap.Find(InParameters);
+			if (SpecToRun != nullptr)
+			{
+				for (int32 Index = 0; Index < (*SpecToRun)->Commands.Num(); ++Index)
+				{
+					FAutomationTestFramework::GetInstance().EnqueueLatentCommand((*SpecToRun)->Commands[Index]);
+				}
+			}
+		}
+		else
+		{
+			TArray<TSharedRef<FSpec>> Specs;
+			IdToSpecMap.GenerateValueArray(Specs);
+
+			for (int32 SpecIndex = 0; SpecIndex < Specs.Num(); SpecIndex++)
+			{
+				for (int32 CommandIndex = 0; CommandIndex < Specs[SpecIndex]->Commands.Num(); ++CommandIndex)
+				{
+					FAutomationTestFramework::GetInstance().EnqueueLatentCommand(Specs[SpecIndex]->Commands[CommandIndex]);
+				}
+			}
+		}
+
+		return true;
+	}
+
+	virtual bool IsStressTest() const 
+	{
+		return false; 
+	}
+
+	virtual uint32 GetRequiredDeviceNum() const override 
+	{ 
+		return 1; 
+	}
+
+	virtual FString GetTestSourceFileName(const FString& InTestName) const override
+	{
+		FString TestId = InTestName;
+		if (TestId.StartsWith(TestName + TEXT(" ")))
+		{
+			TestId = InTestName.RightChop(TestName.Len() + 1);
+		}
+
+		const TSharedRef<FSpec>* Spec = IdToSpecMap.Find(TestId);
+		if (Spec != nullptr)
+		{
+			return (*Spec)->Filename;
+		}
+
+		return FAutomationTestBase::GetTestSourceFileName();
+	}
+
+	virtual int32 GetTestSourceFileLine(const FString& InTestName) const override
+	{ 
+		FString TestId = InTestName;
+		if (TestId.StartsWith(TestName + TEXT(" ")))
+		{
+			TestId = InTestName.RightChop(TestName.Len() + 1);
+		}
+
+		const TSharedRef<FSpec>* Spec = IdToSpecMap.Find(TestId);
+		if (Spec != nullptr)
+		{
+			return (*Spec)->LineNumber;
+		}
+
+		return FAutomationTestBase::GetTestSourceFileLine();
+	}
+
+	virtual void GetTests(TArray<FString>& OutBeautifiedNames, TArray<FString>& OutTestCommands) const override
+	{
+		EnsureDefinitions();
+
+		TArray<TSharedRef<FSpec>> Specs;
+		IdToSpecMap.GenerateValueArray(Specs);
+
+		for (int32 Index = 0; Index < Specs.Num(); Index++)
+		{
+			OutTestCommands.Push(Specs[Index]->Id);
+			OutBeautifiedNames.Push(Specs[Index]->Description);
+		}
+	}
+
+	void xDescribe(const FString& InDescription, TFunction<void()> DoWork)
+	{
+		//disabled
+	}
+
+	void Describe(const FString& InDescription, TFunction<void()> DoWork)
+	{
+		const TSharedRef<FSpecDefinitionScope> ParentScope = DefinitionScopeStack.Last();
+		const TSharedRef<FSpecDefinitionScope> NewScope = MakeShareable(new FSpecDefinitionScope());
+		NewScope->Description = InDescription;
+		ParentScope->Children.Push(NewScope);
+
+		DefinitionScopeStack.Push(NewScope);
+		PushDescription(InDescription);
+		DoWork();
+		PopDescription(InDescription);
+		DefinitionScopeStack.Pop();
+
+		if (NewScope->It.Num() == 0 && NewScope->Children.Num() == 0)
+		{
+			ParentScope->Children.Remove(NewScope);
+		}
+	}
+
+	void xIt(const FString& InDescription, TFunction<void()> DoWork)
+	{
+		//disabled
+	}
+
+	void xIt(const FString& InDescription, EAsyncExecution Execution, TFunction<void()> DoWork)
+	{
+		//disabled
+	}
+
+	void xLatentIt(const FString& InDescription, TFunction<void(const FDoneDelegate&)> DoWork)
+	{
+		//disabled
+	}
+
+	void xLatentIt(const FString& InDescription, EAsyncExecution Execution, TFunction<void(const FDoneDelegate&)> DoWork)
+	{
+		//disabled
+	}
+
+	void It(const FString& InDescription, TFunction<void()> DoWork)
+	{
+		const TSharedRef<FSpecDefinitionScope> CurrentScope = DefinitionScopeStack.Last();
+		const TArray<FProgramCounterSymbolInfo> Stack = FPlatformStackWalk::GetStack(1, 1);
+
+		PushDescription(InDescription);
+		CurrentScope->It.Push(MakeShareable(new FSpecIt(GetDescription(), GetId(), Stack[0].Filename, Stack[0].LineNumber, MakeShareable(new FSingleExecuteLatentCommand(DoWork)))));
+		PopDescription(InDescription);
+	}
+
+	void It(const FString& InDescription, EAsyncExecution Execution, TFunction<void()> DoWork)
+	{
+		const TSharedRef<FSpecDefinitionScope> CurrentScope = DefinitionScopeStack.Last();
+		const TArray<FProgramCounterSymbolInfo> Stack = FPlatformStackWalk::GetStack(1, 1);
+
+		PushDescription(InDescription);
+		CurrentScope->It.Push(MakeShareable(new FSpecIt(GetDescription(), GetId(), Stack[0].Filename, Stack[0].LineNumber, MakeShareable(new FAsyncLatentCommand(Execution, DoWork)))));
+		PopDescription(InDescription);
+	}
+
+	void LatentIt(const FString& InDescription, TFunction<void(const FDoneDelegate&)> DoWork)
+	{
+		const TSharedRef<FSpecDefinitionScope> CurrentScope = DefinitionScopeStack.Last();
+		const TArray<FProgramCounterSymbolInfo> Stack = FPlatformStackWalk::GetStack(1, 1);
+
+		PushDescription(InDescription);
+		CurrentScope->It.Push(MakeShareable(new FSpecIt(GetDescription(), GetId(), Stack[0].Filename, Stack[0].LineNumber, MakeShareable(new FUntilDoneLatentCommand(DoWork)))));
+		PopDescription(InDescription);
+	}
+
+	void LatentIt(const FString& InDescription, EAsyncExecution Execution, TFunction<void(const FDoneDelegate&)> DoWork)
+	{
+		const TSharedRef<FSpecDefinitionScope> CurrentScope = DefinitionScopeStack.Last();
+		const TArray<FProgramCounterSymbolInfo> Stack = FPlatformStackWalk::GetStack(1, 1);
+
+		PushDescription(InDescription);
+		CurrentScope->It.Push(MakeShareable(new FSpecIt(GetDescription(), GetId(), Stack[0].Filename, Stack[0].LineNumber, MakeShareable(new FAsyncUntilDoneLatentCommand(Execution, DoWork)))));
+		PopDescription(InDescription);
+	}
+
+	void xBeforeEach(TFunction<void()> DoWork)
+	{
+		// disabled
+	}
+
+	void xBeforeEach(EAsyncExecution Execution, TFunction<void()> DoWork)
+	{
+		// disabled
+	}
+
+	void xLatentBeforeEach(TFunction<void(const FDoneDelegate&)> DoWork)
+	{
+		// disabled
+	}
+
+	void xLatentBeforeEach(EAsyncExecution Execution, TFunction<void(const FDoneDelegate&)> DoWork)
+	{
+		// disabled
+	}
+
+	void BeforeEach(TFunction<void()> DoWork)
+	{
+		const TSharedRef<FSpecDefinitionScope> CurrentScope = DefinitionScopeStack.Last();
+		CurrentScope->BeforeEach.Push(MakeShareable(new FSingleExecuteLatentCommand(DoWork)));
+	}
+
+	void BeforeEach(EAsyncExecution Execution, TFunction<void()> DoWork)
+	{
+		const TSharedRef<FSpecDefinitionScope> CurrentScope = DefinitionScopeStack.Last();
+		CurrentScope->BeforeEach.Push(MakeShareable(new FAsyncLatentCommand(Execution, DoWork)));
+	}
+
+	void LatentBeforeEach(TFunction<void(const FDoneDelegate&)> DoWork)
+	{
+		const TSharedRef<FSpecDefinitionScope> CurrentScope = DefinitionScopeStack.Last();
+		CurrentScope->BeforeEach.Push(MakeShareable(new FUntilDoneLatentCommand(DoWork)));
+	}
+	     
+	void LatentBeforeEach(EAsyncExecution Execution, TFunction<void(const FDoneDelegate&)> DoWork)
+	{
+		const TSharedRef<FSpecDefinitionScope> CurrentScope = DefinitionScopeStack.Last();
+		CurrentScope->BeforeEach.Push(MakeShareable(new FAsyncUntilDoneLatentCommand(Execution, DoWork)));
+	}
+
+	void xAfterEach(TFunction<void()> DoWork)
+	{
+		// disabled
+	}
+
+	void xAfterEach(EAsyncExecution Execution, TFunction<void()> DoWork)
+	{
+		// disabled
+	}
+
+	void xLatentAfterEach(TFunction<void(const FDoneDelegate&)> DoWork)
+	{
+		// disabled
+	}
+
+	void xLatentAfterEach(EAsyncExecution Execution, TFunction<void(const FDoneDelegate&)> DoWork)
+	{
+		// disabled
+	}
+
+	void AfterEach(TFunction<void()> DoWork)
+	{
+		const TSharedRef<FSpecDefinitionScope> CurrentScope = DefinitionScopeStack.Last();
+		CurrentScope->AfterEach.Push(MakeShareable(new FSingleExecuteLatentCommand(DoWork)));
+	}
+
+	void AfterEach(EAsyncExecution Execution, TFunction<void()> DoWork)
+	{
+		const TSharedRef<FSpecDefinitionScope> CurrentScope = DefinitionScopeStack.Last();
+		CurrentScope->AfterEach.Push(MakeShareable(new FAsyncLatentCommand(Execution, DoWork)));
+	}
+
+	void LatentAfterEach(TFunction<void(const FDoneDelegate&)> DoWork)
+	{
+		const TSharedRef<FSpecDefinitionScope> CurrentScope = DefinitionScopeStack.Last();
+		CurrentScope->AfterEach.Push(MakeShareable(new FUntilDoneLatentCommand(DoWork)));
+	}
+
+	void LatentAfterEach(EAsyncExecution Execution, TFunction<void(const FDoneDelegate&)> DoWork)
+	{
+		const TSharedRef<FSpecDefinitionScope> CurrentScope = DefinitionScopeStack.Last();
+		CurrentScope->AfterEach.Push(MakeShareable(new FAsyncUntilDoneLatentCommand(Execution, DoWork)));
+	}
+
+protected:
+
+	void EnsureDefinitions() const
+	{
+		if (!bHasBeenDefined)
+		{
+			const_cast<FAutomationSpecBase*>(this)->Define();
+			const_cast<FAutomationSpecBase*>(this)->PostDefine();
+		}
+	}
+
+	virtual void Define() = 0;
+
+	void PostDefine()
+	{
+		TArray<TSharedRef<FSpecDefinitionScope>> Stack;
+		Stack.Push(RootDefinitionScope.ToSharedRef());
+
+		TArray<TSharedRef<IAutomationLatentCommand>> BeforeEach;
+		TArray<TSharedRef<IAutomationLatentCommand>> AfterEach;
+
+		while (Stack.Num() > 0)
+		{
+			const TSharedRef<FSpecDefinitionScope> Scope = Stack.Last();
+
+			BeforeEach.Append(Scope->BeforeEach);
+			AfterEach.Append(Scope->AfterEach); // iterate in reverse
+
+			for (int32 ItIndex = 0; ItIndex < Scope->It.Num(); ItIndex++)
+			{
+				TSharedRef<FSpecIt> It = Scope->It[ItIndex];
+
+				TSharedRef<FSpec> Spec = MakeShareable(new FSpec());
+				Spec->Id = It->Id;
+				Spec->Description = It->Description;
+				Spec->Filename = It->Filename;
+				Spec->LineNumber = It->LineNumber;
+				Spec->Commands.Append(BeforeEach);
+				Spec->Commands.Add(It->Command);
+
+				for (int32 AfterEachIndex = AfterEach.Num() - 1; AfterEachIndex >= 0; AfterEachIndex--)
+				{
+					Spec->Commands.Add(AfterEach[AfterEachIndex]);
+				}
+
+				check(!IdToSpecMap.Contains(Spec->Id));
+				IdToSpecMap.Add(Spec->Id, Spec);
+			}
+			Scope->It.Empty();
+
+			if (Scope->Children.Num() > 0)
+			{
+				Stack.Append(Scope->Children);
+				Scope->Children.Empty();
+			}
+			else
+			{
+				while (Stack.Num() > 0 && Stack.Last()->Children.Num() == 0 && Stack.Last()->It.Num() == 0)
+				{
+					const TSharedRef<FSpecDefinitionScope> PoppedScope = Stack.Pop();
+
+					if (PoppedScope->BeforeEach.Num() > 0)
+					{
+						BeforeEach.RemoveAt(BeforeEach.Num() - PoppedScope->BeforeEach.Num(), PoppedScope->BeforeEach.Num());
+					}
+
+					if (PoppedScope->AfterEach.Num() > 0)
+					{
+						AfterEach.RemoveAt(AfterEach.Num() - PoppedScope->AfterEach.Num(), PoppedScope->AfterEach.Num());
+					}
+				}
+			}
+		}
+
+		RootDefinitionScope.Reset();
+		DefinitionScopeStack.Reset();
+		bHasBeenDefined = true;
+	}
+
+	void Redefine()
+	{
+		Description.Empty();
+		IdToSpecMap.Empty();
+		RootDefinitionScope.Reset();
+		DefinitionScopeStack.Empty();
+		bHasBeenDefined = false;
+	}
+
+private:
+
+	void PushDescription(const FString& InDescription)
+	{
+		Description.Add(InDescription);
+	}
+
+	void PopDescription(const FString& InDescription)
+	{
+		Description.RemoveAt(Description.Num() - 1);
+	}
+
+	FString GetDescription() const
+	{
+		FString CompleteDescription;
+		for (int32 Index = 0; Index < Description.Num(); ++Index)
+		{
+			if (Description[Index].IsEmpty())
+			{
+				continue;
+			}
+
+			if (CompleteDescription.IsEmpty())
+			{
+				CompleteDescription = Description[Index];
+			}
+			else if (FChar::IsWhitespace(CompleteDescription[CompleteDescription.Len() - 1]) || FChar::IsWhitespace(Description[Index][0]))
+			{
+				CompleteDescription = CompleteDescription + TEXT(".") + Description[Index];
+			}
+			else
+			{
+				CompleteDescription = FString::Printf(TEXT("%s.%s"), *CompleteDescription, *Description[Index]);
+			}
+		}
+
+		return CompleteDescription;
+	}
+
+	FString GetId() const
+	{
+		if (Description.Last().EndsWith(TEXT("]")))
+		{
+			FString ItDescription = Description.Last();
+			ItDescription.RemoveAt(ItDescription.Len() - 1);
+
+			int32 StartingBraceIndex = INDEX_NONE;
+			if (ItDescription.FindLastChar(TEXT('['), StartingBraceIndex) && StartingBraceIndex != ItDescription.Len() - 1)
+			{
+				FString CommandId = ItDescription.RightChop(StartingBraceIndex + 1);
+				return CommandId;
+			}
+		}
+
+		FString CompleteId;
+		for (int32 Index = 0; Index < Description.Num(); ++Index)
+		{
+			if (Description[Index].IsEmpty())
+			{
+				continue;
+			}
+
+			if (CompleteId.IsEmpty())
+			{
+				CompleteId = Description[Index];
+			}
+			else if (FChar::IsWhitespace(CompleteId[CompleteId.Len() - 1]) || FChar::IsWhitespace(Description[Index][0]))
+			{
+				CompleteId = CompleteId + Description[Index];
+			}
+			else
+			{
+				CompleteId = FString::Printf(TEXT("%s %s"), *CompleteId, *Description[Index]);
+			}
+		}
+
+		return CompleteId;
+	}
+
+private:
+
+	TArray<FString> Description;
+	TMap<FString, TSharedRef<FSpec>> IdToSpecMap;
+	TSharedPtr<FSpecDefinitionScope> RootDefinitionScope;
+	TArray<TSharedRef<FSpecDefinitionScope>> DefinitionScopeStack;
+	bool bHasBeenDefined;
+};
+
 
 //////////////////////////////////////////////////////////////////////////
 // Latent command definition macros
@@ -1564,13 +2322,59 @@ public: \
 		virtual uint32 GetTestFlags() const override { return TFlags; } \
 		virtual bool IsStressTest() const { return false; } \
 		virtual uint32 GetRequiredDeviceNum() const override { return 1; } \
-		virtual FString GetFileName() const { return FileName; } \
-		virtual int32 GetFileLine() const { return LineNumber; } \
+		virtual FString GetTestSourceFileName() const override { return FileName; } \
+		virtual int32 GetTestSourceFileLine() const override { return LineNumber; } \
 	protected: \
 		virtual bool RunTest(const FString& Parameters) override; \
 		virtual FString GetBeautifiedTestName() const override { return PrettyName; } \
+	private: \
+		void Define(); \
 	};
 
+#define DEFINE_SPEC_PRIVATE( TClass, PrettyName, TFlags, FileName, LineNumber ) \
+	class TClass : public FAutomationSpecBase \
+	{ \
+	public: \
+		TClass( const FString& InName ) \
+		: FAutomationSpecBase( InName, false ) {\
+			static_assert((TFlags)&EAutomationTestFlags::ApplicationContextMask, "AutomationTest has no application flag.  It shouldn't run.  See AutomationTest.h."); \
+			static_assert(	(((TFlags)&EAutomationTestFlags::FilterMask) == EAutomationTestFlags::SmokeFilter) || \
+							(((TFlags)&EAutomationTestFlags::FilterMask) == EAutomationTestFlags::EngineFilter) || \
+							(((TFlags)&EAutomationTestFlags::FilterMask) == EAutomationTestFlags::ProductFilter) || \
+							(((TFlags)&EAutomationTestFlags::FilterMask) == EAutomationTestFlags::PerfFilter) || \
+							(((TFlags)&EAutomationTestFlags::FilterMask) == EAutomationTestFlags::StressFilter), \
+							"All AutomationTests must have exactly 1 filter type specified.  See AutomationTest.h."); \
+		} \
+		virtual uint32 GetTestFlags() const override { return TFlags; } \
+		virtual FString GetTestSourceFileName() const override { return FileName; } \
+		virtual int32 GetTestSourceFileLine() const override { return LineNumber; } \
+	protected: \
+		virtual FString GetBeautifiedTestName() const override { return PrettyName; } \
+		virtual void Define() override; \
+	};
+
+#define BEGIN_DEFINE_SPEC_PRIVATE( TClass, PrettyName, TFlags, FileName, LineNumber ) \
+	class TClass : public FAutomationSpecBase \
+	{ \
+	public: \
+		TClass( const FString& InName ) \
+		: FAutomationSpecBase( InName, false ) {\
+			static_assert((TFlags)&EAutomationTestFlags::ApplicationContextMask, "AutomationTest has no application flag.  It shouldn't run.  See AutomationTest.h."); \
+			static_assert(	(((TFlags)&EAutomationTestFlags::FilterMask) == EAutomationTestFlags::SmokeFilter) || \
+							(((TFlags)&EAutomationTestFlags::FilterMask) == EAutomationTestFlags::EngineFilter) || \
+							(((TFlags)&EAutomationTestFlags::FilterMask) == EAutomationTestFlags::ProductFilter) || \
+							(((TFlags)&EAutomationTestFlags::FilterMask) == EAutomationTestFlags::PerfFilter) || \
+							(((TFlags)&EAutomationTestFlags::FilterMask) == EAutomationTestFlags::StressFilter), \
+							"All AutomationTests must have exactly 1 filter type specified.  See AutomationTest.h."); \
+		} \
+		virtual uint32 GetTestFlags() const override { return TFlags; } \
+		using FAutomationSpecBase::GetTestSourceFileName; \
+		virtual FString GetTestSourceFileName() const override { return FileName; } \
+		using FAutomationSpecBase::GetTestSourceFileLine; \
+		virtual int32 GetTestSourceFileLine() const override { return LineNumber; } \
+	protected: \
+		virtual FString GetBeautifiedTestName() const override { return PrettyName; } \
+		virtual void Define() override;
 
 #if WITH_AUTOMATION_WORKER
 	#define IMPLEMENT_SIMPLE_AUTOMATION_TEST( TClass, PrettyName, TFlags ) \
@@ -1613,6 +2417,23 @@ public: \
 			TClass TClass##AutomationTestInstance( TEXT(#TClass) );\
 		}
 
+	#define DEFINE_SPEC( TClass, PrettyName, TFlags ) \
+		DEFINE_SPEC_PRIVATE(TClass, PrettyName, TFlags, __FILE__, __LINE__) \
+		namespace\
+		{\
+			TClass TClass##AutomationSpecInstance( TEXT(#TClass) );\
+		}
+
+	#define BEGIN_DEFINE_SPEC( TClass, PrettyName, TFlags ) \
+		BEGIN_DEFINE_SPEC_PRIVATE(TClass, PrettyName, TFlags, __FILE__, __LINE__) 
+
+	#define END_DEFINE_SPEC( TClass ) \
+		};\
+		namespace\
+		{\
+			TClass TClass##AutomationSpecInstance( TEXT(#TClass) );\
+		}
+
 	//#define BEGIN_CUSTOM_COMPLEX_AUTOMATION_TEST( TClass, TBaseClass, PrettyName, TFlags ) \
 	//	BEGIN_COMPLEX_AUTOMATION_TEST_PRIVATE(TClass, TBaseClass, PrettyName, TFlags, __FILE__, __LINE__)
 	//
@@ -1635,8 +2456,14 @@ public: \
 		IMPLEMENT_SIMPLE_AUTOMATION_TEST_PRIVATE(TClass, TBaseClass, PrettyName, TFlags, __FILE__, __LINE__)
 	#define IMPLEMENT_CUSTOM_COMPLEX_AUTOMATION_TEST( TClass, TBaseClass, PrettyName, TFlags ) \
 		IMPLEMENT_COMPLEX_AUTOMATION_TEST_PRIVATE(TClass, TBaseClass, PrettyName, TFlags, __FILE__, __LINE__)
-	#define IMPLEMENT_BDD_AUTOMATION_TEST(TClass, PrettyName, TFlags, NumParticipants) \
-		IMPLEMENT_BDD_AUTOMATION_TEST_PRIVATE(TClass, PrettyName, TFlags, NumParticipants, __FILE__, __LINE__)
+	#define IMPLEMENT_BDD_AUTOMATION_TEST(TClass, PrettyName, TFlags) \
+		IMPLEMENT_BDD_AUTOMATION_TEST_PRIVATE(TClass, PrettyName, TFlags, __FILE__, __LINE__)
+	#define DEFINE_SPEC(TClass, PrettyName, TFlags) \
+		DEFINE_SPEC_PRIVATE(TClass, PrettyName, TFlags, __FILE__, __LINE__)
+	#define BEGIN_DEFINE_SPEC(TClass, PrettyName, TFlags) \
+		BEGIN_DEFINE_SPEC_PRIVATE(TClass, PrettyName, TFlags, __FILE__, __LINE__)
+	#define END_DEFINE_SPEC(TClass) \
+		}; \
 
 	//#define BEGIN_CUSTOM_COMPLEX_AUTOMATION_TEST( TClass, TBaseClass, PrettyName, TFlags ) \
 	//	BEGIN_CUSTOM_COMPLEX_AUTOMATION_TEST_PRIVATE(TClass, TBaseClass, PrettyName, TFlags, __FILE__, __LINE__)
@@ -1730,4 +2557,4 @@ private:
 	TFunction<bool()> Callback;
 	TFunction<bool()> TimeoutCallback;
 	float Timeout;
-};
+};
\ No newline at end of file
diff --git a/Engine/Source/Runtime/Core/Public/Windows/WindowsCursor.h b/Engine/Source/Runtime/Core/Public/Windows/WindowsCursor.h
--- a/Engine/Source/Runtime/Core/Public/Windows/WindowsCursor.h
+++ b/Engine/Source/Runtime/Core/Public/Windows/WindowsCursor.h
@@ -31,14 +31,7 @@ public:
 
 	virtual void Lock( const RECT* const Bounds ) override;
 
-public:
-
-	/**
-	 * Defines a custom cursor shape for the EMouseCursor::Custom type.
-	 * 
-	 * @param CursorHandle	A native cursor handle to show when EMouseCursor::Custom is selected.
-	 */
-	virtual void SetCustomShape( Windows::HCURSOR CursorHandle );
+	virtual void SetCustomShape(void* CursorHandle) override;
 
 private:
 
diff --git a/Engine/Source/Runtime/InputCore/Private/InputCoreTypes.cpp b/Engine/Source/Runtime/InputCore/Private/InputCoreTypes.cpp
--- a/Engine/Source/Runtime/InputCore/Private/InputCoreTypes.cpp
+++ b/Engine/Source/Runtime/InputCore/Private/InputCoreTypes.cpp
@@ -1133,8 +1133,8 @@ FKey FInputKeyManager::GetKeyFromCodes( const uint32 KeyCode, const uint32 CharC
 
 void FInputKeyManager::GetCodesFromKey(const FKey Key, const uint32*& KeyCode, const uint32*& CharCode) const
 {
-	KeyCode = KeyMapCharToEnum.FindKey(Key);
-	CharCode = KeyMapVirtualToEnum.FindKey(Key);
+	CharCode = KeyMapCharToEnum.FindKey(Key);
+	KeyCode = KeyMapVirtualToEnum.FindKey(Key);
 }
 
 #undef LOCTEXT_NAMESPACE
diff --git a/Engine/Source/Runtime/Online/BuildPatchServices/Private/BuildPatchChunk.cpp b/Engine/Source/Runtime/Online/BuildPatchServices/Private/BuildPatchChunk.cpp
--- a/Engine/Source/Runtime/Online/BuildPatchServices/Private/BuildPatchChunk.cpp
+++ b/Engine/Source/Runtime/Online/BuildPatchServices/Private/BuildPatchChunk.cpp
@@ -235,13 +235,20 @@ uint32 FChunkWriter::FQueuedChunkWriter::Run()
 
 			// Delete the data memory
 			delete ChunkFile;
+
+			// Small sleep for next chunk
+			FPlatformProcess::Sleep(0.0f);
+		}
+		else
+		{
+			// Larger sleep for no work
+			FPlatformProcess::Sleep(0.1f);
 		}
 		double TotalTime = FStatsCollector::CyclesToSeconds(*StatFileCreateTime + *StatSerlialiseTime);
 		if (TotalTime > 0.0)
 		{
 			FStatsCollector::Set(StatDataWriteSpeed, *StatDataWritten / TotalTime);
 		}
-		FPlatformProcess::Sleep(0.0f);
 	}
 	return 0;
 }
diff --git a/Engine/Source/Runtime/Online/BuildPatchServices/Private/BuildPatchDiffManifests.cpp b/Engine/Source/Runtime/Online/BuildPatchServices/Private/BuildPatchDiffManifests.cpp
--- a/Engine/Source/Runtime/Online/BuildPatchServices/Private/BuildPatchDiffManifests.cpp
+++ b/Engine/Source/Runtime/Online/BuildPatchServices/Private/BuildPatchDiffManifests.cpp
@@ -85,7 +85,7 @@ bool FBuildDiffManifests::DiffManifests(const FString& ManifestFilePathA, const
 			++NewChunksCount;
 			TotalChunkSize += ChunkInfoRef.FileSize;
 			NewChunkPaths.Add(FBuildPatchUtils::GetDataFilename(ManifestB.ToSharedRef(), TEXT("."), ChunkInfoRef.Guid));
-			UE_LOG(LogDiffManifests, Log, TEXT("New chunk discovered: Size: %10lld, Path: %s"), ChunkInfoRef.FileSize, *NewChunkPaths.Last());
+			UE_LOG(LogDiffManifests, Verbose, TEXT("New chunk discovered: Size: %10lld, Path: %s"), ChunkInfoRef.FileSize, *NewChunkPaths.Last());
 		}
 	}
 
diff --git a/Engine/Source/Runtime/Online/BuildPatchServices/Private/BuildPatchGeneration.cpp b/Engine/Source/Runtime/Online/BuildPatchServices/Private/BuildPatchGeneration.cpp
--- a/Engine/Source/Runtime/Online/BuildPatchServices/Private/BuildPatchGeneration.cpp
+++ b/Engine/Source/Runtime/Online/BuildPatchServices/Private/BuildPatchGeneration.cpp
@@ -17,6 +17,32 @@
 DECLARE_LOG_CATEGORY_EXTERN(LogPatchGeneration, Log, All);
 DEFINE_LOG_CATEGORY(LogPatchGeneration);
 
+namespace BuildPatchServices
+{
+	struct FScannerDetails
+	{
+	public:
+		FScannerDetails(int32 InLayer, uint64 InLayerOffset, bool bInIsFinalScanner, uint64 InPaddingSize, TArray<uint8> InData, FBlockStructure InStructure, const ICloudEnumerationRef& CloudEnumeration, const FStatsCollectorRef& StatsCollector)
+			: Layer(InLayer)
+			, LayerOffset(InLayerOffset)
+			, bIsFinalScanner(bInIsFinalScanner)
+			, PaddingSize(InPaddingSize)
+			, Data(MoveTemp(InData))
+			, Structure(MoveTemp(InStructure))
+			, Scanner(FDataScannerFactory::Create(Data, CloudEnumeration, StatsCollector))
+		{}
+
+	public:
+		int32 Layer;
+		uint64 LayerOffset;
+		bool bIsFinalScanner;
+		uint64 PaddingSize;
+		TArray<uint8> Data;
+		FBlockStructure Structure;
+		IDataScannerRef Scanner;
+	};
+}
+
 namespace PatchGenerationHelpers
 {
 	bool SerialiseIntersection(const FBlockStructure& ByteStructure, const FBlockStructure& Intersection, FBlockStructure& SerialRanges)
@@ -61,32 +87,20 @@
 		}
 		return Count;
 	}
+
+	int32 GetMaxScannerBacklogCount()
+	{
+		int32 MaxScannerBacklogCount = 75;
+		GConfig->GetInt(TEXT("BuildPatchServices"), TEXT("MaxScannerBacklog"), MaxScannerBacklogCount, GEngineIni);
+		MaxScannerBacklogCount = FMath::Clamp<int32>(MaxScannerBacklogCount, 5, 500);
+		return MaxScannerBacklogCount;
+	}
+
+	bool ScannerArrayFull(const TArray<TUniquePtr<FScannerDetails>>& Scanners)
+	{
+		static int32 MaxScannerBacklogCount = GetMaxScannerBacklogCount();
+		return (FDataScannerCounter::GetNumIncompleteScanners() > FDataScannerCounter::GetNumRunningScanners()) || (Scanners.Num() >= MaxScannerBacklogCount);
+	}
 }
 
-namespace BuildPatchServices
-{
-	struct FScannerDetails
-	{
-	public:
-		FScannerDetails(int32 InLayer, uint64 InLayerOffset, bool bInIsFinalScanner, uint64 InPaddingSize, TArray<uint8> InData, FBlockStructure InStructure, const ICloudEnumerationRef& CloudEnumeration, const FStatsCollectorRef& StatsCollector)
-			: Layer(InLayer)
-			, LayerOffset(InLayerOffset)
-			, bIsFinalScanner(bInIsFinalScanner)
-			, PaddingSize(InPaddingSize)
-			, Data(MoveTemp(InData))
-			, Structure(MoveTemp(InStructure))
-			, Scanner(FDataScannerFactory::Create(Data, CloudEnumeration, StatsCollector))
-		{}
-
-	public:
-		int32 Layer;
-		uint64 LayerOffset;
-		bool bIsFinalScanner;
-		uint64 PaddingSize;
-		TArray<uint8> Data;
-		FBlockStructure Structure;
-		IDataScannerRef Scanner;
-	};
-}
-
 bool FBuildDataGenerator::GenerateChunksManifestFromDirectory(const BuildPatchServices::FGenerationConfiguration& Settings)
@@ -149,16 +163,24 @@
 	double LastProgressLog = FPlatformTime::Seconds();
 	const double TimeGenStarted = LastProgressLog;
 
-	// Load scanner data size from config, default to 32.5 MB.
+	// Load settings from config.
 	float GenerationScannerSizeMegabytes = 32.5f;
+	float StatsLoggerTimeSeconds = 10.0f;
 	GConfig->GetFloat(TEXT("BuildPatchServices"), TEXT("GenerationScannerSizeMegabytes"), GenerationScannerSizeMegabytes, GEngineIni);
+	GConfig->GetFloat(TEXT("BuildPatchServices"), TEXT("StatsLoggerTimeSeconds"), StatsLoggerTimeSeconds, GEngineIni);
 	GenerationScannerSizeMegabytes = FMath::Clamp<float>(GenerationScannerSizeMegabytes, 10.0f, 500.0f);
+	StatsLoggerTimeSeconds = FMath::Clamp<float>(StatsLoggerTimeSeconds, 1.0f, 60.0f);
 	const uint64 ScannerDataSize = GenerationScannerSizeMegabytes * 1048576;
 	const uint64 ScannerOverlapSize = FBuildPatchData::ChunkDataSize - 1;
 	TArray<uint8> DataBuffer;
 
 	// Setup Generation stats.
-	auto* StatTotalTime = StatsCollector->CreateStat(TEXT("Generation: Total Time"), EStatFormat::Timer);
+	volatile int64* StatTotalTime = StatsCollector->CreateStat(TEXT("Generation: Total Time"), EStatFormat::Timer);
+	volatile int64* StatLayers = StatsCollector->CreateStat(TEXT("Generation: Layers"), EStatFormat::Value);
+	volatile int64* StatNumScanners = StatsCollector->CreateStat(TEXT("Generation: Scanner Backlog"), EStatFormat::Value);
+	volatile int64* StatUnknownDataAlloc = StatsCollector->CreateStat(TEXT("Generation: Unmatched Buffers Allocation"), EStatFormat::DataSize);
+	volatile int64* StatUnknownDataNum = StatsCollector->CreateStat(TEXT("Generation: Unmatched Buffers Use"), EStatFormat::DataSize);
+	int64 MaxLayer = 0;
 
 	// List of created scanners.
 	TArray<TUniquePtr<FScannerDetails>> Scanners;
@@ -184,8 +206,7 @@
 	bool bHasUnknownData = true;
 	while (!BuildStream->IsEndOfData() || Scanners.Num() > 0 || bHasUnknownData)
 	{
-		bool bScannerQueueFull = FDataScannerCounter::GetNumIncompleteScanners() > FDataScannerCounter::GetNumRunningScanners();
-		if (!bScannerQueueFull)
+		if (!PatchGenerationHelpers::ScannerArrayFull(Scanners))
 		{
 			// Create a scanner from new build data?
 			if (!BuildStream->IsEndOfData())
@@ -223,6 +244,7 @@
 		}
 
 		// Grab any completed scanners?
+		FStatsCollector::Set(StatNumScanners, Scanners.Num());
 		while (Scanners.Num() > 0 && Scanners[0]->Scanner->IsComplete())
 		{
 			FScannerDetails& Scanner = *Scanners[0];
@@ -487,14 +509,12 @@
 			// Use unknown data to make new scanners.
 			if (UnknownLayerData.Num() > 0)
 			{
-				bScannerQueueFull = FDataScannerCounter::GetNumIncompleteScanners() > FDataScannerCounter::GetNumRunningScanners();
-
 				// Make sure there are enough bytes available for a scanner, plus a chunk, so that we know no more chunks will get
 				// made from this sequential unknown data.
 				uint64 RequiredScannerBytes = ScannerDataSize + FBuildPatchData::ChunkDataSize;
 				// We make a scanner when there's enough data, or if we are at the end of data for this layer.
 				bool bShouldMakeScanner = (ProcessedCount >= RequiredScannerBytes) || bIsFinalData;
-				if (!bScannerQueueFull && bShouldMakeScanner)
+				if (!PatchGenerationHelpers::ScannerArrayFull(Scanners) && bShouldMakeScanner)
 				{
 					// Create data processor.
 					const uint64 SizeToTake = FMath::Min<uint64>(ScannerDataSize, UnknownLayerData.Num());
@@ -512,10 +532,13 @@
 						UE_LOG(LogPatchGeneration, Error, TEXT("Tracked unknown build data inconsistent."));
 						return false;
 					}
-					uint64& DataOffset = LayerToDataOffset.FindOrAdd(Layer + 1);
-					Scanners.Emplace(new FScannerDetails(Layer + 1, DataOffset, bIsFinalScanner, PadSize, MoveTemp(ScannerData), BuildStructure, CloudEnumeration, StatsCollector));
-					LayerToScannerCount.FindOrAdd(Layer + 1)++;
-					UE_LOG(LogPatchGeneration, Verbose, TEXT("Creating scanner on layer %d at %llu. IsFinal:%d."), Layer + 1, DataOffset, bIsFinalScanner);
+					const int32 NextLayer = Layer + 1;
+					uint64& DataOffset = LayerToDataOffset.FindOrAdd(NextLayer);
+					Scanners.Emplace(new FScannerDetails(NextLayer, DataOffset, bIsFinalScanner, PadSize, MoveTemp(ScannerData), BuildStructure, CloudEnumeration, StatsCollector));
+					LayerToScannerCount.FindOrAdd(NextLayer)++;
+					UE_LOG(LogPatchGeneration, Verbose, TEXT("Creating scanner on layer %d at %llu. IsFinal:%d."), NextLayer, DataOffset, bIsFinalScanner);
+					MaxLayer = FMath::Max<int64>(MaxLayer, NextLayer);
+					FStatsCollector::Set(StatLayers, MaxLayer);
 
 					// Remove data we just pulled out from tracking, minus the overlap.
 					uint64 RemoveSize = bIsFinalScanner ? SizeToTake : SizeToTake - ScannerOverlapSize;
@@ -549,27 +572,35 @@
 			}
 		}
 
-		// Log collected stats.
-		GLog->FlushThreadedLogs();
-		FStatsCollector::Set(StatTotalTime, FStatsCollector::GetCycles() - StartTime);
-		StatsCollector->LogStats(10.0f);
-
 		// Set whether we are still processing unknown data.
-		bHasUnknownData = false;
+		int64 UnknownDataAlloc = 0;
+		int64 UnknownDataNum = 0;
 		for (TPair<int32, TArray<uint8>>& UnknownLayerDataPair : LayerToUnknownLayerData)
 		{
-			if (UnknownLayerDataPair.Value.Num() > 0)
-			{
-				bHasUnknownData = true;
-				break;
-			}
+			UnknownDataNum += UnknownLayerDataPair.Value.Num();
+			UnknownDataAlloc += UnknownLayerDataPair.Value.GetAllocatedSize();
 		}
+		bHasUnknownData = UnknownDataNum > 0;
+		FStatsCollector::Set(StatUnknownDataAlloc, UnknownDataAlloc);
+		FStatsCollector::Set(StatUnknownDataNum, UnknownDataNum);
 
+		// Log collected stats.
+		GLog->FlushThreadedLogs();
+		FStatsCollector::Set(StatTotalTime, FStatsCollector::GetCycles() - StartTime);
+		StatsCollector->LogStats(StatsLoggerTimeSeconds);
+
 		// Sleep to allow other threads.
 		FPlatformProcess::Sleep(0.01f);
 	}
 	UE_LOG(LogPatchGeneration, Verbose, TEXT("Scanning complete, waiting for writer thread."));
 
+	// Check that we read some build data
+	if (BuildStream->GetBuildSize() == 0)
+	{
+		UE_LOG(LogPatchGeneration, Error, TEXT("There was no data to process. Please check path: %s."), *Settings.RootDirectory);
+		return false;
+	}
+
 	// Inform no more chunks.
 	ChunkWriter.NoMoreChunks();
 	// Wait for chunk writer.
