Fix TVOS issue (battery level unsupported)
#ue4
#tvos

diff --git a/Engine/Source/Runtime/Core/Private/IOS/IOSAppDelegate.cpp b/Engine/Source/Runtime/Core/Private/IOS/IOSAppDelegate.cpp
--- a/Engine/Source/Runtime/Core/Private/IOS/IOSAppDelegate.cpp
+++ b/Engine/Source/Runtime/Core/Private/IOS/IOSAppDelegate.cpp
@@ -384,10 +384,15 @@ void InstallSignalHandlers()
 {
 	UIDevice *myDevice = [UIDevice currentDevice];
 
+#if PLATFORM_TVOS
+	// TVOS does not have a battery, return fully charged
+	return 100;
+#else
 	//must enable battery monitoring in order to get a valid value here
 	[myDevice setBatteryMonitoringEnabled : YES];
 	//battery level is from 0.0 to 1.0, get it in terms of 0-100
 	return ((int)([myDevice batteryLevel] * 100));
+#endif
 }
 
 /**
