Revert CL3215720 from dev-core for now (VS2015 doesn't like it)

diff --git a/Engine/Source/Runtime/Core/Public/Delegates/Tuple.h b/Engine/Source/Runtime/Core/Public/Delegates/Tuple.h
--- a/Engine/Source/Runtime/Core/Public/Delegates/Tuple.h
+++ b/Engine/Source/Runtime/Core/Public/Delegates/Tuple.h
@@ -26,11 +26,6 @@ struct TNthTypeFromParameterPack<0, T, OtherTypes...>
 template <typename T, uint32 Index>
 struct TTupleElement
 {
-	TTupleElement()
-		: Value()
-	{
-	}
-
 	template <typename... ArgTypes>
 	explicit TTupleElement(ArgTypes&&... Args)
 		: Value(Forward<ArgTypes>(Args)...)
@@ -80,11 +75,6 @@ struct TTupleImpl;
 template <uint32... Indices, typename... Types>
 struct TTupleImpl<TIntegerSequence<uint32, Indices...>, Types...> : TTupleElement<Types, Indices>...
 {
-	TTupleImpl()
-		: TTupleElement<Types, Indices>()...
-	{
-	}
-
 	template <typename... ArgTypes>
 	explicit TTupleImpl(ArgTypes&&... Args)
 		: TTupleElement<Types, Indices>(Forward<ArgTypes>(Args))...
