Tuple default constructor reintroduced after being backed out in CL# 3231861.
Erroneous 'multiple default constructors' warning disabled in VS2013.

diff --git a/Engine/Source/Runtime/Core/Public/Delegates/Tuple.h b/Engine/Source/Runtime/Core/Public/Delegates/Tuple.h
--- a/Engine/Source/Runtime/Core/Public/Delegates/Tuple.h
+++ b/Engine/Source/Runtime/Core/Public/Delegates/Tuple.h
@@ -72,9 +72,21 @@ struct TTuple;
 template <typename Indices, typename... Types>
 struct TTupleImpl;
 
+// VS2013 gives erroneous warnings about multiple default constructors for TTupleElement and TTupleImpl,
+// even though the variadic one can never be called with no arguments.
+#if defined(_MSC_VER) && _MSC_VER <= 1800
+	#pragma warning(push)
+	#pragma warning(disable:4520) // multiple default constructors specified
+#endif
+
 template <uint32... Indices, typename... Types>
 struct TTupleImpl<TIntegerSequence<uint32, Indices...>, Types...> : TTupleElement<Types, Indices>...
 {
+	TTupleImpl()
+		: TTupleElement<Types, Indices>()...
+	{
+	}
+
 	template <typename... ArgTypes>
 	explicit TTupleImpl(ArgTypes&&... Args)
 		: TTupleElement<Types, Indices>(Forward<ArgTypes>(Args))...
@@ -97,6 +109,10 @@ struct TTupleImpl<TIntegerSequence<uint32, Indices...>, Types...> : TTupleElemen
 	}
 };
 
+#if defined(_MSC_VER) && _MSC_VER <= 1800
+	#pragma warning(pop)
+#endif
+
 #ifdef _MSC_VER
 
 	// Not strictly necessary, but some VC versions give a 'syntax error: <fake-expression>' error
