Made sure the async log writer flushes the log archive even if there's no serialization requests to prevent situations where some of the recent requests haven't been flushed to disk.

Fixes Mac Project Launcher hang on waiting for server due to unflushed logs.

diff --git a/Engine/Source/Runtime/Core/Private/Misc/OutputDeviceFile.cpp b/Engine/Source/Runtime/Core/Private/Misc/OutputDeviceFile.cpp
--- a/Engine/Source/Runtime/Core/Private/Misc/OutputDeviceFile.cpp
+++ b/Engine/Source/Runtime/Core/Private/Misc/OutputDeviceFile.cpp
@@ -63,6 +63,9 @@ class CORE_API FAsyncWriter : public FRunnable, public FArchive
 	/** [WRITER THREAD] Last time the archive was flushed. used in threaded situations to flush the underlying archive at a certain maximum rate. */
 	double LastArchiveFlushTime;
 
+	/** [WRITER THREAD] Archive flush interval. */
+	double ArchiveFlushIntervalSec;
+
 	/** [WRITER THREAD] Serialize the contents of the ring buffer to disk */
 	void SerializeBufferToArchive()
 	{
@@ -89,7 +92,6 @@ class CORE_API FAsyncWriter : public FRunnable, public FArchive
 			// Flush the archive periodically if running on a separate thread
 			if (Thread)
 			{
-				const double ArchiveFlushIntervalSec = 0.2;
 				if ((FPlatformTime::Seconds() - LastArchiveFlushTime) > ArchiveFlushIntervalSec)
 				{
 					Ar.Flush();
@@ -127,9 +129,16 @@ public:
 		, BufferStartPos(0)
 		, BufferEndPos(0)
 		, LastArchiveFlushTime(0.0)
+		, ArchiveFlushIntervalSec(0.2)
 	{
 		Buffer.AddUninitialized(InitialBufferSize);
 
+		float CommandLineInterval = 0.0;
+		if (FParse::Value(FCommandLine::Get(), TEXT("LOGFLUSHINTERVAL="), CommandLineInterval))
+		{
+			ArchiveFlushIntervalSec = CommandLineInterval;
+		}
+
 		if (FPlatformProcess::SupportsMultithreading())
 		{
 			FString WriterName = FString::Printf(TEXT("FAsyncWriter_%s"), *FPaths::GetBaseFilename(Ar.GetArchiveName()));
@@ -227,6 +236,10 @@ public:
 			{
 				SerializeBufferToArchive();
 			}
+			else if ((FPlatformTime::Seconds() - LastArchiveFlushTime) > ArchiveFlushIntervalSec)
+			{
+				SerializeRequestCounter.Increment();
+			}
 			else
 			{
 				FPlatformProcess::Sleep(0.01f);
