Fix the SectionInfoMap serialization issue, we redo the same change as the one in the release 4.15 branch (cl:3242088).

diff --git a/Engine/Source/Runtime/Core/Public/UObject/ReleaseObjectVersion.h b/Engine/Source/Runtime/Core/Public/UObject/ReleaseObjectVersion.h
--- a/Engine/Source/Runtime/Core/Public/UObject/ReleaseObjectVersion.h
+++ b/Engine/Source/Runtime/Core/Public/UObject/ReleaseObjectVersion.h
@@ -25,6 +25,9 @@ struct CORE_API FReleaseObjectVersion
 		// Add Component node templates now use their own unique naming scheme to ensure more reliable archetype lookups.
 		AddComponentNodeTemplateUniqueNames,
 
+		// Fix a serialization issue with static mesh FMeshSectionInfoMap UProperty
+		UPropertryForMeshSectionSerialize,
+
 		// -----<new versions can be added above this line>-------------------------------------------------
 		VersionPlusOne,
 		LatestVersion = VersionPlusOne - 1
diff --git a/Engine/Source/Runtime/Engine/Classes/Engine/StaticMesh.h b/Engine/Source/Runtime/Engine/Classes/Engine/StaticMesh.h
--- a/Engine/Source/Runtime/Engine/Classes/Engine/StaticMesh.h
+++ b/Engine/Source/Runtime/Engine/Classes/Engine/StaticMesh.h
@@ -238,6 +238,7 @@ struct FMeshSectionInfoMap
 	GENERATED_USTRUCT_BODY()
 
 	/** Maps an LOD+Section to the material it should render with. */
+	UPROPERTY()
 	TMap<uint32,FMeshSectionInfo> Map;
 
 	/** Serialize. */
diff --git a/Engine/Source/Runtime/Engine/Private/StaticMesh.cpp b/Engine/Source/Runtime/Engine/Private/StaticMesh.cpp
--- a/Engine/Source/Runtime/Engine/Private/StaticMesh.cpp
+++ b/Engine/Source/Runtime/Engine/Private/StaticMesh.cpp
@@ -2225,7 +2225,12 @@ FArchive& operator<<(FArchive& Ar, FMeshSectionInfo& Info)
 
 void FMeshSectionInfoMap::Serialize(FArchive& Ar)
 {
-	Ar << Map;
+	Ar.UsingCustomVersion(FReleaseObjectVersion::GUID);
+
+	if (Ar.CustomVer(FReleaseObjectVersion::GUID) < FReleaseObjectVersion::UPropertryForMeshSectionSerialize)
+	{
+		Ar << Map;
+	}
 }
 
 #endif // #if WITH_EDITORONLY_DATA
