Copying //UE4/Dev-Editor to //UE4/Dev-Main (Source: //UE4/Dev-Editor @ 3279756)

diff --git a/Engine/Source/Editor/UnrealEd/Private/ObjectTools.cpp b/Engine/Source/Editor/UnrealEd/Private/ObjectTools.cpp
--- a/Engine/Source/Editor/UnrealEd/Private/ObjectTools.cpp
+++ b/Engine/Source/Editor/UnrealEd/Private/ObjectTools.cpp
@@ -1310,7 +1310,7 @@ namespace ObjectTools
 	 */
 	void SelectActorsInLevelDirectlyReferencingObject( UObject* RefObj )
 	{
-		UPackage* Package = Cast<UPackage>(RefObj->GetOutermost());
+		UPackage* Package = RefObj->GetOutermost();
 		if (Package && Package->ContainsMap())
 		{
 			// Walk the chain of outers to find the object that is 'in' the level...
@@ -1492,6 +1492,7 @@ namespace ObjectTools
 	void CleanupAfterSuccessfulDelete (const TArray<UPackage*>& PotentialPackagesToDelete, bool bPerformReferenceCheck)
 	{
 		TArray<UPackage*> PackagesToDelete = PotentialPackagesToDelete;
+		TArray<UPackage*> EmptyPackagesToUnload;
 		TArray<FString> PackageFilesToDelete;
 		TArray<TSharedRef<ISourceControlState, ESPMode::ThreadSafe> > PackageSCCStates;
 		ISourceControlProvider& SourceControlProvider = ISourceControlModule::Get().GetProvider();
@@ -1507,7 +1508,7 @@ namespace ObjectTools
 
 			bool bIsReferenced = false;
 			
-			if ( bPerformReferenceCheck )
+			if ( Package != nullptr && bPerformReferenceCheck )
 			{
 				FReferencerInformationList FoundReferences;
 				bIsReferenced = IsReferenced(Package, GARBAGE_COLLECTION_KEEPFLAGS, EInternalObjectFlags::GarbageCollectionKeepFlags,  true, &FoundReferences);
@@ -1533,16 +1534,26 @@ namespace ObjectTools
 			}
 			else
 			{
+				UPackage* CurrentPackage = Cast<UPackage>(Package);
+
 				FString PackageFilename;
-				if( !FPackageName::DoesPackageExist( Package->GetName(), NULL, &PackageFilename ) )
+				if( Package != nullptr && FPackageName::DoesPackageExist( Package->GetName(), NULL, &PackageFilename ) )
+				{
+					PackageFilesToDelete.Add(PackageFilename);
+					CurrentPackage->SetDirtyFlag(false);
+				}
+				else
 				{
 					// Could not determine filename for package so we can not delete
 					PackagesToDelete.RemoveAt(PackageIdx);
-					continue;
-				}
 
-				PackageFilesToDelete.Add(PackageFilename);
-				Cast<UPackage>(Package)->SetDirtyFlag(false);
+					if (UPackage::IsEmptyPackage(CurrentPackage))
+					{
+						// If the package is empty, unload it anyway
+						EmptyPackagesToUnload.Add(CurrentPackage);
+						CurrentPackage->SetDirtyFlag(false);
+					}
+				}
 			}
 		}
 
@@ -1561,10 +1572,14 @@ namespace ObjectTools
 		}
 
 		// Unload the packages and collect garbage.
-		if ( PackagesToDelete.Num() > 0 )
+		if ( PackagesToDelete.Num() > 0 || EmptyPackagesToUnload.Num() > 0 )
 		{
-			PackageTools::UnloadPackages(PackagesToDelete);
-			PackagesToDelete.Reset();
+			TArray<UPackage*> AllPackagesToUnload;
+			AllPackagesToUnload.Reserve(PackagesToDelete.Num() + EmptyPackagesToUnload.Num());
+			AllPackagesToUnload.Append(PackagesToDelete);
+			AllPackagesToUnload.Append(EmptyPackagesToUnload);
+
+			PackageTools::UnloadPackages(AllPackagesToUnload);
 		}
 		CollectGarbage( GARBAGE_COLLECTION_KEEPFLAGS );
 
diff --git a/Engine/Source/Editor/UnrealEd/Private/PlayLevel.cpp b/Engine/Source/Editor/UnrealEd/Private/PlayLevel.cpp
--- a/Engine/Source/Editor/UnrealEd/Private/PlayLevel.cpp
+++ b/Engine/Source/Editor/UnrealEd/Private/PlayLevel.cpp
@@ -264,7 +264,7 @@ void UEditorEngine::EndPlayMap()
 		GameViewport->OnGameViewportInputKey().Unbind();
 
 		// Remove close handler binding
-		GameViewport->OnCloseRequested().Unbind();
+		GameViewport->OnCloseRequested().Remove(ViewportCloseRequestedDelegateHandle);
 
 		GameViewport->CloseRequested(GameViewport->Viewport);
 	}
@@ -2512,7 +2512,7 @@ void UEditorEngine::PlayInEditor( UWorld* InWorld, bool bInSimulateInEditor )
 		PIEInstance = ( !CanRunUnderOneProcess && PlayNetMode == EPlayNetMode::PIE_Client ) ? INDEX_NONE : 0;
 		UGameInstance* const GameInstance = CreatePIEGameInstance(PIEInstance, bInSimulateInEditor, bAnyBlueprintErrors, bStartInSpectatorMode, false, PIEStartTime);
 
-		if (!PlayInSettings->EnableSound)
+		if (!PlayInSettings->EnableGameSound)
 		{
 			UWorld* GameInstanceWorld = GameInstance->GetWorld();
 			if (FAudioDevice* GameInstanceAudioDevice = GameInstanceWorld->GetAudioDevice())
@@ -3164,7 +3164,7 @@ UGameInstance* UEditorEngine::CreatePIEGameInstance(int32 InPIEInstance, bool bI
 		ViewportClient->OnGameViewportInputKey().BindUObject(this, &UEditorEngine::ProcessDebuggerCommands);
 
 		// Add a handler for viewport close requests
-		ViewportClient->OnCloseRequested().BindUObject(this, &UEditorEngine::OnViewportCloseRequested);
+		ViewportCloseRequestedDelegateHandle = ViewportClient->OnCloseRequested().AddUObject(this, &UEditorEngine::OnViewportCloseRequested);
 		FSlatePlayInEditorInfo& SlatePlayInEditorSession = SlatePlayInEditorMap.Add(PieWorldContext->ContextHandle, FSlatePlayInEditorInfo());
 		SlatePlayInEditorSession.DestinationSlateViewport = RequestedDestinationSlateViewport;	// Might be invalid depending how pie was launched. Code below handles this.
 		RequestedDestinationSlateViewport = NULL;
@@ -3770,7 +3770,7 @@ UWorld* UEditorEngine::CreatePIEWorldBySavingToTemp(FWorldContext& WorldContext,
 UWorld* UEditorEngine::CreatePIEWorldByDuplication(FWorldContext &WorldContext, UWorld* InWorld, FString &PlayWorldMapName)
 {
 	double StartTime = FPlatformTime::Seconds();
-	UPackage* InPackage = Cast<UPackage>(InWorld->GetOutermost());
+	UPackage* InPackage = InWorld->GetOutermost();
 	UWorld* NewPIEWorld = NULL;
 	
 	const FString WorldPackageName = InPackage->GetName();
diff --git a/Engine/Source/Editor/UnrealEd/Private/SCSVImportOptions.cpp b/Engine/Source/Editor/UnrealEd/Private/SCSVImportOptions.cpp
--- a/Engine/Source/Editor/UnrealEd/Private/SCSVImportOptions.cpp
+++ b/Engine/Source/Editor/UnrealEd/Private/SCSVImportOptions.cpp
@@ -53,9 +53,38 @@ void SCSVImportOptions::Construct(const FArguments& InArgs)
 		. Padding(10)
 		[
 			SNew(SVerticalBox)
+			+SVerticalBox::Slot()
+			.AutoHeight()
+			[
+				SNew(SBorder)
+				.Padding(FMargin(3))
+				.BorderImage(FEditorStyle::GetBrush("ToolPanel.GroupBorder"))
+				.Visibility( InArgs._FullPath.IsEmpty() ? EVisibility::Collapsed : EVisibility::Visible )
+				[
+					SNew(SHorizontalBox)
+					+SHorizontalBox::Slot()
+					.AutoWidth()
+					[
+						SNew(STextBlock)
+						.Font(FEditorStyle::GetFontStyle("CurveEd.LabelFont"))
+						.Text(LOCTEXT("Import_CurrentFileTitle", "Current File: "))
+					]
+					+SHorizontalBox::Slot()
+					.Padding(5, 0, 0, 0)
+					.AutoWidth()
+					.VAlign(VAlign_Center)
+					[
+						SNew(STextBlock)
+						.Font(FEditorStyle::GetFontStyle("CurveEd.InfoFont"))
+						.Text(InArgs._FullPath)
+					]
+				]
+			]
+
 			// Import type
 			+SVerticalBox::Slot()
 			.AutoHeight()
+			.Padding(2)
 			[
 				SNew(STextBlock)
 				.Text( LOCTEXT("ChooseAssetType", "Import As:") )
diff --git a/Engine/Source/Editor/UnrealEd/Private/SColorGradientEditor.cpp b/Engine/Source/Editor/UnrealEd/Private/SColorGradientEditor.cpp
--- a/Engine/Source/Editor/UnrealEd/Private/SColorGradientEditor.cpp
+++ b/Engine/Source/Editor/UnrealEd/Private/SColorGradientEditor.cpp
@@ -151,7 +151,7 @@ int32 SColorGradientEditor::OnPaint( const FPaintArgs& Args, const FGeometry& Al
 		FGeometry GradientAreaGeometry = AllottedGeometry.MakeChild( FVector2D(0.0f, 16.0f), FVector2D( AllottedGeometry.Size.X, AllottedGeometry.Size.Y - 30.0f ) );
 
 		bool bEnabled = ShouldBeEnabled( bParentEnabled );
-		ESlateDrawEffect::Type DrawEffects = bEnabled ? ESlateDrawEffect::None : ESlateDrawEffect::DisabledEffect;
+		ESlateDrawEffect DrawEffects = bEnabled ? ESlateDrawEffect::None : ESlateDrawEffect::DisabledEffect;
 
 		// Pixel to value input converter
 		FTrackScaleInfo ScaleInfo(ViewMinInput.Get(),  ViewMaxInput.Get(), 0.0f, 1.0f, GradientAreaGeometry.Size);
@@ -689,7 +689,7 @@ void SColorGradientEditor::OnSetGradientStopTimeFromPopup( const FText& NewText,
 	}
 }
 
-void SColorGradientEditor::DrawGradientStopMark( const FGradientStopMark& Mark, const FGeometry& Geometry, float XPos, const FLinearColor& Color, FSlateWindowElementList& OutDrawElements, int32 LayerId, const FSlateRect& InClippingRect, ESlateDrawEffect::Type DrawEffects, bool bColor, const FWidgetStyle& InWidgetStyle ) const
+void SColorGradientEditor::DrawGradientStopMark( const FGradientStopMark& Mark, const FGeometry& Geometry, float XPos, const FLinearColor& Color, FSlateWindowElementList& OutDrawElements, int32 LayerId, const FSlateRect& InClippingRect, ESlateDrawEffect DrawEffects, bool bColor, const FWidgetStyle& InWidgetStyle ) const
 {
 	static const FSlateBrush* WhiteBrush = FEditorStyle::GetBrush("WhiteBrush");
 	static const FSlateBrush* ColorStopBrush = FEditorStyle::GetBrush("Sequencer.Timeline.ScrubHandleDown");
diff --git a/Engine/Source/Editor/UnrealEd/Private/SCurveEditor.cpp b/Engine/Source/Editor/UnrealEd/Private/SCurveEditor.cpp
--- a/Engine/Source/Editor/UnrealEd/Private/SCurveEditor.cpp
+++ b/Engine/Source/Editor/UnrealEd/Private/SCurveEditor.cpp
@@ -748,7 +748,7 @@ int32 SCurveEditor::OnPaint(const FPaintArgs& Args, const FGeometry& AllottedGeo
 {
 	// Rendering info
 	bool bEnabled = ShouldBeEnabled( bParentEnabled );
-	ESlateDrawEffect::Type DrawEffects = bEnabled ? ESlateDrawEffect::None : ESlateDrawEffect::DisabledEffect;
+	ESlateDrawEffect DrawEffects = bEnabled ? ESlateDrawEffect::None : ESlateDrawEffect::DisabledEffect;
 	const FSlateBrush* TimelineAreaBrush = FEditorStyle::GetBrush("CurveEd.TimelineArea");
 	const FSlateBrush* WhiteBrush = FEditorStyle::GetBrush("WhiteTexture");
 
@@ -857,7 +857,7 @@ int32 SCurveEditor::OnPaint(const FPaintArgs& Args, const FGeometry& AllottedGeo
 }
 
 void SCurveEditor::PaintCurve(TSharedPtr<FCurveViewModel> CurveViewModel, const FGeometry &AllottedGeometry, FTrackScaleInfo &ScaleInfo, FSlateWindowElementList &OutDrawElements,
-	int32 LayerId, const FSlateRect& MyClippingRect, ESlateDrawEffect::Type DrawEffects, const FWidgetStyle &InWidgetStyle, bool bAnyCurveViewModelsSelected )const
+	int32 LayerId, const FSlateRect& MyClippingRect, ESlateDrawEffect DrawEffects, const FWidgetStyle &InWidgetStyle, bool bAnyCurveViewModelsSelected )const
 {
 	if (CurveViewModel.IsValid())
 	{
@@ -977,7 +977,7 @@ void SCurveEditor::CreateLinesForSegment( FRichCurve* Curve, const FRichCurveKey
 	}
 }
 
-void SCurveEditor::PaintKeys(TSharedPtr<FCurveViewModel> CurveViewModel, FTrackScaleInfo &ScaleInfo, FSlateWindowElementList &OutDrawElements, int32 LayerId, int32 SelectedLayerId, const FGeometry &AllottedGeometry, const FSlateRect& MyClippingRect, ESlateDrawEffect::Type DrawEffects, const FWidgetStyle &InWidgetStyle, bool bAnyCurveViewModelsSelected ) const
+void SCurveEditor::PaintKeys(TSharedPtr<FCurveViewModel> CurveViewModel, FTrackScaleInfo &ScaleInfo, FSlateWindowElementList &OutDrawElements, int32 LayerId, int32 SelectedLayerId, const FGeometry &AllottedGeometry, const FSlateRect& MyClippingRect, ESlateDrawEffect DrawEffects, const FWidgetStyle &InWidgetStyle, bool bAnyCurveViewModelsSelected ) const
 {
 	FLinearColor KeyColor = CurveViewModel->bIsLocked ? FLinearColor(0.1f,0.1f,0.1f,1.f) : InWidgetStyle.GetColorAndOpacityTint();
 
@@ -1026,7 +1026,7 @@ void SCurveEditor::PaintKeys(TSharedPtr<FCurveViewModel> CurveViewModel, FTrackS
 }
 
 
-void SCurveEditor::PaintTangent( TSharedPtr<FCurveViewModel> CurveViewModel, FTrackScaleInfo &ScaleInfo, FRichCurve* Curve, FKeyHandle KeyHandle, FVector2D KeyLocation, FSlateWindowElementList &OutDrawElements, int32 LayerId, const FGeometry &AllottedGeometry, const FSlateRect& MyClippingRect, ESlateDrawEffect::Type DrawEffects, int32 LayerToUse, const FWidgetStyle &InWidgetStyle, bool bTangentSelected, bool bIsArrivalSelected, bool bIsLeaveSelected, bool bAnyCurveViewModelsSelected ) const
+void SCurveEditor::PaintTangent( TSharedPtr<FCurveViewModel> CurveViewModel, FTrackScaleInfo &ScaleInfo, FRichCurve* Curve, FKeyHandle KeyHandle, FVector2D KeyLocation, FSlateWindowElementList &OutDrawElements, int32 LayerId, const FGeometry &AllottedGeometry, const FSlateRect& MyClippingRect, ESlateDrawEffect DrawEffects, int32 LayerToUse, const FWidgetStyle &InWidgetStyle, bool bTangentSelected, bool bIsArrivalSelected, bool bIsLeaveSelected, bool bAnyCurveViewModelsSelected ) const
 {
 	FVector2D ArriveTangentLocation, LeaveTangentLocation;
 	GetTangentPoints(ScaleInfo, FSelectedCurveKey(Curve,KeyHandle), ArriveTangentLocation, LeaveTangentLocation);
@@ -1112,7 +1112,7 @@ float SCurveEditor::GetTimeStep(FTrackScaleInfo &ScaleInfo) const
 }
 
 void SCurveEditor::PaintGridLines(const FGeometry &AllottedGeometry, FTrackScaleInfo &ScaleInfo, FSlateWindowElementList &OutDrawElements, 
-	int32 LayerId, const FSlateRect& MyClippingRect, ESlateDrawEffect::Type DrawEffects )const
+	int32 LayerId, const FSlateRect& MyClippingRect, ESlateDrawEffect DrawEffects )const
 {
 	const float MaxGridPixelSpacing = 150.0f;
 
diff --git a/Engine/Source/Editor/UnrealEd/Private/Settings/LevelEditorPlaySettingsCustomization.h b/Engine/Source/Editor/UnrealEd/Private/Settings/LevelEditorPlaySettingsCustomization.h
--- a/Engine/Source/Editor/UnrealEd/Private/Settings/LevelEditorPlaySettingsCustomization.h
+++ b/Engine/Source/Editor/UnrealEd/Private/Settings/LevelEditorPlaySettingsCustomization.h
@@ -322,7 +322,7 @@ public:
 			TArray<TSharedRef<IPropertyHandle>> PIECategoryProperties;
 			PlayInEditorCategory.GetDefaultProperties(PIECategoryProperties, true, false);
 
-			TSharedPtr<IPropertyHandle> PIEEnableSoundHandle = LayoutBuilder.GetProperty(GET_MEMBER_NAME_CHECKED(ULevelEditorPlaySettings, EnableSound));
+			TSharedPtr<IPropertyHandle> PIEEnableSoundHandle = LayoutBuilder.GetProperty(GET_MEMBER_NAME_CHECKED(ULevelEditorPlaySettings, EnableGameSound));
 			PIESoundQualityLevelHandle = LayoutBuilder.GetProperty(GET_MEMBER_NAME_CHECKED(ULevelEditorPlaySettings, PlayInEditorSoundQualityLevel));
 			PIESoundQualityLevelHandle->MarkHiddenByCustomization();
 
diff --git a/Engine/Source/Editor/UnrealEd/Private/Settings/SettingsClasses.cpp b/Engine/Source/Editor/UnrealEd/Private/Settings/SettingsClasses.cpp
--- a/Engine/Source/Editor/UnrealEd/Private/Settings/SettingsClasses.cpp
+++ b/Engine/Source/Editor/UnrealEd/Private/Settings/SettingsClasses.cpp
@@ -237,6 +237,7 @@ FAutoReimportDirectoryConfig::FParseContext::FParseContext(bool bInEnableLogging
 bool FAutoReimportDirectoryConfig::ParseSourceDirectoryAndMountPoint(FString& SourceDirectory, FString& MountPoint, const FParseContext& InContext)
 {
 	SourceDirectory.ReplaceInline(TEXT("\\"), TEXT("/"));
+	MountPoint.ReplaceInline(TEXT("\\"), TEXT("/"));
 
 	// Check if the source directory is actually a mount point
 	if (!FPackageName::GetPackageMountPoint(SourceDirectory).IsNone())
@@ -257,6 +258,7 @@ bool FAutoReimportDirectoryConfig::ParseSourceDirectoryAndMountPoint(FString& So
 		if (FPackageName::GetPackageMountPoint(MountPoint).IsNone())
 		{
 			UE_CLOG(InContext.bEnableLogging, LogAutoReimportManager, Warning, TEXT("Unable to setup directory %s to map to %s, as it's not a valid mounted path. Continuing without mounted path (auto reimports will still work, but auto add won't)."), *SourceDirectory, *MountPoint);
+			MountPoint = FString();
 		}
 	}
 	else if(!MountPoint.IsEmpty())
@@ -359,6 +361,7 @@ ULevelEditorPlaySettings::ULevelEditorPlaySettings( const FObjectInitializer& Ob
 	bAutoCompileBlueprintsOnLaunch = true;
 	CenterNewWindow = true;
 	CenterStandaloneWindow = true;
+	EnablePIEEnterAndExitSounds = false;
 }
 
 void ULevelEditorPlaySettings::PostEditChangeProperty(struct FPropertyChangedEvent& PropertyChangedEvent)
diff --git a/Engine/Source/Editor/UnrealEd/Private/UnrealEdMisc.cpp b/Engine/Source/Editor/UnrealEd/Private/UnrealEdMisc.cpp
--- a/Engine/Source/Editor/UnrealEd/Private/UnrealEdMisc.cpp
+++ b/Engine/Source/Editor/UnrealEd/Private/UnrealEdMisc.cpp
@@ -1174,7 +1174,7 @@ void FUnrealEdMisc::OnMessageTokenActivated(const TSharedRef<IMessageToken>& Tok
 						while (Blueprint == nullptr && ParentObject != nullptr)
 						{
 							Blueprint = UBlueprint::GetBlueprintFromClass(ParentObject->GetClass());
-							ParentObject = Object->GetOuter();
+							ParentObject = ParentObject->GetOuter();
 						}
 
 						if (Blueprint != nullptr)
diff --git a/Engine/Source/Editor/UnrealEd/Public/AnimationEditorUtils.h b/Engine/Source/Editor/UnrealEd/Public/AnimationEditorUtils.h
--- a/Engine/Source/Editor/UnrealEd/Public/AnimationEditorUtils.h
+++ b/Engine/Source/Editor/UnrealEd/Public/AnimationEditorUtils.h
@@ -1,7 +1,6 @@
 // Copyright 1998-2017 Epic Games, Inc. All Rights Reserved.
 
-#ifndef __AnimationEditorUtils_h__
-#define __AnimationEditorUtils_h__
+#pragma once
 
 #include "CoreMinimal.h"
 #include "UObject/Object.h"
@@ -163,5 +162,3 @@ namespace AnimationEditorUtils
 		}
 	}
 } // namespace AnimationEditorUtils
-
-#endif //__AnimationEditorUtils_h__
diff --git a/Engine/Source/Editor/UnrealEd/Public/ApexClothingUtils.h b/Engine/Source/Editor/UnrealEd/Public/ApexClothingUtils.h
--- a/Engine/Source/Editor/UnrealEd/Public/ApexClothingUtils.h
+++ b/Engine/Source/Editor/UnrealEd/Public/ApexClothingUtils.h
@@ -1,8 +1,6 @@
 // Copyright 1998-2017 Epic Games, Inc. All Rights Reserved.
 
-
-#ifndef __ApexClothingUtils_h__
-#define __ApexClothingUtils_h__
+#pragma once
 
 #include "CoreMinimal.h"
 #include "EngineDefines.h"
@@ -123,5 +121,3 @@ namespace ApexClothingUtils
 
 #endif // #if WITH_APEX_CLOTHING
 } // namespace ApexClothingUtils
-
-#endif //__ApexClothingUtils_h__
diff --git a/Engine/Source/Editor/UnrealEd/Public/AssetNotifications.h b/Engine/Source/Editor/UnrealEd/Public/AssetNotifications.h
--- a/Engine/Source/Editor/UnrealEd/Public/AssetNotifications.h
+++ b/Engine/Source/Editor/UnrealEd/Public/AssetNotifications.h
@@ -1,7 +1,6 @@
 // Copyright 1998-2017 Epic Games, Inc. All Rights Reserved.
 
-#ifndef __ASSETNOTIFICATIONS_H__
-#define __ASSETNOTIFICATIONS_H__
+#pragma once
 
 #include "CoreMinimal.h"
 
@@ -11,8 +10,3 @@ public:
 	/** Inform about change in skeleton asset that requires saving*/
 	UNREALED_API static void SkeletonNeedsToBeSaved(const class USkeleton* Skeleton);
 };
-
-#endif
-
-// EOF
-
diff --git a/Engine/Source/Editor/UnrealEd/Public/AssetSelection.h b/Engine/Source/Editor/UnrealEd/Public/AssetSelection.h
--- a/Engine/Source/Editor/UnrealEd/Public/AssetSelection.h
+++ b/Engine/Source/Editor/UnrealEd/Public/AssetSelection.h
@@ -1,10 +1,6 @@
 // Copyright 1998-2017 Epic Games, Inc. All Rights Reserved.
 
-
-
-
-#ifndef __ASSETSELECTION_H__
-#define __ASSETSELECTION_H__
+#pragma once
 
 #include "CoreMinimal.h"
 #include "AssetData.h"
@@ -323,8 +319,3 @@ private:
 	{
 	}
 };
-
-#endif
-
-// EOF
-
diff --git a/Engine/Source/Editor/UnrealEd/Public/BSPOps.h b/Engine/Source/Editor/UnrealEd/Public/BSPOps.h
--- a/Engine/Source/Editor/UnrealEd/Public/BSPOps.h
+++ b/Engine/Source/Editor/UnrealEd/Public/BSPOps.h
@@ -1,8 +1,6 @@
 // Copyright 1998-2017 Epic Games, Inc. All Rights Reserved.
 
-
-#ifndef __BSPOPS_H__
-#define __BSPOPS_H__
+#pragma once
 
 #include "CoreMinimal.h"
 #include "Engine/Brush.h"
@@ -153,7 +151,3 @@ private:
 	typedef TMap<FBspPointsKey, FBspPointsGridItem> FGridMap;
 	FGridMap GridMap;
 };
-
-
-
-#endif // __BSPOPS_H__
diff --git a/Engine/Source/Editor/UnrealEd/Public/BusyCursor.h b/Engine/Source/Editor/UnrealEd/Public/BusyCursor.h
--- a/Engine/Source/Editor/UnrealEd/Public/BusyCursor.h
+++ b/Engine/Source/Editor/UnrealEd/Public/BusyCursor.h
@@ -1,8 +1,6 @@
 // Copyright 1998-2017 Epic Games, Inc. All Rights Reserved.
 
-
-#ifndef __BUSYCURSOR_H__
-#define __BUSYCURSOR_H__
+#pragma once
 
 #include "CoreMinimal.h"
 
@@ -15,5 +13,3 @@ public:
 	FScopedBusyCursor();
 	~FScopedBusyCursor();
 };
-
-#endif // __BUSYCURSOR_H__
diff --git a/Engine/Source/Editor/UnrealEd/Public/CameraController.h b/Engine/Source/Editor/UnrealEd/Public/CameraController.h
--- a/Engine/Source/Editor/UnrealEd/Public/CameraController.h
+++ b/Engine/Source/Editor/UnrealEd/Public/CameraController.h
@@ -6,9 +6,6 @@
 
 #pragma once
 
-#ifndef __CameraController_h__
-#define __CameraController_h__
-
 #include "CoreMinimal.h"
 
 /**
@@ -344,8 +341,3 @@ private:
 	float OriginalFOVForRecoil;
 
 };
-
-
-
-
-#endif // __CameraController_h__
diff --git a/Engine/Source/Editor/UnrealEd/Public/Commandlets/EditorCommandlets.h b/Engine/Source/Editor/UnrealEd/Public/Commandlets/EditorCommandlets.h
--- a/Engine/Source/Editor/UnrealEd/Public/Commandlets/EditorCommandlets.h
+++ b/Engine/Source/Editor/UnrealEd/Public/Commandlets/EditorCommandlets.h
@@ -4,8 +4,7 @@
 	Editor.h: Unreal editor public header file.
 =============================================================================*/
 
-#ifndef _INC_EDITOR_COMMANDLETS
-#define _INC_EDITOR_COMMANDLETS
+#pragma once
 
 #include "CoreMinimal.h"
 #include "UObject/Object.h"
@@ -269,5 +268,3 @@ struct FNativePropertyData
 		return PropertyData.Num() || PropertyText.Num();
 	}
 };
-
-#endif
diff --git a/Engine/Source/Editor/UnrealEd/Public/ConsolidateWindow.h b/Engine/Source/Editor/UnrealEd/Public/ConsolidateWindow.h
--- a/Engine/Source/Editor/UnrealEd/Public/ConsolidateWindow.h
+++ b/Engine/Source/Editor/UnrealEd/Public/ConsolidateWindow.h
@@ -4,8 +4,7 @@
 	ConsolidationWindow.h: Dialog for displaying the asset consolidation tool.
 =============================================================================*/
 
-#ifndef __SCONSOLIDATEWINDOW_H__
-#define __SCONSOLIDATEWINDOW_H__
+#pragma once
 
 #include "CoreMinimal.h"
 
@@ -39,5 +38,3 @@ private:
 	 */
 	static TWeakPtr<SConsolidateToolWidget> WidgetInstance;
 };
-
-#endif
diff --git a/Engine/Source/Editor/UnrealEd/Public/DataTableEditorUtils.h b/Engine/Source/Editor/UnrealEd/Public/DataTableEditorUtils.h
--- a/Engine/Source/Editor/UnrealEd/Public/DataTableEditorUtils.h
+++ b/Engine/Source/Editor/UnrealEd/Public/DataTableEditorUtils.h
@@ -68,6 +68,8 @@ struct UNREALED_API FDataTableEditorUtils
 	static uint8* AddRow(UDataTable* DataTable, FName RowName);
 	static bool RenameRow(UDataTable* DataTable, FName OldName, FName NewName);
 	static bool MoveRow(UDataTable* DataTable, FName RowName, ERowMoveDirection Direction, int32 NumRowsToMoveBy = 1);
+	static bool DiffersFromDefault(UDataTable* DataTable, FName RowName);
+	static bool ResetToDefault(UDataTable* DataTable, FName RowName);
 
 	static void BroadcastPreChange(UDataTable* DataTable, EDataTableChangeInfo Info);
 	static void BroadcastPostChange(UDataTable* DataTable, EDataTableChangeInfo Info);
diff --git a/Engine/Source/Editor/UnrealEd/Public/DebugToolExec.h b/Engine/Source/Editor/UnrealEd/Public/DebugToolExec.h
--- a/Engine/Source/Editor/UnrealEd/Public/DebugToolExec.h
+++ b/Engine/Source/Editor/UnrealEd/Public/DebugToolExec.h
@@ -4,8 +4,7 @@
 	DebugToolExec.h: Game debug tool implementation.
 =============================================================================*/ 
 
-#ifndef __DEBUGTOOLEXEC_H__
-#define __DEBUGTOOLEXEC_H__
+#pragma once
 
 #include "CoreMinimal.h"
 
@@ -31,5 +30,3 @@ public:
 	//~ End Exec Interface
 
 };
-
-#endif
diff --git a/Engine/Source/Editor/UnrealEd/Public/Dialogs/Dialogs.h b/Engine/Source/Editor/UnrealEd/Public/Dialogs/Dialogs.h
--- a/Engine/Source/Editor/UnrealEd/Public/Dialogs/Dialogs.h
+++ b/Engine/Source/Editor/UnrealEd/Public/Dialogs/Dialogs.h
@@ -1,8 +1,6 @@
 // Copyright 1998-2017 Epic Games, Inc. All Rights Reserved.
 
-
-#ifndef __DIALOGS_H__
-#define __DIALOGS_H__
+#pragma once
 
 #include "CoreMinimal.h"
 #include "Input/Reply.h"
@@ -194,6 +192,3 @@ UNREALED_API bool PromptUserIfExistingObject(const FString& Name, const FString&
 * @return					true if the user selected the OK button, false otherwise.
 */
 UNREALED_API bool PromptUserForDirectory(FString& OutDirectory, const FString& Message, const FString& DefaultPath);
-
-
-#endif // __DIALOGS_H__
diff --git a/Engine/Source/Editor/UnrealEd/Public/Dialogs/DlgMoveAssets.h b/Engine/Source/Editor/UnrealEd/Public/Dialogs/DlgMoveAssets.h
--- a/Engine/Source/Editor/UnrealEd/Public/Dialogs/DlgMoveAssets.h
+++ b/Engine/Source/Editor/UnrealEd/Public/Dialogs/DlgMoveAssets.h
@@ -4,8 +4,7 @@
 	DlgMoveAssets.h: UnrealEd dialog for moving assets.
 =============================================================================*/
 
-#ifndef __DLGMOVEASSETS_H__
-#define __DLGMOVEASSETS_H__
+#pragma once
 
 #include "CoreMinimal.h"
 
@@ -64,5 +63,3 @@ private:
 	/** Cached pointer to the message box held within the window */
 	TSharedPtr<class SDlgMoveAsset> MoveAssetWidget;
 };
-
-#endif // __DLGMOVEASSETS_H__
diff --git a/Engine/Source/Editor/UnrealEd/Public/Dialogs/DlgReferenceTree.h b/Engine/Source/Editor/UnrealEd/Public/Dialogs/DlgReferenceTree.h
--- a/Engine/Source/Editor/UnrealEd/Public/Dialogs/DlgReferenceTree.h
+++ b/Engine/Source/Editor/UnrealEd/Public/Dialogs/DlgReferenceTree.h
@@ -1,8 +1,6 @@
 // Copyright 1998-2017 Epic Games, Inc. All Rights Reserved.
 
-
-#ifndef __DLGREFERENCETREE_H__
-#define __DLGREFERENCETREE_H__
+#pragma once
 
 #include "CoreMinimal.h"
 #include "Serialization/ArchiveUObject.h"
@@ -225,5 +223,3 @@ private:
 	/** Pointer to the containing window. */
 	TWeakPtr< SWindow > MyWindow;
 };
-
-#endif
diff --git a/Engine/Source/Editor/UnrealEd/Public/Dialogs/SBuildProgress.h b/Engine/Source/Editor/UnrealEd/Public/Dialogs/SBuildProgress.h
--- a/Engine/Source/Editor/UnrealEd/Public/Dialogs/SBuildProgress.h
+++ b/Engine/Source/Editor/UnrealEd/Public/Dialogs/SBuildProgress.h
@@ -4,8 +4,7 @@
 	DlgBuildProgress.h: UnrealEd dialog for displaying map build progress and cancelling builds.
 =============================================================================*/
 
-#ifndef __DLGBUILDPROGRESS_H__
-#define __DLGBUILDPROGRESS_H__
+#pragma once
 
 #include "CoreMinimal.h"
 #include "Input/Reply.h"
@@ -130,5 +129,3 @@ private:
 	int32				ErrorCount;
 	int32				CriticalErrorCount;
 };
-
-#endif // __DLGBUILDPROGRESS_H__
diff --git a/Engine/Source/Editor/UnrealEd/Public/DragTool_Measure.h b/Engine/Source/Editor/UnrealEd/Public/DragTool_Measure.h
--- a/Engine/Source/Editor/UnrealEd/Public/DragTool_Measure.h
+++ b/Engine/Source/Editor/UnrealEd/Public/DragTool_Measure.h
@@ -1,8 +1,6 @@
 // Copyright 1998-2017 Epic Games, Inc. All Rights Reserved.
 
-
-#ifndef __DRAGTOOL_MEASURE_H__
-#define __DRAGTOOL_MEASURE_H__
+#pragma once
 
 #include "CoreMinimal.h"
 #include "EditorDragTools.h"
@@ -49,5 +47,3 @@ private:
 
 	FEditorViewportClient* ViewportClient;
 };
-
-#endif // __DRAGTOOL_MEASURE_H__
diff --git a/Engine/Source/Editor/UnrealEd/Public/DragTool_ViewportChange.h b/Engine/Source/Editor/UnrealEd/Public/DragTool_ViewportChange.h
--- a/Engine/Source/Editor/UnrealEd/Public/DragTool_ViewportChange.h
+++ b/Engine/Source/Editor/UnrealEd/Public/DragTool_ViewportChange.h
@@ -1,8 +1,6 @@
 // Copyright 1998-2017 Epic Games, Inc. All Rights Reserved.
 
-
-#ifndef __DRAGTOOL_VIEWPORTCHANGE_H__
-#define __DRAGTOOL_VIEWPORTCHANGE_H__
+#pragma once
 
 #include "CoreMinimal.h"
 #include "Editor/UnrealEdTypes.h"
@@ -51,5 +49,3 @@ private:
 
 	FVector2D ViewOptionOffset;
 };
-
-#endif // __DRAGTOOL_VIEWPORTCHANGE_H__
diff --git a/Engine/Source/Editor/UnrealEd/Public/EditorBuildUtils.h b/Engine/Source/Editor/UnrealEd/Public/EditorBuildUtils.h
--- a/Engine/Source/Editor/UnrealEd/Public/EditorBuildUtils.h
+++ b/Engine/Source/Editor/UnrealEd/Public/EditorBuildUtils.h
@@ -4,8 +4,7 @@
 	EditorBuildUtils.h: Utilities for building in the editor
 =============================================================================*/
 
-#ifndef __EditorBuildUtils_h__
-#define __EditorBuildUtils_h__
+#pragma once
 
 #include "CoreMinimal.h"
 #include "Engine/EngineBaseTypes.h"
@@ -272,5 +271,3 @@ private:
 	/** Set to a valid build type if an async build is in progress. */
 	static FName InProgressBuildId;
 };
-
-#endif // #ifndef __EditorBuildUtils_h__
diff --git a/Engine/Source/Editor/UnrealEd/Public/EditorLevelUtils.h b/Engine/Source/Editor/UnrealEd/Public/EditorLevelUtils.h
--- a/Engine/Source/Editor/UnrealEd/Public/EditorLevelUtils.h
+++ b/Engine/Source/Editor/UnrealEd/Public/EditorLevelUtils.h
@@ -6,9 +6,6 @@
 
 #pragma once
 
-#ifndef __EditorLevelUtils_h__
-#define __EditorLevelUtils_h__
-
 #include "CoreMinimal.h"
 
 class AActor;
@@ -97,15 +94,6 @@ namespace EditorLevelUtils
 	 * @param	bForceLayersVisible		If true and the level is visible, force the level's layers to be visible.
 	 */
 	UNREALED_API void SetLevelVisibility(ULevel* Level, bool bShouldBeVisible, bool bForceLayersVisible);
-	
-	/**
-	 * Returns a class that represents the required streaming type (User selected via a dialog)
-	 *
-	 * @param	LevelPackageName		The name of the level (its displayed in the dialog)
-	 *
-	 * @returns	The streaming class or NULL if the user clicked cancel
-	 */
-	UClass* SelectLevelStreamType( const TCHAR* LevelPackageName );
 
 	/**
 	 * Removes a level from the world.  Returns true if the level was removed successfully.
@@ -151,8 +139,3 @@ namespace EditorLevelUtils
 	UNREALED_API void GetWorlds(UWorld* InWorld, TArray<UWorld*>& OutWorlds, bool bIncludeInWorld, bool bOnlyEditorVisible = false);
 
 }
-
-
-
-
-#endif	// __EditorLevelUtils_h__
diff --git a/Engine/Source/Editor/UnrealEd/Public/EditorReimportHandler.h b/Engine/Source/Editor/UnrealEd/Public/EditorReimportHandler.h
--- a/Engine/Source/Editor/UnrealEd/Public/EditorReimportHandler.h
+++ b/Engine/Source/Editor/UnrealEd/Public/EditorReimportHandler.h
@@ -50,11 +50,11 @@ public:
 	 * @param	Obj	Object to try reimporting
 	 * @param	bAskForNewFileIfMissing If the file is missing, open a dialog to ask for a new one
 	 * @param	bShowNotification True to show a notification when complete, false otherwise
-	 * @param	PreferedReimportFile if not empty, will be use in case the original file is missing and bAskForNewFileIfMissing is set to false
+	 * @param	PreferredReimportFile if not empty, will be use in case the original file is missing and bAskForNewFileIfMissing is set to false
 	 *
 	 * @return	true if the object was handled by one of the reimport handlers; false otherwise
 	 */
-	UNREALED_API virtual bool Reimport( UObject* Obj, bool bAskForNewFileIfMissing = false, bool bShowNotification = true, FString PreferedReimportFile = TEXT(""), FReimportHandler* SpecifiedReimportHandler = nullptr );
+	UNREALED_API virtual bool Reimport( UObject* Obj, bool bAskForNewFileIfMissing = false, bool bShowNotification = true, FString PreferredReimportFile = TEXT(""), FReimportHandler* SpecifiedReimportHandler = nullptr );
 
 	/**
 	 * Attemp to reimport all specified objects. This function will verify that all source file exist and ask the user
@@ -75,11 +75,11 @@ public:
 	* @param	Objects	Objects to try reimporting
 	* @param	bAskForNewFileIfMissing If the file is missing, open a dialog to ask for a new one
 	* @param	bShowNotification True to show a notification when complete, false otherwise
-	* @param	PreferedReimportFile if not empty, will be use in case the original file is missing and bAskForNewFileIfMissing is set to false
+	* @param	PreferredReimportFile if not empty, will be use in case the original file is missing and bAskForNewFileIfMissing is set to false
 	*
 	* @return	true if the objects all imported successfully, for more granular success reporting use FReimportManager::Reimport
 	*/
-	UNREALED_API virtual bool ReimportMultiple( TArrayView<UObject*> Objects, bool bAskForNewFileIfMissing = false, bool bShowNotification = true, FString PreferedReimportFile = TEXT(""), FReimportHandler* SpecifiedReimportHandler = nullptr );
+	UNREALED_API virtual bool ReimportMultiple( TArrayView<UObject*> Objects, bool bAskForNewFileIfMissing = false, bool bShowNotification = true, FString PreferredReimportFile = TEXT(""), FReimportHandler* SpecifiedReimportHandler = nullptr );
 
 	/**
 	 * Update the reimport paths for the specified object
diff --git a/Engine/Source/Editor/UnrealEd/Public/EditorShowFlags.h b/Engine/Source/Editor/UnrealEd/Public/EditorShowFlags.h
--- a/Engine/Source/Editor/UnrealEd/Public/EditorShowFlags.h
+++ b/Engine/Source/Editor/UnrealEd/Public/EditorShowFlags.h
@@ -1,8 +1,6 @@
 // Copyright 1998-2017 Epic Games, Inc. All Rights Reserved.
 
-
-#ifndef _EditorShowFlags_h_
-#define _EditorShowFlags_h_
+#pragma once
 
 #include "CoreMinimal.h"
 #include "Framework/Commands/InputChord.h"
@@ -56,5 +54,3 @@ public:
 };
 
 UNREALED_API TArray<FShowFlagData>& GetShowFlagMenuItems();
-
-#endif
diff --git a/Engine/Source/Editor/UnrealEd/Public/Factories.h b/Engine/Source/Editor/UnrealEd/Public/Factories.h
--- a/Engine/Source/Editor/UnrealEd/Public/Factories.h
+++ b/Engine/Source/Editor/UnrealEd/Public/Factories.h
@@ -4,8 +4,7 @@
 	Factories.h: Unreal Engine factory types.
 =============================================================================*/
 
-#ifndef __FACTORIES_H__
-#define __FACTORIES_H__
+#pragma once
 
 #include "CoreMinimal.h"
 #include "UObject/Class.h"
@@ -148,5 +147,3 @@ UNREALED_API bool DecompressTGA_helper(
 	uint32*& TextureData,
 	const int32 TextureDataSize,
 	FFeedbackContext* Warn );
-
-#endif // __FACTORIES_H__
diff --git a/Engine/Source/Editor/UnrealEd/Public/FbxAnimUtils.h b/Engine/Source/Editor/UnrealEd/Public/FbxAnimUtils.h
--- a/Engine/Source/Editor/UnrealEd/Public/FbxAnimUtils.h
+++ b/Engine/Source/Editor/UnrealEd/Public/FbxAnimUtils.h
@@ -1,8 +1,6 @@
 // Copyright 1998-2017 Epic Games, Inc. All Rights Reserved.
 
-
-#ifndef __FbxAnimUtils_h__
-#define __FbxAnimUtils_h__
+#pragma once
 
 #include "CoreMinimal.h"
 
@@ -26,5 +24,3 @@ namespace FbxAnimUtils
 	UNREALED_API bool ImportCurveTableFromNode(const FString& InFbxFilename, const FString& InCurveNodeName, UCurveTable* InOutCurveTable, float& OutPreRoll);
 
 } // namespace FbxAnimUtils
-
-#endif //__FbxAnimUtils_h__
diff --git a/Engine/Source/Editor/UnrealEd/Public/Kismet2/DebuggerCommands.h b/Engine/Source/Editor/UnrealEd/Public/Kismet2/DebuggerCommands.h
--- a/Engine/Source/Editor/UnrealEd/Public/Kismet2/DebuggerCommands.h
+++ b/Engine/Source/Editor/UnrealEd/Public/Kismet2/DebuggerCommands.h
@@ -1,7 +1,6 @@
 // Copyright 1998-2017 Epic Games, Inc. All Rights Reserved.
 
-#ifndef __DebuggerCommands_h__
-#define __DebuggerCommands_h__
+#pragma once
 
 #include "CoreMinimal.h"
 #include "Widgets/DeclarativeSyntaxSupport.h"
@@ -150,5 +149,3 @@ public:
 	static bool IsInSIE_AndRunning();
 	static bool IsInPIE_AndRunning();
 };
-
-#endif // __DebuggerCommands_h__
diff --git a/Engine/Source/Editor/UnrealEd/Public/Kismet2/EnumEditorUtils.h b/Engine/Source/Editor/UnrealEd/Public/Kismet2/EnumEditorUtils.h
--- a/Engine/Source/Editor/UnrealEd/Public/Kismet2/EnumEditorUtils.h
+++ b/Engine/Source/Editor/UnrealEd/Public/Kismet2/EnumEditorUtils.h
@@ -1,8 +1,5 @@
 // Copyright 1998-2017 Epic Games, Inc. All Rights Reserved.
 
-#ifndef __EnumEditorUtils_h__
-#define __EnumEditorUtils_h__
-
 #pragma once
 
 #include "CoreMinimal.h"
@@ -84,5 +81,3 @@ public:
 	static void EnsureAllDisplayNamesExist(class UUserDefinedEnum* Enum);
 	static void UpgradeDisplayNamesFromMetaData(class UUserDefinedEnum* Enum);
 };
-
-#endif // __EnumEditorUtils_h__
diff --git a/Engine/Source/Editor/UnrealEd/Public/Kismet2/KismetDebugUtilities.h b/Engine/Source/Editor/UnrealEd/Public/Kismet2/KismetDebugUtilities.h
--- a/Engine/Source/Editor/UnrealEd/Public/Kismet2/KismetDebugUtilities.h
+++ b/Engine/Source/Editor/UnrealEd/Public/Kismet2/KismetDebugUtilities.h
@@ -1,9 +1,5 @@
 // Copyright 1998-2017 Epic Games, Inc. All Rights Reserved.
 
-
-#ifndef __KismetDebugUtilities_h__
-#define __KismetDebugUtilities_h__
-
 #pragma once
 
 #include "CoreMinimal.h"
@@ -235,5 +231,3 @@ private:
 };
 
 //////////////////////////////////////////////////////////////////////////
-
-#endif
diff --git a/Engine/Source/Editor/UnrealEd/Public/Kismet2/KismetEditorUtilities.h b/Engine/Source/Editor/UnrealEd/Public/Kismet2/KismetEditorUtilities.h
--- a/Engine/Source/Editor/UnrealEd/Public/Kismet2/KismetEditorUtilities.h
+++ b/Engine/Source/Editor/UnrealEd/Public/Kismet2/KismetEditorUtilities.h
@@ -1,9 +1,5 @@
 // Copyright 1998-2017 Epic Games, Inc. All Rights Reserved.
 
-
-#ifndef __KismetEditorUtilities_h__
-#define __KismetEditorUtilities_h__
-
 #pragma once
 
 #include "CoreMinimal.h"
@@ -314,7 +310,3 @@ private:
 
 	FKismetEditorUtilities() {}
 };
-
-
-
-#endif//__KismetEditorUtilities_h__
diff --git a/Engine/Source/Editor/UnrealEd/Public/Kismet2/KismetReinstanceUtilities.h b/Engine/Source/Editor/UnrealEd/Public/Kismet2/KismetReinstanceUtilities.h
--- a/Engine/Source/Editor/UnrealEd/Public/Kismet2/KismetReinstanceUtilities.h
+++ b/Engine/Source/Editor/UnrealEd/Public/Kismet2/KismetReinstanceUtilities.h
@@ -1,9 +1,5 @@
 // Copyright 1998-2017 Epic Games, Inc. All Rights Reserved.
 
-
-#ifndef __KismetReinstanceUtilities_h__
-#define __KismetReinstanceUtilities_h__
-
 #pragma once
 
 #include "CoreMinimal.h"
@@ -201,6 +197,3 @@ private:
 	/** Handles the work of ReplaceInstancesOfClass, handling both normal replacement of instances and batch */
 	static void ReplaceInstancesOfClass_Inner(TMap<UClass*, UClass*>& InOldToNewClassMap, UObject* InOriginalCDO, TSet<UObject*>* ObjectsThatShouldUseOldStuff = NULL, bool bClassObjectReplaced = false, bool bPreserveRootComponent = true);
 };
-
-
-#endif//__KismetEditorUtilities_h__
diff --git a/Engine/Source/Editor/UnrealEd/Public/Kismet2/StructureEditorUtils.h b/Engine/Source/Editor/UnrealEd/Public/Kismet2/StructureEditorUtils.h
--- a/Engine/Source/Editor/UnrealEd/Public/Kismet2/StructureEditorUtils.h
+++ b/Engine/Source/Editor/UnrealEd/Public/Kismet2/StructureEditorUtils.h
@@ -169,6 +169,8 @@ public:
 
 	static void RecreateDefaultInstanceInEditorData(UUserDefinedStruct* Struct);
 
+	static bool DiffersFromDefaultValue(const UUserDefinedStruct* Struct, uint8* StructData);
+
 	//VALIDATION
 	static bool CanHaveAMemberVariableOfType(const UUserDefinedStruct* Struct, const FEdGraphPinType& VarType, FString* OutMsg = NULL);
 
diff --git a/Engine/Source/Editor/UnrealEd/Public/MRUFavoritesList.h b/Engine/Source/Editor/UnrealEd/Public/MRUFavoritesList.h
--- a/Engine/Source/Editor/UnrealEd/Public/MRUFavoritesList.h
+++ b/Engine/Source/Editor/UnrealEd/Public/MRUFavoritesList.h
@@ -4,8 +4,8 @@
 	MRUFavoritesList : Helper class for handling MRU and favorited maps
 
 =============================================================================*/
-#ifndef __MAINMRUFAVORITESLIST_H__
-#define __MAINMRUFAVORITESLIST_H__
+
+#pragma once
 
 #include "CoreMinimal.h"
 #include "MRUList.h"
@@ -93,5 +93,3 @@ private:
 	/** INI section to read/write favorite items to */
 	static const FString FAVORITES_INI_SECTION;
 };
-
-#endif
diff --git a/Engine/Source/Editor/UnrealEd/Public/MRUList.h b/Engine/Source/Editor/UnrealEd/Public/MRUList.h
--- a/Engine/Source/Editor/UnrealEd/Public/MRUList.h
+++ b/Engine/Source/Editor/UnrealEd/Public/MRUList.h
@@ -5,8 +5,7 @@
 
 =============================================================================*/
 
-#ifndef __MRULIST_H__
-#define __MRULIST_H__
+#pragma once
 
 #include "CoreMinimal.h"
 
@@ -129,5 +128,3 @@ protected:
 	/** The INI section we read/write the filenames to. */
 	FString INISection;
 };
-
-#endif // #ifndef __MRULIST_H__
diff --git a/Engine/Source/Editor/UnrealEd/Public/MatineeExporter.h b/Engine/Source/Editor/UnrealEd/Public/MatineeExporter.h
--- a/Engine/Source/Editor/UnrealEd/Public/MatineeExporter.h
+++ b/Engine/Source/Editor/UnrealEd/Public/MatineeExporter.h
@@ -4,8 +4,7 @@
 	Matinee exporter for Unreal Engine 3.
 =============================================================================*/
 
-#ifndef __MATINEEEXPORTER_H__
-#define __MATINEEEXPORTER_H__
+#pragma once
 
 #include "CoreMinimal.h"
 #include "GameFramework/Actor.h"
@@ -97,5 +96,3 @@ protected:
 	/** When true, we'll export with hierarchical relation of attachment with relative transform */
 	bool bKeepHierarchy;
 };
-
-#endif // __MATINEEEXPORTER_H__
diff --git a/Engine/Source/Editor/UnrealEd/Public/MeshPaintRendering.h b/Engine/Source/Editor/UnrealEd/Public/MeshPaintRendering.h
--- a/Engine/Source/Editor/UnrealEd/Public/MeshPaintRendering.h
+++ b/Engine/Source/Editor/UnrealEd/Public/MeshPaintRendering.h
@@ -4,9 +4,6 @@
 	MeshPaintRendering.h: Mesh texture paint brush rendering
 ================================================================================*/
 
-#ifndef __MeshPaintRendering_h__
-#define __MeshPaintRendering_h__
-
 #pragma once
 
 #include "CoreMinimal.h"
@@ -75,7 +72,3 @@ namespace MeshPaintRendering
 													const FMeshPaintDilateShaderParameters& InShaderParams );
 
 }
-
-
-
-#endif	// __MeshPaintRendering_h__
diff --git a/Engine/Source/Editor/UnrealEd/Public/MouseDeltaTracker.h b/Engine/Source/Editor/UnrealEd/Public/MouseDeltaTracker.h
--- a/Engine/Source/Editor/UnrealEd/Public/MouseDeltaTracker.h
+++ b/Engine/Source/Editor/UnrealEd/Public/MouseDeltaTracker.h
@@ -1,8 +1,6 @@
 // Copyright 1998-2017 Epic Games, Inc. All Rights Reserved.
 
-
-#ifndef __MOUSEDELTATRACKER_H__
-#define __MOUSEDELTATRACKER_H__
+#pragma once
 
 #include "CoreMinimal.h"
 #include "InputCoreTypes.h"
@@ -192,5 +190,3 @@ private:
 };
 
 static const float MOUSE_CLICK_DRAG_DELTA = 4.0f*4.0f;
-
-#endif // __MOUSEDELTATRACKER_H__
diff --git a/Engine/Source/Editor/UnrealEd/Public/PackageBackup.h b/Engine/Source/Editor/UnrealEd/Public/PackageBackup.h
--- a/Engine/Source/Editor/UnrealEd/Public/PackageBackup.h
+++ b/Engine/Source/Editor/UnrealEd/Public/PackageBackup.h
@@ -4,8 +4,7 @@
 	PackageBackup.h: Utility class for backing up a package.
 =============================================================================*/
 
-#ifndef _PACKAGEBACKUP_H
-#define _PACKAGEBACKUP_H
+#pragma once
 
 #include "CoreMinimal.h"
 
@@ -112,5 +111,3 @@ private:
 	~FAutoPackageBackup();
 	FAutoPackageBackup& operator=( const FAutoPackageBackup& );
 };
-
-#endif // #ifndef _PACKAGEBACKUP_H
diff --git a/Engine/Source/Editor/UnrealEd/Public/PackageUtilityWorkers.h b/Engine/Source/Editor/UnrealEd/Public/PackageUtilityWorkers.h
--- a/Engine/Source/Editor/UnrealEd/Public/PackageUtilityWorkers.h
+++ b/Engine/Source/Editor/UnrealEd/Public/PackageUtilityWorkers.h
@@ -4,8 +4,7 @@
 	PackageUtilityWorkers.cpp: Declarations for structs and classes used by package commandlets.
 =============================================================================*/
 
-#ifndef __PACKAGEUTILITYWORKERS_H__
-#define __PACKAGEUTILITYWORKERS_H__
+#pragma once
 
 #include "CoreMinimal.h"
 
@@ -102,16 +101,3 @@ struct FPkgInfoReporter_Log : public FPkgInfoReporter
 	 */
 	virtual void GeneratePackageReport( class FLinkerLoad* InLinker=NULL );
 };
-
-#endif
-
-
-//EOF
-
-
-
-
-
-
-
-
diff --git a/Engine/Source/Editor/UnrealEd/Public/ReferencedAssetsUtils.h b/Engine/Source/Editor/UnrealEd/Public/ReferencedAssetsUtils.h
--- a/Engine/Source/Editor/UnrealEd/Public/ReferencedAssetsUtils.h
+++ b/Engine/Source/Editor/UnrealEd/Public/ReferencedAssetsUtils.h
@@ -1,8 +1,6 @@
 // Copyright 1998-2017 Epic Games, Inc. All Rights Reserved.
 
-
-#ifndef __REFERENCEDASSETSBROWSER_H__
-#define __REFERENCEDASSETSBROWSER_H__
+#pragma once
 
 #include "CoreMinimal.h"
 #include "Serialization/ArchiveUObject.h"
@@ -191,5 +189,3 @@ protected:
 	/** Handle to the registered OnEditorMapChange delegate. */
 	FDelegateHandle OnEditorMapChangeDelegateHandle;
 };
-
-#endif // __REFERENCEDASSETSBROWSER_H__
diff --git a/Engine/Source/Editor/UnrealEd/Public/SCSVImportOptions.h b/Engine/Source/Editor/UnrealEd/Public/SCSVImportOptions.h
--- a/Engine/Source/Editor/UnrealEd/Public/SCSVImportOptions.h
+++ b/Engine/Source/Editor/UnrealEd/Public/SCSVImportOptions.h
@@ -24,16 +24,18 @@ private:
 public:
 	SLATE_BEGIN_ARGS(SCSVImportOptions)
 		: _WidgetWindow()
-	{}
+		, _FullPath()
+		{}
 
-	SLATE_ARGUMENT(TSharedPtr<SWindow>, WidgetWindow)
-		SLATE_END_ARGS()
+		SLATE_ARGUMENT(TSharedPtr<SWindow>, WidgetWindow)
+		SLATE_ARGUMENT(FText, FullPath)
+	SLATE_END_ARGS()
 
-		SCSVImportOptions()
+	SCSVImportOptions()
 		: bImport(false)
 		, SelectedImportType(ECSVImportType::ECSV_DataTable)
 		, SelectedStruct(NULL)
-	{}
+		{}
 
 	void Construct(const FArguments& InArgs);
 
diff --git a/Engine/Source/Editor/UnrealEd/Public/SColorGradientEditor.h b/Engine/Source/Editor/UnrealEd/Public/SColorGradientEditor.h
--- a/Engine/Source/Editor/UnrealEd/Public/SColorGradientEditor.h
+++ b/Engine/Source/Editor/UnrealEd/Public/SColorGradientEditor.h
@@ -148,7 +148,7 @@ private:
 	 * @param DrawEffects		Any draw effects to apply
 	 * @param bColor			If true RGB of the Color param will be used, otherwise the A value will be used
 	 */
-	void DrawGradientStopMark( const FGradientStopMark& Mark, const FGeometry& Geometry, float XPos, const FLinearColor& Color, FSlateWindowElementList& OutDrawElements, int32 LayerId, const FSlateRect& InClippingRect, ESlateDrawEffect::Type DrawEffects, bool bColor, const FWidgetStyle& InWidgetStyle ) const;
+	void DrawGradientStopMark( const FGradientStopMark& Mark, const FGeometry& Geometry, float XPos, const FLinearColor& Color, FSlateWindowElementList& OutDrawElements, int32 LayerId, const FSlateRect& InClippingRect, ESlateDrawEffect DrawEffects, bool bColor, const FWidgetStyle& InWidgetStyle ) const;
 
 	/**
 	 * Calculates the geometry of the gradient stop color mark area
diff --git a/Engine/Source/Editor/UnrealEd/Public/SCurveEditor.h b/Engine/Source/Editor/UnrealEd/Public/SCurveEditor.h
--- a/Engine/Source/Editor/UnrealEd/Public/SCurveEditor.h
+++ b/Engine/Source/Editor/UnrealEd/Public/SCurveEditor.h
@@ -1,8 +1,6 @@
 // Copyright 1998-2017 Epic Games, Inc. All Rights Reserved.
 
-
-#ifndef __SCurveEditor_h__
-#define __SCurveEditor_h__
+#pragma once
 
 #include "CoreMinimal.h"
 #include "Misc/Attribute.h"
@@ -463,17 +461,17 @@ private:
 
 	/** Paint a curve */
 	void PaintCurve(TSharedPtr<FCurveViewModel> CurveViewModel, const FGeometry &AllottedGeometry, FTrackScaleInfo &ScaleInfo, FSlateWindowElementList &OutDrawElements, 
-					int32 LayerId, const FSlateRect& MyClippingRect, ESlateDrawEffect::Type DrawEffects, const FWidgetStyle &InWidgetStyle, bool bAnyCurveViewModelsSelected) const;
+					int32 LayerId, const FSlateRect& MyClippingRect, ESlateDrawEffect DrawEffects, const FWidgetStyle &InWidgetStyle, bool bAnyCurveViewModelsSelected) const;
 
 	/** Paint the keys that make up a curve */
-	void PaintKeys(TSharedPtr<FCurveViewModel> CurveViewModel, FTrackScaleInfo &ScaleInfo, FSlateWindowElementList &OutDrawElements, int32 LayerId, int32 SelectedLayerId, const FGeometry &AllottedGeometry, const FSlateRect& MyClippingRect, ESlateDrawEffect::Type DrawEffects, const FWidgetStyle &InWidgetStyle, bool bAnyCurveViewModelsSelected ) const;
+	void PaintKeys(TSharedPtr<FCurveViewModel> CurveViewModel, FTrackScaleInfo &ScaleInfo, FSlateWindowElementList &OutDrawElements, int32 LayerId, int32 SelectedLayerId, const FGeometry &AllottedGeometry, const FSlateRect& MyClippingRect, ESlateDrawEffect DrawEffects, const FWidgetStyle &InWidgetStyle, bool bAnyCurveViewModelsSelected ) const;
 
 	/** Paint the tangent for a key with cubic curves */
 	void PaintTangent( TSharedPtr<FCurveViewModel> CurveViewModel, FTrackScaleInfo &ScaleInfo, FRichCurve* Curve, FKeyHandle KeyHandle, FVector2D KeyLocation, FSlateWindowElementList &OutDrawElements, int32 LayerId, 
-					   const FGeometry &AllottedGeometry, const FSlateRect& MyClippingRect, ESlateDrawEffect::Type DrawEffects, int32 LayerToUse, const FWidgetStyle &InWidgetStyle, bool bTangentSelected, bool bIsArrivalSelected, bool bIsLeaveSelected, bool bAnyCurveViewModelsSelected ) const;
+					   const FGeometry &AllottedGeometry, const FSlateRect& MyClippingRect, ESlateDrawEffect DrawEffects, int32 LayerToUse, const FWidgetStyle &InWidgetStyle, bool bTangentSelected, bool bIsArrivalSelected, bool bIsLeaveSelected, bool bAnyCurveViewModelsSelected ) const;
 
 	/** Paint Grid lines, these make it easier to visualize relative distance */
-	void PaintGridLines( const FGeometry &AllottedGeometry, FTrackScaleInfo &ScaleInfo, FSlateWindowElementList &OutDrawElements, int32 LayerId, const FSlateRect& MyClippingRect, ESlateDrawEffect::Type DrawEffects )const;
+	void PaintGridLines( const FGeometry &AllottedGeometry, FTrackScaleInfo &ScaleInfo, FSlateWindowElementList &OutDrawElements, int32 LayerId, const FSlateRect& MyClippingRect, ESlateDrawEffect DrawEffects )const;
 
 	/** Paints the marquee for selection */
 	void PaintMarquee(const FGeometry& AllottedGeometry, const FSlateRect& MyClippingRect, FSlateWindowElementList& OutDrawElements, int32 LayerId) const;
@@ -812,5 +810,3 @@ protected:
 	/** Generic Popup Entry */
 	TWeakPtr<IMenu> EntryPopupMenu;
 };
-
-#endif // __SCurveEditor_h__
diff --git a/Engine/Source/Editor/UnrealEd/Public/SSkeletonWidget.h b/Engine/Source/Editor/UnrealEd/Public/SSkeletonWidget.h
--- a/Engine/Source/Editor/UnrealEd/Public/SSkeletonWidget.h
+++ b/Engine/Source/Editor/UnrealEd/Public/SSkeletonWidget.h
@@ -1,8 +1,6 @@
 // Copyright 1998-2017 Epic Games, Inc. All Rights Reserved.
 
-
-#ifndef __SSkeletonWidget_h__
-#define __SSkeletonWidget_h__
+#pragma once
 
 #include "CoreMinimal.h"
 #include "Input/Reply.h"
@@ -746,5 +744,3 @@ protected:
 };
 
 #undef LOCTEXT_NAMESPACE
-
-#endif // __SSkeletonWidget_h__
diff --git a/Engine/Source/Editor/UnrealEd/Public/STransformViewportToolbar.h b/Engine/Source/Editor/UnrealEd/Public/STransformViewportToolbar.h
--- a/Engine/Source/Editor/UnrealEd/Public/STransformViewportToolbar.h
+++ b/Engine/Source/Editor/UnrealEd/Public/STransformViewportToolbar.h
@@ -1,8 +1,6 @@
 // Copyright 1998-2017 Epic Games, Inc. All Rights Reserved.
 
-
-#ifndef __STransformViewportToolBar_h__
-#define __STransformViewportToolBar_h__
+#pragma once
 
 #include "CoreMinimal.h"
 #include "Layout/Visibility.h"
@@ -118,6 +116,3 @@ private:
 	/** Command list */
 	TSharedPtr<FUICommandList> CommandList;
 };
-
-#endif //__STransformViewportToolBar_h__
-
diff --git a/Engine/Source/Editor/UnrealEd/Public/ScopedTransaction.h b/Engine/Source/Editor/UnrealEd/Public/ScopedTransaction.h
--- a/Engine/Source/Editor/UnrealEd/Public/ScopedTransaction.h
+++ b/Engine/Source/Editor/UnrealEd/Public/ScopedTransaction.h
@@ -1,8 +1,6 @@
 // Copyright 1998-2017 Epic Games, Inc. All Rights Reserved.
 
-
-#ifndef __SCOPEDTRANSACTION_H__
-#define __SCOPEDTRANSACTION_H__
+#pragma once
 
 #include "CoreMinimal.h"
 
@@ -47,5 +45,3 @@ private:
 	/** Stores the transaction index, which is used to cancel the transaction. */
 	int32 Index;
 };
-
-#endif // __SCOPEDTRANSACTION_H__
diff --git a/Engine/Source/Editor/UnrealEd/Public/ScriptDisassembler.h b/Engine/Source/Editor/UnrealEd/Public/ScriptDisassembler.h
--- a/Engine/Source/Editor/UnrealEd/Public/ScriptDisassembler.h
+++ b/Engine/Source/Editor/UnrealEd/Public/ScriptDisassembler.h
@@ -4,8 +4,7 @@
 	ScriptDisassembler.h: Disassembler for Kismet bytecode.
 =============================================================================*/
 
-#ifndef __SCRIPTDISASSEMBLER_H__
-#define __SCRIPTDISASSEMBLER_H__
+#pragma once
 
 #include "CoreMinimal.h"
 #include "UObject/Script.h"
@@ -85,5 +84,3 @@ private:
 		return (T*)ReadQWORD(ScriptIndex);
 	}
 };
-
-#endif
diff --git a/Engine/Source/Editor/UnrealEd/Public/SourceCodeNavigation.h b/Engine/Source/Editor/UnrealEd/Public/SourceCodeNavigation.h
--- a/Engine/Source/Editor/UnrealEd/Public/SourceCodeNavigation.h
+++ b/Engine/Source/Editor/UnrealEd/Public/SourceCodeNavigation.h
@@ -1,9 +1,5 @@
 // Copyright 1998-2017 Epic Games, Inc. All Rights Reserved.
 
-
-#ifndef __SourceCodeNavigation_h__
-#define __SourceCodeNavigation_h__
-
 #pragma once
 
 #include "CoreMinimal.h"
@@ -216,7 +212,3 @@ private:
 	/** Source file database instance. */
 	static FSourceFileDatabase Instance;
 };
-
-
-
-#endif		// __SourceCodeNavigation_h__
diff --git a/Engine/Source/Editor/UnrealEd/Public/Tests/FbxAutomationCommon.h b/Engine/Source/Editor/UnrealEd/Public/Tests/FbxAutomationCommon.h
--- a/Engine/Source/Editor/UnrealEd/Public/Tests/FbxAutomationCommon.h
+++ b/Engine/Source/Editor/UnrealEd/Public/Tests/FbxAutomationCommon.h
@@ -23,6 +23,8 @@ enum EFBXExpectedResultPreset
 	Created_Skeletalmesh_Number,
 	/** Data should contain the number of Material created [int0] under the target content folder. */
 	Materials_Created_Number,
+	/** Data should be the slot index [int0], and the expected original imported material slot name [string0]. */
+	Material_Slot_Imported_Name,
 
 	/** Data should be the total number of vertex for all LOD [int0]. */
 	Vertex_Number,
@@ -45,6 +47,8 @@ enum EFBXExpectedResultPreset
 	Mesh_LOD_Section_Material_Name,
 	/** Data should be the LOD index [int0], section index [int1] and the expected material index of mesh materials [int2]. */
 	Mesh_LOD_Section_Material_Index,
+	/** Data should be the LOD index [int0], section index [int1] and the expected original imported material slot name [string0]. */
+	Mesh_LOD_Section_Material_Imported_Name,
 	
 	/** Data should be the LOD index [int0] and the number of UV channel [int1] for the specified LOD. */
 	LOD_UV_Channel_Number,
@@ -59,10 +63,16 @@ enum EFBXExpectedResultPreset
 UENUM()
 enum EFBXTestPlanActionType
 {
+	/*Normal import*/
 	Import,
+	/*Re-import the previous import, this is mandatory to make an import before*/
 	Reimport,
+	/*Add a new LOD*/
 	AddLOD,
+	/*Reimport an existing LOD*/
 	ReimportLOD,
+	/*The fbx will be imported, package will be save, object will be delete from memory then reload from the saved package. This mode force a delete of the asset after the test, no reimport is possible after*/
+	ImportReload,
 };
 
 /**
diff --git a/Engine/Source/Editor/UnrealEd/Public/Viewports.h b/Engine/Source/Editor/UnrealEd/Public/Viewports.h
--- a/Engine/Source/Editor/UnrealEd/Public/Viewports.h
+++ b/Engine/Source/Editor/UnrealEd/Public/Viewports.h
@@ -4,8 +4,7 @@
 	Viewports.h: The viewport windows used by the editor
 =============================================================================*/
 
-#ifndef __VIEWPORTS_H__
-#define __VIEWPORTS_H__
+#pragma once
 
 #include "CoreMinimal.h"
 
@@ -20,8 +19,3 @@ namespace EditorViewportDefs
 	/** Default camera field of view angle for level editor perspective viewports */
 	const float DefaultPerspectiveFOVAngle( 90.0f );
 }
-
-
-
-
-#endif // __VIEWPORTS_H__
diff --git a/Engine/Source/Editor/UnrealEd/UnrealEd.Build.cs b/Engine/Source/Editor/UnrealEd/UnrealEd.Build.cs
--- a/Engine/Source/Editor/UnrealEd/UnrealEd.Build.cs
+++ b/Engine/Source/Editor/UnrealEd/UnrealEd.Build.cs
@@ -189,7 +189,8 @@ public class UnrealEd : ModuleRules
 				"SettingsEditor",
 				"SessionFrontend",
 				"Sequencer",
-                "SuperSearch",
+				"StringTableEditor",
+				"SuperSearch",
                 "GeometryMode",
 				"TextureAlignMode",
 				"FoliageEdit",
diff --git a/Engine/Source/Editor/VREditor/VREditorStyle.cpp b/Engine/Source/Editor/VREditor/VREditorStyle.cpp
--- a/Engine/Source/Editor/VREditor/VREditorStyle.cpp
+++ b/Engine/Source/Editor/VREditor/VREditorStyle.cpp
@@ -49,6 +49,9 @@ const FVector2D Icon40x40(40.0f, 40.0f);
 TSharedRef< FSlateStyleSet > FVREditorStyle::Create()
 {
 	TSharedRef< FSlateStyleSet > Style = MakeShareable(new FSlateStyleSet(FVREditorStyle::GetStyleSetName()));
+	Style->SetContentRoot(FPaths::EngineContentDir() / TEXT("Editor/Slate"));
+	Style->SetCoreContentRoot(FPaths::EngineContentDir() / TEXT("Slate"));
+
 	// Use the default menu button style, but set the background to dark grey.
 	const FButtonStyle NormalButton = FEditorStyle::GetWidgetStyle<FButtonStyle>("Menu.Button");
 	Style->Set("VREditorStyle.Button", FButtonStyle(NormalButton)
@@ -56,13 +59,13 @@ TSharedRef< FSlateStyleSet > FVREditorStyle::Create()
 
 	const FTextBlockStyle NormalText = FEditorStyle::GetWidgetStyle<FTextBlockStyle>("NormalText");
 	Style->Set("VREditorStyle.Label", FTextBlockStyle(NormalText)
-		.SetFont(TTF_CORE_FONT("Slate/Fonts/Roboto-Regular", 9)));
+		.SetFont(TTF_CORE_FONT("Fonts/Roboto-Regular", 9)));
 
 	// Headings will have a font outline
 	FFontOutlineSettings HeadingOutline;
 	HeadingOutline.OutlineColor = FLinearColor(0.2f, 0.2f, 0.2f, 0.5f);
 	HeadingOutline.OutlineSize = 2.0f;
-	FSlateFontInfo HeadlineFont = TTF_CORE_FONT("Slate/Fonts/Roboto-Regular", 18);
+	FSlateFontInfo HeadlineFont = TTF_CORE_FONT("Fonts/Roboto-Regular", 18);
 	HeadlineFont.OutlineSettings = HeadingOutline;
 
 	Style->Set("VREditorStyle.Heading", FTextBlockStyle(NormalText)
@@ -73,7 +76,7 @@ TSharedRef< FSlateStyleSet > FVREditorStyle::Create()
 	FFontOutlineSettings HelperOutline;
 	HelperOutline.OutlineColor = FLinearColor( 0.2f, 0.2f, 0.2f, 0.5f );
 	HelperOutline.OutlineSize = 3.0f;
-	FSlateFontInfo HelperFont = TTF_CORE_FONT( "Slate/Fonts/Roboto-Regular", 24 );
+	FSlateFontInfo HelperFont = TTF_CORE_FONT( "Fonts/Roboto-Regular", 24 );
 	HelperFont.OutlineSettings = HelperOutline;
 
 	Style->Set( "VREditorStyle.HelperText", FTextBlockStyle( NormalText )
diff --git a/Engine/Source/Editor/WorldBrowser/Private/StreamingLevels/StreamingLevelCollectionModel.cpp b/Engine/Source/Editor/WorldBrowser/Private/StreamingLevels/StreamingLevelCollectionModel.cpp
--- a/Engine/Source/Editor/WorldBrowser/Private/StreamingLevels/StreamingLevelCollectionModel.cpp
+++ b/Engine/Source/Editor/WorldBrowser/Private/StreamingLevels/StreamingLevelCollectionModel.cpp
@@ -36,11 +36,14 @@ FStreamingLevelCollectionModel::FStreamingLevelCollectionModel()
 	{
 		AddedLevelStreamingClass = DefaultLevelStreamingClass;
 	}
+
+	FEditorDelegates::RefreshLevelBrowser.AddRaw(this, &FStreamingLevelCollectionModel::PopulateLevelsList);
 }
 
 FStreamingLevelCollectionModel::~FStreamingLevelCollectionModel()
 {
 	GEditor->UnregisterForUndo( this );
+	FEditorDelegates::RefreshLevelBrowser.RemoveAll(this);
 }
 
 void FStreamingLevelCollectionModel::Initialize(UWorld* InWorld)
diff --git a/Engine/Source/Editor/WorldBrowser/Private/Tiles/SWorldTileItem.cpp b/Engine/Source/Editor/WorldBrowser/Private/Tiles/SWorldTileItem.cpp
--- a/Engine/Source/Editor/WorldBrowser/Private/Tiles/SWorldTileItem.cpp
+++ b/Engine/Source/Editor/WorldBrowser/Private/Tiles/SWorldTileItem.cpp
@@ -26,14 +26,14 @@ int32 SWorldTileImage::OnPaint(const FPaintArgs& Args, const FGeometry& Allotted
 	if ((ImageBrush != nullptr) && (ImageBrush->DrawAs != ESlateBrushDrawType::NoDrawType))
 	{
 		const bool bIsEnabled = EditableTile.Get() && ShouldBeEnabled(bParentEnabled);
-		const uint32 DrawEffects = bIsEnabled ? 0 : ESlateDrawEffect::DisabledEffect;
+		const ESlateDrawEffect DrawEffects = bIsEnabled ? ESlateDrawEffect::None : ESlateDrawEffect::DisabledEffect;
 		FSlateDrawElement::MakeBox(
 			OutDrawElements, 
 			LayerId, 
 			AllottedGeometry.ToPaintGeometry(), 
 			ImageBrush, 
 			MyClippingRect, 
-			DrawEffects|ESlateDrawEffect::IgnoreTextureAlpha, 
+			DrawEffects | ESlateDrawEffect::IgnoreTextureAlpha, 
 			FColor::White);
 	}
 	return LayerId;
diff --git a/Engine/Source/Programs/CrashReporter/CrashReportClient/CrashReportClient.Target.cs b/Engine/Source/Programs/CrashReporter/CrashReportClient/CrashReportClient.Target.cs
--- a/Engine/Source/Programs/CrashReporter/CrashReportClient/CrashReportClient.Target.cs
+++ b/Engine/Source/Programs/CrashReporter/CrashReportClient/CrashReportClient.Target.cs
@@ -58,5 +58,7 @@ public class CrashReportClientTarget : TargetRules
 
 		// Epic Games Launcher needs to run on OS X 10.9, so CrashReportClient needs this as well
 		OutCPPEnvironmentConfiguration.bEnableOSX109Support = true;
+
+		OutCPPEnvironmentConfiguration.Definitions.Add("NOINITCRASHREPORTER=1");
 	}
 }
diff --git a/Engine/Source/Programs/CrashReporter/CrashReportClient/Private/CrashUpload.cpp b/Engine/Source/Programs/CrashReporter/CrashReportClient/Private/CrashUpload.cpp
--- a/Engine/Source/Programs/CrashReporter/CrashReportClient/Private/CrashUpload.cpp
+++ b/Engine/Source/Programs/CrashReporter/CrashReportClient/Private/CrashUpload.cpp
@@ -571,21 +571,8 @@ void FCrashUploadToReceiver::OnProcessRequestComplete(FHttpRequestPtr HttpReques
 			}
 			else
 			{
-				if (State == EUploadState::CheckingReport)
-				{
-					SetCurrentState(EUploadState::CheckingReportDetail);
-					if (!SendCheckReportRequest())
-					{
-						// Skip this report for now
-						AddReportToFailedList();
-						CheckPendingReportsForFilesToUpload();
-					}
-				}
-				else
-				{
-					SetCurrentState(EUploadState::CompressAndSendData);
-					CompressAndSendData();
-				}
+				SetCurrentState(EUploadState::CompressAndSendData);
+				CompressAndSendData();
 			}
 		}
 		break;
diff --git a/Engine/Source/Programs/IncludeTool/IncludeTool/Rules.cs b/Engine/Source/Programs/IncludeTool/IncludeTool/Rules.cs
--- a/Engine/Source/Programs/IncludeTool/IncludeTool/Rules.cs
+++ b/Engine/Source/Programs/IncludeTool/IncludeTool/Rules.cs
@@ -433,6 +433,7 @@ namespace IncludeTool
 			"/Engine/Source/Runtime/SlateCore/Public/Fonts/ShapedTextFwd.h", // Typedef isn't a forward declaration
 			"/Engine/Source/Runtime/Slate/Public/Framework/Text/ShapedTextCacheFwd.h", // Typedef isn't a forward declaration
 			"/Engine/Source/Runtime/MovieScene/Public/MovieSceneFwd.h",
+			"/Engine/Source/Runtime/Core/Public/Internationalization/StringTableCoreFwd.h", // Typedef isn't a forward declaration
 		};
 
 		/// <summary>
diff --git a/Engine/Source/Programs/IncludeTool/IncludeTool/Support/TextBuffer.cs b/Engine/Source/Programs/IncludeTool/IncludeTool/Support/TextBuffer.cs
--- a/Engine/Source/Programs/IncludeTool/IncludeTool/Support/TextBuffer.cs
+++ b/Engine/Source/Programs/IncludeTool/IncludeTool/Support/TextBuffer.cs
@@ -263,6 +263,26 @@ namespace IncludeTool.Support
 			return Extract(Begin.LineIdx, Begin.ColumnIdx, End.LineIdx, End.ColumnIdx);
 		}
 
+		/// <summary>
+		/// Access a single line within the file
+		/// </summary>
+		/// <param name="LineIdx">The line  to return</param>
+		/// <returns>The line at the given position</returns>
+		public string this[int LineIdx]
+		{
+			get
+			{
+				if (LineIdx < 0 || LineIdx >= Lines.Length)
+				{
+					return "";
+				}
+				else
+				{
+					return Lines[LineIdx];
+				}
+			}
+		}
+
 		/// <summary>
 		/// Access a single character within the file, by line and column
 		/// </summary>
diff --git a/Engine/Source/Programs/IncludeTool/IncludeTool/SymbolTable.cs b/Engine/Source/Programs/IncludeTool/IncludeTool/SymbolTable.cs
--- a/Engine/Source/Programs/IncludeTool/IncludeTool/SymbolTable.cs
+++ b/Engine/Source/Programs/IncludeTool/IncludeTool/SymbolTable.cs
@@ -672,7 +672,7 @@ namespace IncludeTool
 						&& !ReadTemplateClassOrStructForwardDeclaration(Reader, HeaderFile, SymbolToHeader, Log)
 						&& !ReadEnumClassForwardDeclaration(Reader, HeaderFile, SymbolToHeader, Log))
 					{
-						Log.WriteLine("{0}({1}): error: invalid forward declaration", HeaderFile.Location, Reader.CurrentLine);
+						Log.WriteLine("{0}({1}): error: invalid forward declaration - '{2}'", HeaderFile.Location, Reader.CurrentLine + 1, HeaderFile.Text[Reader.CurrentLine]);
 						return false;
 					}
 				}
diff --git a/Engine/Source/Programs/SlateViewer/Private/SlateViewerApp.cpp b/Engine/Source/Programs/SlateViewer/Private/SlateViewerApp.cpp
--- a/Engine/Source/Programs/SlateViewer/Private/SlateViewerApp.cpp
+++ b/Engine/Source/Programs/SlateViewer/Private/SlateViewerApp.cpp
@@ -141,7 +141,7 @@ namespace
 			const int Height = PopupFeaturesSP->IsHeightSet() ? PopupFeaturesSP->GetHeight() : 600;
 			const FVector2D BrowserWindowSize(Width, Height);
 
-			const ESizingRule::Type SizeingRule = PopupFeaturesSP->IsResizable() ? ESizingRule::UserSized : ESizingRule::FixedSize;
+			const ESizingRule SizeingRule = PopupFeaturesSP->IsResizable() ? ESizingRule::UserSized : ESizingRule::FixedSize;
 
 			TSharedPtr<IWebBrowserWindow> NewBrowserWindowSP = NewBrowserWindow.Pin();
 			check(NewBrowserWindowSP.IsValid())
diff --git a/Engine/Source/Programs/UnrealFileServer/Private/UnrealFileServer.h b/Engine/Source/Programs/UnrealFileServer/Private/UnrealFileServer.h
--- a/Engine/Source/Programs/UnrealFileServer/Private/UnrealFileServer.h
+++ b/Engine/Source/Programs/UnrealFileServer/Private/UnrealFileServer.h
@@ -1,9 +1,5 @@
 // Copyright 1998-2017 Epic Games, Inc. All Rights Reserved.
 
-
-#ifndef __UnrealFileServer_h__
-#define __UnrealFileServer_h__
+#pragma once
 
 #include "CoreMinimal.h"
-
-#endif		// __UnrealFileServer_h__
\ No newline at end of file
diff --git a/Engine/Source/Programs/UnrealHeaderTool/Private/UHTMakefile/PropertyBaseArchiveProxy.cpp b/Engine/Source/Programs/UnrealHeaderTool/Private/UHTMakefile/PropertyBaseArchiveProxy.cpp
--- a/Engine/Source/Programs/UnrealHeaderTool/Private/UHTMakefile/PropertyBaseArchiveProxy.cpp
+++ b/Engine/Source/Programs/UnrealHeaderTool/Private/UHTMakefile/PropertyBaseArchiveProxy.cpp
@@ -76,7 +76,7 @@ void FPropertyBaseArchiveProxy::PostConstruct(FPropertyBase* PropertyBase) const
 void FPropertyBaseArchiveProxy::Resolve(FPropertyBase* PropertyBase, const FUHTMakefile& UHTMakefile) const
 {
 	FPropertyBase* MapKeyProp = UHTMakefile.GetPropertyBaseByIndex(MapKeyPropIndex);
-	PropertyBase->MapKeyProp = MapKeyProp && MapKeyProp->HasBeenAlreadyMadeSharable() ? MapKeyProp->AsShared() : TSharedPtr<FPropertyBase>(MapKeyProp);
+	PropertyBase->MapKeyProp = MapKeyProp && MapKeyProp->DoesSharedInstanceExist() ? MapKeyProp->AsShared() : TSharedPtr<FPropertyBase>(MapKeyProp);
 
 	switch (PropertyBase->Type)
 	{
diff --git a/Engine/Source/Programs/UnrealHeaderTool/Private/UHTMakefile/StructScopeArchiveProxy.cpp b/Engine/Source/Programs/UnrealHeaderTool/Private/UHTMakefile/StructScopeArchiveProxy.cpp
--- a/Engine/Source/Programs/UnrealHeaderTool/Private/UHTMakefile/StructScopeArchiveProxy.cpp
+++ b/Engine/Source/Programs/UnrealHeaderTool/Private/UHTMakefile/StructScopeArchiveProxy.cpp
@@ -34,5 +34,5 @@ void FStructScopeArchiveProxy::Resolve(FStructScope* StructScope, const FUHTMake
 	FScopeArchiveProxy::Resolve(StructScope, UHTMakefile);
 	StructScope->Struct = UHTMakefile.GetStructByIndex(StructIndex);
 
-	FScope::ScopeMap.Add(StructScope->Struct, StructScope->HasBeenAlreadyMadeSharable() ? StructScope->AsShared() : MakeShareable(StructScope));
+	FScope::ScopeMap.Add(StructScope->Struct, StructScope->DoesSharedInstanceExist() ? StructScope->AsShared() : MakeShareable(StructScope));
 }
diff --git a/Engine/Source/Programs/UnrealHeaderTool/Private/UnrealHeaderTool.h b/Engine/Source/Programs/UnrealHeaderTool/Private/UnrealHeaderTool.h
--- a/Engine/Source/Programs/UnrealHeaderTool/Private/UnrealHeaderTool.h
+++ b/Engine/Source/Programs/UnrealHeaderTool/Private/UnrealHeaderTool.h
@@ -1,8 +1,3 @@
 // Copyright 1998-2017 Epic Games, Inc. All Rights Reserved.
 
-
-#ifndef __UnrealHeaderTool_h__
-#define __UnrealHeaderTool_h__
-
-
-#endif		// __UnrealHeaderTool_h__
+#pragma once
diff --git a/Engine/Source/Programs/UnrealPak/Private/UnrealPak.h b/Engine/Source/Programs/UnrealPak/Private/UnrealPak.h
--- a/Engine/Source/Programs/UnrealPak/Private/UnrealPak.h
+++ b/Engine/Source/Programs/UnrealPak/Private/UnrealPak.h
@@ -1,9 +1,5 @@
 // Copyright 1998-2017 Epic Games, Inc. All Rights Reserved.
 
+#pragma once
 
-#ifndef __UnrealPak_h__
-#define __UnrealPak_h__
-
-#include "CoreMinimal.h"
-
-#endif		// __UnrealPak_h__
\ No newline at end of file
+#include "CoreMinimal.h"
\ No newline at end of file
diff --git a/Engine/Source/Programs/UnrealWatchdog/Private/UnrealWatchdog.cpp b/Engine/Source/Programs/UnrealWatchdog/Private/UnrealWatchdog.cpp
--- a/Engine/Source/Programs/UnrealWatchdog/Private/UnrealWatchdog.cpp
+++ b/Engine/Source/Programs/UnrealWatchdog/Private/UnrealWatchdog.cpp
@@ -27,6 +27,10 @@ namespace WatchdogDefs
 	static const FString ShutdownSessionToken(TEXT("Shutdown"));
 	static const FString CrashSessionToken(TEXT("Crashed"));
 	static const FString TrueValueString(TEXT("1"));
+
+	static const FTimespan SendWatchdogHeartbeatPeriod(0, 5, 0);
+	static const FTimespan CheckParentRunningPeriod(0, 0, 10);
+	static const float TickSleepSeconds = 2.0f;
 }
 
 struct FWatchdogStoredValues
@@ -297,6 +301,98 @@ void SendShutdownEvent(IAnalyticsProviderET& Analytics, const FWatchdogCommandLi
 	Analytics.RecordEvent(TEXT("UnrealWatchdog.Shutdown"), MoveTemp(ShutdownAttributes));
 }
 
+void TickHeartbeat(IAnalyticsProviderET& Analytics, const FWatchdogCommandLine& CommandLine, FDateTime& InOutNextHeartbeatSend)
+{
+	// Send heartbeat due?
+	while (FDateTime::UtcNow() >= InOutNextHeartbeatSend)		// "while" allows watchdog to send heartbeats missed during a modal dialog popup 
+	{
+		InOutNextHeartbeatSend += WatchdogDefs::SendWatchdogHeartbeatPeriod;
+		SendHeartbeatEvent(Analytics, CommandLine);
+	}
+}
+
+bool TickProcessCheck(const FWatchdogCommandLine& CommandLine, FProcHandle& InParentProcess, FDateTime& InOutNextProcessCheck, bool& OutGotProcReturnCode, int32& OutReturnCode)
+{
+	// Parent application still running?
+	if (FDateTime::UtcNow() >= InOutNextProcessCheck)
+	{	
+		if (!FPlatformProcess::IsApplicationRunning(CommandLine.ParentProcessId))
+		{
+			UE_LOG(UnrealWatchdogLog, Log, TEXT("Watchdog detected terminated process PID %u"), CommandLine.ParentProcessId);
+			OutGotProcReturnCode = FPlatformProcess::GetProcReturnCode(InParentProcess, &OutReturnCode);
+			return false;
+		}
+
+		InOutNextProcessCheck = FDateTime::UtcNow() + WatchdogDefs::CheckParentRunningPeriod;
+	}
+	return true;
+}
+
+void TickHangCheck(IAnalyticsProviderET& Analytics, const FWatchdogCommandLine& CommandLine, FDateTime& InOutNextHeartbeatCheck, const FString& WatchdogSectionName, bool& bOutHang)
+{
+	if (FDateTime::UtcNow() >= InOutNextHeartbeatCheck)
+	{
+		bool bHangDetected = !CheckParentHeartbeat(Analytics, CommandLine, WatchdogSectionName);
+
+		if (bHangDetected && !bOutHang)
+		{
+			// New hang
+			SendHangDetectedEvent(Analytics, CommandLine);
+			bOutHang = true;
+		}
+		else if (!bHangDetected && bOutHang)
+		{
+			// Previous hang recovered
+			SendHangRecoveredEvent(Analytics, CommandLine);
+			bOutHang = false;
+		}
+
+		static const FTimespan CheckParentHeartbeatPeriod(0, 0, CommandLine.HangThresholdSeconds);
+		InOutNextHeartbeatCheck = FDateTime::UtcNow() + CheckParentHeartbeatPeriod;
+	}
+}
+
+bool WaitForProcess(IAnalyticsProviderET& Analytics, const FWatchdogCommandLine& CommandLine, int32& OutReturnCode, bool& bOutHang, const FString& WatchdogSectionName)
+{
+	FDateTime NextHeartbeatSend = FDateTime::UtcNow();
+	FDateTime NextProcessCheck = FDateTime::UtcNow();
+	FDateTime NextHeartbeatCheck = FDateTime::UtcNow();
+	bool bSuccess = false;
+	bOutHang = false;
+	OutReturnCode = -1;
+
+	FProcHandle ParentProcess = GetProcessHandle(CommandLine);
+	if (ParentProcess.IsValid())
+	{
+		while (!GIsRequestingExit)
+		{
+			TickHeartbeat(Analytics, CommandLine, NextHeartbeatSend);
+
+			if (!TickProcessCheck(CommandLine, ParentProcess, NextProcessCheck, bSuccess, OutReturnCode))
+			{
+				// process not running
+				break;
+			}
+
+			if (CommandLine.bAllowDetectHangs)
+			{
+				// Check parent heartbeat for hang?
+				TickHangCheck(Analytics, CommandLine, NextHeartbeatCheck, WatchdogSectionName, bOutHang);
+			}
+
+			FPlatformProcess::Sleep(WatchdogDefs::TickSleepSeconds);
+		}
+
+		FPlatformProcess::CloseProc(ParentProcess);
+	}
+	else
+	{
+		UE_LOG(UnrealWatchdogLog, Error, TEXT("Watchdog failed to get handle to process PID %u"), CommandLine.ParentProcessId);
+	}
+
+	return bSuccess;
+}
+
 /**
 * @EventName UnrealWatchdog.AbnormalShutdown
 *
diff --git a/Engine/Source/Programs/UnrealWatchdog/Private/UnrealWatchdog.h b/Engine/Source/Programs/UnrealWatchdog/Private/UnrealWatchdog.h
--- a/Engine/Source/Programs/UnrealWatchdog/Private/UnrealWatchdog.h
+++ b/Engine/Source/Programs/UnrealWatchdog/Private/UnrealWatchdog.h
@@ -15,7 +15,8 @@ struct FWatchdogCommandLine
 	FString EngineVersion;
 	uint32 ParentProcessId;
 	int32 SuccessReturnCode;
-	int32 HeartbeatSeconds;
+	bool bAllowDetectHangs;
+	int32 HangThresholdSeconds;
 	bool bAllowDialogs;
 	bool bHasProcessId;
 
@@ -30,16 +31,13 @@ struct FWatchdogCommandLine
 		bHasProcessId = FParse::Value(InCommandLine, TEXT("PID="), ParentProcessId);
 		SuccessReturnCode = 0;
 		FParse::Value(InCommandLine, TEXT("SuccessfulRtnCode="), SuccessReturnCode);
-		HeartbeatSeconds = 300;
-		FParse::Value(InCommandLine, TEXT("HeartbeatSeconds="), HeartbeatSeconds);
+		bAllowDetectHangs = FParse::Param(InCommandLine, TEXT("DetectHangs"));
+		HangThresholdSeconds = 120;
+		FParse::Value(InCommandLine, TEXT("HangSeconds="), HangThresholdSeconds);
 		bAllowDialogs = FParse::Param(InCommandLine, TEXT("AllowDialogs"));
 	}
 };
 
 bool GetWatchdogStoredDebuggerValue(const FString& WatchdogSectionName);
 int RunUnrealWatchdog(const TCHAR* CommandLine);
-bool WaitForProcess(class IAnalyticsProviderET& Analytics, const FWatchdogCommandLine& CommandLine, int32& OutReturnCode, bool& bOutHang, const FString& WatchdogSectionName);
-void SendHeartbeatEvent(class IAnalyticsProviderET& Analytics, const FWatchdogCommandLine& CommandLine);
-bool CheckParentHeartbeat(class IAnalyticsProviderET& Analytics, const FWatchdogCommandLine& CommandLine, const FString& WatchdogSectionName);
-void SendHangDetectedEvent(class IAnalyticsProviderET& Analytics, const FWatchdogCommandLine& CommandLine);
-void SendHangRecoveredEvent(class IAnalyticsProviderET& Analytics, const FWatchdogCommandLine& CommandLine);
+FProcHandle GetProcessHandle(const FWatchdogCommandLine& CommandLine);
diff --git a/Engine/Source/Programs/UnrealWatchdog/Private/Windows/UnrealWatchdogMainWindows.cpp b/Engine/Source/Programs/UnrealWatchdog/Private/Windows/UnrealWatchdogMainWindows.cpp
--- a/Engine/Source/Programs/UnrealWatchdog/Private/Windows/UnrealWatchdogMainWindows.cpp
+++ b/Engine/Source/Programs/UnrealWatchdog/Private/Windows/UnrealWatchdogMainWindows.cpp
@@ -41,77 +41,7 @@ INT32_MAIN_INT32_ARGC_TCHAR_ARGV()
 	return 0;
 }
 
-bool WaitForProcess(IAnalyticsProviderET& Analytics, const FWatchdogCommandLine& CommandLine, int32& OutReturnCode, bool& bOutHang, const FString& WatchdogSectionName)
+FProcHandle GetProcessHandle(const FWatchdogCommandLine& CommandLine)
 {
-	const int HangDetectionGraceSeconds = 10;
-	const FTimespan CheckParentHeartbeatPeriod(0, 0, CommandLine.HeartbeatSeconds + HangDetectionGraceSeconds);
-	const FTimespan SendWatchdogHeartbeatPeriod(0, 5, 0);
-	const FTimespan CheckParentRunningPeriod(0, 0, 10);
-
-	FDateTime NextHeartbeatSend = FDateTime::UtcNow();
-	FDateTime NextHeartbeatCheck = FDateTime::UtcNow();
-	bool bSuccess = false;
-	bOutHang = false;
-	OutReturnCode = -1;
-
-	FProcHandle ParentProcess = FPlatformProcess::OpenProcess(CommandLine.ParentProcessId);
-	if (ParentProcess.IsValid())
-	{
-		while (!GIsRequestingExit)
-		{
-			// Send heartbeat due?
-			if (FDateTime::UtcNow() >= NextHeartbeatSend)
-			{
-				NextHeartbeatSend += SendWatchdogHeartbeatPeriod;
-				SendHeartbeatEvent(Analytics, CommandLine);
-			}
-
-			// Parent application still running?
-			if (!FPlatformProcess::IsApplicationRunning(CommandLine.ParentProcessId))
-			{
-				UE_LOG(UnrealWatchdogLog, Log, TEXT("Watchdog detected terminated process PID %u"), CommandLine.ParentProcessId);
-				bSuccess = FPlatformProcess::GetProcReturnCode(ParentProcess, &OutReturnCode);
-				break;
-			}
-
-			// Check parent heartbeat for hang?
-			if (FDateTime::UtcNow() >= NextHeartbeatCheck)
-			{
-				NextHeartbeatCheck += CheckParentHeartbeatPeriod;
-
-				if (!GetWatchdogStoredDebuggerValue(WatchdogSectionName))
-				{
-					bool bHangDetected = !CheckParentHeartbeat(Analytics, CommandLine, WatchdogSectionName);
-
-					if (bHangDetected && !bOutHang)
-					{
-						// New hang
-						SendHangDetectedEvent(Analytics, CommandLine);
-						bOutHang = true;
-					}
-					else if (!bHangDetected && bOutHang)
-					{
-						// Previous hang recovered
-						SendHangRecoveredEvent(Analytics, CommandLine);
-						bOutHang = false;
-					}
-				}
-			}
-
-			// Sleep until next event
-			FDateTime EndLoopTime = FDateTime::UtcNow();
-			FTimespan TimeToHeartbeatCheck = NextHeartbeatCheck - EndLoopTime;
-			FTimespan TimeToHeartbeatSend = NextHeartbeatSend - EndLoopTime;
-			double SleepSeconds = FMath::Min3(TimeToHeartbeatCheck.GetTotalSeconds(), TimeToHeartbeatSend.GetTotalSeconds(), CheckParentRunningPeriod.GetTotalSeconds());
-			FPlatformProcess::Sleep((float)SleepSeconds);
-		}
-
-		FPlatformProcess::CloseProc(ParentProcess);
-	}
-	else
-	{
-		UE_LOG(UnrealWatchdogLog, Error, TEXT("Watchdog failed to get handle to process PID %u"), CommandLine.ParentProcessId);
-	}
-
-	return bSuccess;
+	return FPlatformProcess::OpenProcess(CommandLine.ParentProcessId);
 }
\ No newline at end of file
diff --git a/Engine/Source/Runtime/AppFramework/Private/Widgets/Colors/SColorValueSlider.h b/Engine/Source/Runtime/AppFramework/Private/Widgets/Colors/SColorValueSlider.h
--- a/Engine/Source/Runtime/AppFramework/Private/Widgets/Colors/SColorValueSlider.h
+++ b/Engine/Source/Runtime/AppFramework/Private/Widgets/Colors/SColorValueSlider.h
@@ -61,7 +61,7 @@ public:
 	virtual int32 OnPaint( const FPaintArgs& Args, const FGeometry& AllottedGeometry, const FSlateRect& MyClippingRect, FSlateWindowElementList& OutDrawElements, int32 LayerId, const FWidgetStyle& InWidgetStyle, bool bParentEnabled ) const
 	{
 		const bool bIsEnabled = ShouldBeEnabled(bParentEnabled);
-		const uint32 DrawEffects = bIsEnabled ? ESlateDrawEffect::None : ESlateDrawEffect::DisabledEffect;
+		const ESlateDrawEffect DrawEffects = bIsEnabled ? ESlateDrawEffect::None : ESlateDrawEffect::DisabledEffect;
 	
 		FLinearColor FullValueColor = SelectedColor.Get();
 		FullValueColor.B = FullValueColor.A = 1.f;
diff --git a/Engine/Source/Runtime/AppFramework/Private/Widgets/Colors/SComplexGradient.cpp b/Engine/Source/Runtime/AppFramework/Private/Widgets/Colors/SComplexGradient.cpp
--- a/Engine/Source/Runtime/AppFramework/Private/Widgets/Colors/SComplexGradient.cpp
+++ b/Engine/Source/Runtime/AppFramework/Private/Widgets/Colors/SComplexGradient.cpp
@@ -22,7 +22,7 @@ void SComplexGradient::Construct( const FArguments& InArgs )
 
 int32 SComplexGradient::OnPaint( const FPaintArgs& Args, const FGeometry& AllottedGeometry, const FSlateRect& MyClippingRect, FSlateWindowElementList& OutDrawElements, int32 LayerId, const FWidgetStyle& InWidgetStyle, bool bParentEnabled ) const
 {
-	ESlateDrawEffect::Type DrawEffects = (bParentEnabled && IsEnabled()) ? ESlateDrawEffect::None : ESlateDrawEffect::DisabledEffect;
+	ESlateDrawEffect DrawEffects = (bParentEnabled && IsEnabled()) ? ESlateDrawEffect::None : ESlateDrawEffect::DisabledEffect;
 
 	if (bHasAlphaBackground)
 	{
diff --git a/Engine/Source/Runtime/AppFramework/Private/Widgets/Colors/SSimpleGradient.cpp b/Engine/Source/Runtime/AppFramework/Private/Widgets/Colors/SSimpleGradient.cpp
--- a/Engine/Source/Runtime/AppFramework/Private/Widgets/Colors/SSimpleGradient.cpp
+++ b/Engine/Source/Runtime/AppFramework/Private/Widgets/Colors/SSimpleGradient.cpp
@@ -23,7 +23,7 @@ void SSimpleGradient::Construct( const FArguments& InArgs )
 
 int32 SSimpleGradient::OnPaint( const FPaintArgs& Args, const FGeometry& AllottedGeometry, const FSlateRect& MyClippingRect, FSlateWindowElementList& OutDrawElements, int32 LayerId, const FWidgetStyle& InWidgetStyle, bool bParentEnabled ) const
 {
-	const ESlateDrawEffect::Type DrawEffects = this->ShouldBeEnabled(bParentEnabled) ? ESlateDrawEffect::None : ESlateDrawEffect::DisabledEffect;
+	const ESlateDrawEffect DrawEffects = this->ShouldBeEnabled(bParentEnabled) ? ESlateDrawEffect::None : ESlateDrawEffect::DisabledEffect;
 
 	if (bHasAlphaBackground)
 	{
diff --git a/Engine/Source/Runtime/AppFramework/Public/Widgets/Colors/SColorGradingPicker.h b/Engine/Source/Runtime/AppFramework/Public/Widgets/Colors/SColorGradingPicker.h
--- a/Engine/Source/Runtime/AppFramework/Public/Widgets/Colors/SColorGradingPicker.h
+++ b/Engine/Source/Runtime/AppFramework/Public/Widgets/Colors/SColorGradingPicker.h
@@ -89,6 +89,13 @@ protected:
 	// Callback for value changes in the color spectrum picker.
 	void HandleCurrentColorValueChanged(FLinearColor NewValue);
 
+	void OnBeginSliderMovement();
+	void OnEndSliderMovement(float NewValue);
+	void AdjustRatioValue(FVector4 &NewValue);
+
+	bool bIsMouseDragging;
+	FVector4 StartDragRatio;
+
 	float ValuesMin;
 	float ValuesMax;
 	float MainDelta;
diff --git a/Engine/Source/Runtime/AssetRegistry/Public/AssetData.h b/Engine/Source/Runtime/AssetRegistry/Public/AssetData.h
--- a/Engine/Source/Runtime/AssetRegistry/Public/AssetData.h
+++ b/Engine/Source/Runtime/AssetRegistry/Public/AssetData.h
@@ -13,7 +13,8 @@
 #include "UObject/LinkerLoad.h"
 #include "SharedMapView.h"
 
-DECLARE_LOG_CATEGORY_EXTERN(LogAssetData, Log, All);
+ASSETREGISTRY_API DECLARE_LOG_CATEGORY_EXTERN(LogAssetData, Log, All);
+
 /** A class to hold important information about an assets found by the Asset Registry */
 class FAssetData
 {
diff --git a/Engine/Source/Runtime/AutomationMessages/Classes/AutomationWorkerMessages.h b/Engine/Source/Runtime/AutomationMessages/Classes/AutomationWorkerMessages.h
--- a/Engine/Source/Runtime/AutomationMessages/Classes/AutomationWorkerMessages.h
+++ b/Engine/Source/Runtime/AutomationMessages/Classes/AutomationWorkerMessages.h
@@ -419,6 +419,8 @@ public:
 
 	UPROPERTY(EditAnywhere, Category="Message")
 	FGuid Id;
+	UPROPERTY(EditAnywhere, Category="Message")
+	FString Commit;
 
 	UPROPERTY(EditAnywhere, Category="Message")
 	int32 Width;
@@ -504,6 +506,7 @@ public:
 
 		// Unique Id so we know if this screenshot has already been imported.
 		Id = Data.Id;
+		Commit = Data.Commit;
 
 		// Resolution Details
 		Width = Data.Width;
diff --git a/Engine/Source/Runtime/Core/Core.Build.cs b/Engine/Source/Runtime/Core/Core.Build.cs
--- a/Engine/Source/Runtime/Core/Core.Build.cs
+++ b/Engine/Source/Runtime/Core/Core.Build.cs
@@ -63,6 +63,9 @@ public class Core : ModuleRules
 		if (UEBuildConfiguration.bBuildEditor == true)
 		{
 			DynamicallyLoadedModuleNames.Add("SourceCodeAccess");
+
+			PrivateIncludePathModuleNames.Add("DirectoryWatcher");
+			DynamicallyLoadedModuleNames.Add("DirectoryWatcher");
 		}
 
 		if ((Target.Platform == UnrealTargetPlatform.Win64) ||
diff --git a/Engine/Source/Runtime/Core/Private/GenericPlatform/GenericPlatformCrashContext.cpp b/Engine/Source/Runtime/Core/Private/GenericPlatform/GenericPlatformCrashContext.cpp
--- a/Engine/Source/Runtime/Core/Private/GenericPlatform/GenericPlatformCrashContext.cpp
+++ b/Engine/Source/Runtime/Core/Private/GenericPlatform/GenericPlatformCrashContext.cpp
@@ -32,6 +32,8 @@ const ANSICHAR* FGenericCrashContext::CrashConfigFileNameA = "CrashReportClient.
 const TCHAR* FGenericCrashContext::CrashConfigFileNameW = TEXT("CrashReportClient.ini");
 const FString FGenericCrashContext::CrashConfigExtension = TEXT(".ini");
 const FString FGenericCrashContext::ConfigSectionName = TEXT("CrashReportClient");
+const FString FGenericCrashContext::CrashConfigPurgeDays = TEXT("CrashConfigPurgeDays");
+const FString FGenericCrashContext::CrashGUIDRootPrefix = TEXT("UE4CC-");
 
 const FString FGenericCrashContext::CrashContextExtension = TEXT(".runtime-xml");
 const FString FGenericCrashContext::RuntimePropertiesTag = TEXT( "RuntimeProperties" );
@@ -148,7 +150,7 @@ void FGenericCrashContext::Initialize()
 	}
 
 	const FGuid Guid = FGuid::NewGuid();
-	NCachedCrashContextProperties::CrashGUIDRoot = FString::Printf(TEXT("UE4CC-%s-%s"), *NCachedCrashContextProperties::PlatformNameIni, *Guid.ToString(EGuidFormats::Digits));
+	NCachedCrashContextProperties::CrashGUIDRoot = FString::Printf(TEXT("%s%s-%s"), *CrashGUIDRootPrefix, *NCachedCrashContextProperties::PlatformNameIni, *Guid.ToString(EGuidFormats::Digits));
 
 	// Initialize delegate for updating SecondsSinceStart, because FPlatformTime::Seconds() is not POSIX safe.
 	const float PollingInterval = 1.0f;
@@ -181,12 +183,14 @@ void FGenericCrashContext::Initialize()
 	FCoreDelegates::ConfigReadyForUse.AddStatic(FGenericCrashContext::InitializeFromConfig);
 
 	bIsInitialized = true;
-#endif
+#endif	// !NOINITCRASHREPORTER
 }
 
 void FGenericCrashContext::InitializeFromConfig()
 {
-#if !NO_LOGGING
+#if !NOINITCRASHREPORTER
+	PurgeOldCrashConfig();
+
 	const bool bForceGetSection = false;
 	const bool bConstSection = true;
 	FConfigSection* CRCConfigSection = GConfig->GetSectionPrivate(*ConfigSectionName, bForceGetSection, bConstSection, GEngineIni);
@@ -203,7 +207,7 @@ void FGenericCrashContext::InitializeFromConfig()
 		CrashConfigFile.Dirty = true;
 		CrashConfigFile.Write(GetCrashConfigFilePath());
 	}
-#endif
+#endif	// !NOINITCRASHREPORTER
 }
 
 FGenericCrashContext::FGenericCrashContext()
@@ -450,11 +454,46 @@ const TCHAR* FGenericCrashContext::GetCrashConfigFilePath()
 	static FString CrashConfigFilePath;
 	if (CrashConfigFilePath.IsEmpty())
 	{
-		CrashConfigFilePath = FPaths::Combine(*FPaths::GameLogDir(), *NCachedCrashContextProperties::CrashGUIDRoot, FGenericCrashContext::CrashConfigFileNameW);
+		CrashConfigFilePath = FPaths::Combine(GetCrashConfigFolder(), *NCachedCrashContextProperties::CrashGUIDRoot, FGenericCrashContext::CrashConfigFileNameW);
 	}
 	return *CrashConfigFilePath;
 }
 
+const TCHAR* FGenericCrashContext::GetCrashConfigFolder()
+{
+	static FString CrashConfigFolder;
+	if (CrashConfigFolder.IsEmpty())
+	{
+		CrashConfigFolder = FPaths::Combine(*FPaths::GeneratedConfigDir(), TEXT("CrashReportClient"));
+	}
+	return *CrashConfigFolder;
+}
+
+void FGenericCrashContext::PurgeOldCrashConfig()
+{
+	int32 PurgeDays = 2;
+	GConfig->GetInt(*ConfigSectionName, *CrashConfigPurgeDays, PurgeDays, GEngineIni);
+
+	if (PurgeDays > 0)
+	{
+		IFileManager& FileManager = IFileManager::Get();
+
+		// Delete items older than PurgeDays
+		TArray<FString> Directories;
+		FileManager.FindFiles(Directories, *(FPaths::Combine(GetCrashConfigFolder(), *CrashGUIDRootPrefix) + TEXT("*")), false, true);
+
+		for (const FString& Dir : Directories)
+		{
+			const FString CrashConfigDirectory = FPaths::Combine(GetCrashConfigFolder(), *Dir);
+			const FDateTime DirectoryAccessTime = FileManager.GetTimeStamp(*CrashConfigDirectory);
+			if (FDateTime::Now() - DirectoryAccessTime > FTimespan::FromDays(PurgeDays))
+			{
+				FileManager.DeleteDirectory(*CrashConfigDirectory, false, true);
+			}
+		}
+	}
+}
+
 FProgramCounterSymbolInfoEx::FProgramCounterSymbolInfoEx( FString InModuleName, FString InFunctionName, FString InFilename, uint32 InLineNumber, uint64 InSymbolDisplacement, uint64 InOffsetInModule, uint64 InProgramCounter ) :
 	ModuleName( InModuleName ),
 	FunctionName( InFunctionName ),
diff --git a/Engine/Source/Runtime/Core/Private/HTML5/HTML5Cursor.cpp b/Engine/Source/Runtime/Core/Private/HTML5/HTML5Cursor.cpp
--- a/Engine/Source/Runtime/Core/Private/HTML5/HTML5Cursor.cpp
+++ b/Engine/Source/Runtime/Core/Private/HTML5/HTML5Cursor.cpp
@@ -60,8 +60,3 @@ void FHTML5Cursor::Lock( const RECT* const Bounds )
 		LockStatus = true; 
 	}
 }
-
-void FHTML5Cursor::SetCustomShape(void* CursorHandle)
-{
-	
-}
\ No newline at end of file
