Copying //UE4/Dev-Editor to //UE4/Dev-Main (Source: //UE4/Dev-Editor @ 3279756)

diff --git a/Engine/Source/Runtime/Core/Private/Internationalization/StringTableCore.cpp b/Engine/Source/Runtime/Core/Private/Internationalization/StringTableCore.cpp
new file mode 100644
--- /dev/null
+++ b/Engine/Source/Runtime/Core/Private/Internationalization/StringTableCore.cpp
@@ -0,0 +1,596 @@
+// Copyright 1998-2017 Epic Games, Inc. All Rights Reserved.
+
+#include "StringTableCore.h"
+#include "Misc/ScopeLock.h"
+#include "Misc/FileHelper.h"
+#include "Misc/ConfigCacheIni.h"
+#include "Serialization/Csv/CsvParser.h"
+
+DEFINE_LOG_CATEGORY(LogStringTable);
+
+
+IStringTableEngineBridge* IStringTableEngineBridge::InstancePtr = nullptr;
+
+
+namespace StringTableRedirects
+{
+	static TMap<FName, FName> TableIdRedirects;
+	static TMap<FName, TMap<FString, FString, FDefaultSetAllocator, FLocKeyMapFuncs<FString>>> TableKeyRedirects;
+}
+
+
+FStringTableEntry::FStringTableEntry()
+{
+}
+
+FStringTableEntry::FStringTableEntry(FStringTableConstRef InOwnerTable, FString InSourceString, FTextDisplayStringPtr InDisplayString)
+	: OwnerTable(MoveTemp(InOwnerTable))
+	, SourceString(MoveTemp(InSourceString))
+	, DisplayString(MoveTemp(InDisplayString))
+{
+}
+
+bool FStringTableEntry::IsOwned() const
+{
+	return OwnerTable.IsValid();
+}
+
+void FStringTableEntry::Disown()
+{
+	OwnerTable.Reset();
+}
+
+const FString& FStringTableEntry::GetSourceString() const
+{
+	return SourceString;
+}
+
+FTextDisplayStringPtr FStringTableEntry::GetDisplayString() const
+{
+	return DisplayString;
+}
+
+
+FStringTable::FStringTable()
+	: OwnerAsset(nullptr)
+	, bIsLoaded(true)
+{
+}
+
+FStringTable::~FStringTable()
+{
+	// Make sure our entries are disowned correctly
+	ClearSourceStrings();
+}
+
+UStringTable* FStringTable::GetOwnerAsset() const
+{
+	return OwnerAsset;
+}
+
+void FStringTable::SetOwnerAsset(UStringTable* InOwnerAsset)
+{
+	OwnerAsset = InOwnerAsset;
+}
+
+bool FStringTable::IsLoaded() const
+{
+	return bIsLoaded;
+}
+
+void FStringTable::IsLoaded(const bool bInIsLoaded)
+{
+	bIsLoaded = bInIsLoaded;
+}
+
+FString FStringTable::GetNamespace() const
+{
+	return TableNamespace;
+}
+
+void FStringTable::SetNamespace(const FString& InNamespace)
+{
+	FScopeLock KeyMappingLock(&KeyMappingCS);
+
+	if (!TableNamespace.Equals(InNamespace, ESearchCase::CaseSensitive))
+	{
+		TableNamespace = InNamespace;
+
+		// Changing the namespace affects the display string pointers, so update those now
+		for (auto& KeyToEntryPair : KeysToEntries)
+		{
+			FStringTableEntryPtr OldEntry = KeyToEntryPair.Value;
+			OldEntry->Disown();
+
+			DisplayStringsToKeys.Remove(OldEntry->GetDisplayString());
+			KeyToEntryPair.Value = FStringTableEntry::NewStringTableEntry(AsShared(), OldEntry->GetSourceString(), FTextLocalizationManager::Get().GetDisplayString(TableNamespace, KeyToEntryPair.Key, &OldEntry->GetSourceString()));
+			DisplayStringsToKeys.Emplace(KeyToEntryPair.Value->GetDisplayString(), KeyToEntryPair.Key);
+		}
+	}
+}
+
+bool FStringTable::GetSourceString(const FString& InKey, FString& OutSourceString) const
+{
+	FScopeLock KeyMappingLock(&KeyMappingCS);
+
+	FStringTableEntryPtr TableEntry = KeysToEntries.FindRef(InKey);
+	if (TableEntry.IsValid())
+	{
+		OutSourceString = TableEntry->GetSourceString();
+		return true;
+	}
+	return false;
+}
+
+void FStringTable::SetSourceString(const FString& InKey, const FString& InSourceString)
+{
+	checkf(!InKey.IsEmpty(), TEXT("String table key cannot be empty!"));
+
+	FScopeLock KeyMappingLock(&KeyMappingCS);
+	
+	FStringTableEntryPtr TableEntry = KeysToEntries.FindRef(InKey);
+	if (TableEntry.IsValid())
+	{
+		TableEntry->Disown();
+		DisplayStringsToKeys.Remove(TableEntry->GetDisplayString());
+	}
+	
+	TableEntry = FStringTableEntry::NewStringTableEntry(AsShared(), InSourceString, FTextLocalizationManager::Get().GetDisplayString(TableNamespace, InKey, &InSourceString));
+	KeysToEntries.Emplace(InKey, TableEntry);
+	DisplayStringsToKeys.Emplace(TableEntry->GetDisplayString(), InKey);
+}
+
+void FStringTable::RemoveSourceString(const FString& InKey)
+{
+	FScopeLock KeyMappingLock(&KeyMappingCS);
+
+	FStringTableEntryPtr TableEntry = KeysToEntries.FindRef(InKey);
+	if (TableEntry.IsValid())
+	{
+		TableEntry->Disown();
+		KeysToEntries.Remove(InKey);
+		DisplayStringsToKeys.Remove(TableEntry->GetDisplayString());
+		ClearMetaData(InKey);
+	}
+}
+
+void FStringTable::EnumerateSourceStrings(const TFunctionRef<bool(const FString&, const FString&)>& InEnumerator) const
+{
+	FScopeLock KeyMappingLock(&KeyMappingCS);
+
+	for (const auto& KeyToEntryPair : KeysToEntries)
+	{
+		if (!InEnumerator(KeyToEntryPair.Key, KeyToEntryPair.Value->GetSourceString()))
+		{
+			break;
+		}
+	}
+}
+
+void FStringTable::ClearSourceStrings(const int32 InSlack)
+{
+	FScopeLock KeyMappingLock(&KeyMappingCS);
+
+	for (const auto& KeyToEntryPair : KeysToEntries)
+	{
+		KeyToEntryPair.Value->Disown();
+	}
+
+	KeysToEntries.Empty(InSlack);
+	DisplayStringsToKeys.Empty(InSlack);
+
+	ClearMetaData(InSlack);
+}
+
+FStringTableEntryConstPtr FStringTable::FindEntry(const FString& InKey) const
+{
+	FScopeLock KeyMappingLock(&KeyMappingCS);
+	return KeysToEntries.FindRef(InKey);
+}
+
+bool FStringTable::FindKey(const FStringTableEntryConstRef& InEntry, FString& OutKey) const
+{
+	return FindKey(InEntry->GetDisplayString().ToSharedRef(), OutKey);
+}
+
+bool FStringTable::FindKey(const FTextDisplayStringRef& InDisplayString, FString& OutKey) const
+{
+	FScopeLock KeyMappingLock(&KeyMappingCS);
+
+	const FString* FoundKey = DisplayStringsToKeys.Find(InDisplayString);
+	if (FoundKey)
+	{
+		OutKey = *FoundKey;
+		return true;
+	}
+	return false;
+}
+
+FString FStringTable::GetMetaData(const FString& InKey, const FName InMetaDataId) const
+{
+	FScopeLock MetaDataLock(&KeysToMetaDataCS);
+
+	const FMetaDataMap* MetaDataMap = KeysToMetaData.Find(InKey);
+	if (MetaDataMap)
+	{
+		return MetaDataMap->FindRef(InMetaDataId);
+	}
+	return FString();
+}
+
+void FStringTable::SetMetaData(const FString& InKey, const FName InMetaDataId, const FString& InMetaDataValue)
+{
+	FScopeLock MetaDataLock(&KeysToMetaDataCS);
+
+	FMetaDataMap& MetaDataMap = KeysToMetaData.FindOrAdd(InKey);
+	MetaDataMap.Add(InMetaDataId, InMetaDataValue);
+}
+
+void FStringTable::RemoveMetaData(const FString& InKey, const FName InMetaDataId)
+{
+	FScopeLock MetaDataLock(&KeysToMetaDataCS);
+
+	FMetaDataMap* MetaDataMap = KeysToMetaData.Find(InKey);
+	if (MetaDataMap)
+	{
+		MetaDataMap->Remove(InMetaDataId);
+		if (MetaDataMap->Num() == 0)
+		{
+			KeysToMetaData.Remove(InKey);
+		}
+	}
+}
+
+void FStringTable::EnumerateMetaData(const FString& InKey, const TFunctionRef<bool(FName, const FString&)>& InEnumerator) const
+{
+	FScopeLock MetaDataLock(&KeysToMetaDataCS);
+
+	const FMetaDataMap* MetaDataMap = KeysToMetaData.Find(InKey);
+	if (MetaDataMap)
+	{
+		for (const auto& IdToMetaDataPair : *MetaDataMap)
+		{
+			if (!InEnumerator(IdToMetaDataPair.Key, IdToMetaDataPair.Value))
+			{
+				break;
+			}
+		}
+	}
+}
+
+void FStringTable::ClearMetaData(const FString& InKey)
+{
+	FScopeLock MetaDataLock(&KeysToMetaDataCS);
+	KeysToMetaData.Remove(InKey);
+}
+
+void FStringTable::ClearMetaData(const int32 InSlack)
+{
+	FScopeLock MetaDataLock(&KeysToMetaDataCS);
+	KeysToMetaData.Empty(InSlack);
+}
+
+void FStringTable::Serialize(FArchive& Ar)
+{
+	FScopeLock KeyMappingLock(&KeyMappingCS);
+	FScopeLock MetaDataLock(&KeysToMetaDataCS);
+
+	Ar << TableNamespace;
+
+	if (Ar.IsSaving())
+	{
+		// Save entries
+		{
+			int32 NumEntries = KeysToEntries.Num();
+			Ar << NumEntries;
+
+			for (const auto& KeyToEntryPair : KeysToEntries)
+			{
+				FString Key = KeyToEntryPair.Key;
+				Ar << Key;
+
+				FString SourceString = KeyToEntryPair.Value->GetSourceString();
+				Ar << SourceString;
+			}
+		}
+
+		// Save meta-data
+		{
+			Ar << KeysToMetaData;
+		}
+	}
+	else if (Ar.IsLoading())
+	{
+		// Load entries
+		{
+			int32 NumEntries = 0;
+			Ar << NumEntries;
+
+			ClearSourceStrings(NumEntries);
+			for (int32 EntryIndex = 0; EntryIndex < NumEntries; ++EntryIndex)
+			{
+				FString Key;
+				Ar << Key;
+
+				FString SourceString;
+				Ar << SourceString;
+
+				FStringTableEntryRef TableEntry = FStringTableEntry::NewStringTableEntry(AsShared(), SourceString, FTextLocalizationManager::Get().GetDisplayString(TableNamespace, Key, &SourceString));
+				KeysToEntries.Emplace(Key, TableEntry);
+				DisplayStringsToKeys.Emplace(TableEntry->GetDisplayString(), Key);
+			}
+		}
+
+		// Load meta-data
+		{
+			Ar << KeysToMetaData;
+		}
+	}
+}
+
+bool FStringTable::ExportStrings(const FString& InFilename) const
+{
+	FString ExportedStrings;
+
+	{
+		FScopeLock KeyMappingLock(&KeyMappingCS);
+		FScopeLock MetaDataLock(&KeysToMetaDataCS);
+
+		// Collect meta-data column names
+		TSet<FName> MetaDataColumnNames;
+		for (const auto& KeyToMetaDataPair : KeysToMetaData)
+		{
+			for (const auto& IdToMetaDataPair : KeyToMetaDataPair.Value)
+			{
+				MetaDataColumnNames.Add(IdToMetaDataPair.Key);
+			}
+		}
+
+		// Write header
+		ExportedStrings += TEXT("Key,SourceString");
+		for (const FName MetaDataColumnName : MetaDataColumnNames)
+		{
+			ExportedStrings += TEXT(",");
+			ExportedStrings += MetaDataColumnName.ToString();
+		}
+		ExportedStrings += TEXT("\n");
+
+		// Write entries
+		for (const auto& KeyToEntryPair : KeysToEntries)
+		{
+			FString ExportedKey = KeyToEntryPair.Key.ReplaceCharWithEscapedChar();
+			ExportedKey.ReplaceInline(TEXT("\""), TEXT("\"\""));
+
+			FString ExportedSourceString = KeyToEntryPair.Value->GetSourceString().ReplaceCharWithEscapedChar();
+			ExportedSourceString.ReplaceInline(TEXT("\""), TEXT("\"\""));
+
+			ExportedStrings += TEXT("\"");
+			ExportedStrings += ExportedKey;
+			ExportedStrings += TEXT("\"");
+
+			ExportedStrings += TEXT(",");
+
+			ExportedStrings += TEXT("\"");
+			ExportedStrings += ExportedSourceString;
+			ExportedStrings += TEXT("\"");
+
+			for (const FName MetaDataColumnName : MetaDataColumnNames)
+			{
+				FString ExportedMetaData = GetMetaData(KeyToEntryPair.Key, MetaDataColumnName);
+				ExportedMetaData.ReplaceInline(TEXT("\""), TEXT("\"\""));
+
+				ExportedStrings += TEXT(",");
+
+				ExportedStrings += TEXT("\"");
+				ExportedStrings += ExportedMetaData;
+				ExportedStrings += TEXT("\"");
+			}
+
+			ExportedStrings += TEXT("\n");
+		}
+	}
+
+	return FFileHelper::SaveStringToFile(ExportedStrings, *InFilename);
+}
+
+bool FStringTable::ImportStrings(const FString& InFilename)
+{
+	FString ImportedStrings;
+	if (!FFileHelper::LoadFileToString(ImportedStrings, *InFilename))
+	{
+		UE_LOG(LogStringTable, Warning, TEXT("Failed to import string table from '%s'. Could not open file."), *InFilename);
+		return false;
+	}
+
+	const FCsvParser ImportedStringsParser(ImportedStrings);
+	const FCsvParser::FRows& Rows = ImportedStringsParser.GetRows();
+
+	// Must have at least 2 rows (header and content)
+	if (Rows.Num() <= 1)
+	{
+		UE_LOG(LogStringTable, Warning, TEXT("Failed to import string table from '%s'. Incorrect number of rows (must be at least 2)."), *InFilename);
+		return false;
+	}
+
+	int32 KeyColumn = INDEX_NONE;
+	int32 SourceStringColumn = INDEX_NONE;
+	TMap<FName, int32> MetaDataColumns;
+
+	// Validate header
+	{
+		const TArray<const TCHAR*>& Cells = Rows[0];
+
+		for (int32 CellIdx = 0; CellIdx < Cells.Num(); ++CellIdx)
+		{
+			const TCHAR* Cell = Cells[CellIdx];
+			if (FCString::Stricmp(Cell, TEXT("Key")) == 0 && KeyColumn == INDEX_NONE)
+			{
+				KeyColumn = CellIdx;
+			}
+			else if(FCString::Stricmp(Cell, TEXT("SourceString")) == 0 && SourceStringColumn == INDEX_NONE)
+			{
+				SourceStringColumn = CellIdx;
+			}
+			else
+			{
+				const FName MetaDataName = Cell;
+				if (!MetaDataName.IsNone())
+				{
+					MetaDataColumns.Add(MetaDataName, CellIdx);
+				}
+			}
+		}
+
+		bool bValidHeader = true;
+		if (KeyColumn == INDEX_NONE)
+		{
+			bValidHeader = false;
+			UE_LOG(LogStringTable, Warning, TEXT("Failed to import string table from '%s'. Failed to find required column 'Key'."), *InFilename);
+		}
+		if (SourceStringColumn == INDEX_NONE)
+		{
+			bValidHeader = false;
+			UE_LOG(LogStringTable, Warning, TEXT("Failed to import string table from '%s'. Failed to find required column 'SourceString'."), *InFilename);
+		}
+		if (!bValidHeader)
+		{
+			return false;
+		}
+	}
+
+	// Import rows
+	{
+		FScopeLock KeyMappingLock(&KeyMappingCS);
+		FScopeLock MetaDataLock(&KeysToMetaDataCS);
+
+		ClearSourceStrings(Rows.Num() - 1);
+		for (int32 RowIdx = 1; RowIdx < Rows.Num(); ++RowIdx)
+		{
+			const TArray<const TCHAR*>& Cells = Rows[RowIdx];
+
+			// Must have at least an entry for the Key and SourceString columns
+			if (Cells.IsValidIndex(KeyColumn) && Cells.IsValidIndex(SourceStringColumn))
+			{
+				FString Key = Cells[KeyColumn];
+				Key = Key.ReplaceEscapedCharWithChar();
+
+				FString SourceString = Cells[SourceStringColumn];
+				SourceString = SourceString.ReplaceEscapedCharWithChar();
+
+				FStringTableEntryRef TableEntry = FStringTableEntry::NewStringTableEntry(AsShared(), SourceString, FTextLocalizationManager::Get().GetDisplayString(TableNamespace, Key, &SourceString));
+				KeysToEntries.Emplace(Key, TableEntry);
+				DisplayStringsToKeys.Emplace(TableEntry->GetDisplayString(), Key);
+
+				for (const auto& MetaDataColumnPair : MetaDataColumns)
+				{
+					if (Cells.IsValidIndex(MetaDataColumnPair.Value))
+					{
+						FString MetaData = Cells[MetaDataColumnPair.Value];
+						MetaData = MetaData.ReplaceEscapedCharWithChar();
+
+						if (!MetaData.IsEmpty())
+						{
+							FMetaDataMap& MetaDataMap = KeysToMetaData.FindOrAdd(Key);
+							MetaDataMap.Add(MetaDataColumnPair.Key, MetaData);
+						}
+					}
+				}
+			}
+		}
+	}
+
+	return true;
+}
+
+
+void FStringTableRedirects::InitStringTableRedirects()
+{
+	check(GConfig);
+
+	FConfigSection* CoreStringTableSection = GConfig->GetSectionPrivate(TEXT("Core.StringTable"), false, true, GEngineIni);
+	if (CoreStringTableSection)
+	{
+		for (FConfigSection::TIterator It(*CoreStringTableSection); It; ++It)
+		{
+			if (It.Key() == TEXT("StringTableRedirects"))
+			{
+				const FString& ConfigValue = It.Value().GetValue();
+
+				FName OldStringTable;
+				FName NewStringTable;
+
+				FString OldKey;
+				FString NewKey;
+
+				if (FParse::Value(*ConfigValue, TEXT("OldStringTable="), OldStringTable))
+				{
+					FParse::Value(*ConfigValue, TEXT("NewStringTable="), NewStringTable);
+					UE_CLOG(NewStringTable.IsNone(), LogStringTable, Warning, TEXT("Failed to parse string table redirect '%s'. Missing or empty 'NewStringTable'."), *ConfigValue);
+
+					if (!NewStringTable.IsNone())
+					{
+						StringTableRedirects::TableIdRedirects.Add(OldStringTable, NewStringTable);
+					}
+				}
+				else if (FParse::Value(*ConfigValue, TEXT("StringTable="), OldStringTable))
+				{
+					FParse::Value(*ConfigValue, TEXT("OldKey="), OldKey);
+					UE_CLOG(OldKey.IsEmpty(), LogStringTable, Warning, TEXT("Failed to parse string table redirect '%s'. Missing or empty 'OldKey'."), *ConfigValue);
+
+					FParse::Value(*ConfigValue, TEXT("NewKey="), NewKey);
+					UE_CLOG(NewKey.IsEmpty(), LogStringTable, Warning, TEXT("Failed to parse string table redirect '%s'. Missing or empty 'NewKey'."), *ConfigValue);
+
+					if (!OldKey.IsEmpty() && !NewKey.IsEmpty())
+					{
+						StringTableRedirects::TableKeyRedirects.FindOrAdd(OldStringTable).Add(OldKey, NewKey);
+					}
+				}
+				else
+				{
+					UE_LOG(LogStringTable, Warning, TEXT("Failed to parse string table redirect '%s'. Expected 'OldStringTable' and 'NewStringTable' for a table ID redirect, or 'StringTable', 'OldKey', 'NewKey' for a key redirect."), *ConfigValue);
+				}
+			}
+		}
+	}
+}
+
+void FStringTableRedirects::RedirectTableId(FName& InOutTableId, const EStringTableLoadingPolicy InLoadingPolicy)
+{
+	// Process the static redirect
+	const FName* RedirectedTableId = StringTableRedirects::TableIdRedirects.Find(InOutTableId);
+	if (RedirectedTableId)
+	{
+		InOutTableId = *RedirectedTableId;
+	}
+
+	// Process the asset redirect
+	IStringTableEngineBridge::RedirectAndLoadStringTableAsset(InOutTableId, InLoadingPolicy);
+}
+
+void FStringTableRedirects::RedirectKey(const FName InTableId, FString& InOutKey)
+{
+	const auto* RedirectedKeyMap = StringTableRedirects::TableKeyRedirects.Find(InTableId);
+	if (RedirectedKeyMap)
+	{
+		const FString* RedirectedKey = RedirectedKeyMap->Find(InOutKey);
+		if (RedirectedKey)
+		{
+			InOutKey = *RedirectedKey;
+		}
+	}
+}
+
+void FStringTableRedirects::RedirectTableIdAndKey(FName& InOutTableId, FString& InOutKey, const EStringTableLoadingPolicy InLoadingPolicy)
+{
+	RedirectTableId(InOutTableId, InLoadingPolicy);
+	RedirectKey(InOutTableId, InOutKey);
+}
+
+
+void FStringTableReferenceCollection::CollectAssetReferences(const FName InTableId, FArchive& InAr)
+{
+	if (InAr.IsObjectReferenceCollector())
+	{
+		IStringTableEngineBridge::CollectStringTableAssetReferences(InTableId, InAr);
+	}
+}
diff --git a/Engine/Source/Runtime/Core/Private/Internationalization/StringTableRegistry.cpp b/Engine/Source/Runtime/Core/Private/Internationalization/StringTableRegistry.cpp
new file mode 100644
--- /dev/null
+++ b/Engine/Source/Runtime/Core/Private/Internationalization/StringTableRegistry.cpp
@@ -0,0 +1,242 @@
+// Copyright 1998-2017 Epic Games, Inc. All Rights Reserved.
+
+#include "StringTableRegistry.h"
+#include "StringTableCore.h"
+#include "Misc/ScopeLock.h"
+#include "ModuleManager.h"
+
+#if WITH_EDITOR
+#include "IDirectoryWatcher.h"
+#include "DirectoryWatcherModule.h"
+#endif // WITH_EDITOR
+
+FStringTableRegistry& FStringTableRegistry::Get()
+{
+	static FStringTableRegistry Instance;
+	return Instance;
+}
+
+FStringTableRegistry::FStringTableRegistry()
+{
+#if WITH_EDITOR
+	// Commandlets and in-game don't listen for directory changes
+	if (!IsRunningCommandlet() && GIsEditor)
+	{
+		FDirectoryWatcherModule& DirectoryWatcherModule = FModuleManager::LoadModuleChecked<FDirectoryWatcherModule>("DirectoryWatcher");
+		IDirectoryWatcher* DirectoryWatcher = DirectoryWatcherModule.Get();
+
+		if (DirectoryWatcher)
+		{
+			DirectoryWatcher->RegisterDirectoryChangedCallback_Handle(FPaths::EngineContentDir(), IDirectoryWatcher::FDirectoryChanged::CreateRaw(this, &FStringTableRegistry::OnDirectoryChanged), EngineDirectoryWatcherHandle);
+
+			DirectoryWatcher->RegisterDirectoryChangedCallback_Handle(FPaths::GameContentDir(), IDirectoryWatcher::FDirectoryChanged::CreateRaw(this, &FStringTableRegistry::OnDirectoryChanged), GameDirectoryWatcherHandle);
+		}
+	}
+#endif // WITH_EDITOR
+}
+
+FStringTableRegistry::~FStringTableRegistry()
+{
+#if WITH_EDITOR
+	// Commandlets and in-game don't listen for directory changes
+	if (!IsRunningCommandlet() && GIsEditor)
+	{
+		// If the directory module is still loaded, unregister any delegates
+		if (FModuleManager::Get().IsModuleLoaded("DirectoryWatcher"))
+		{
+			FDirectoryWatcherModule& DirectoryWatcherModule = FModuleManager::GetModuleChecked<FDirectoryWatcherModule>("DirectoryWatcher");
+			IDirectoryWatcher* DirectoryWatcher = DirectoryWatcherModule.Get();
+
+			if (DirectoryWatcher)
+			{
+				DirectoryWatcher->UnregisterDirectoryChangedCallback_Handle(FPaths::EngineContentDir(), EngineDirectoryWatcherHandle);
+				EngineDirectoryWatcherHandle.Reset();
+
+				DirectoryWatcher->UnregisterDirectoryChangedCallback_Handle(FPaths::GameContentDir(), GameDirectoryWatcherHandle);
+				GameDirectoryWatcherHandle.Reset();
+			}
+		}
+	}
+#endif // WITH_EDITOR
+}
+
+void FStringTableRegistry::RegisterStringTable(const FName InTableId, FStringTableRef InTable)
+{
+	FScopeLock RegisteredStringTablesLock(&RegisteredStringTablesCS);
+
+	checkf(!InTableId.IsNone(), TEXT("String table ID cannot be 'None'!"));
+	checkf(!RegisteredStringTables.Contains(InTableId), TEXT("String table ID '%s' is already in use!"), *InTableId.ToString());
+
+	RegisteredStringTables.Add(InTableId, MoveTemp(InTable));
+}
+
+void FStringTableRegistry::UnregisterStringTable(const FName InTableId)
+{
+	FScopeLock RegisteredStringTablesLock(&RegisteredStringTablesCS);
+	RegisteredStringTables.Remove(InTableId);
+}
+
+FStringTablePtr FStringTableRegistry::FindMutableStringTable(const FName InTableId) const
+{
+	FScopeLock RegisteredStringTablesLock(&RegisteredStringTablesCS);
+	return RegisteredStringTables.FindRef(InTableId);
+}
+
+FStringTableConstPtr FStringTableRegistry::FindStringTable(const FName InTableId) const
+{
+	FScopeLock RegisteredStringTablesLock(&RegisteredStringTablesCS);
+	return RegisteredStringTables.FindRef(InTableId);
+}
+
+UStringTable* FStringTableRegistry::FindStringTableAsset(const FName InTableId) const
+{
+	UStringTable* StringTableAsset = nullptr;
+	{
+		FStringTableConstPtr StringTable = FindStringTable(InTableId);
+		if (StringTable.IsValid())
+		{
+			StringTableAsset = StringTable->GetOwnerAsset();
+		}
+	}
+	return StringTableAsset;
+}
+
+void FStringTableRegistry::EnumerateStringTables(const TFunctionRef<bool(const FName&, const FStringTableConstRef&)>& InEnumerator) const
+{
+	FScopeLock RegisteredStringTablesLock(&RegisteredStringTablesCS);
+
+	for (const auto& RegisteredStringTablePair : RegisteredStringTables)
+	{
+		if (!InEnumerator(RegisteredStringTablePair.Key, RegisteredStringTablePair.Value.ToSharedRef()))
+		{
+			break;
+		}
+	}
+}
+
+bool FStringTableRegistry::FindTableIdAndKey(const FText& InText, FName& OutTableId, FString& OutKey) const
+{
+	if (InText.IsFromStringTable())
+	{
+		return FindTableIdAndKey(FTextInspector::GetSharedDisplayString(InText), OutTableId, OutKey) || FTextInspector::GetTableIdAndKey(InText, OutTableId, OutKey);
+	}
+
+	return false;
+}
+
+bool FStringTableRegistry::FindTableIdAndKey(const FTextDisplayStringRef& InDisplayString, FName& OutTableId, FString& OutKey) const
+{
+	FScopeLock RegisteredStringTablesLock(&RegisteredStringTablesCS);
+
+	for (const auto& RegisteredStringTablePair : RegisteredStringTables)
+	{
+		FString FoundKey;
+		if (RegisteredStringTablePair.Value->FindKey(InDisplayString, FoundKey))
+		{
+			OutTableId = RegisteredStringTablePair.Key;
+			OutKey = MoveTemp(FoundKey);
+			return true;
+		}
+	}
+
+	return false;
+}
+
+void FStringTableRegistry::LogMissingStringTableEntry(const FName InTableId, const FString& InKey)
+{
+	FScopeLock LoggedMissingEntriesLock(&LoggedMissingEntriesCS);
+
+	if (const FLocKeySet* LoggedMissingKeys = LoggedMissingEntries.Find(InTableId))
+	{
+		if (LoggedMissingKeys->Contains(InKey))
+		{
+			return;
+		}
+	}
+
+	FLocKeySet& LoggedMissingKeys = LoggedMissingEntries.FindOrAdd(InTableId);
+	LoggedMissingKeys.Add(InKey);
+
+	UE_LOG(LogStringTable, Warning, TEXT("Failed to find string table entry for '%s' '%s'. Did you forget to add a string table redirector?"), *InTableId.ToString(), *InKey);
+}
+
+#if WITH_EDITOR
+void FStringTableRegistry::OnDirectoryChanged(const TArray<FFileChangeData>& InFileChanges)
+{
+	FScopeLock CSVFilesToWatchLock(&CSVFilesToWatchCS);
+
+	for (const FFileChangeData& FileChange : InFileChanges)
+	{
+		if (FileChange.Action == FFileChangeData::FCA_Added || FileChange.Action == FFileChangeData::FCA_Modified)
+		{
+			const FName TableId = CSVFilesToWatch.FindRef(FPaths::ConvertRelativePathToFull(FileChange.Filename));
+			if (!TableId.IsNone())
+			{
+				FStringTablePtr StringTable = FindMutableStringTable(TableId);
+				if (StringTable.IsValid())
+				{
+					if (!StringTable->ImportStrings(FileChange.Filename))
+					{
+						UE_LOG(LogStringTable, Warning, TEXT("Failed to import strings from '%s'"), *FileChange.Filename);
+					}
+				}
+			}
+		}
+	}
+}
+#endif // WITH_EDITOR
+
+void FStringTableRegistry::Internal_NewLocTable(const FName InTableId, const FString& InNamespace)
+{
+	FStringTableRef StringTable = FStringTable::NewStringTable();
+	StringTable->SetNamespace(InNamespace);
+
+	RegisterStringTable(InTableId, StringTable);
+}
+
+void FStringTableRegistry::Internal_LocTableFromFile(const FName InTableId, const FString& InNamespace, const FString& InFilePath, const FString& InRootPath)
+{
+	FStringTableRef StringTable = FStringTable::NewStringTable();
+	StringTable->SetNamespace(InNamespace);
+
+	const FString FullPath = InRootPath / InFilePath;
+	if (!StringTable->ImportStrings(FullPath))
+	{
+		UE_LOG(LogStringTable, Warning, TEXT("Failed to import strings from '%s'"), *FullPath);
+	}
+
+	RegisterStringTable(InTableId, StringTable);
+
+#if WITH_EDITOR
+	{
+		FScopeLock CSVFilesToWatchLock(&CSVFilesToWatchCS);
+		CSVFilesToWatch.Add(FPaths::ConvertRelativePathToFull(FullPath), InTableId);
+	}
+#endif // WITH_EDITOR
+}
+
+void FStringTableRegistry::Internal_SetLocTableEntry(const FName InTableId, const FString& InKey, const FString& InSourceString)
+{
+	FStringTablePtr StringTable = FindMutableStringTable(InTableId);
+	checkf(StringTable.IsValid(), TEXT("Attempting to add a string table entry to the unknown string table '%s'"), InTableId);
+
+	StringTable->SetSourceString(InKey, InSourceString);
+}
+
+void FStringTableRegistry::Internal_SetLocTableEntryMetaData(const FName InTableId, const FString& InKey, const FName InMetaDataId, const FString& InMetaData)
+{
+	FStringTablePtr StringTable = FindMutableStringTable(InTableId);
+	checkf(StringTable.IsValid(), TEXT("Attempting to add string table entry meta-data to the unknown string table '%s'"), InTableId);
+
+	StringTable->SetMetaData(InKey, InMetaDataId, InMetaData);
+}
+
+FText FStringTableRegistry::Internal_FindLocTableEntry(const FName InTableId, const FString& InKey, const EStringTableLoadingPolicy InLoadingPolicy) const
+{
+	// RedirectTableIdAndKey also causes string table assets to be loaded (as it has to do this to process asset redirects)
+	FName RedirectedTableId = InTableId;
+	FString RedirectedKey = InKey;
+	FStringTableRedirects::RedirectTableIdAndKey(RedirectedTableId, RedirectedKey, InLoadingPolicy);
+
+	return FText(RedirectedTableId, RedirectedKey);
+}
diff --git a/Engine/Source/Runtime/Core/Private/Internationalization/Text.cpp b/Engine/Source/Runtime/Core/Private/Internationalization/Text.cpp
--- a/Engine/Source/Runtime/Core/Private/Internationalization/Text.cpp
+++ b/Engine/Source/Runtime/Core/Private/Internationalization/Text.cpp
@@ -6,6 +6,7 @@
 #include "UObject/DebugSerializationFlags.h"
 #include "Internationalization/Culture.h"
 #include "Internationalization/Internationalization.h"
+#include "Internationalization/StringTableRegistry.h"
 
 #include "Internationalization/TextHistory.h"
 #include "Internationalization/ITextData.h"
@@ -77,6 +78,17 @@ const FTextDisplayStringRef FTextInspector::GetSharedDisplayString(const FText&
 	return Text.TextData->GetLocalizedString().ToSharedRef();
 }
 
+bool FTextInspector::GetTableIdAndKey(const FText& Text, FName& OutTableId, FString& OutKey)
+{
+	if (Text.IsFromStringTable())
+	{
+		static_cast<const FTextHistory_StringTableEntry&>(Text.TextData->GetTextHistory()).GetTableIdAndKey(OutTableId, OutKey);
+		return true;
+	}
+
+	return false;
+}
+
 uint32 FTextInspector::GetFlags(const FText& Text)
 {
 	return Text.Flags;
@@ -229,6 +241,12 @@ FText::FText( FString&& InSourceString )
 	TextData->SetTextHistory(FTextHistory_Base(MoveTemp(InSourceString)));
 }
 
+FText::FText( FName InTableId, FString InKey )
+	: TextData(new TIndirectTextData<FTextHistory_StringTableEntry>(FTextHistory_StringTableEntry(InTableId, MoveTemp(InKey))))
+	, Flags(0)
+{
+}
+
 FText::FText( FString&& InSourceString, FTextDisplayStringRef InDisplayString )
 	: TextData(new TLocalizedTextData<FTextHistory_Base>(MoveTemp(InDisplayString)))
 	, Flags(0)
@@ -664,7 +682,7 @@ void FText::SerializeText(FArchive& Ar, FText& Value)
 	}
 	Ar << Value.Flags;
 
-	if (Ar.IsLoading() && Ar.ArIsPersistent)
+	if (Ar.IsLoading() && Ar.IsPersistent())
 	{
 		Value.Flags &= ~(ETextFlag::ConvertedProperty | ETextFlag::InitializedFromString); // Remove conversion flag before saving.
 	}
@@ -753,6 +771,12 @@ void FText::SerializeText(FArchive& Ar, FText& Value)
 					Value.TextData = MakeShareable(new TLocalizedTextData<FTextHistory_Transform>());
 					break;
 				}
+
+			case ETextHistoryType::StringTableEntry:
+				{
+					Value.TextData = MakeShareable(new TIndirectTextData<FTextHistory_StringTableEntry>());
+					break;
+				}
 			default:
 				{
 					bSerializeHistory = false;
@@ -765,7 +789,11 @@ void FText::SerializeText(FArchive& Ar, FText& Value)
 		{
 			FTextHistory& MutableTextHistory = Value.TextData->GetMutableTextHistory();
 			MutableTextHistory.Serialize(Ar);
-			MutableTextHistory.SerializeForDisplayString(Ar, Value.TextData->GetMutableLocalizedString());
+
+			if (Value.TextData->OwnsLocalizedString())
+			{
+				MutableTextHistory.SerializeForDisplayString(Ar, Value.TextData->GetMutableLocalizedString());
+			}
 		}
 	}
 
@@ -789,7 +817,7 @@ void FText::SerializeText(FArchive& Ar, FText& Value)
 #if WITH_EDITOR
 FText FText::ChangeKey( const FString& Namespace, const FString& Key, const FText& Text )
 {
-	return FText(FString(*Text.TextData->GetTextHistory().GetSourceString()), Namespace, Key, Text.Flags);
+	return FText(FString(*Text.TextData->GetTextHistory().GetSourceString()), Namespace, Key, 0);
 }
 #endif
 
@@ -813,6 +841,11 @@ FText FText::CreateChronologicalText(TSharedRef<ITextData, ESPMode::ThreadSafe>
 	return NewText;
 }
 
+FText FText::FromStringTable(const FName InTableId, const FString& InKey, const EStringTableLoadingPolicy InLoadingPolicy)
+{
+	return FStringTableRegistry::Get().Internal_FindLocTableEntry(InTableId, InKey, InLoadingPolicy);
+}
+
 FText FText::FromName( const FName& Val) 
 {
 	return FText::FromString(Val.ToString());
@@ -906,6 +939,11 @@ bool FText::IsCultureInvariant() const
 	return (Flags & ETextFlag::CultureInvariant) != 0;
 }
 
+bool FText::IsFromStringTable() const
+{
+	return TextData->GetTextHistory().GetType() == ETextHistoryType::StringTableEntry;
+}
+
 bool FText::ShouldGatherForLocalization() const
 {
 	const FString& SourceString = GetSourceString();
@@ -922,7 +960,7 @@ bool FText::ShouldGatherForLocalization() const
 		return true;
 	};
 
-	return !((Flags & ETextFlag::CultureInvariant) || (Flags & ETextFlag::Transient)) && !SourceString.IsEmpty() && !IsAllWhitespace(SourceString);
+	return !((Flags & ETextFlag::CultureInvariant) || (Flags & ETextFlag::Transient)) && !IsFromStringTable() && !SourceString.IsEmpty() && !IsAllWhitespace(SourceString);
 }
 
 const FString& FText::GetSourceString() const
@@ -1220,9 +1258,10 @@ bool TextBiDi::IsControlCharacter(const TCHAR InChar)
 const FString FTextStringHelper::InvTextMarker = TEXT("INVTEXT");
 const FString FTextStringHelper::NsLocTextMarker = TEXT("NSLOCTEXT");
 const FString FTextStringHelper::LocTextMarker = TEXT("LOCTEXT");
+const FString FTextStringHelper::LocTableMarker = TEXT("LOCTABLE");
 #undef LOC_DEFINE_REGION
 
-bool FTextStringHelper::ReadFromString_ComplexText(const TCHAR* Buffer, FText& OutValue, const TCHAR* TextNamespace, const TCHAR* PackageNamespace, int32* OutNumCharsRead)
+bool FTextStringHelper::ReadFromString_ComplexText(const TCHAR* Buffer, FText& OutValue, const TCHAR* TextNamespace, const TCHAR* PackageNamespace, int32* OutNumCharsRead, const EStringTableLoadingPolicy InLoadingPolicy)
 {
 #define LOC_DEFINE_REGION
 	const TCHAR* const Start = Buffer;
@@ -1268,7 +1307,38 @@ bool FTextStringHelper::ReadFromString_ComplexText(const TCHAR* Buffer, FText& O
 			return false;					\
 		}
 
-	if (FCString::Strstr(Buffer, *InvTextMarker))
+	if (FCString::Strstr(Buffer, *LocTableMarker))
+	{
+		// Parsing something of the form: LOCTABLE("...", "...")
+		Buffer += LocTableMarker.Len();
+
+		// Walk to the opening bracket
+		WALK_TO_CHARACTER('(');
+
+		// Walk to the opening quote, and then parse out the quoted table ID
+		FString TableIdString;
+		WALK_TO_CHARACTER('"');
+		EXTRACT_QUOTED_STRING(TableIdString);
+
+		// Walk to the opening quote, and then parse out the quoted key
+		FString KeyString;
+		WALK_TO_CHARACTER('"');
+		EXTRACT_QUOTED_STRING(KeyString);
+
+		// Walk to the closing bracket, and then move past it to indicate that the value was successfully imported
+		WALK_TO_CHARACTER(')');
+		++Buffer;
+
+		OutValue = FText::FromStringTable(FName(*TableIdString), KeyString, InLoadingPolicy);
+
+		if (OutNumCharsRead)
+		{
+			*OutNumCharsRead = (Buffer - Start);
+		}
+
+		return true;
+	}
+	else if (FCString::Strstr(Buffer, *InvTextMarker))
 	{
 		// Parsing something of the form: INVTEXT("...")
 		Buffer += InvTextMarker.Len();
@@ -1417,7 +1487,7 @@ bool FTextStringHelper::ReadFromString_ComplexText(const TCHAR* Buffer, FText& O
 	return false;
 }
 
-bool FTextStringHelper::ReadFromString(const TCHAR* Buffer, FText& OutValue, const TCHAR* TextNamespace, const TCHAR* PackageNamespace, int32* OutNumCharsRead, const bool bRequiresQuotes)
+bool FTextStringHelper::ReadFromString(const TCHAR* Buffer, FText& OutValue, const TCHAR* TextNamespace, const TCHAR* PackageNamespace, int32* OutNumCharsRead, const bool bRequiresQuotes, const EStringTableLoadingPolicy InLoadingPolicy)
 {
 	const TCHAR* const Start = Buffer;
 
@@ -1429,7 +1499,7 @@ bool FTextStringHelper::ReadFromString(const TCHAR* Buffer, FText& OutValue, con
 	// First, try and parse the text as a complex text export
 	{
 		int32 SubNumCharsRead = 0;
-		if (FTextStringHelper::ReadFromString_ComplexText(Buffer, OutValue, TextNamespace, PackageNamespace, &SubNumCharsRead))
+		if (FTextStringHelper::ReadFromString_ComplexText(Buffer, OutValue, TextNamespace, PackageNamespace, &SubNumCharsRead, InLoadingPolicy))
 		{
 			Buffer += SubNumCharsRead;
 			if (OutNumCharsRead)
@@ -1484,7 +1554,20 @@ bool FTextStringHelper::WriteToString(FString& Buffer, const FText& Value, const
 #define LOC_DEFINE_REGION
 	const FString& StringValue = FTextInspector::GetDisplayString(Value);
 
-	if (Value.IsCultureInvariant())
+	if (Value.IsFromStringTable())
+	{
+		FName TableId;
+		FString Key;
+		FStringTableRegistry::Get().FindTableIdAndKey(Value, TableId, Key);
+
+		// Produces LOCTABLE("...", "...")
+		Buffer += TEXT("LOCTABLE(\"");
+		Buffer += TableId.ToString().ReplaceCharWithEscapedChar();
+		Buffer += TEXT("\", \"");
+		Buffer += Key.ReplaceCharWithEscapedChar();
+		Buffer += TEXT("\")");
+	}
+	else if (Value.IsCultureInvariant())
 	{
 		// Produces INVTEXT("...")
 		Buffer += TEXT("INVTEXT(\"");
@@ -1536,7 +1619,7 @@ bool FTextStringHelper::WriteToString(FString& Buffer, const FText& Value, const
 bool FTextStringHelper::IsComplexText(const TCHAR* Buffer)
 {
 #define LOC_DEFINE_REGION
-	return FCString::Strstr(Buffer, *InvTextMarker) || FCString::Strstr(Buffer, *NsLocTextMarker) || FCString::Strstr(Buffer, *LocTextMarker);
+	return FCString::Strstr(Buffer, *InvTextMarker) || FCString::Strstr(Buffer, *NsLocTextMarker) || FCString::Strstr(Buffer, *LocTextMarker) || FCString::Strstr(Buffer, *LocTableMarker);
 #undef LOC_DEFINE_REGION
 }
 
diff --git a/Engine/Source/Runtime/Core/Private/Internationalization/TextData.h b/Engine/Source/Runtime/Core/Private/Internationalization/TextData.h
--- a/Engine/Source/Runtime/Core/Private/Internationalization/TextData.h
+++ b/Engine/Source/Runtime/Core/Private/Internationalization/TextData.h
@@ -8,7 +8,6 @@
 #include "Internationalization/ITextData.h"
 #include "Misc/ScopeLock.h"
 
-class FTextHistory;
 
 /** 
  * Implementation of ITextData that stores the common data and functionality.
@@ -39,6 +38,11 @@ public:
 	{
 	}
 
+	virtual bool OwnsLocalizedString() const override
+	{
+		return true;
+	}
+
 	virtual FTextDisplayStringPtr GetLocalizedString() const override
 	{
 		return LocalizedString;
@@ -163,3 +167,77 @@ protected:
 	FString DisplayString;
 	FCriticalSection PersistTextCS;
 };
+
+
+/** 
+ * Implementation of ITextData optimized for storing indirect display string references via its text history.
+ * The history type used must provide a GetDisplayString function that returns a FTextDisplayStringRef.
+ */
+template <typename THistoryType>
+class TIndirectTextData : public ITextData
+{
+public:
+	TIndirectTextData()
+		: History()
+	{
+	}
+
+	explicit TIndirectTextData(THistoryType&& InHistory)
+		: History(MoveTemp(InHistory))
+	{
+	}
+
+	virtual ~TIndirectTextData()
+	{
+	}
+
+	virtual bool OwnsLocalizedString() const override
+	{
+		return false;
+	}
+
+	virtual const FString& GetDisplayString() const override
+	{
+		return *History.GetDisplayString();
+	}
+
+	virtual FTextDisplayStringPtr GetLocalizedString() const override
+	{
+		return History.GetDisplayString();
+	}
+
+	virtual FTextDisplayStringPtr& GetMutableLocalizedString() override
+	{
+		check(0); // Should never be called since OwnsLocalizedString returns false
+		
+		static FTextDisplayStringPtr DummyPtr;
+		return DummyPtr;
+	}
+
+	virtual const FTextHistory& GetTextHistory() const override
+	{
+		return History;
+	}
+
+	virtual FTextHistory& GetMutableTextHistory() override
+	{
+		return History;
+	}
+
+	virtual void PersistText() override
+	{
+	}
+
+	virtual uint16 GetGlobalHistoryRevision() const override
+	{
+		return History.GetRevision();
+	}
+
+	virtual uint16 GetLocalHistoryRevision() const override
+	{
+		return FTextLocalizationManager::Get().GetLocalRevisionForDisplayString(History.GetDisplayString());
+	}
+
+protected:
+	THistoryType History;
+};
diff --git a/Engine/Source/Runtime/Core/Private/Internationalization/TextHistory.cpp b/Engine/Source/Runtime/Core/Private/Internationalization/TextHistory.cpp
--- a/Engine/Source/Runtime/Core/Private/Internationalization/TextHistory.cpp
+++ b/Engine/Source/Runtime/Core/Private/Internationalization/TextHistory.cpp
@@ -5,12 +5,18 @@
 #include "Internationalization/ITextData.h"
 #include "Internationalization/Culture.h"
 #include "Internationalization/Internationalization.h"
+#include "Internationalization/StringTableCore.h"
+#include "Internationalization/StringTableRegistry.h"
 #include "Misc/Guid.h"
+#include "Misc/ScopeLock.h"
 
 #include "Internationalization/TextFormatter.h"
 #include "Internationalization/TextNamespaceUtil.h"
 #include "UObject/PropertyPortFlags.h"
 
+DECLARE_LOG_CATEGORY_EXTERN(LogTextHistory, Log, All);
+DEFINE_LOG_CATEGORY(LogTextHistory);
+
 ///////////////////////////////////////
 // FTextHistory
 
@@ -1059,3 +1065,168 @@ bool FTextHistory_Transform::GetHistoricNumericData(const FText& InText, FHistor
 {
 	return FTextInspector::GetHistoricNumericData(SourceText, OutHistoricNumericData);
 }
+
+// FTextHistory_StringTableEntry
+
+FTextHistory_StringTableEntry::FTextHistory_StringTableEntry(FName InTableId, FString&& InKey)
+	: TableId(InTableId)
+	, Key(MoveTemp(InKey))
+{
+	GetStringTableEntry();
+}
+
+FTextHistory_StringTableEntry::FTextHistory_StringTableEntry(FTextHistory_StringTableEntry&& Other)
+	: FTextHistory(MoveTemp(Other))
+	, TableId(Other.TableId)
+	, Key(MoveTemp(Other.Key))
+	, StringTableEntry(MoveTemp(Other.StringTableEntry))
+{
+	Other.TableId = FName();
+}
+
+FTextHistory_StringTableEntry& FTextHistory_StringTableEntry::operator=(FTextHistory_StringTableEntry&& Other)
+{
+	FTextHistory::operator=(MoveTemp(Other));
+	if (this != &Other)
+	{
+		TableId = Other.TableId;
+		Key = MoveTemp(Other.Key);
+		StringTableEntry = MoveTemp(Other.StringTableEntry);
+	
+		Other.TableId = FName();
+	}
+	return *this;
+}
+
+FText FTextHistory_StringTableEntry::ToText(bool bInAsSource) const
+{
+	if (bInAsSource)
+	{
+		FStringTableEntryConstPtr StringTableEntryPin = GetStringTableEntry();
+		return StringTableEntryPin.IsValid() ? FText::FromString(StringTableEntryPin->GetSourceString()) : FText::GetEmpty();
+	}
+
+	// This should never be called except to build as source
+	check(0);
+
+	return FText::GetEmpty();
+}
+
+const FString* FTextHistory_StringTableEntry::GetSourceString() const
+{
+	FStringTableEntryConstPtr StringTableEntryPin = GetStringTableEntry();
+	if (StringTableEntryPin.IsValid())
+	{
+		return &StringTableEntryPin->GetSourceString();
+	}
+
+	static const FString MissingSourceString = TEXT("<MISSING STRING TABLE ENTRY>");
+	return &MissingSourceString;
+}
+
+FTextDisplayStringRef FTextHistory_StringTableEntry::GetDisplayString() const
+{
+	FStringTableEntryConstPtr StringTableEntryPin = GetStringTableEntry();
+	if (StringTableEntryPin.IsValid())
+	{
+		FTextDisplayStringPtr DisplayString = StringTableEntryPin->GetDisplayString();
+		if (DisplayString.IsValid())
+		{
+			return DisplayString.ToSharedRef();
+		}
+	}
+
+	static const FTextDisplayStringRef MissingDisplayString = MakeShareable(new FString(TEXT("<MISSING STRING TABLE ENTRY>")));
+	return MissingDisplayString;
+}
+
+void FTextHistory_StringTableEntry::GetTableIdAndKey(FName& OutTableId, FString& OutKey) const
+{
+	OutTableId = TableId;
+	OutKey = Key;
+}
+
+void FTextHistory_StringTableEntry::Serialize(FArchive& Ar)
+{
+	if (Ar.IsSaving())
+	{
+		int8 HistoryType = (int8)ETextHistoryType::StringTableEntry;
+		Ar << HistoryType;
+	}
+
+	if (Ar.IsLoading())
+	{
+		// We will definitely need to do a rebuild later
+		Revision = 0;
+
+		Ar << TableId;
+		Ar << Key;
+
+		// String Table assets should already have been created via dependency loading when using the EDL (although they may not be fully loaded yet)
+		FStringTableRedirects::RedirectTableIdAndKey(TableId, Key, GEventDrivenLoaderEnabled ? EStringTableLoadingPolicy::Find : EStringTableLoadingPolicy::FindOrLoad);
+
+		// Re-cache the pointer
+		StringTableEntry.Reset();
+		GetStringTableEntry();
+	}
+	else if (Ar.IsSaving())
+	{
+		// Update the table ID and key on save to make sure they're up-to-date
+		FStringTableEntryConstPtr StringTableEntryPin = GetStringTableEntry();
+		if (StringTableEntryPin.IsValid())
+		{
+			FTextDisplayStringPtr DisplayString = StringTableEntryPin->GetDisplayString();
+			FStringTableRegistry::Get().FindTableIdAndKey(DisplayString.ToSharedRef(), TableId, Key);
+		}
+
+		Ar << TableId;
+		Ar << Key;
+	}
+
+	// Collect string table asset references
+	FStringTableReferenceCollection::CollectAssetReferences(TableId, Ar);
+}
+
+void FTextHistory_StringTableEntry::SerializeForDisplayString(FArchive& Ar, FTextDisplayStringPtr& InOutDisplayString)
+{
+	if (Ar.IsLoading())
+	{
+		// We will definitely need to do a rebuild later
+		Revision = 0;
+	}
+}
+
+FStringTableEntryConstPtr FTextHistory_StringTableEntry::GetStringTableEntry() const
+{
+	FStringTableEntryConstPtr StringTableEntryPin = StringTableEntry.Pin();
+
+	bool bSuppressMissingEntryWarning = false;
+
+	if (!StringTableEntryPin.IsValid() || !StringTableEntryPin->IsOwned())
+	{
+		FScopeLock ScopeLock(&StringTableEntryCS);
+
+		// Test again now that we've taken the lock in-case another thread beat us to it
+		StringTableEntryPin = StringTableEntry.Pin();
+		if (!StringTableEntryPin.IsValid() || !StringTableEntryPin->IsOwned())
+		{
+			StringTableEntryPin.Reset();
+
+			FStringTableConstPtr StringTable = FStringTableRegistry::Get().FindStringTable(TableId);
+			if (StringTable.IsValid())
+			{
+				bSuppressMissingEntryWarning = !StringTable->IsLoaded();
+				StringTableEntryPin = StringTable->FindEntry(Key);
+			}
+
+			StringTableEntry = StringTableEntryPin;
+		}
+	}
+
+	if (!StringTableEntryPin.IsValid() && !bSuppressMissingEntryWarning)
+	{
+		FStringTableRegistry::Get().LogMissingStringTableEntry(TableId, Key);
+	}
+
+	return StringTableEntryPin;
+}
diff --git a/Engine/Source/Runtime/Core/Private/Internationalization/TextHistory.h b/Engine/Source/Runtime/Core/Private/Internationalization/TextHistory.h
--- a/Engine/Source/Runtime/Core/Private/Internationalization/TextHistory.h
+++ b/Engine/Source/Runtime/Core/Private/Internationalization/TextHistory.h
@@ -7,6 +7,7 @@
 #include "Templates/SharedPointer.h"
 #include "Misc/Optional.h"
 #include "Internationalization/Text.h"
+#include "Internationalization/StringTableCoreFwd.h"
 #include "Misc/DateTime.h"
 
 enum class ETextHistoryType
@@ -22,6 +23,7 @@ enum class ETextHistoryType
 	AsTime,
 	AsDateTime,
 	Transform,
+	StringTableEntry,
 
 	// Add new enum types at the end only! They are serialized by index.
 };
@@ -38,6 +40,9 @@ public:
 	FTextHistory(FTextHistory&& Other);
 	FTextHistory& operator=(FTextHistory&& Other);
 
+	/** Get the type of this history */
+	virtual ETextHistoryType GetType() const = 0;
+
 	/** Rebuilds the FText from the hierarchical history, the result should be in the current locale */
 	virtual FText ToText(bool bInAsSource) const = 0;
 	
@@ -90,6 +95,7 @@ public:
 	FTextHistory_Base& operator=(FTextHistory_Base&& Other);
 
 	//~ Begin FTextHistory Interface
+	virtual ETextHistoryType GetType() const override { return ETextHistoryType::Base; }
 	virtual FText ToText(bool bInAsSource) const override;
 	virtual void Serialize(FArchive& Ar) override;
 	virtual void SerializeForDisplayString(FArchive& Ar, FTextDisplayStringPtr& InOutDisplayString) override;
@@ -122,6 +128,7 @@ public:
 	FTextHistory_NamedFormat& operator=(FTextHistory_NamedFormat&& Other);
 
 	//~ Begin FTextHistory Interface
+	virtual ETextHistoryType GetType() const override { return ETextHistoryType::NamedFormat; }
 	virtual FText ToText(bool bInAsSource) const override;
 	virtual void Serialize(FArchive& Ar) override;
 	virtual void GetHistoricFormatData(const FText& InText, TArray<FHistoricTextFormatData>& OutHistoricFormatData) const override;
@@ -150,6 +157,7 @@ public:
 	FTextHistory_OrderedFormat& operator=(FTextHistory_OrderedFormat&& Other);
 
 	//~ Begin FTextHistory Interface
+	virtual ETextHistoryType GetType() const override { return ETextHistoryType::OrderedFormat; }
 	virtual FText ToText(bool bInAsSource) const override;
 	virtual void Serialize(FArchive& Ar) override;
 	virtual void GetHistoricFormatData(const FText& InText, TArray<FHistoricTextFormatData>& OutHistoricFormatData) const override;
@@ -178,6 +186,7 @@ public:
 	FTextHistory_ArgumentDataFormat& operator=(FTextHistory_ArgumentDataFormat&& Other);
 
 	//~ Begin FTextHistory Interface
+	virtual ETextHistoryType GetType() const override { return ETextHistoryType::ArgumentFormat; }
 	virtual FText ToText(bool bInAsSource) const override;
 	virtual void Serialize(FArchive& Ar) override;
 	virtual void GetHistoricFormatData(const FText& InText, TArray<FHistoricTextFormatData>& OutHistoricFormatData) const override;
@@ -235,6 +244,7 @@ public:
 	FTextHistory_AsNumber& operator=(FTextHistory_AsNumber&& Other);
 
 	//~ Begin FTextHistory Interface
+	virtual ETextHistoryType GetType() const override { return ETextHistoryType::AsNumber; }
 	virtual FText ToText(bool bInAsSource) const override;
 	virtual void Serialize(FArchive& Ar) override;
 	virtual bool GetHistoricNumericData(const FText& InText, FHistoricTextNumericData& OutHistoricNumericData) const override;
@@ -258,6 +268,7 @@ public:
 	FTextHistory_AsPercent& operator=(FTextHistory_AsPercent&& Other);
 
 	//~ Begin FTextHistory Interface
+	virtual ETextHistoryType GetType() const override { return ETextHistoryType::AsPercent; }
 	virtual FText ToText(bool bInAsSource) const override;
 	virtual void Serialize(FArchive& Ar) override;
 	virtual bool GetHistoricNumericData(const FText& InText, FHistoricTextNumericData& OutHistoricNumericData) const override;
@@ -281,6 +292,7 @@ public:
 	FTextHistory_AsCurrency& operator=(FTextHistory_AsCurrency&& Other);
 
 	//~ Begin FTextHistory Interface
+	virtual ETextHistoryType GetType() const override { return ETextHistoryType::AsCurrency; }
 	virtual FText ToText(bool bInAsSource) const override;
 	virtual void Serialize(FArchive& Ar) override;
 	//~ End FTextHistory Interface
@@ -306,6 +318,7 @@ public:
 	FTextHistory_AsDate& operator=(FTextHistory_AsDate&& Other);
 
 	//~ Begin FTextHistory Interface
+	virtual ETextHistoryType GetType() const override { return ETextHistoryType::AsDate; }
 	virtual FText ToText(bool bInAsSource) const override;
 	virtual void Serialize(FArchive& Ar) override;
 	//~ End FTextHistory Interface
@@ -337,6 +350,7 @@ public:
 	FTextHistory_AsTime& operator=(FTextHistory_AsTime&& Other);
 
 	//~ Begin FTextHistory Interface
+	virtual ETextHistoryType GetType() const override { return ETextHistoryType::AsTime; }
 	virtual FText ToText(bool bInAsSource) const override;
 	virtual void Serialize(FArchive& Ar) override;
 	//~ End FTextHistory Interface
@@ -368,6 +382,7 @@ public:
 	FTextHistory_AsDateTime& operator=(FTextHistory_AsDateTime&& Other);
 
 	//~ Begin FTextHistory Interface
+	virtual ETextHistoryType GetType() const override { return ETextHistoryType::AsDateTime; }
 	virtual FText ToText(bool bInAsSource) const override;
 	virtual void Serialize(FArchive& Ar) override;
 	//~ End FTextHistory Interfaces
@@ -409,6 +424,7 @@ public:
 	FTextHistory_Transform& operator=(FTextHistory_Transform&& Other);
 
 	//~ Begin FTextHistory Interface
+	virtual ETextHistoryType GetType() const override { return ETextHistoryType::Transform; }
 	virtual FText ToText(bool bInAsSource) const override;
 	virtual void Serialize(FArchive& Ar) override;
 	virtual void GetHistoricFormatData(const FText& InText, TArray<FHistoricTextFormatData>& OutHistoricFormatData) const override;
@@ -425,3 +441,52 @@ private:
 	/** How the source text was transformed */
 	ETransformType TransformType;
 };
+
+/** Holds a pointer to a referenced display string from a string table. */
+class CORE_API FTextHistory_StringTableEntry : public FTextHistory
+{
+public:
+	FTextHistory_StringTableEntry() {}
+	FTextHistory_StringTableEntry(FName InTableId, FString&& InKey);
+
+	/** Allow moving */
+	FTextHistory_StringTableEntry(FTextHistory_StringTableEntry&& Other);
+	FTextHistory_StringTableEntry& operator=(FTextHistory_StringTableEntry&& Other);
+
+	//~ Begin FTextHistory Interface
+	virtual ETextHistoryType GetType() const override { return ETextHistoryType::StringTableEntry; }
+	virtual FText ToText(bool bInAsSource) const override;
+	virtual void Serialize(FArchive& Ar) override;
+	virtual void SerializeForDisplayString(FArchive& Ar, FTextDisplayStringPtr& InOutDisplayString) override;
+	virtual const FString* GetSourceString() const override;
+	//~ End FTextHistory Interface
+
+	FTextDisplayStringRef GetDisplayString() const;
+
+	void GetTableIdAndKey(FName& OutTableId, FString& OutKey) const;
+
+protected:
+	//~ Begin FTextHistory Interface
+	virtual bool CanRebuildText() { return false; }
+	//~ End FTextHistory Interface
+
+private:
+	/** Disallow copying */
+	FTextHistory_StringTableEntry(const FTextHistory_StringTableEntry&);
+	FTextHistory_StringTableEntry& operator=(FTextHistory_StringTableEntry&);
+
+	/** Get the string table pointer, potentially re-caching it if it's missing or stale */
+	FStringTableEntryConstPtr GetStringTableEntry() const;
+
+	/** The string table ID being referenced */
+	FName TableId;
+
+	/** The key within the string table being referenced */
+	FString Key;
+
+	/** Cached string table entry pointer */
+	mutable FStringTableEntryConstWeakPtr StringTableEntry;
+
+	/** Critical section preventing concurrent access when re-caching StringTableEntry */
+	mutable FCriticalSection StringTableEntryCS;
+};
diff --git a/Engine/Source/Runtime/Core/Private/Internationalization/TextLocalizationManager.cpp b/Engine/Source/Runtime/Core/Private/Internationalization/TextLocalizationManager.cpp
--- a/Engine/Source/Runtime/Core/Private/Internationalization/TextLocalizationManager.cpp
+++ b/Engine/Source/Runtime/Core/Private/Internationalization/TextLocalizationManager.cpp
@@ -10,6 +10,7 @@
 #include "Misc/Paths.h"
 #include "Internationalization/Culture.h"
 #include "Internationalization/Internationalization.h"
+#include "Internationalization/StringTableCore.h"
 #include "Stats/Stats.h"
 #include "Misc/ConfigCacheIni.h"
 #include "Misc/App.h"
@@ -32,6 +33,8 @@ void EndInitTextLocalization()
 {
 	DECLARE_SCOPE_CYCLE_COUNTER(TEXT("EndInitTextLocalization"), STAT_EndInitTextLocalization, STATGROUP_LoadTime);
 
+	FStringTableRedirects::InitStringTableRedirects();
+
 	const bool ShouldLoadEditor = WITH_EDITOR;
 	const bool ShouldLoadGame = FApp::IsGame();
 
diff --git a/Engine/Source/Runtime/Core/Private/Linux/LinuxCursor.cpp b/Engine/Source/Runtime/Core/Private/Linux/LinuxCursor.cpp
--- a/Engine/Source/Runtime/Core/Private/Linux/LinuxCursor.cpp
+++ b/Engine/Source/Runtime/Core/Private/Linux/LinuxCursor.cpp
@@ -33,6 +33,7 @@ FLinuxCursor::FLinuxCursor()
 	for( int32 CurCursorIndex = 0; CurCursorIndex < EMouseCursor::TotalCursorCount; ++CurCursorIndex )
 	{
 		CursorHandles[ CurCursorIndex ] = NULL;
+		CursorOverrideHandles[ CurCursorIndex ] = NULL;
 
 		SDL_HCursor CursorHandle = NULL;
 		switch( CurCursorIndex )
@@ -229,11 +230,14 @@ void FLinuxCursor::SetPosition( const int32 X, const int32 Y )
 	SetCachedPosition(X, Y);
 }
 
-void FLinuxCursor::SetType( const EMouseCursor::Type InNewCursor )
-{	
-	checkf( InNewCursor < EMouseCursor::TotalCursorCount, TEXT("Invalid cursor(%d) supplied"), InNewCursor );
+void FLinuxCursor::SetType(const EMouseCursor::Type InNewCursor)
+{
+	checkf(InNewCursor < EMouseCursor::TotalCursorCount, TEXT("Invalid cursor(%d) supplied"), InNewCursor);
 	CurrentType = InNewCursor;
-	if(CursorHandles[InNewCursor] == nullptr)
+
+	SDL_HCursor CurrentCursor = CursorOverrideHandles[InNewCursor] ? CursorOverrideHandles[InNewCursor] : CursorHandles[InNewCursor];
+
+	if (CurrentCursor == nullptr)
 	{
 		if (InNewCursor != EMouseCursor::Custom)
 		{
@@ -246,7 +250,7 @@ void FLinuxCursor::SetType( const EMouseCursor::Type InNewCursor )
 	{
 		bHidden = false;
 		SDL_ShowCursor(SDL_ENABLE);
-		SDL_SetCursor(CursorHandles[InNewCursor]);
+		SDL_SetCursor(CurrentCursor);
 	}
 }
 
@@ -333,7 +337,15 @@ bool FLinuxCursor::IsHidden()
 	return bHidden;
 }
 
-void FLinuxCursor::SetCustomShape(void* CursorHandle)
+void FLinuxCursor::SetTypeShape(EMouseCursor::Type InCursorType, void* InCursorHandle)
 {
-	
+	checkf(InCursorType < EMouseCursor::TotalCursorCount, TEXT("Invalid cursor(%d) supplied"), (int)InCursorType);
+
+	SDL_Cursor* CursorHandle = (SDL_Cursor*)InCursorHandle;
+	CursorOverrideHandles[InCursorType] = CursorHandle;
+
+	if (CurrentType == InCursorType)
+	{
+		SetType(CurrentType);
+	}
 }
\ No newline at end of file
diff --git a/Engine/Source/Runtime/Core/Private/Mac/MacCursor.cpp b/Engine/Source/Runtime/Core/Private/Mac/MacCursor.cpp
--- a/Engine/Source/Runtime/Core/Private/Mac/MacCursor.cpp
+++ b/Engine/Source/Runtime/Core/Private/Mac/MacCursor.cpp
@@ -41,6 +41,7 @@ FMacCursor::FMacCursor()
 	for (int32 CursorIndex = 0; CursorIndex < EMouseCursor::TotalCursorCount; ++CursorIndex)
 	{
 		CursorHandles[CursorIndex] = NULL;
+		CursorOverrideHandles[CursorIndex] = NULL;
 
 		NSCursor *CursorHandle = NULL;
 		switch (CursorIndex)
@@ -288,17 +289,21 @@ void FMacCursor::SetType(const EMouseCursor::Type InNewCursor)
 	{
 		MacApplication->SetHighPrecisionMouseMode(false, nullptr);
 	}
+	
 	CurrentType = InNewCursor;
-	CurrentCursor = CursorHandles[InNewCursor];
+	CurrentCursor = CursorOverrideHandles[InNewCursor] ? CursorOverrideHandles[InNewCursor] : CursorHandles[InNewCursor];
+
 	if (CurrentCursor)
 	{
 		[CurrentCursor set];
 	}
+
 	UpdateVisibility();
 }
 
 void FMacCursor::GetSize(int32& Width, int32& Height) const
 {
+	// TODO: This isn't accurate
 	Width = 16;
 	Height = 16;
 }
@@ -534,15 +539,22 @@ const FVector2D& FMacCursor::GetMouseScaling() const
 	}
 }
 
-void FMacCursor::SetCustomShape(void* InCursorHandle)
+void FMacCursor::SetTypeShape(EMouseCursor::Type InCursorType, void* InCursorHandle)
 {
 	NSCursor* CursorHandle = (NSCursor*)InCursorHandle;
+	
+	{
+		SCOPED_AUTORELEASE_POOL;
+		[CursorHandle retain];
+		if (CursorOverrideHandles[InCursorType] != NULL)
+		{
+			[CursorOverrideHandles[InCursorType] release];
+		}
+		CursorOverrideHandles[InCursorType] = CursorHandle;
+	}
 
-	SCOPED_AUTORELEASE_POOL;
-	[CursorHandle retain];
-	if (CursorHandles[EMouseCursor::Custom] != NULL)
+	if (CurrentType == InCursorType)
 	{
-		[CursorHandles[EMouseCursor::Custom] release];
+		SetType(CurrentType);
 	}
-	CursorHandles[EMouseCursor::Custom] = CursorHandle;
 }
\ No newline at end of file
diff --git a/Engine/Source/Runtime/Core/Private/Misc/AutomationTest.cpp b/Engine/Source/Runtime/Core/Private/Misc/AutomationTest.cpp
--- a/Engine/Source/Runtime/Core/Private/Misc/AutomationTest.cpp
+++ b/Engine/Source/Runtime/Core/Private/Misc/AutomationTest.cpp
@@ -730,7 +730,7 @@ void FAutomationTestBase::AddError(const FString& InError, int32 StackOffset)
 
 void FAutomationTestBase::AddError(const FString& InError, const FString& InFilename, int32 InLineNumber)
 {
-	if (!bSuppressLogs)
+	if ( !bSuppressLogs )
 	{
 		ExecutionInfo.Errors.Add(FAutomationEvent(InError, ExecutionInfo.Context, InFilename, InLineNumber));
 	}
@@ -738,17 +738,17 @@ void FAutomationTestBase::AddError(const FString& InError, const FString& InFile
 
 void FAutomationTestBase::AddWarning( const FString& InWarning )
 {
-	if( !bSuppressLogs )
+	if ( !bSuppressLogs )
 	{
-		ExecutionInfo.Warnings.Add( InWarning );
+		ExecutionInfo.Warnings.Add(InWarning);
 	}
 }
 
 void FAutomationTestBase::AddLogItem( const FString& InLogItem )
 {
-	if( !bSuppressLogs )
+	if ( !bSuppressLogs )
 	{
-		ExecutionInfo.LogItems.Add( InLogItem );
+		ExecutionInfo.LogItems.Add(InLogItem);
 	}
 }
 
@@ -757,6 +757,11 @@ void FAutomationTestBase::AddAnalyticsItem(const FString& InAnalyticsItem)
 	ExecutionInfo.AnalyticsItems.Add(InAnalyticsItem);
 }
 
+void FAutomationTestBase::AddRemoteArtifact(const FString& ArtifactName, const FString& RemoteArtifactPath)
+{
+	ExecutionInfo.Artifacts.Add(FAutomationTestArtifact(ArtifactName, RemoteArtifactPath));
+}
+
 bool FAutomationTestBase::HasAnyErrors() const
 {
 	return ExecutionInfo.Errors.Num() > 0;
diff --git a/Engine/Source/Runtime/Core/Private/Misc/FileHelper.cpp b/Engine/Source/Runtime/Core/Private/Misc/FileHelper.cpp
--- a/Engine/Source/Runtime/Core/Private/Misc/FileHelper.cpp
+++ b/Engine/Source/Runtime/Core/Private/Misc/FileHelper.cpp
@@ -559,21 +559,16 @@ void FMaintenance::DeleteOldLogs()
 				}
 			}
 		}
+	}
 
-		// Remove old UE4 crash contexts
-		TArray<FString> Directories;
-		IFileManager::Get().FindFiles( Directories, *FString::Printf( TEXT( "%s/UE4CC*" ), *FPaths::GameLogDir() ), false, true );
+	// Remove all legacy UE4 crash contexts (regardless of age and purge settings, these are deprecated)
+	TArray<FString> Directories;
+	IFileManager::Get().FindFiles(Directories, *FString::Printf(TEXT("%s/UE4CC*"), *FPaths::GameLogDir()), false, true);
 
-		for (const FString& Dir : Directories)
-		{
-			const FString CrashContextDirectory = FPaths::GameLogDir() / Dir;
-			const FDateTime DirectoryAccessTime = IFileManager::Get().GetTimeStamp( *CrashContextDirectory );
-			if (FDateTime::Now() - DirectoryAccessTime > FTimespan::FromDays( PurgeLogsDays ))
-			{
-				UE_LOG( LogStreaming, Log, TEXT( "Deleting old crash context %s" ), *Dir );
-				IFileManager::Get().DeleteDirectory( *CrashContextDirectory, false, true );
-			}
-		}
+	for (const FString& Dir : Directories)
+	{
+		const FString CrashConfigDirectory = FPaths::GameLogDir() / Dir;
+		IFileManager::Get().DeleteDirectory(*CrashConfigDirectory, false, true);
 	}
 }
 
diff --git a/Engine/Source/Runtime/Core/Private/Misc/Parse.cpp b/Engine/Source/Runtime/Core/Private/Misc/Parse.cpp
--- a/Engine/Source/Runtime/Core/Private/Misc/Parse.cpp
+++ b/Engine/Source/Runtime/Core/Private/Misc/Parse.cpp
@@ -371,6 +371,11 @@ bool FParse::QuotedString( const TCHAR* Buffer, FString& Value, int32* OutNumCha
 			Value += TCHAR('\r');
 			++Buffer;
 		}
+		else if (*Buffer == TCHAR('t')) // escaped tab
+		{
+			Value += TCHAR('\t');
+			++Buffer;
+		}
 		else // some other escape sequence, assume it's a hex character value
 		{
 			Value += FString::Printf(TEXT("%c"), (HexDigit(Buffer[0]) * 16) + HexDigit(Buffer[1]));
diff --git a/Engine/Source/Runtime/Core/Private/Modules/ModuleManager.cpp b/Engine/Source/Runtime/Core/Private/Modules/ModuleManager.cpp
--- a/Engine/Source/Runtime/Core/Private/Modules/ModuleManager.cpp
+++ b/Engine/Source/Runtime/Core/Private/Modules/ModuleManager.cpp
@@ -21,6 +21,14 @@ DEFINE_LOG_CATEGORY_STATIC(LogModuleManager, Log, All);
 
 int32 FModuleManager::FModuleInfo::CurrentLoadOrder = 1;
 
+void FModuleManager::WarnIfItWasntSafeToLoadHere(const FName InModuleName)
+{
+	if ( !IsInGameThread() )
+	{
+		UE_LOG(LogModuleManager, Warning, TEXT("ModuleManager: Attempting to load '%s' outside the main thread.  This module was already loaded - so we didn't crash but this isn't safe.  Please call LoadModule on the main/game thread only.  You can use GetModule or GetModuleChecked instead, those are safe to call outside the game thread."), *InModuleName.ToString());
+	}
+}
+
 FModuleManager::ModuleInfoPtr FModuleManager::FindModule(FName InModuleName)
 {
 	FModuleManager::ModuleInfoPtr Result = nullptr;
diff --git a/Engine/Source/Runtime/Core/Private/Windows/WindowsCursor.cpp b/Engine/Source/Runtime/Core/Private/Windows/WindowsCursor.cpp
--- a/Engine/Source/Runtime/Core/Private/Windows/WindowsCursor.cpp
+++ b/Engine/Source/Runtime/Core/Private/Windows/WindowsCursor.cpp
@@ -21,6 +21,7 @@ FWindowsCursor::FWindowsCursor()
 	for( int32 CurCursorIndex = 0; CurCursorIndex < EMouseCursor::TotalCursorCount; ++CurCursorIndex )
 	{
 		CursorHandles[ CurCursorIndex ] = NULL;
+		CursorOverrideHandles[ CurCursorIndex ] = NULL;
 
 		HCURSOR CursorHandle = NULL;
 		switch( CurCursorIndex )
@@ -161,11 +162,21 @@ void FWindowsCursor::SetType( const EMouseCursor::Type InNewCursor )
 	// NOTE: Watch out for contention with FWindowsViewport::UpdateMouseCursor
 	checkf( InNewCursor < EMouseCursor::TotalCursorCount, TEXT("Invalid cursor(%d) supplied"), (int)InNewCursor );
 	CurrentType = InNewCursor;
-	::SetCursor( CursorHandles[ InNewCursor ] );
+
+	if (CursorOverrideHandles[InNewCursor])
+	{
+		::SetCursor(CursorOverrideHandles[InNewCursor]);
+	}
+	else
+	{
+		::SetCursor(CursorHandles[InNewCursor]);
+	}
 }
 
 void FWindowsCursor::GetSize( int32& Width, int32& Height ) const
 {
+	//TODO this is wrong, this should query the size of the cursor on the platform.
+
 	Width = 16;
 	Height = 16;
 }
@@ -191,8 +202,15 @@ void FWindowsCursor::Lock( const RECT* const Bounds )
 	// If the cursor is not visible and we're running game, assume we're in a mode where the mouse is controlling the camera and lock it to the center of the widget.
 }
 
-void FWindowsCursor::SetCustomShape(void* InCursorHandle)
+void FWindowsCursor::SetTypeShape(EMouseCursor::Type InCursorType, void* InCursorHandle)
 {
+	checkf(InCursorType < EMouseCursor::TotalCursorCount, TEXT("Invalid cursor(%d) supplied"), (int)InCursorType);
+
 	HCURSOR CursorHandle = (HCURSOR)InCursorHandle;
-	CursorHandles[EMouseCursor::Custom] = CursorHandle;
+	CursorOverrideHandles[InCursorType] = CursorHandle;
+
+	if (CurrentType == InCursorType)
+	{
+		SetType(CurrentType);
+	}
 }
diff --git a/Engine/Source/Runtime/Core/Private/Windows/WindowsPlatformCrashContext.cpp b/Engine/Source/Runtime/Core/Private/Windows/WindowsPlatformCrashContext.cpp
--- a/Engine/Source/Runtime/Core/Private/Windows/WindowsPlatformCrashContext.cpp
+++ b/Engine/Source/Runtime/Core/Private/Windows/WindowsPlatformCrashContext.cpp
@@ -137,7 +137,7 @@ int32 ReportCrashUsingCrashReportClient(FWindowsPlatformCrashContext& InContext,
 		InContext.GetUniqueCrashName(CrashGUID, FGenericCrashContext::CrashGUIDLength);
 		const FString AppName = FString::Printf(TEXT("UE4-%s"), FApp::GetGameName());
 
-		FString CrashFolder = FPaths::Combine(*FPaths::GameLogDir(), CrashGUID);
+		FString CrashFolder = FPaths::Combine(*FPaths::GameSavedDir(), TEXT("Crashes"), CrashGUID);
 		FString CrashFolderAbsolute = IFileManager::Get().ConvertToAbsolutePathForExternalAppForWrite(*CrashFolder);
 		if (IFileManager::Get().MakeDirectory(*CrashFolderAbsolute, true))
 		{
diff --git a/Engine/Source/Runtime/Core/Public/GenericPlatform/GenericPlatformCrashContext.h b/Engine/Source/Runtime/Core/Public/GenericPlatform/GenericPlatformCrashContext.h
--- a/Engine/Source/Runtime/Core/Public/GenericPlatform/GenericPlatformCrashContext.h
+++ b/Engine/Source/Runtime/Core/Public/GenericPlatform/GenericPlatformCrashContext.h
@@ -86,6 +86,8 @@ public:
 	static const TCHAR* CrashConfigFileNameW;
 	static const FString CrashConfigExtension;
 	static const FString ConfigSectionName;
+	static const FString CrashConfigPurgeDays;
+	static const FString CrashGUIDRootPrefix;
 
 	static const FString CrashContextExtension;
 	static const FString RuntimePropertiesTag;
@@ -170,6 +172,12 @@ public:
 	/** Helper to get the crash report client config filepath saved by this instance and copied to each crash report folder. */
 	static const TCHAR* GetCrashConfigFilePath();
 
+	/** Helper to get the crash report client config folder used by GetCrashConfigFilePath(). */
+	static const TCHAR* GetCrashConfigFolder();
+
+	/** Helper to clean out old files in the crash report client config folder. */
+	static void PurgeOldCrashConfig();
+
 	/**
 	 * @return whether this crash is a non-crash event
 	 */
diff --git a/Engine/Source/Runtime/Core/Public/GenericPlatform/ICursor.h b/Engine/Source/Runtime/Core/Public/GenericPlatform/ICursor.h
--- a/Engine/Source/Runtime/Core/Public/GenericPlatform/ICursor.h
+++ b/Engine/Source/Runtime/Core/Public/GenericPlatform/ICursor.h
@@ -67,6 +67,7 @@ namespace EMouseCursor
 }
 
 
+
 class ICursor
 {
 public:
@@ -100,5 +101,9 @@ public:
 	 */
 	virtual void Lock( const RECT* const Bounds ) = 0;
 
-	virtual void SetCustomShape(void* CursorHandle) = 0;
+	/**  */
+	virtual void SetTypeShape(EMouseCursor::Type InCursorType, void* CursorHandle) = 0;
+
+	DEPRECATED(4.16, "Use SetTypeShape instead.")
+	void SetCustomShape(void* CursorHandle) { SetTypeShape(EMouseCursor::Custom, CursorHandle); }
 };
diff --git a/Engine/Source/Runtime/Core/Public/HTML5/HTML5Cursor.h b/Engine/Source/Runtime/Core/Public/HTML5/HTML5Cursor.h
--- a/Engine/Source/Runtime/Core/Public/HTML5/HTML5Cursor.h
+++ b/Engine/Source/Runtime/Core/Public/HTML5/HTML5Cursor.h
@@ -30,7 +30,7 @@ public:
 
 	virtual void Lock( const RECT* const Bounds ) override;
 
-	virtual void SetCustomShape(void* CursorHandle) override;
+	virtual void SetTypeShape(EMouseCursor::Type InCursorType, void* CursorHandle) override { }
 
 private:
 	EMouseCursor::Type CurrentType;
diff --git a/Engine/Source/Runtime/Core/Public/Internationalization/ITextData.h b/Engine/Source/Runtime/Core/Public/Internationalization/ITextData.h
--- a/Engine/Source/Runtime/Core/Public/Internationalization/ITextData.h
+++ b/Engine/Source/Runtime/Core/Public/Internationalization/ITextData.h
@@ -18,6 +18,11 @@ public:
 	{
 	}
 
+	/**
+	 * True if this text data owns its localized string pointer, and allows you to call GetMutableLocalizedString on it
+	 */
+	virtual bool OwnsLocalizedString() const = 0;
+
 	/**
 	 * Get the string to use for display purposes.
 	 * This may have come from the localization manager, or may been generated at runtime (eg, via FText::AsNumber).
diff --git a/Engine/Source/Runtime/Core/Public/Internationalization/LocKeyFuncs.h b/Engine/Source/Runtime/Core/Public/Internationalization/LocKeyFuncs.h
new file mode 100644
--- /dev/null
+++ b/Engine/Source/Runtime/Core/Public/Internationalization/LocKeyFuncs.h
@@ -0,0 +1,41 @@
+// Copyright 1998-2017 Epic Games, Inc. All Rights Reserved.
+#pragma once
+
+#include "CoreTypes.h"
+#include "Containers/Set.h"
+#include "Containers/Map.h"
+
+/** Case sensitive hashing function for TSet */
+struct FLocKeySetFuncs : BaseKeyFuncs<FString, FString>
+{
+	static FORCEINLINE const FString& GetSetKey(const FString& Element)
+	{
+		return Element;
+	}
+	static FORCEINLINE bool Matches(const FString& A, const FString& B)
+	{
+		return A.Equals(B, ESearchCase::CaseSensitive);
+	}
+	static FORCEINLINE uint32 GetKeyHash(const FString& Key)
+	{
+		return FCrc::StrCrc32<TCHAR>(*Key);
+	}
+};
+
+/** Case sensitive hashing function for TMap */
+template <typename ValueType>
+struct FLocKeyMapFuncs : BaseKeyFuncs<ValueType, FString>
+{
+	static FORCEINLINE const FString& GetSetKey(const TPair<FString, ValueType>& Element)
+	{
+		return Element.Key;
+	}
+	static FORCEINLINE bool Matches(const FString& A, const FString& B)
+	{
+		return A.Equals(B, ESearchCase::CaseSensitive);
+	}
+	static FORCEINLINE uint32 GetKeyHash(const FString& Key)
+	{
+		return FCrc::StrCrc32<TCHAR>(*Key);
+	}
+};
diff --git a/Engine/Source/Runtime/Core/Public/Internationalization/StringTableCore.h b/Engine/Source/Runtime/Core/Public/Internationalization/StringTableCore.h
new file mode 100644
--- /dev/null
+++ b/Engine/Source/Runtime/Core/Public/Internationalization/StringTableCore.h
@@ -0,0 +1,217 @@
+// Copyright 1998-2017 Epic Games, Inc. All Rights Reserved.
+#pragma once
+
+#include "StringTableCoreFwd.h"
+#include "Containers/Map.h"
+#include "UObject/NameTypes.h"
+#include "Containers/UnrealString.h"
+#include "Internationalization/Text.h"
+#include "LocKeyFuncs.h"
+
+DECLARE_LOG_CATEGORY_EXTERN(LogStringTable, Log, All);
+
+/** Entry within a string table */
+class CORE_API FStringTableEntry
+{
+public:
+	/** Create a new string table entry using the given data */
+	static FStringTableEntryRef NewStringTableEntry(FStringTableConstRef InOwnerTable, FString InSourceString, FTextDisplayStringPtr InDisplayString)
+	{
+		return MakeShared<FStringTableEntry, ESPMode::ThreadSafe>(MoveTemp(InOwnerTable), MoveTemp(InSourceString), MoveTemp(InDisplayString));
+	}
+
+	/** Default constructor */
+	FStringTableEntry();
+
+	/** Create a new string table entry using the given data */
+	FStringTableEntry(FStringTableConstRef InOwnerTable, FString InSourceString, FTextDisplayStringPtr InDisplayString);
+
+	/** @return true if this entry is currently owned by a string table, false if it's been disowned (and should be re-cached) */
+	bool IsOwned() const;
+
+	/** Disown this string table entry. This is used to notify external code that has cached this entry that it needs to re-cache it from the string table */
+	void Disown();
+
+	/** Get the source string of this string table entry */
+	const FString& GetSourceString() const;
+
+	/** Get the display string of this string table entry */
+	FTextDisplayStringPtr GetDisplayString() const;
+
+private:
+	/** The string table that owns us (if any) */
+	FStringTableConstWeakPtr OwnerTable;
+
+	/** The source string of this entry */
+	FString SourceString;
+
+	/** The display string of this entry */
+	FTextDisplayStringPtr DisplayString;
+};
+
+/** String table implementation. Holds Key->SourceString pairs of text. */
+class CORE_API FStringTable : public TSharedFromThis<FStringTable, ESPMode::ThreadSafe>
+{
+public:
+	/** Create a new string table */
+	static FStringTableRef NewStringTable()
+	{
+		return MakeShared<FStringTable, ESPMode::ThreadSafe>();
+	}
+
+	/** Default constructor */
+	FStringTable();
+
+	/** Destructor */
+	~FStringTable();
+
+	/** @return The asset that owns this string table instance (if any) */
+	UStringTable* GetOwnerAsset() const;
+
+	/** Set the asset that owns this string table instance (if any) */
+	void SetOwnerAsset(UStringTable* InOwnerAsset);
+
+	/** Has this string table been fully loaded yet? (used during asset loading) */
+	bool IsLoaded() const;
+
+	/** Set whether this string table has been fully loaded yet */
+	void IsLoaded(const bool bInIsLoaded);
+
+	/** @return The namespace used by all entries in this string table */
+	FString GetNamespace() const;
+
+	/** Set the namespace used by all entries in this string table */
+	void SetNamespace(const FString& InNamespace);
+
+	/** Get the source string used by the given entry (if any) */
+	bool GetSourceString(const FString& InKey, FString& OutSourceString) const;
+
+	/** Set the source string used by the given entry (will replace any existing data for that entry) */
+	void SetSourceString(const FString& InKey, const FString& InSourceString);
+
+	/** Remove the given entry (including its meta-data) */
+	void RemoveSourceString(const FString& InKey);
+
+	/** Enumerate all source strings in the table. Return true from the enumerator to continue, or false to stop */
+	void EnumerateSourceStrings(const TFunctionRef<bool(const FString&, const FString&)>& InEnumerator) const;
+
+	/** Clear all entries from the table (including their meta-data) */
+	void ClearSourceStrings(const int32 InSlack = 0);
+
+	/** Find the entry with the given key (if any) */
+	FStringTableEntryConstPtr FindEntry(const FString& InKey) const;
+
+	/** Given an entry, check to see if it exists in this table, and if so, get its key */
+	bool FindKey(const FStringTableEntryConstRef& InEntry, FString& OutKey) const;
+
+	/** Given the display string of an entry, check to see if it exists in this table, and if so, get its key */
+	bool FindKey(const FTextDisplayStringRef& InDisplayString, FString& OutKey) const;
+
+	/** Get the meta-data with the given ID associated with the given entry, or an empty string if not found */
+	FString GetMetaData(const FString& InKey, const FName InMetaDataId) const;
+
+	/** Set the meta-data with the given ID associated with the given entry */
+	void SetMetaData(const FString& InKey, const FName InMetaDataId, const FString& InMetaDataValue);
+
+	/** Remove the meta-data with the given ID associated with the given entry */
+	void RemoveMetaData(const FString& InKey, const FName InMetaDataId);
+
+	/** Enumerate all meta-data associated with the given entry. Return true from the enumerator to continue, or false to stop */
+	void EnumerateMetaData(const FString& InKey, const TFunctionRef<bool(FName, const FString&)>& InEnumerator) const;
+
+	/** Remove all meta-data associated with the given entry */
+	void ClearMetaData(const FString& InKey);
+
+	/** Clear all meta-data from the table */
+	void ClearMetaData(const int32 InSlack = 0);
+
+	/** Serialize this string table to/from an archive */
+	void Serialize(FArchive& Ar);
+
+	/** Export the key, string, and meta-data information in this string table to a CSV file (does not export the namespace) */
+	bool ExportStrings(const FString& InFilename) const;
+
+	/** Import key, string, and meta-data information from a CSV file to this string table (does not import the namespace) */
+	bool ImportStrings(const FString& InFilename);
+
+private:
+	/** Pointer back to the asset that owns this table */
+	UStringTable* OwnerAsset;
+
+	/** True if this table has been fully loaded (used for assets) */
+	bool bIsLoaded;
+
+	/** The namespace to use for all the strings in this table */
+	FString TableNamespace;
+
+	/** Mapping between the text key and entry data for the strings within this table */
+	TMap<FString, FStringTableEntryPtr, FDefaultSetAllocator, FLocKeyMapFuncs<FStringTableEntryPtr>> KeysToEntries;
+
+	/** Mapping between the display string and the text key for the strings within this table */
+	TMap<FTextDisplayStringPtr, FString> DisplayStringsToKeys;
+
+	/** Critical section preventing concurrent modification of KeysToEntries */
+	mutable FCriticalSection KeyMappingCS;
+
+	/** Mapping between the text key and its meta-data map */
+	typedef TMap<FName, FString> FMetaDataMap;
+	TMap<FString, FMetaDataMap, FDefaultSetAllocator, FLocKeyMapFuncs<FMetaDataMap>> KeysToMetaData;
+
+	/** Critical section preventing concurrent modification of KeysToMetaData */
+	mutable FCriticalSection KeysToMetaDataCS;
+};
+
+/** Interface to allow Core code to access String Table assets from the Engine */
+class CORE_API IStringTableEngineBridge
+{
+public:
+	/** Redirect and load a string table asset by its name */
+	static void RedirectAndLoadStringTableAsset(FName& InOutTableId, const EStringTableLoadingPolicy InLoadingPolicy)
+	{
+		if (InstancePtr)
+		{
+			InstancePtr->RedirectAndLoadStringTableAssetImpl(InOutTableId, InLoadingPolicy);
+		}
+	}
+
+	/** Collect a string table asset reference */
+	static void CollectStringTableAssetReferences(const FName InTableId, FArchive& InAr)
+	{
+		if (InstancePtr)
+		{
+			InstancePtr->CollectStringTableAssetReferencesImpl(InTableId, InAr);
+		}
+	}
+
+protected:
+	virtual ~IStringTableEngineBridge() {}
+
+	virtual void RedirectAndLoadStringTableAssetImpl(FName& InOutTableId, const EStringTableLoadingPolicy InLoadingPolicy) = 0;
+	virtual void CollectStringTableAssetReferencesImpl(const FName InTableId, FArchive& InAr) = 0;
+
+	/** Singleton instance, populated by the derived type */
+	static IStringTableEngineBridge* InstancePtr;
+};
+
+/** String table redirect utils */
+struct CORE_API FStringTableRedirects
+{
+	/** Initialize the string table redirects */
+	static void InitStringTableRedirects();
+
+	/** Redirect a table ID */
+	static void RedirectTableId(FName& InOutTableId, const EStringTableLoadingPolicy InLoadingPolicy);
+
+	/** Redirect a key */
+	static void RedirectKey(const FName InTableId, FString& InOutKey);
+
+	/** Redirect a table ID and key */
+	static void RedirectTableIdAndKey(FName& InOutTableId, FString& InOutKey, const EStringTableLoadingPolicy InLoadingPolicy);
+};
+
+/** String table reference collection utils */
+struct CORE_API FStringTableReferenceCollection
+{
+	/** Collect asset references from the given table ID */
+	static void CollectAssetReferences(const FName InTableId, FArchive& InAr);
+};
diff --git a/Engine/Source/Runtime/Core/Public/Internationalization/StringTableRegistry.h b/Engine/Source/Runtime/Core/Public/Internationalization/StringTableRegistry.h
new file mode 100644
--- /dev/null
+++ b/Engine/Source/Runtime/Core/Public/Internationalization/StringTableRegistry.h
@@ -0,0 +1,123 @@
+// Copyright 1998-2017 Epic Games, Inc. All Rights Reserved.
+#pragma once
+
+#include "StringTableCoreFwd.h"
+#include "Containers/Map.h"
+#include "UObject/NameTypes.h"
+#include "Delegates/Delegate.h"
+#include "Containers/UnrealString.h"
+#include "TextLocalizationManager.h"
+#include "LocKeyFuncs.h"
+
+/** Singleton registry of known string table instances */
+class CORE_API FStringTableRegistry
+{
+public:
+	/** Singleton accessor */
+	static FStringTableRegistry& Get();
+
+	/** Register a string table with the given ID */
+	void RegisterStringTable(const FName InTableId, FStringTableRef InTable);
+
+	/** Unregister a string table with the given ID */
+	void UnregisterStringTable(const FName InTableId);
+
+	/** Try and find a string table with the given ID */
+	FStringTablePtr FindMutableStringTable(const FName InTableId) const;
+
+	/** Try and find a string table with the given ID */
+	FStringTableConstPtr FindStringTable(const FName InTableId) const;
+
+	/** Try and find a string table asset with the given ID */
+	UStringTable* FindStringTableAsset(const FName InTableId) const;
+
+	/** Enumerate all registered string tables */
+	void EnumerateStringTables(const TFunctionRef<bool(const FName&, const FStringTableConstRef&)>& InEnumerator) const;
+
+	/** Given an FText, try and find the table ID and key for it */
+	bool FindTableIdAndKey(const FText& InText, FName& OutTableId, FString& OutKey) const;
+
+	/** Given an FText display string, try and find the table ID and key for it */
+	bool FindTableIdAndKey(const FTextDisplayStringRef& InDisplayString, FName& OutTableId, FString& OutKey) const;
+
+	/** Log a missing string table entry (will only log each missing entry once to avoid spam) */
+	void LogMissingStringTableEntry(const FName InTableId, const FString& InKey);
+
+	/** Internal function called by LOCTABLE_NEW to create and register a new FStringTable instance */
+	void Internal_NewLocTable(const FName InTableId, const FString& InNamespace);
+
+	/** Internal function called by LOCTABLE_FROMFILE_X to create and register a FStringTable instance that has been populated from a file */
+	void Internal_LocTableFromFile(const FName InTableId, const FString& InNamespace, const FString& InFilePath, const FString& InRootPath);
+
+	/** Internal function called by LOCTABLE_SETSTRING to set the entry denoted by the given key to the given source string, within the given string table (table must have been registered already) */
+	void Internal_SetLocTableEntry(const FName InTableId, const FString& InKey, const FString& InSourceString);
+
+	/** Internal function called by LOCTABLE_SETMETA to set meta-data for the entry denoted by the given key, within the given string table (table must have been registered already) */
+	void Internal_SetLocTableEntryMetaData(const FName InTableId, const FString& InKey, const FName InMetaDataId, const FString& InMetaData);
+
+	/** Internal function called by LOCTABLE to find the entry with by the given key within the given string table (redirects, will load assets if needed, and returns an empty FText if not found) */
+	FText Internal_FindLocTableEntry(const FName InTableId, const FString& InKey, const EStringTableLoadingPolicy InLoadingPolicy) const;
+
+private:
+	/** Private constructor - use singleton accessor */
+	FStringTableRegistry();
+	~FStringTableRegistry();
+
+	/** Mapping from a table ID to a string table instance */
+	TMap<FName, FStringTablePtr> RegisteredStringTables;
+
+	/** Critical section preventing concurrent modification of RegisteredStringTables */
+	mutable FCriticalSection RegisteredStringTablesCS;
+
+	/** Mapping from a table ID to a set of keys that we've logged warnings for */
+	typedef TSet<FString, FLocKeySetFuncs> FLocKeySet;
+	TMap<FName, FLocKeySet> LoggedMissingEntries;
+
+	/** Critical section preventing concurrent modification of LoggedMissingEntries */
+	mutable FCriticalSection LoggedMissingEntriesCS;
+
+#if WITH_EDITOR
+	/** Callback handler for a directory change notification */
+	void OnDirectoryChanged(const TArray<struct FFileChangeData>& InFileChanges);
+
+	/** Mapping of absolute CSV file paths to the table ID that imported them using Internal_LocTableFromFile */
+	TMap<FString, FName> CSVFilesToWatch;
+
+	/** Critical section preventing concurrent modification of CSVFilesToWatch */
+	mutable FCriticalSection CSVFilesToWatchCS;
+
+	/** Delegate handle watching the Engine directory */
+	FDelegateHandle EngineDirectoryWatcherHandle;
+
+	/** Delegate handle watching the Game directory */
+	FDelegateHandle GameDirectoryWatcherHandle;
+#endif // WITH_EDITOR
+};
+
+#define LOC_DEFINE_REGION
+
+/** Creates and registers a new string table instance. */
+#define LOCTABLE_NEW(ID, NAMESPACE) \
+	FStringTableRegistry::Get().Internal_NewLocTable(TEXT(ID), TEXT(NAMESPACE))
+
+/** Creates and registers a new string table instance, loading strings from the given file (the path is relative to the engine content directory). */
+#define LOCTABLE_FROMFILE_ENGINE(ID, NAMESPACE, FILEPATH) \
+	FStringTableRegistry::Get().Internal_LocTableFromFile(TEXT(ID), TEXT(NAMESPACE), TEXT(FILEPATH), FPaths::EngineContentDir())
+
+/** Creates and registers a new string table instance, loading strings from the given file (the path is relative to the game content directory). */
+#define LOCTABLE_FROMFILE_GAME(ID, NAMESPACE, FILEPATH) \
+	FStringTableRegistry::Get().Internal_LocTableFromFile(TEXT(ID), TEXT(NAMESPACE), TEXT(FILEPATH), FPaths::GameContentDir())
+
+/** Add a string table entry with the given key and source string. */
+#define LOCTABLE_SETSTRING(ID, KEY, SRC) \
+	FStringTableRegistry::Get().Internal_SetLocTableEntry(TEXT(ID), TEXT(KEY), TEXT(SRC))
+
+/** Add a meta-data for the entry with the given key and source string. */
+#define LOCTABLE_SETMETA(ID, KEY, METAID, META) \
+	FStringTableRegistry::Get().Internal_SetLocTableEntryMetaData(TEXT(ID), TEXT(KEY), TEXT(METAID), TEXT(META))
+
+/** Find a string table with the given ID, and try and find an entry within it using the given key. Returns a dummy FText if not found. */
+#define LOCTABLE(ID, KEY) \
+	FStringTableRegistry::Get().Internal_FindLocTableEntry(TEXT(ID), TEXT(KEY), EStringTableLoadingPolicy::FindOrFullyLoad)
+
+#undef LOC_DEFINE_REGION
diff --git a/Engine/Source/Runtime/Core/Public/Internationalization/Text.h b/Engine/Source/Runtime/Core/Public/Internationalization/Text.h
--- a/Engine/Source/Runtime/Core/Public/Internationalization/Text.h
+++ b/Engine/Source/Runtime/Core/Public/Internationalization/Text.h
@@ -12,6 +12,7 @@
 #include "Templates/SharedPointer.h"
 #include "Internationalization/CulturePointer.h"
 #include "Internationalization/TextLocalizationManager.h"
+#include "Internationalization/StringTableCoreFwd.h"
 #include "Misc/Optional.h"
 #include "Templates/UniquePtr.h"
 
@@ -343,6 +344,12 @@ public:
 	 */
 	static bool FindText( const FString& Namespace, const FString& Key, FText& OutText, const FString* const SourceString = nullptr );
 
+	/**
+	 * Attempts to create an FText instance from a string table ID and key (this is the same as the LOCTABLE macro, except this can also work with non-literal string values).
+	 * @return The found text, or an dummy FText if not found.
+	 */
+	static FText FromStringTable(const FName InTableId, const FString& InKey, const EStringTableLoadingPolicy InLoadingPolicy = EStringTableLoadingPolicy::FindOrFullyLoad);
+
 	/**
 	 * Generate an FText representing the pass name
 	 */
@@ -479,6 +486,7 @@ public:
 
 	bool IsTransient() const;
 	bool IsCultureInvariant() const;
+	bool IsFromStringTable() const;
 
 	bool ShouldGatherForLocalization() const;
 
@@ -499,6 +507,8 @@ private:
 
 	explicit FText( FString&& InSourceString );
 
+	FText( FName InTableId, FString InKey );
+
 	FText( FString&& InSourceString, FTextDisplayStringRef InDisplayString );
 
 	FText( FString&& InSourceString, const FString& InNamespace, const FString& InKey, uint32 InFlags=0 );
@@ -554,9 +564,9 @@ public:
 	friend class FTextFormatData;
 	friend class FTextSnapshot;
 	friend class FTextInspector;
+	friend class FStringTableRegistry;
 	friend class FArchive;
 	friend class UTextProperty;
-	friend class UGatherTextFromAssetsCommandlet;
 	friend class FFormatArgumentValue;
 	friend class FTextHistory_NamedFormat;
 	friend class FTextHistory_ArgumentDataFormat;
@@ -856,6 +866,7 @@ public:
 	static const FString* GetSourceString(const FText& Text);
 	static const FString& GetDisplayString(const FText& Text);
 	static const FTextDisplayStringRef GetSharedDisplayString(const FText& Text);
+	static bool GetTableIdAndKey(const FText& Text, FName& OutTableId, FString& OutKey);
 	static uint32 GetFlags(const FText& Text);
 	static void GetHistoricFormatData(const FText& Text, TArray<FHistoricTextFormatData>& OutHistoricFormatData);
 	static bool GetHistoricNumericData(const FText& Text, FHistoricTextNumericData& OutHistoricNumericData);
@@ -873,10 +884,11 @@ public:
 	 * @param PackageNamespace	The package namespace of the containing object (if loading for a property - see TextNamespaceUtil::GetPackageNamespace).
 	 * @param OutNumCharsRead	An optional output parameter to fill with the number of characters we read from the given buffer.
 	 * @param bRequiresQuotes	True if the read text literal must be surrounded by quotes (eg, when loading from a delimited list).
+	 * @param InLoadingPolicy	Controls how we should load any referenced string table assets.
 	 *
 	 * @return True if we read a valid FText instance into OutValue, false otherwise
 	 */
-	static bool ReadFromString(const TCHAR* Buffer, FText& OutValue, const TCHAR* TextNamespace = nullptr, const TCHAR* PackageNamespace = nullptr, int32* OutNumCharsRead = nullptr, const bool bRequiresQuotes = false);
+	static bool ReadFromString(const TCHAR* Buffer, FText& OutValue, const TCHAR* TextNamespace = nullptr, const TCHAR* PackageNamespace = nullptr, int32* OutNumCharsRead = nullptr, const bool bRequiresQuotes = false, const EStringTableLoadingPolicy InLoadingPolicy = EStringTableLoadingPolicy::FindOrFullyLoad);
 
 	/**
 	 * Write the given FText instance to a stream of text
@@ -897,12 +909,13 @@ public:
 	static bool IsComplexText(const TCHAR* Buffer);
 
 private:
-	static bool ReadFromString_ComplexText(const TCHAR* Buffer, FText& OutValue, const TCHAR* TextNamespace, const TCHAR* PackageNamespace, int32* OutNumCharsRead);
+	static bool ReadFromString_ComplexText(const TCHAR* Buffer, FText& OutValue, const TCHAR* TextNamespace, const TCHAR* PackageNamespace, int32* OutNumCharsRead, const EStringTableLoadingPolicy InLoadingPolicy);
 
 #define LOC_DEFINE_REGION
 	static const FString InvTextMarker;
 	static const FString NsLocTextMarker;
 	static const FString LocTextMarker;
+	static const FString LocTableMarker;
 #undef LOC_DEFINE_REGION
 };
 
diff --git a/Engine/Source/Runtime/Core/Public/Linux/LinuxCursor.h b/Engine/Source/Runtime/Core/Public/Linux/LinuxCursor.h
--- a/Engine/Source/Runtime/Core/Public/Linux/LinuxCursor.h
+++ b/Engine/Source/Runtime/Core/Public/Linux/LinuxCursor.h
@@ -34,7 +34,7 @@ public:
 
 	virtual void Lock( const RECT* const Bounds ) override;
 
-	virtual void SetCustomShape(void* CursorHandle) override;
+	virtual void SetTypeShape(EMouseCursor::Type InCursorType, void* CursorHandle) override;
 
 public:
 
@@ -49,14 +49,6 @@ public:
 
 	bool IsHidden();
 
-	uint32 GetCursorEvent();
-
-	/** Updates accumulated offset which is used for maintaining a "shadow" cursor position */
-	void AddOffset(const int32 DX, const int32 DY);
-
-	/** Resets accumulated offsets to (0,0), */
-	void ResetOffset();
-
 	/**
 	 * Invalidates whatever cached data cursor may have
 	 */
@@ -73,6 +65,9 @@ private:
 	/** Cursors */
 	SDL_HCursor CursorHandles[ EMouseCursor::TotalCursorCount ];
 
+	/** Override Cursors */
+	SDL_HCursor CursorOverrideHandles[ EMouseCursor::TotalCursorCount ];
+
 	FIntRect CursorClipRect;
 
 	uint32 CursorEvent;
diff --git a/Engine/Source/Runtime/Core/Public/Mac/MacCursor.h b/Engine/Source/Runtime/Core/Public/Mac/MacCursor.h
--- a/Engine/Source/Runtime/Core/Public/Mac/MacCursor.h
+++ b/Engine/Source/Runtime/Core/Public/Mac/MacCursor.h
@@ -29,7 +29,7 @@ public:
 
 	virtual void Lock(const RECT* const Bounds) override;
 
-	virtual void SetCustomShape(void* CursorHandle) override;
+	virtual void SetTypeShape(EMouseCursor::Type InCursorType, void* CursorHandle) override;
 
 public:
 
@@ -64,6 +64,7 @@ private:
 	EMouseCursor::Type CurrentType;
 
 	NSCursor* CursorHandles[EMouseCursor::TotalCursorCount];
+	NSCursor* CursorOverrideHandles[EMouseCursor::TotalCursorCount];
 
 	FIntRect CursorClipRect;
 
diff --git a/Engine/Source/Runtime/Core/Public/Math/TransformCalculus2D.h b/Engine/Source/Runtime/Core/Public/Math/TransformCalculus2D.h
--- a/Engine/Source/Runtime/Core/Public/Math/TransformCalculus2D.h
+++ b/Engine/Source/Runtime/Core/Public/Math/TransformCalculus2D.h
@@ -462,13 +462,21 @@ public:
 		return FScale2D(A*A + B*B, C*C + D*D);
 	}
 
-	/** Gets the squared scale from the matrix (avoids sqrt). */
+	/** Gets the scale from the matrix. */
 	FScale2D GetScale() const
 	{
 		FScale2D ScaleSquared = GetScaleSquared();
 		return FScale2D(FMath::Sqrt(ScaleSquared.GetVector().X), FMath::Sqrt(ScaleSquared.GetVector().Y));
 	}
 
+	/** Gets the rotation angle of the matrix. */
+	float GetRotationAngle() const
+	{
+		float A, B, C, D;
+		GetMatrix(A, B, C, D);
+		return FMath::Atan(C / D);
+	}
+
 	/** Determines if the matrix is identity or not. Uses exact float comparison, so rounding error is not considered. */
 	bool IsIdentity() const
 	{
diff --git a/Engine/Source/Runtime/Core/Public/Misc/AutomationTest.h b/Engine/Source/Runtime/Core/Public/Misc/AutomationTest.h
--- a/Engine/Source/Runtime/Core/Public/Misc/AutomationTest.h
+++ b/Engine/Source/Runtime/Core/Public/Misc/AutomationTest.h
@@ -92,6 +92,18 @@ struct CORE_API FAutomationEvent
 	FString ToString() const;
 };
 
+class FAutomationTestArtifact
+{
+public:
+	FString Name;
+	FString FilePath;
+
+	FAutomationTestArtifact(FString InName, FString InFilePath)
+		: Name(InName)
+		, FilePath(InFilePath)
+	{
+	}
+};
 
 /** Simple class to store the results of the execution of a automation test */
 class FAutomationTestExecutionInfo
@@ -118,6 +130,7 @@ public:
 		Warnings.Empty();
 		LogItems.Empty();
 		AnalyticsItems.Empty();
+		Artifacts.Empty();
 	}
 
 	/**
@@ -128,14 +141,22 @@ public:
 
 	/** Whether the automation test completed successfully or not */
 	bool bSuccessful;
+	
 	/** Any errors that occurred during execution */
 	TArray<FAutomationEvent> Errors;
+	
 	/** Any warnings that occurred during execution */
 	TArray<FString> Warnings;
+	
 	/** Any log items that occurred during execution */
 	TArray<FString> LogItems;
+	
 	/** Any analytics items that occurred during execution */
 	TArray<FString> AnalyticsItems;
+
+	/** Any items that need to be reference by reports and other systems. */
+	TArray<FAutomationTestArtifact> Artifacts;
+
 	/** Time to complete the task */
 	float Duration;
 };
@@ -467,6 +488,7 @@ struct FAutomationScreenshotData
 	FString Context;
 
 	FGuid Id;
+	FString Commit;
 
 	int32 Width;
 	int32 Height;
@@ -510,7 +532,9 @@ struct FAutomationScreenshotData
 	FString Path;
 
 	FAutomationScreenshotData()
-		: Width(0)
+		: Id()
+		, Commit()
+		, Width(0)
 		, Height(0)
 		, bIsStereo(false)
 		, ResolutionQuality(1.0f)
@@ -920,7 +944,6 @@ public:
 	 */
 	virtual void AddError( const FString& InError, int32 StackOffset = 0 );
 
-
 	/**
 	 * Adds an error message to this test
 	 *
@@ -945,12 +968,17 @@ public:
 	virtual void AddLogItem( const FString& InLogItem );
 
 	/**
-	* Adds a analytics string to parse later
-	*
-	* @param	InLogItem	Log item to add to this test
-	*/
+	 * Adds a analytics string to parse later
+	 *
+	 * @param	InLogItem	Log item to add to this test
+	 */
 	virtual void AddAnalyticsItem(const FString& InAnalyticsItem);
 
+	/**
+	 * Occasionally a test may need to provide artifacts that should appear in the build report for the automation test.
+	 */
+	virtual void AddRemoteArtifact(const FString& ArtifactName, const FString& RemoteArtifactPath);
+
 	/**
 	 * Returns whether this test has any errors associated with it or not
 	 *
@@ -1249,7 +1277,6 @@ public:
 		}
 	}
 
-
 protected:
 	/**
 	 * Asks the test to enumerate variants that will all go through the "RunTest" function with different parameters (for load all maps, this should enumerate all maps to load)\
@@ -1289,7 +1316,6 @@ protected:
 
 	//allow framework to call protected function
 	friend class FAutomationTestFramework;
-
 };
 
 class CORE_API FBDDAutomationTestBase : public FAutomationTestBase
diff --git a/Engine/Source/Runtime/Core/Public/Misc/EnumClassFlags.h b/Engine/Source/Runtime/Core/Public/Misc/EnumClassFlags.h
--- a/Engine/Source/Runtime/Core/Public/Misc/EnumClassFlags.h
+++ b/Engine/Source/Runtime/Core/Public/Misc/EnumClassFlags.h
@@ -14,3 +14,15 @@
 	inline CONSTEXPR Enum  operator^ (Enum  Lhs, Enum Rhs) { return (Enum)((__underlying_type(Enum))Lhs ^ (__underlying_type(Enum))Rhs); } \
 	inline CONSTEXPR bool  operator! (Enum  E)             { return !(__underlying_type(Enum))E; } \
 	inline CONSTEXPR Enum  operator~ (Enum  E)             { return (Enum)~(__underlying_type(Enum))E; }
+
+template<typename Enum>
+inline bool EnumHasAllFlags(Enum Flags, Enum Contains)
+{
+	return ( ( ( __underlying_type(Enum) )Flags ) & ( __underlying_type(Enum) )Contains ) == ( ( __underlying_type(Enum) )Contains );
+}
+
+template<typename Enum>
+inline bool EnumHasAnyFlags(Enum Flags, Enum Contains)
+{
+	return ( ( ( __underlying_type(Enum) )Flags ) & ( __underlying_type(Enum) )Contains ) != 0;
+}
diff --git a/Engine/Source/Runtime/Core/Public/Modules/ModuleManager.h b/Engine/Source/Runtime/Core/Public/Modules/ModuleManager.h
--- a/Engine/Source/Runtime/Core/Public/Modules/ModuleManager.h
+++ b/Engine/Source/Runtime/Core/Public/Modules/ModuleManager.h
@@ -298,6 +298,12 @@ public:
 		{
 			ModuleManager.LoadModule(ModuleName);
 		}
+		else
+		{
+#if !(UE_BUILD_SHIPPING || UE_BUILD_TEST)
+			WarnIfItWasntSafeToLoadHere(ModuleName);
+#endif
+		}
 
 		return GetModuleChecked<TModuleInterface>(ModuleName);
 	}
@@ -318,6 +324,12 @@ public:
 		{
 			ModuleManager.LoadModule(ModuleName);
 		}
+		else
+		{
+#if !(UE_BUILD_SHIPPING || UE_BUILD_TEST)
+			WarnIfItWasntSafeToLoadHere(ModuleName);
+#endif
+		}
 
 		return GetModulePtr<TModuleInterface>(ModuleName);
 	}
@@ -549,6 +561,8 @@ public:
 	void ResetModulePathsCache();
 
 private:
+	static void WarnIfItWasntSafeToLoadHere(const FName InModuleName);
+
 	/** Thread safe module finding routine. */
 	ModuleInfoPtr FindModule(FName InModuleName);
 	ModuleInfoRef FindModuleChecked(FName InModuleName);
diff --git a/Engine/Source/Runtime/Core/Public/Templates/SharedPointer.h b/Engine/Source/Runtime/Core/Public/Templates/SharedPointer.h
--- a/Engine/Source/Runtime/Core/Public/Templates/SharedPointer.h
+++ b/Engine/Source/Runtime/Core/Public/Templates/SharedPointer.h
@@ -1309,10 +1309,13 @@ public:		// @todo: Ideally this would be private, but template sharing problems
 	}
 
 	/**
-	 * Checks whether given instance has been already made sharable (use in checks to detect when it 
-	 * happened, since it's a straight way to crashing
+	 * Checks whether our referenced instance is valid (ie, whether it's safe to call AsShared).
+	 * If this returns false, it means that your instance has either:
+	 *  - Not yet been assigned to a shared pointer (via MakeShared or MakeShareable).
+	 *  - Is currently within its constructor (so the shared instance isn't yet available).
+	 *  - Is currently within its destructor (so the shared instance is no longer available).
 	 */
-	FORCEINLINE bool HasBeenAlreadyMadeSharable() const
+	FORCEINLINE bool DoesSharedInstanceExist() const
 	{
 		return WeakThis.IsValid();
 	}
diff --git a/Engine/Source/Runtime/Core/Public/UObject/EditorObjectVersion.h b/Engine/Source/Runtime/Core/Public/UObject/EditorObjectVersion.h
--- a/Engine/Source/Runtime/Core/Public/UObject/EditorObjectVersion.h
+++ b/Engine/Source/Runtime/Core/Public/UObject/EditorObjectVersion.h
@@ -39,6 +39,8 @@ struct CORE_API FEditorObjectVersion
 		StableUserDefinedEnumDisplayNames,
 		// Added "Inline" option to UFontFace assets
 		AddedInlineFontFaceAssets,
+		// Fix a serialization issue with static mesh FMeshSectionInfoMap UProperty
+		UPropertryForMeshSectionSerialize,
 		// -----<new versions can be added above this line>-------------------------------------------------
 		VersionPlusOne,
 		LatestVersion = VersionPlusOne - 1
diff --git a/Engine/Source/Runtime/Core/Public/Windows/WindowsCursor.h b/Engine/Source/Runtime/Core/Public/Windows/WindowsCursor.h
--- a/Engine/Source/Runtime/Core/Public/Windows/WindowsCursor.h
+++ b/Engine/Source/Runtime/Core/Public/Windows/WindowsCursor.h
@@ -31,7 +31,7 @@ public:
 
 	virtual void Lock( const RECT* const Bounds ) override;
 
-	virtual void SetCustomShape(void* CursorHandle) override;
+	virtual void SetTypeShape(EMouseCursor::Type InCursorType, void* CursorHandle) override;
 
 private:
 
@@ -39,4 +39,7 @@ private:
 
 	/** Cursors */
 	Windows::HCURSOR CursorHandles[ EMouseCursor::TotalCursorCount ];
+
+	/** Override Cursors */
+	Windows::HCURSOR CursorOverrideHandles[EMouseCursor::TotalCursorCount];
 };
\ No newline at end of file
