Build - Fixing incomplete type.

diff --git a/Engine/Source/Runtime/Core/Public/Misc/AutomationTest.h b/Engine/Source/Runtime/Core/Public/Misc/AutomationTest.h
--- a/Engine/Source/Runtime/Core/Public/Misc/AutomationTest.h
+++ b/Engine/Source/Runtime/Core/Public/Misc/AutomationTest.h
@@ -22,6 +22,8 @@
 #include "Future.h"
 #include "Async.h"
 #include "Misc/Guid.h"
+#include "Math/Vector.h"
+#include "Math/Color.h"
 
 /** Flags for specifying automation test requirements/behavior */
 namespace EAutomationTestFlags
