Updating EditorTutorial target to depend on the other targets

diff --git a/Engine/Config/Localization/EditorTutorials.ini b/Engine/Config/Localization/EditorTutorials.ini
--- a/Engine/Config/Localization/EditorTutorials.ini
+++ b/Engine/Config/Localization/EditorTutorials.ini
@@ -1,4 +1,10 @@
 ﻿[CommonSettings]
+ManifestDependencies=Content/Localization/Engine/Engine.manifest
+ManifestDependencies=Content/Localization/Editor/Editor.manifest
+ManifestDependencies=Content/Localization/Category/Category.manifest
+ManifestDependencies=Content/Localization/Keywords/Keywords.manifest
+ManifestDependencies=Content/Localization/PropertyNames/PropertyNames.manifest
+ManifestDependencies=Content/Localization/ToolTips/ToolTips.manifest
 SourcePath=Content/Localization/EditorTutorials
 DestinationPath=Content/Localization/EditorTutorials
 ManifestName=EditorTutorials.manifest
diff --git a/Engine/Source/Runtime/Core/Private/Internationalization/TextLocalizationManager.cpp b/Engine/Source/Runtime/Core/Private/Internationalization/TextLocalizationManager.cpp
--- a/Engine/Source/Runtime/Core/Private/Internationalization/TextLocalizationManager.cpp
+++ b/Engine/Source/Runtime/Core/Private/Internationalization/TextLocalizationManager.cpp
@@ -229,7 +229,7 @@ void FTextLocalizationManager::FLocalizationEntryTracker::DetectAndLogConflicts(
 				{
 					const FEntry& RightEntry = EntryArray[l];
 					const bool bDoesSourceStringHashDiffer = LeftEntry.SourceStringHash != RightEntry.SourceStringHash;
-					const bool bDoesLocalizedStringDiffer = LeftEntry.LocalizedString != RightEntry.LocalizedString;
+					const bool bDoesLocalizedStringDiffer = !LeftEntry.LocalizedString.Equals(RightEntry.LocalizedString, ESearchCase::CaseSensitive);
 					WasConflictDetected = bDoesSourceStringHashDiffer || bDoesLocalizedStringDiffer;
 				}
 			}
@@ -246,7 +246,7 @@ void FTextLocalizationManager::FLocalizationEntryTracker::DetectAndLogConflicts(
 						CollidingEntryListString += TEXT('\n');
 					}
 
-					CollidingEntryListString += FString::Printf( TEXT("Localization Resource: (%s) Source String Hash: (%d) Localized String: (%s)"), *(Entry.LocResID), Entry.SourceStringHash, *(Entry.LocalizedString) );
+					CollidingEntryListString += FString::Printf( TEXT("    Localization Resource: (%s) Source String Hash: (%d) Localized String: (%s)"), *(Entry.LocResID), Entry.SourceStringHash, *(Entry.LocalizedString) );
 				}
 
 				UE_LOG(LogTextLocalizationManager, Warning, TEXT("Loaded localization resources contain conflicting entries for (Namespace:%s, Key:%s):\n%s"), *NamespaceName, *KeyName, *CollidingEntryListString);
