Downgrading a warning until UE-41908 is fixed

diff --git a/Engine/Source/Runtime/Core/Private/Internationalization/TextHistory.cpp b/Engine/Source/Runtime/Core/Private/Internationalization/TextHistory.cpp
--- a/Engine/Source/Runtime/Core/Private/Internationalization/TextHistory.cpp
+++ b/Engine/Source/Runtime/Core/Private/Internationalization/TextHistory.cpp
@@ -202,7 +202,7 @@ void FTextHistory_Base::SerializeForDisplayString(FArchive& Ar, FTextDisplayStri
 					{
 						const FString CurPackageNamespace = TextNamespaceUtil::ExtractPackageNamespace(Namespace);
 						const FString CurCleanNamespace = TextNamespaceUtil::StripPackageNamespace(Namespace);
-						UE_LOG(LogCore, Warning, TEXT("Package localization ID mismatch during cook! Expected '%s', got '%s'. Namespace: '%s', Key: '%s', Source: '%s'."), *PackageNamespace, *CurPackageNamespace, *CurCleanNamespace, *Key, *SourceString);
+						UE_LOG(LogCore, Display, TEXT("Package localization ID mismatch during cook! Expected '%s', got '%s'. Namespace: '%s', Key: '%s', Source: '%s'."), *PackageNamespace, *CurPackageNamespace, *CurCleanNamespace, *Key, *SourceString);
 					}
 					else
 					{
