Fixing overflow error when loading an empty v1 locres file

diff --git a/Engine/Source/Runtime/Core/Private/Internationalization/TextLocalizationResource.cpp b/Engine/Source/Runtime/Core/Private/Internationalization/TextLocalizationResource.cpp
--- a/Engine/Source/Runtime/Core/Private/Internationalization/TextLocalizationResource.cpp
+++ b/Engine/Source/Runtime/Core/Private/Internationalization/TextLocalizationResource.cpp
@@ -115,7 +115,10 @@ bool FTextLocalizationResource::LoadFromArchive(FArchive& Archive, const FString
 
 	// Read magic number
 	FGuid MagicNumber;
-	Archive << MagicNumber;
+	if (Archive.TotalSize() >= sizeof(FGuid))
+	{
+		Archive << MagicNumber;
+	}
 
 	FTextLocalizationResourceVersion::ELocResVersion VersionNumber = FTextLocalizationResourceVersion::ELocResVersion::Legacy;
 	if (MagicNumber == FTextLocalizationResourceVersion::LocResMagic)
