Fixed log of network version where not all hashed values were logged.
Changed incompatible-network log in OnlineBeacon to an error and display the associated valued

diff --git a/Engine/Source/Runtime/Core/Private/Misc/NetworkVersion.cpp b/Engine/Source/Runtime/Core/Private/Misc/NetworkVersion.cpp
--- a/Engine/Source/Runtime/Core/Private/Misc/NetworkVersion.cpp
+++ b/Engine/Source/Runtime/Core/Private/Misc/NetworkVersion.cpp
@@ -80,8 +80,11 @@ uint32 FNetworkVersion::GetLocalNetworkVersion( bool AllowOverrideDelegate /*=tr
 	// Get the project name (NOT case sensitive)
 	const FString ProjectName( FString( FApp::GetGameName() ).ToLower() );
 
+	// network compatible CL
+	const uint32 NetworkCompatibleChangelist = GetNetworkCompatibleChangelist();
+
 	// Start with project name+compatible changelist as seed
-	CachedNetworkChecksum = FCrc::StrCrc32( *ProjectName, GetNetworkCompatibleChangelist() );
+	CachedNetworkChecksum = FCrc::StrCrc32( *ProjectName, NetworkCompatibleChangelist);
 
 	// Next, hash with project version
 	CachedNetworkChecksum = FCrc::StrCrc32( *ProjectVersion, CachedNetworkChecksum );
@@ -93,7 +96,8 @@ uint32 FNetworkVersion::GetLocalNetworkVersion( bool AllowOverrideDelegate /*=tr
 	CachedNetworkChecksum = FCrc::MemCrc32( &EngineNetworkVersion, sizeof( EngineNetworkVersion ), CachedNetworkChecksum );
 	CachedNetworkChecksum = FCrc::MemCrc32( &GameNetworkVersion, sizeof( GameNetworkVersion ), CachedNetworkChecksum );
 
-	UE_LOG( LogNetVersion, Log, TEXT( "GetLocalNetworkVersion: CL: %u, ProjectName: %s, ProjectVersion: %s, EngineNetworkVersion: %i, GameNetworkVersion: %i, NetworkChecksum: %u" ), FEngineVersion::CompatibleWith().GetChangelist(), *ProjectName, *ProjectVersion, EngineNetworkVersion, GameNetworkVersion, CachedNetworkChecksum );
+	UE_LOG( LogNetVersion, Log, TEXT( "GetNetworkCompatibleChangelist: %u, ProjectName: '%s', ProjectVersion: '%s', EngineNetworkVersion: %i, GameNetworkVersion: %i, NetworkChecksum: %u" ), 
+		NetworkCompatibleChangelist, *ProjectName, *ProjectVersion, EngineNetworkVersion, GameNetworkVersion, CachedNetworkChecksum );
 
 	bHasCachedNetworkChecksum = true;
 
