Temporarly added more logging to netversion calculation for debugging if the Orion tests should fail again

diff --git a/Engine/Source/Runtime/Core/Private/Misc/NetworkVersion.cpp b/Engine/Source/Runtime/Core/Private/Misc/NetworkVersion.cpp
--- a/Engine/Source/Runtime/Core/Private/Misc/NetworkVersion.cpp
+++ b/Engine/Source/Runtime/Core/Private/Misc/NetworkVersion.cpp
@@ -85,16 +85,21 @@ uint32 FNetworkVersion::GetLocalNetworkVersion( bool AllowOverrideDelegate /*=tr
 
 	// Start with project name+compatible changelist as seed
 	CachedNetworkChecksum = FCrc::StrCrc32( *ProjectName, NetworkCompatibleChangelist);
+	UE_LOG(LogNetVersion, Log, TEXT("Hash1: %u"), CachedNetworkChecksum);
 
 	// Next, hash with project version
 	CachedNetworkChecksum = FCrc::StrCrc32( *ProjectVersion, CachedNetworkChecksum );
+	UE_LOG(LogNetVersion, Log, TEXT("Hash2: %u"), CachedNetworkChecksum);
 
 	// Finally, hash with engine/game network version
 	const uint32 EngineNetworkVersion	= FNetworkVersion::GetEngineNetworkProtocolVersion();
 	const uint32 GameNetworkVersion		= FNetworkVersion::GetGameNetworkProtocolVersion();
 
 	CachedNetworkChecksum = FCrc::MemCrc32( &EngineNetworkVersion, sizeof( EngineNetworkVersion ), CachedNetworkChecksum );
+	UE_LOG(LogNetVersion, Log, TEXT("Hash3: %u"), CachedNetworkChecksum);
+
 	CachedNetworkChecksum = FCrc::MemCrc32( &GameNetworkVersion, sizeof( GameNetworkVersion ), CachedNetworkChecksum );
+	UE_LOG(LogNetVersion, Log, TEXT("Hash4: %u"), CachedNetworkChecksum);
 
 	UE_LOG( LogNetVersion, Log, TEXT( "GetNetworkCompatibleChangelist: %u, ProjectName: '%s', ProjectVersion: '%s', EngineNetworkVersion: %i, GameNetworkVersion: %i, NetworkChecksum: %u" ), 
 		NetworkCompatibleChangelist, *ProjectName, *ProjectVersion, EngineNetworkVersion, GameNetworkVersion, CachedNetworkChecksum );
