Disabling optimization around suspect GetLocalNetworkVersion code to narrow down suspected compiler issue

diff --git a/Engine/Source/Runtime/Core/Private/Misc/NetworkVersion.cpp b/Engine/Source/Runtime/Core/Private/Misc/NetworkVersion.cpp
--- a/Engine/Source/Runtime/Core/Private/Misc/NetworkVersion.cpp
+++ b/Engine/Source/Runtime/Core/Private/Misc/NetworkVersion.cpp
@@ -59,8 +59,11 @@ uint32 FNetworkVersion::GetGameCompatibleNetworkProtocolVersion()
 	return GameCompatibleNetworkProtocolVersion;
 }
 
+PRAGMA_DISABLE_OPTIMIZATION
+
 uint32 FNetworkVersion::GetLocalNetworkVersion( bool AllowOverrideDelegate /*=true*/ )
 {
+	
 	if ( bHasCachedNetworkChecksum )
 	{
 		return CachedNetworkChecksum;
@@ -104,6 +107,8 @@ uint32 FNetworkVersion::GetLocalNetworkVersion( bool AllowOverrideDelegate /*=tr
 	return CachedNetworkChecksum;
 }
 
+PRAGMA_ENABLE_OPTIMIZATION
+
 bool FNetworkVersion::IsNetworkCompatible( const uint32 LocalNetworkVersion, const uint32 RemoteNetworkVersion )
 {
 	if ( IsNetworkCompatibleOverride.IsBound() )
