Modified calculation of network version to address what appeared to be a compiler bug generating different hashes on Win64 Server in Development/Test configurations.

Code now generates a long string with all version info, CRCs it to get the network checksum, and displays both results.

diff --git a/Engine/Source/Runtime/Core/Private/Misc/NetworkVersion.cpp b/Engine/Source/Runtime/Core/Private/Misc/NetworkVersion.cpp
--- a/Engine/Source/Runtime/Core/Private/Misc/NetworkVersion.cpp
+++ b/Engine/Source/Runtime/Core/Private/Misc/NetworkVersion.cpp
@@ -59,11 +59,8 @@ uint32 FNetworkVersion::GetGameCompatibleNetworkProtocolVersion()
 	return GameCompatibleNetworkProtocolVersion;
 }
 
-PRAGMA_DISABLE_OPTIMIZATION
-
 uint32 FNetworkVersion::GetLocalNetworkVersion( bool AllowOverrideDelegate /*=true*/ )
 {
-	
 	if ( bHasCachedNetworkChecksum )
 	{
 		return CachedNetworkChecksum;
@@ -73,42 +70,29 @@ uint32 FNetworkVersion::GetLocalNetworkVersion( bool AllowOverrideDelegate /*=tr
 	{
 		CachedNetworkChecksum = GetLocalNetworkVersionOverride.Execute();
 
-		UE_LOG( LogNetVersion, Log, TEXT( "GetLocalNetworkVersionOverride: NetworkChecksum: %u" ), CachedNetworkChecksum );
+		UE_LOG( LogNetVersion, Log, TEXT( "Checksum from delegate: %u" ), CachedNetworkChecksum );
 
 		bHasCachedNetworkChecksum = true;
 
 		return CachedNetworkChecksum;
 	}
 
-	// Get the project name (NOT case sensitive)
-	const FString ProjectName( FString( FApp::GetGameName() ).ToLower() );
-
-	// network compatible CL
-	const uint32 NetworkCompatibleChangelist = GetNetworkCompatibleChangelist();
-
-	// Start with project name+compatible changelist as seed
-	CachedNetworkChecksum = FCrc::StrCrc32( *ProjectName, NetworkCompatibleChangelist);
+	FString VersionString = FString::Printf(TEXT("%s %s, NetCL: %d, EngineNetVer: %d, GameNetVer: %d"),
+		FApp::GetGameName(),
+		*ProjectVersion,
+		GetNetworkCompatibleChangelist(),
+		FNetworkVersion::GetEngineNetworkProtocolVersion(),
+		FNetworkVersion::GetGameNetworkProtocolVersion());
 
-	// Next, hash with project version
-	CachedNetworkChecksum = FCrc::StrCrc32( *ProjectVersion, CachedNetworkChecksum );
+	CachedNetworkChecksum = FCrc::StrCrc32(*VersionString.ToLower());
 
-	// Finally, hash with engine/game network version
-	const uint32 EngineNetworkVersion	= FNetworkVersion::GetEngineNetworkProtocolVersion();
-	const uint32 GameNetworkVersion		= FNetworkVersion::GetGameNetworkProtocolVersion();
-
-	CachedNetworkChecksum = FCrc::MemCrc32( &EngineNetworkVersion, sizeof( EngineNetworkVersion ), CachedNetworkChecksum );
-	CachedNetworkChecksum = FCrc::MemCrc32( &GameNetworkVersion, sizeof( GameNetworkVersion ), CachedNetworkChecksum );
-
-	UE_LOG( LogNetVersion, Log, TEXT( "GetNetworkCompatibleChangelist: %u, ProjectName: '%s', ProjectVersion: '%s', EngineNetworkVersion: %i, GameNetworkVersion: %i, NetworkChecksum: %u" ), 
-		NetworkCompatibleChangelist, *ProjectName, *ProjectVersion, EngineNetworkVersion, GameNetworkVersion, CachedNetworkChecksum );
+	UE_LOG(LogNetVersion, Log, TEXT("%s (Checksum: %u)"), *VersionString, CachedNetworkChecksum);
 
 	bHasCachedNetworkChecksum = true;
 
 	return CachedNetworkChecksum;
 }
 
-PRAGMA_ENABLE_OPTIMIZATION
-
 bool FNetworkVersion::IsNetworkCompatible( const uint32 LocalNetworkVersion, const uint32 RemoteNetworkVersion )
 {
 	if ( IsNetworkCompatibleOverride.IsBound() )
