Fixes for ApplePlatformMemory unity build errors & tvOS compilation of ApplePlatformDebugEvents.

diff --git a/Engine/Source/Runtime/Core/Private/Apple/ApplePlatformDebugEvents.cpp b/Engine/Source/Runtime/Core/Private/Apple/ApplePlatformDebugEvents.cpp
--- a/Engine/Source/Runtime/Core/Private/Apple/ApplePlatformDebugEvents.cpp
+++ b/Engine/Source/Runtime/Core/Private/Apple/ApplePlatformDebugEvents.cpp
@@ -35,7 +35,7 @@
 
 #if APPLE_PROFILING_ENABLED
 
-#define APPLE_PROFILING_FALLBACKS (!PLATFORM_MAC && __IPHONE_OS_VERSION_MIN_REQUIRED < __IPHONE_10_0) || (PLATFORM_MAC && __MAC_OS_X_VERSION_MIN_REQUIRED < __MAC_10_12)
+#define APPLE_PROFILING_FALLBACKS !PLATFORM_TVOS && ((!PLATFORM_MAC && __IPHONE_OS_VERSION_MIN_REQUIRED < __IPHONE_10_0) || (PLATFORM_MAC && __MAC_OS_X_VERSION_MIN_REQUIRED < __MAC_10_12))
 
 DEFINE_LOG_CATEGORY(LogInstruments)
 
diff --git a/Engine/Source/Runtime/Core/Private/Apple/ApplePlatformMemory.cpp b/Engine/Source/Runtime/Core/Private/Apple/ApplePlatformMemory.cpp
--- a/Engine/Source/Runtime/Core/Private/Apple/ApplePlatformMemory.cpp
+++ b/Engine/Source/Runtime/Core/Private/Apple/ApplePlatformMemory.cpp
@@ -6,7 +6,11 @@
 
 #include "ApplePlatformMemory.h"
 #include "HAL/PlatformMemory.h"
+#include "HAL/PlatformMath.h"
+#include "HAL/UnrealMemory.h"
+#include "Templates/AlignmentTemplates.h"
 
+#include <stdlib.h>
 #include <objc/runtime.h>
 #include <CoreFoundation/CFBase.h>
 
@@ -62,7 +66,7 @@ NS_ASSUME_NONNULL_BEGIN
 	
 	// Allocate the correct size & zero it
 	// All allocations must be 16 byte aligned
-	SIZE_T Size = Align(FMath::Max(class_getInstanceSize(NewClass), class_getInstanceSize([FApplePlatformObjectZombie class])), 16);
+	SIZE_T Size = Align(FPlatformMath::Max(class_getInstanceSize(NewClass), class_getInstanceSize([FApplePlatformObjectZombie class])), 16);
 	void* Mem = nullptr;
 	
 	OSQueueHead* Alloc = [NewClass classAllocator];
