Fix bug with .ini updates

diff --git a/Engine/Source/Runtime/Core/Private/Misc/ConfigCacheIni.cpp b/Engine/Source/Runtime/Core/Private/Misc/ConfigCacheIni.cpp
--- a/Engine/Source/Runtime/Core/Private/Misc/ConfigCacheIni.cpp
+++ b/Engine/Source/Runtime/Core/Private/Misc/ConfigCacheIni.cpp
@@ -3454,7 +3454,6 @@ private:
 				{
 					UpdatedSection += FConfigFile::GenerateExportedPropertyLine(PropertyName, PropertyValue);
 					bWrotePropertyOnPass = true;
-					break;
 				}
 				else
 				{
