Ignore Clang's nonportable-include-path warning on Mac

diff --git a/Engine/Source/Runtime/Core/Public/Mac/MacPlatformCompilerPreSetup.h b/Engine/Source/Runtime/Core/Public/Mac/MacPlatformCompilerPreSetup.h
--- a/Engine/Source/Runtime/Core/Public/Mac/MacPlatformCompilerPreSetup.h
+++ b/Engine/Source/Runtime/Core/Public/Mac/MacPlatformCompilerPreSetup.h
@@ -151,6 +151,7 @@
 #pragma clang diagnostic ignored "-Wobjc-protocol-property-synthesis"
 #pragma clang diagnostic ignored "-Wnullability-completeness-on-arrays"
 #pragma clang diagnostic ignored "-Wnull-dereference"
+#pragma clang diagnostic ignored "-Wnonportable-include-path" // Ideally this one would be set in MacToolChain, but we don't have a way to check the compiler version in there yet
 #endif
 
 // We can use pragma optimisation's on and off as of Apple LLVM 7.3.0 but not before.
