Copying //UE4/Dev-Framework to //UE4/Dev-Main (Source: //UE4/Dev-Framework @ 3431384)

diff --git a/Engine/Source/Runtime/Core/Public/Containers/SortedMap.h b/Engine/Source/Runtime/Core/Public/Containers/SortedMap.h
new file mode 100644
--- /dev/null
+++ b/Engine/Source/Runtime/Core/Public/Containers/SortedMap.h
@@ -0,0 +1,669 @@
+// Copyright 1998-2017 Epic Games, Inc. All Rights Reserved.
+
+#pragma once
+
+#include "Containers/Map.h"
+#include "Containers/Algo/BinarySearch.h"
+#include "Containers/Algo/Sort.h"
+#include "UObject/NameTypes.h"
+
+/** 
+ * A Map of keys to value, implemented as a sorted TArray of TPairs. It has a mostly identical interface to TMap and is designed as a drop in replacement
+ * Keys must be unique, there is no equivalent sorted version of TMultiMap
+ * It uses half as much memory as TMap, but adding/removing elements is O(n) and finding is O(Log n)
+ * In practice it is faster than TMap for low element counts, and slower as n increases
+ * This map is always kept sorted by the key type so cannot be sorted manually
+ */
+template <typename KeyType, typename ValueType, typename ArrayAllocator /*= FDefaultAllocator*/, typename SortPredicate /* = TLess<KeyType> */ >
+class TSortedMap
+{
+	template <typename OtherKeyType, typename OtherValueType, typename OtherArrayAllocator, typename OtherSortPredicate>
+	friend class TSortedMap;
+	friend struct TContainerTraits<TSortedMap>;
+
+public:
+	typedef typename TTypeTraits<KeyType  >::ConstPointerType KeyConstPointerType;
+	typedef typename TTypeTraits<KeyType  >::ConstInitType    KeyInitType;
+	typedef typename TTypeTraits<ValueType>::ConstInitType    ValueInitType;
+	typedef TPair<KeyType, ValueType> ElementType;
+	
+	TSortedMap() = default;
+	TSortedMap(TSortedMap&&) = default;
+	TSortedMap(const TSortedMap&) = default;
+	TSortedMap& operator=(TSortedMap&&) = default;
+	TSortedMap& operator=(const TSortedMap&) = default;
+
+	/** Constructor for moving elements from a TSortedMap with a different ArrayAllocator */
+	template<typename OtherArrayAllocator>
+	TSortedMap(TSortedMap<KeyType, ValueType, OtherArrayAllocator, SortPredicate>&& Other)
+		: Pairs(MoveTemp(Other.Pairs))
+	{
+	}
+
+	/** Constructor for copying elements from a TSortedMap with a different ArrayAllocator */
+	template<typename OtherArrayAllocator>
+	TSortedMap(const TSortedMap<KeyType, ValueType, OtherArrayAllocator, SortPredicate>& Other)
+		: Pairs(Other.Pairs)
+	{
+	}
+
+	/** Assignment operator for moving elements from a TSortedMap with a different ArrayAllocator */
+	template<typename OtherArrayAllocator>
+	TSortedMap& operator=(TSortedMap<KeyType, ValueType, OtherArrayAllocator, SortPredicate>&& Other)
+	{
+		Pairs = MoveTemp(Other.Pairs);
+		return *this;
+	}
+
+	/** Assignment operator for copying elements from a TSortedMap with a different ArrayAllocator */
+	template<typename OtherArrayAllocator>
+	TSortedMap& operator=(const TSortedMap<KeyType, ValueType, OtherArrayAllocator, SortPredicate>& Other)
+	{
+		Pairs = Other.Pairs;
+		return *this;
+	}
+
+	/** Equality operator. This is efficient because pairs are always sorted */
+	FORCEINLINE bool operator==(const TSortedMap& Other) const
+	{
+		return Pairs == Other.Pairs;
+	}
+
+	/** Inequality operator. This is efficient because pairs are always sorted */
+	FORCEINLINE bool operator!=(const TSortedMap& Other) const
+	{
+		return Pairs != Other.Pairs;
+	}
+
+	/**
+	 * Removes all elements from the map, potentially leaving space allocated for an expected number of elements about to be added.
+	 * @param ExpectedNumElements - The number of elements about to be added to the set.
+	 */
+	FORCEINLINE void Empty(int32 ExpectedNumElements = 0)
+	{
+		Pairs.Empty(ExpectedNumElements);
+	}
+
+    /** Efficiently empties out the map but preserves all allocations and capacities */
+    FORCEINLINE void Reset()
+    {
+        Pairs.Reset();
+    }
+
+	/** Shrinks the pair set to avoid slack. */
+	FORCEINLINE void Shrink()
+	{
+		Pairs.Shrink();
+	}
+
+	/** Preallocates enough memory to contain Number elements */
+	FORCEINLINE void Reserve(int32 Number)
+	{
+		Pairs.Reserve(Number);
+	}
+
+	/** @return The number of elements in the map. */
+	FORCEINLINE int32 Num() const
+	{
+		return Pairs.Num();
+	}
+
+	/** 
+	 * Helper function to return the amount of memory allocated by this container 
+	 * @return number of bytes allocated by this container
+	 */
+	FORCEINLINE uint32 GetAllocatedSize() const
+	{
+		return Pairs.GetAllocatedSize();
+	}
+
+	/** Tracks the container's memory use through an archive. */
+	FORCEINLINE void CountBytes(FArchive& Ar)
+	{
+		Pairs.CountBytes(Ar);
+	}
+
+	/**
+	 * Sets the value associated with a key.
+	 *
+	 * @param InKey - The key to associate the value with.
+	 * @param InValue - The value to associate with the key.
+	 * @return A reference to the value as stored in the map.  The reference is only valid until the next change to any key in the map.
+	 */
+	FORCEINLINE ValueType& Add(const KeyType&  InKey, const ValueType&  InValue) { return Emplace(         InKey ,          InValue ); }
+	FORCEINLINE ValueType& Add(const KeyType&  InKey,       ValueType&& InValue) { return Emplace(         InKey , MoveTemp(InValue)); }
+	FORCEINLINE ValueType& Add(      KeyType&& InKey, const ValueType&  InValue) { return Emplace(MoveTemp(InKey),          InValue ); }
+	FORCEINLINE ValueType& Add(      KeyType&& InKey,       ValueType&& InValue) { return Emplace(MoveTemp(InKey), MoveTemp(InValue)); }
+
+	/**
+	 * Sets a default value associated with a key.
+	 *
+	 * @param InKey - The key to associate the value with.
+	 * @return A reference to the value as stored in the map.  The reference is only valid until the next change to any key in the map.
+	 */
+	FORCEINLINE ValueType& Add(const KeyType&  InKey) { return Emplace(         InKey ); }
+	FORCEINLINE ValueType& Add(      KeyType&& InKey) { return Emplace(MoveTemp(InKey)); }
+
+	/**
+	 * Sets the value associated with a key.
+	 *
+	 * @param InKey - The key to associate the value with.
+	 * @param InValue - The value to associate with the key.
+	 * @return A reference to the value as stored in the map.  The reference is only valid until the next change to any key in the map.
+	 */
+	template <typename InitKeyType, typename InitValueType>
+	ValueType& Emplace(InitKeyType&& InKey, InitValueType&& InValue)
+	{
+		ElementType* DataPtr = AllocateMemoryForEmplace(InKey);
+	
+		new(DataPtr) ElementType(TPairInitializer<InitKeyType&&, InitValueType&&>(Forward<InitKeyType>(InKey), Forward<InitValueType>(InValue)));
+
+		return DataPtr->Value;
+	}
+
+	/**
+	 * Sets a default value associated with a key.
+	 *
+	 * @param InKey - The key to associate the value with.
+	 * @return A reference to the value as stored in the map.  The reference is only valid until the next change to any key in the map.
+	 */
+	template <typename InitKeyType>
+	ValueType& Emplace(InitKeyType&& InKey)
+	{
+		ElementType* DataPtr = AllocateMemoryForEmplace(InKey);
+
+		new(DataPtr) ElementType(TKeyInitializer<InitKeyType&&>(Forward<InitKeyType>(InKey)));
+
+		return DataPtr->Value;
+	}
+
+	/**
+	 * Removes all value associations for a key.
+	 * @param InKey - The key to remove associated values for.
+	 * @return The number of values that were associated with the key.
+	 */
+	FORCEINLINE int32 Remove(KeyConstPointerType InKey)
+	{
+		int32 RemoveIndex = FindIndex(InKey);
+
+		if (RemoveIndex == INDEX_NONE)
+		{
+			return 0;
+		}
+
+		Pairs.RemoveAt(RemoveIndex);
+
+		return 1;
+	}
+
+	/**
+	 * Returns the key associated with the specified value.  The time taken is O(N) in the number of pairs.
+	 * @param	Value - The value to search for
+	 * @return	A pointer to the key associated with the specified value, or nullptr if the value isn't contained in this map.  The pointer
+	 *			is only valid until the next change to any key in the map.
+	 */
+	const KeyType* FindKey(ValueInitType Value) const
+	{
+		for(typename ElementArrayType::TConstIterator PairIt(Pairs);PairIt;++PairIt)
+		{
+			if(PairIt->Value == Value)
+			{
+				return &PairIt->Key;
+			}
+		}
+		return nullptr;
+	}
+
+	/**
+	 * Returns the value associated with a specified key.
+	 * @param	Key - The key to search for.
+	 * @return	A pointer to the value associated with the specified key, or nullptr if the key isn't contained in this map.  The pointer
+	 *			is only valid until the next change to any key in the map.
+	 */
+	FORCEINLINE ValueType* Find(KeyConstPointerType Key)
+	{
+		int32 FoundIndex = FindIndex(Key);
+
+		if (FoundIndex != INDEX_NONE)
+		{
+			return &Pairs[FoundIndex].Value;
+		}
+
+		return nullptr;
+	}
+
+	FORCEINLINE const ValueType* Find(KeyConstPointerType Key) const
+	{
+		return const_cast<TSortedMap*>(this)->Find(Key);
+	}
+
+	/**
+	 * Returns the value associated with a specified key, or if none exists, 
+	 * adds a value using the default constructor.
+	 * @param	Key - The key to search for.
+	 * @return	A reference to the value associated with the specified key.
+	 */
+	FORCEINLINE ValueType& FindOrAdd(const KeyType&  Key) { return FindOrAddImpl(         Key ); }
+	FORCEINLINE ValueType& FindOrAdd(      KeyType&& Key) { return FindOrAddImpl(MoveTemp(Key)); }
+
+	/**
+	 * Returns a reference to the value associated with a specified key.
+	 * @param	Key - The key to search for.
+	 * @return	The value associated with the specified key, or triggers an assertion if the key does not exist.
+	 */
+	FORCEINLINE ValueType& FindChecked(KeyConstPointerType Key)
+	{
+		ValueType* Value = Find(Key);
+		check(Value != nullptr);
+		return *Value;
+	}
+
+	FORCEINLINE const ValueType& FindChecked(KeyConstPointerType Key) const
+	{
+		const ValueType* Value = Find(Key);
+		check(Value != nullptr);
+		return *Value;
+	}
+
+	/**
+	 * Returns the value associated with a specified key.
+	 * @param	Key - The key to search for.
+	 * @return	The value associated with the specified key, or the default value for the ValueType if the key isn't contained in this map.
+	 */
+	FORCEINLINE ValueType FindRef(KeyConstPointerType Key) const
+	{
+		if (const ValueType* Value = Find(Key))
+		{
+			return *Value;
+		}
+
+		return ValueType();
+	}
+
+	/**
+	 * Checks if map contains the specified key.
+	 * @param Key - The key to check for.
+	 * @return true if the map contains the key.
+	 */
+	FORCEINLINE bool Contains(KeyConstPointerType Key) const
+	{
+		if (Find(Key))
+		{
+			return true;
+		}
+		return false;
+	}
+
+	/**
+	 * Returns the unique keys contained within this map
+	 * @param OutKeys - Upon return, contains the set of unique keys in this map.
+	 * @return The number of unique keys in the map.
+	 */
+	template<typename Allocator> int32 GetKeys(TArray<KeyType, Allocator>& OutKeys) const
+	{
+		for (typename ElementArrayType::TConstIterator PairIt(Pairs); PairIt; ++PairIt)
+		{
+			new(OutKeys) KeyType(PairIt->Key);
+		}
+
+		return OutKeys.Num();
+	}
+
+	/**
+	 * Generates an array from the keys in this map.
+	 */
+	template<typename Allocator> void GenerateKeyArray(TArray<KeyType, Allocator>& OutArray) const
+	{
+		OutArray.Empty(Pairs.Num());
+		for(typename ElementArrayType::TConstIterator PairIt(Pairs);PairIt;++PairIt)
+		{
+			new(OutArray) KeyType(PairIt->Key);
+		}
+	}
+
+	/**
+	 * Generates an array from the values in this map.
+	 */
+	template<typename Allocator> void GenerateValueArray(TArray<ValueType, Allocator>& OutArray) const
+	{
+		OutArray.Empty(Pairs.Num());
+		for(typename ElementArrayType::TConstIterator PairIt(Pairs);PairIt;++PairIt)
+		{
+			new(OutArray) ValueType(PairIt->Value);
+		}
+	}
+
+	/** Serializer. */
+	FORCEINLINE friend FArchive& operator<<(FArchive& Ar, TSortedMap& Map)
+	{
+		Ar << Map.Pairs;
+
+		if (Ar.IsLoading())
+		{
+			// We need to resort, in case the sorting is not consistent with what it was before
+			Algo::SortBy(Map.Pairs, FKeyForward(), SortPredicate());
+		}
+		return Ar;
+	}
+
+	/**
+	 * Describes the map's contents through an output device.
+	 * @param Ar - The output device to describe the map's contents through.
+	 */
+	void Dump(FOutputDevice& Ar)
+	{
+		Pairs.Dump(Ar);
+	}
+
+	/**
+	 * Removes the pair with the specified key and copies the value that was removed to the ref parameter
+	 * @param Key - the key to search for
+	 * @param OutRemovedValue - if found, the value that was removed (not modified if the key was not found)
+	 * @return whether or not the key was found
+	 */
+	FORCEINLINE bool RemoveAndCopyValue(KeyInitType Key, ValueType& OutRemovedValue)
+	{
+		int32 FoundIndex = FindIndex(Key);
+		if (FoundIndex == INDEX_NONE)
+		{
+			return false;
+		}
+
+		OutRemovedValue = MoveTemp(Pairs[FoundIndex].Value);
+		Pairs.RemoveAt(FoundIndex);
+		return true;
+	}
+	
+	/**
+	 * Finds a pair with the specified key, removes it from the map, and returns the value part of the pair.
+	 * If no pair was found, an exception is thrown.
+	 * @param Key - the key to search for
+	 * @return whether or not the key was found
+	 */
+	FORCEINLINE ValueType FindAndRemoveChecked(KeyConstPointerType Key)
+	{
+		int32 FoundIndex = FindIndex(Key);
+		check(FoundIndex != INDEX_NONE);
+		
+		ValueType OutRemovedValue = MoveTemp(Pairs[FoundIndex].Value);
+		Pairs.RemoveAt(FoundIndex);
+		return OutRemovedValue;
+	}
+
+	/**
+	 * Move all items from another map into our map (if any keys are in both, the value from the other map wins) and empty the other map.
+	 * @param OtherMap - The other map of items to move the elements from.
+	 */
+	template<typename OtherArrayAllocator, typename OtherSortPredicate>
+	void Append(TSortedMap<KeyType, ValueType, OtherArrayAllocator, OtherSortPredicate>&& OtherMap)
+	{
+		this->Reserve(this->Num() + OtherMap.Num());
+		for (auto& Pair : OtherMap)
+		{
+			this->Add(MoveTemp(Pair.Key), MoveTemp(Pair.Value));
+		}
+
+		OtherMap.Reset();
+	}
+
+	/**
+	 * Add all items from another map to our map (if any keys are in both, the value from the other map wins)
+	 * @param OtherMap - The other map of items to add.
+	 */
+	template<typename OtherArrayAllocator, typename OtherSortPredicate>
+	void Append(const TSortedMap<KeyType, ValueType, OtherArrayAllocator, OtherSortPredicate>& OtherMap)
+	{
+		this->Reserve(this->Num() + OtherMap.Num());
+		for (auto& Pair : OtherMap)
+		{
+			this->Add(Pair.Key, Pair.Value);
+		}
+	}
+
+	FORCEINLINE       ValueType& operator[](KeyConstPointerType Key)       { return this->FindChecked(Key); }
+	FORCEINLINE const ValueType& operator[](KeyConstPointerType Key) const { return this->FindChecked(Key); }
+
+private:
+	typedef TArray<ElementType, ArrayAllocator> ElementArrayType;
+
+	/** Implementation of find and add */
+	template <typename ArgType>
+	FORCEINLINE ValueType& FindOrAddImpl(ArgType&& Key)
+	{
+		if (ValueType* Value = Find(Key))
+		{
+			return *Value;
+		}
+
+		return Add(Forward<ArgType>(Key));
+	}
+
+	/** Find index of key */
+	FORCEINLINE int32 FindIndex(KeyConstPointerType Key)
+	{
+		return Algo::BinarySearchBy(Pairs, Key, FKeyForward(), SortPredicate());
+	}
+
+	/** Allocates raw memory for emplacing */
+	template <typename InitKeyType>
+	FORCEINLINE ElementType* AllocateMemoryForEmplace(InitKeyType&& InKey)
+	{
+		int32 InsertIndex = Algo::LowerBoundBy(Pairs, InKey, FKeyForward(), SortPredicate());
+		check(InsertIndex >= 0 && InsertIndex <= Pairs.Num());
+
+		ElementType* DataPtr = nullptr;
+		// Since we returned lower bound we already know InKey <= InsertIndex key. So if InKey is not < InsertIndex key, they must be equal
+		if (Pairs.IsValidIndex(InsertIndex) && !SortPredicate()(InKey, Pairs[InsertIndex].Key))
+		{
+			// Replacing element, delete old one
+			DataPtr = Pairs.GetData() + InsertIndex;
+			DestructItems(DataPtr, 1);
+		}
+		else
+		{
+			// Adding new one, this may reallocate Pairs
+			Pairs.InsertUninitialized(InsertIndex, 1);
+			DataPtr = Pairs.GetData() + InsertIndex;
+		}
+
+		return DataPtr;
+	}
+
+	/** Forwards sorting into Key of pair */
+	struct FKeyForward
+	{
+		FORCEINLINE const KeyType& operator()(const ElementType& Pair) const
+		{
+			return Pair.Key;
+		}
+	};
+
+	/** The base of TSortedMap iterators */
+	template<bool bConst>
+	class TBaseIterator
+	{
+	public:
+		typedef typename TChooseClass<bConst,typename ElementArrayType::TConstIterator,typename ElementArrayType::TIterator>::Result PairItType;
+	private:
+		typedef typename TChooseClass<bConst,const TSortedMap,TSortedMap>::Result MapType;
+		typedef typename TChooseClass<bConst,const KeyType,KeyType>::Result ItKeyType;
+		typedef typename TChooseClass<bConst,const ValueType,ValueType>::Result ItValueType;
+		typedef typename TChooseClass<bConst,const typename ElementArrayType::ElementType, typename ElementArrayType::ElementType>::Result PairType;
+
+	protected:
+		FORCEINLINE TBaseIterator(const PairItType& InElementIt)
+			: PairIt(InElementIt)
+		{
+		}
+
+	public:
+		FORCEINLINE TBaseIterator& operator++()
+		{
+			++PairIt;
+			return *this;
+		}
+
+		/** conversion to "bool" returning true if the iterator is valid. */
+		FORCEINLINE explicit operator bool() const
+		{ 
+			return !!PairIt; 
+		}
+
+		FORCEINLINE friend bool operator==(const TBaseIterator& Lhs, const TBaseIterator& Rhs) { return Lhs.PairIt == Rhs.PairIt; }
+		FORCEINLINE friend bool operator!=(const TBaseIterator& Lhs, const TBaseIterator& Rhs) { return Lhs.PairIt != Rhs.PairIt; }
+
+		FORCEINLINE ItKeyType&   Key()   const { return PairIt->Key; }
+		FORCEINLINE ItValueType& Value() const { return PairIt->Value; }
+
+		FORCEINLINE PairType& operator* () const { return  *PairIt; }
+		FORCEINLINE PairType* operator->() const { return &*PairIt; }
+
+	protected:
+		PairItType PairIt;
+	};
+
+	/** An array of the key-value pairs in the map */
+	ElementArrayType Pairs;
+
+public:
+
+	/** Map iterator */
+	class TIterator : public TBaseIterator<false>
+	{
+	public:
+
+		FORCEINLINE TIterator(TSortedMap& InMap)
+			: TBaseIterator<false>(InMap.Pairs.CreateIterator())
+		{
+		}
+
+		FORCEINLINE TIterator(const typename TBaseIterator<false>::PairItType& InPairIt)
+			: TBaseIterator<false>(InPairIt)
+		{
+		}
+
+		/** Removes the current pair from the map */
+		FORCEINLINE void RemoveCurrent()
+		{
+			TBaseIterator<false>::PairIt.RemoveCurrent();
+		}
+	};
+
+	/** Const map iterator */
+	class TConstIterator : public TBaseIterator<true>
+	{
+	public:
+		FORCEINLINE TConstIterator(const TSortedMap& InMap)
+			: TBaseIterator<true>(InMap.Pairs.CreateConstIterator())
+		{
+		}
+
+		FORCEINLINE TConstIterator(const typename TBaseIterator<true>::PairItType& InPairIt)
+			: TBaseIterator<true>(InPairIt)
+		{
+		}
+	};
+
+	/** Iterates over values associated with a specified key in a const map. This will be at most one value because keys must be unique */
+	class TConstKeyIterator : public TBaseIterator<true>
+	{
+	public:
+		FORCEINLINE TConstKeyIterator(const TSortedMap& InMap, KeyInitType InKey)
+			: TBaseIterator<true>(InMap.Pairs.CreateIterator())
+		{
+			int32 NewIndex = FindIndex(InKey);
+		
+			if (NewIndex != INDEX_NONE)
+			{
+				TBaseIterator<true>::PairIt += NewIndex;
+			}
+			else
+			{
+				TBaseIterator<true>::PairIt.SetToEnd();
+			}
+		}
+
+		FORCEINLINE TConstKeyIterator& operator++()
+		{
+			TBaseIterator<true>::PairIt.SetToEnd();
+			return *this;
+		}
+	};
+
+	/** Iterates over values associated with a specified key in a map. This will be at most one value because keys must be unique */
+	class TKeyIterator : public TBaseIterator<false>
+	{
+	public:
+		FORCEINLINE TKeyIterator(TSortedMap& InMap, KeyInitType InKey)
+			: TBaseIterator<false>(InMap.Pairs.CreateConstIterator())
+		{
+			int32 NewIndex = FindIndex(InKey);
+
+			if (NewIndex != INDEX_NONE)
+			{
+				TBaseIterator<true>::PairIt += NewIndex;
+			}
+			else
+			{
+				TBaseIterator<true>::PairIt.SetToEnd();
+			}
+		}
+
+		FORCEINLINE TKeyIterator& operator++()
+		{
+			TBaseIterator<false>::PairIt.SetToEnd();
+			return *this;
+		}
+
+		/** Removes the current key-value pair from the map. */
+		FORCEINLINE void RemoveCurrent()
+		{
+			TBaseIterator<false>::PairIt.RemoveCurrent();
+			TBaseIterator<false>::PairIt.SetToEnd();
+		}
+	};
+
+	/** Creates an iterator over all the pairs in this map */
+	FORCEINLINE TIterator CreateIterator()
+	{
+		return TIterator(*this);
+	}
+
+	/** Creates a const iterator over all the pairs in this map */
+	FORCEINLINE TConstIterator CreateConstIterator() const
+	{
+		return TConstIterator(*this);
+	}
+
+	/** Creates an iterator over the values associated with a specified key in a map. This will be at most one value because keys must be unique */
+	FORCEINLINE TKeyIterator CreateKeyIterator(KeyInitType InKey)
+	{
+		return TKeyIterator(*this, InKey);
+	}
+
+	/** Creates a const iterator over the values associated with a specified key in a map. This will be at most one value because keys must be unique */
+	FORCEINLINE TConstKeyIterator CreateConstKeyIterator(KeyInitType InKey) const
+	{
+		return TConstKeyIterator(*this, InKey);
+	}
+
+	/** Ranged For iterators. Unlike normal TMap these are not the same as the normal iterator for performance reasons */
+	typedef typename ElementArrayType::RangedForIteratorType RangedForIteratorType;
+	typedef typename ElementArrayType::RangedForConstIteratorType RangedForConstIteratorType;
+
+private:
+	/**
+	 * DO NOT USE DIRECTLY
+	 * STL-like iterators to enable range-based for loop support.
+	 */
+	FORCEINLINE friend RangedForIteratorType		begin(TSortedMap& MapBase)		{ return begin(MapBase.Pairs); }
+	FORCEINLINE friend RangedForConstIteratorType	begin(const TSortedMap& MapBase){ return begin(MapBase.Pairs); }
+	FORCEINLINE friend RangedForIteratorType		end(TSortedMap& MapBase)		{ return end(MapBase.Pairs); }
+	FORCEINLINE friend RangedForConstIteratorType	end(const TSortedMap& MapBase)	{ return end(MapBase.Pairs); }
+};
+
+template <typename KeyType, typename ValueType, typename ArrayAllocator, typename SortPredicate>
+struct TContainerTraits<TSortedMap<KeyType, ValueType, ArrayAllocator, SortPredicate>> : public TContainerTraitsBase<TSortedMap<KeyType, ValueType, ArrayAllocator, SortPredicate>>
+{
+	enum { MoveWillEmptyContainer = TContainerTraits<typename TSortedMap<KeyType, ValueType, ArrayAllocator, SortPredicate>::ElementArrayType>::MoveWillEmptyContainer };
+};
diff --git a/Engine/Source/Runtime/Core/Public/CoreGlobals.h b/Engine/Source/Runtime/Core/Public/CoreGlobals.h
--- a/Engine/Source/Runtime/Core/Public/CoreGlobals.h
+++ b/Engine/Source/Runtime/Core/Public/CoreGlobals.h
@@ -89,9 +89,6 @@ extern CORE_API bool GCompilingBlueprint;
 /** True if we're reconstructing blueprint instances. Should never be true on cooked builds */
 extern CORE_API bool GIsReconstructingBlueprintInstances;
 
-/** Force blueprints to not compile on load */
-extern CORE_API bool GForceDisableBlueprintCompileOnLoad;
-
 /** True if actors and objects are being re-instanced. */
 extern CORE_API bool GIsReinstancing;
 
diff --git a/Engine/Source/Runtime/Core/Public/Math/NumericLimits.h b/Engine/Source/Runtime/Core/Public/Math/NumericLimits.h
--- a/Engine/Source/Runtime/Core/Public/Math/NumericLimits.h
+++ b/Engine/Source/Runtime/Core/Public/Math/NumericLimits.h
@@ -73,17 +73,17 @@ struct TNumericLimits<uint8>
 {
 	typedef uint8 NumericType;
 	
-	static NumericType Min()
+	static constexpr NumericType Min()
 	{
 		return MIN_uint8;
 	}
 
-	static NumericType Max()
+	static constexpr NumericType Max()
 	{
 		return MAX_uint8;
 	}
 
-	static NumericType Lowest()
+	static constexpr NumericType Lowest()
 	{
 		return Min();
 	}
@@ -95,17 +95,17 @@ struct TNumericLimits<uint16>
 {
 	typedef uint16 NumericType;
 	
-	static NumericType Min()
+	static constexpr NumericType Min()
 	{
 		return MIN_uint16;
 	}
 
-	static NumericType Max()
+	static constexpr NumericType Max()
 	{
 		return MAX_uint16;
 	}
 
-	static NumericType Lowest()
+	static constexpr NumericType Lowest()
 	{
 		return Min();
 	}
@@ -117,17 +117,17 @@ struct TNumericLimits<uint32>
 {
 	typedef uint32 NumericType;
 	
-	static NumericType Min()
+	static constexpr NumericType Min()
 	{
 		return MIN_uint32;
 	}
 
-	static NumericType Max()
+	static constexpr NumericType Max()
 	{
 		return MAX_uint32;
 	}
 	
-	static NumericType Lowest()
+	static constexpr NumericType Lowest()
 	{
 		return Min();
 	}
@@ -139,17 +139,17 @@ struct TNumericLimits<uint64>
 {
 	typedef uint64 NumericType;
 	
-	static NumericType Min()
+	static constexpr NumericType Min()
 	{
 		return MIN_uint64;
 	}
 
-	static NumericType Max()
+	static constexpr NumericType Max()
 	{
 		return MAX_uint64;
 	}
 
-	static NumericType Lowest()
+	static constexpr NumericType Lowest()
 	{
 		return Min();
 	}
@@ -161,17 +161,17 @@ struct TNumericLimits<int8>
 {
 	typedef int8 NumericType;
 	
-	static NumericType Min()
+	static constexpr NumericType Min()
 	{
 		return MIN_int8;
 	}
 
-	static NumericType Max()
+	static constexpr NumericType Max()
 	{
 		return MAX_int8;
 	}
 
-	static NumericType Lowest()
+	static constexpr NumericType Lowest()
 	{
 		return Min();
 	}
@@ -183,17 +183,17 @@ struct TNumericLimits<int16>
 {
 	typedef int16 NumericType;
 
-	static NumericType Min()
+	static constexpr NumericType Min()
 	{
 		return MIN_int16;
 	}
 
-	static NumericType Max()
+	static constexpr NumericType Max()
 	{
 		return MAX_int16;
 	}
 
-	static NumericType Lowest()
+	static constexpr NumericType Lowest()
 	{
 		return Min();
 	}
@@ -205,17 +205,17 @@ struct TNumericLimits<int32>
 {
 	typedef int32 NumericType;
 
-	static NumericType Min()
+	static constexpr NumericType Min()
 	{
 		return MIN_int32;
 	}
 
-	static NumericType Max()
+	static constexpr NumericType Max()
 	{
 		return MAX_int32;
 	}
 
-	static NumericType Lowest()
+	static constexpr NumericType Lowest()
 	{
 		return Min();
 	}
@@ -227,17 +227,17 @@ struct TNumericLimits<int64>
 {
 	typedef int64 NumericType;
 	
-	static NumericType Min()
+	static constexpr NumericType Min()
 	{
 		return MIN_int64;
 	}
 
-	static NumericType Max()
+	static constexpr NumericType Max()
 	{
 		return MAX_int64;
 	}
 
-	static NumericType Lowest()
+	static constexpr NumericType Lowest()
 	{
 		return Min();
 	}
@@ -249,17 +249,17 @@ struct TNumericLimits<float>
 {
 	typedef float NumericType;
 
-	static NumericType Min()
+	static constexpr NumericType Min()
 	{
 		return MIN_flt;
 	}
 
-	static NumericType Max()
+	static constexpr NumericType Max()
 	{
 		return MAX_flt;
 	}
 
-	static NumericType Lowest()
+	static constexpr NumericType Lowest()
 	{
 		return -Max();
 	}
@@ -271,17 +271,17 @@ struct TNumericLimits<double>
 {
 	typedef double NumericType;
 
-	static NumericType Min()
+	static constexpr NumericType Min()
 	{
 		return MIN_dbl;
 	}
 
-	static NumericType Max()
+	static constexpr NumericType Max()
 	{
 		return MAX_dbl;
 	}
 
-	static NumericType Lowest()
+	static constexpr NumericType Lowest()
 	{
 		return -Max();
 	}
diff --git a/Engine/Source/Runtime/Core/Public/Misc/CoreDelegates.h b/Engine/Source/Runtime/Core/Public/Misc/CoreDelegates.h
--- a/Engine/Source/Runtime/Core/Public/Misc/CoreDelegates.h
+++ b/Engine/Source/Runtime/Core/Public/Misc/CoreDelegates.h
@@ -129,9 +129,15 @@ public:
 	// Called when an error occurred.
 	static FSimpleMulticastDelegate OnShutdownAfterError;
 
-	// Called when appInit is called.
+	// Called when appInit is called, very early in startup
 	static FSimpleMulticastDelegate OnInit;
 
+	// Called at the end of UEngine::Init, right before loading PostEngineInit modules for both normal execution and commandlets
+	static FSimpleMulticastDelegate OnPostEngineInit;
+
+	// Called at the very end of engine initialization, right before the engine starts ticking. This is not called for commandlets
+	static FSimpleMulticastDelegate OnFEngineLoopInitComplete;
+
 	// Called when the application is about to exit.
 	static FSimpleMulticastDelegate OnExit;
 
@@ -313,9 +319,6 @@ public:
 	/* Sent just before the rendering thread is destroyed. */
 	static FRenderingThreadChanged PreRenderingThreadDestroyed;
 
-	// Called when appInit is called.
-	static FSimpleMulticastDelegate OnFEngineLoopInitComplete;
-
 	// Callback to allow custom resolution of package names. Arguments are InRequestedName, OutResolvedName.
 	// Should return True of resolution occured.
 	DECLARE_DELEGATE_RetVal_TwoParams(bool, FResolvePackageNameDelegate, const FString&, FString&);
diff --git a/Engine/Source/Runtime/Core/Public/Templates/Less.h b/Engine/Source/Runtime/Core/Public/Templates/Less.h
--- a/Engine/Source/Runtime/Core/Public/Templates/Less.h
+++ b/Engine/Source/Runtime/Core/Public/Templates/Less.h
@@ -6,10 +6,11 @@
 
 /**
  * Binary predicate class for sorting elements in order.  Assumes < operator is defined for the template type.
+ * Forward declaration exists in ContainersFwd.h
  *
  * See: http://en.cppreference.com/w/cpp/utility/functional/less
  */
-template <typename T = void>
+template <typename T /*= void */>
 struct TLess
 {
 	FORCEINLINE bool operator()(const T& A, const T& B) const
diff --git a/Engine/Source/Runtime/Core/Public/Templates/Sorting.h b/Engine/Source/Runtime/Core/Public/Templates/Sorting.h
--- a/Engine/Source/Runtime/Core/Public/Templates/Sorting.h
+++ b/Engine/Source/Runtime/Core/Public/Templates/Sorting.h
@@ -5,6 +5,7 @@
 #include "CoreTypes.h"
 #include "HAL/PlatformMath.h"
 #include "Templates/Less.h"
+#include "Algo/BinarySearch.h"
 
 /**
  * Helper class for dereferencing pointer types in Sort function
@@ -313,75 +314,20 @@ public:
 
 		while (AStart < BStart && BStart < Num)
 		{
-			int32 NewAOffset = BinarySearchLast(First + AStart, BStart - AStart, First[BStart], Predicate) + 1;
+			// Index after the last value == First[BStart]
+			int32 NewAOffset = AlgoImpl::UpperBoundInternal(First + AStart, BStart - AStart, First[BStart], FIdentityFunctor(), Predicate);
 			AStart += NewAOffset;
 
 			if (AStart >= BStart) // done
 				break;
 
-			int32 NewBOffset = BinarySearchFirst(First + BStart, Num - BStart, First[AStart], Predicate);
+			// Index of the first value == First[AStart]
+			int32 NewBOffset = AlgoImpl::LowerBoundInternal(First + BStart, Num - BStart, First[AStart], FIdentityFunctor(), Predicate);
 			TRotationPolicy::Rotate(First, AStart, BStart + NewBOffset, NewBOffset);
 			BStart += NewBOffset;
 			AStart += NewBOffset + 1;
 		}
 	}
-
-private:
-	/**
-	 * Performs binary search, resulting in position of the first element with given value in an array.
-	 *
-	 * @param First Pointer to array.
-	 * @param Num Number of elements in array.
-	 * @param Value Value to look for.
-	 * @param Predicate Predicate for comparison.
-	 *
-	 * @returns Position of the first element with value Value.
-	 */
-	template <class T, class PREDICATE_CLASS>
-	static int32 BinarySearchFirst(T* First, const int32 Num, const T& Value, const PREDICATE_CLASS& Predicate)
-	{
-		int32 Start = 0;
-		int32 End = Num;
-		
-		while (End - Start > 1)
-		{
-			int32 Mid = (Start + End) / 2;
-			bool bComparison = Predicate(First[Mid], Value);
-
-			Start = bComparison ? Mid : Start;
-			End = bComparison ? End : Mid;
-		}
-
-		return Predicate(First[Start], Value) ? Start + 1 : Start;
-	}
-
-	/**
-	 * Performs binary search, resulting in position of the last element with given value in an array.
-	 *
-	 * @param First Pointer to array.
-	 * @param Num Number of elements in array.
-	 * @param Value Value to look for.
-	 * @param Predicate Predicate for comparison.
-	 *
-	 * @returns Position of the last element with value Value.
-	 */
-	template <class T, class PREDICATE_CLASS>
-	static int32 BinarySearchLast(T* First, const int32 Num, const T& Value, const PREDICATE_CLASS& Predicate)
-	{
-		int32 Start = 0;
-		int32 End = Num;
-
-		while (End - Start > 1)
-		{
-			int32 Mid = (Start + End) / 2;
-			bool bComparison = !Predicate(Value, First[Mid]);
-			
-			Start = bComparison ? Mid : Start;
-			End = bComparison ? End : Mid;
-		}
-
-		return Predicate(Value, First[Start]) ? Start - 1 : Start;
-	}
 };
 
 /**
diff --git a/Engine/Source/Runtime/Core/Public/UObject/BlueprintsObjectVersion.h b/Engine/Source/Runtime/Core/Public/UObject/BlueprintsObjectVersion.h
--- a/Engine/Source/Runtime/Core/Public/UObject/BlueprintsObjectVersion.h
+++ b/Engine/Source/Runtime/Core/Public/UObject/BlueprintsObjectVersion.h
@@ -20,6 +20,7 @@ struct CORE_API FBlueprintsObjectVersion
 		SCSHasComponentTemplateClass,
 		ComponentTemplateClassSupport,
 		ArrayGetFuncsReplacedByCustomNode,
+		DisallowObjectConfigVars,
 
 		// -----<new versions can be added above this line>-------------------------------------------------
 		VersionPlusOne,
diff --git a/Engine/Source/Runtime/Core/Public/UObject/FrameworkObjectVersion.h b/Engine/Source/Runtime/Core/Public/UObject/FrameworkObjectVersion.h
--- a/Engine/Source/Runtime/Core/Public/UObject/FrameworkObjectVersion.h
+++ b/Engine/Source/Runtime/Core/Public/UObject/FrameworkObjectVersion.h
@@ -85,6 +85,18 @@ struct CORE_API FFrameworkObjectVersion
 		// Serialize hard references to sound files when possible
 		HardSoundReferences,
 
+		// Enforce const correctness in Animation Blueprint function graphs
+		EnforceConstInAnimBlueprintFunctionGraphs,
+		
+		// Upgrade the InputKeySelector to use a text style
+		InputKeySelectorTextStyle,
+
+		// Represent a pins container type as an enum not 3 independent booleans
+		EdGraphPinContainerType,
+
+		// Switch asset pins to store as string instead of hard object referece
+		ChangeAssetPinsToString,
+
 		// -----<new versions can be added above this line>-------------------------------------------------
 		VersionPlusOne,
 		LatestVersion = VersionPlusOne - 1
diff --git a/Engine/Source/Runtime/Core/Public/UObject/NameTypes.h b/Engine/Source/Runtime/Core/Public/UObject/NameTypes.h
--- a/Engine/Source/Runtime/Core/Public/UObject/NameTypes.h
+++ b/Engine/Source/Runtime/Core/Public/UObject/NameTypes.h
@@ -59,8 +59,8 @@ typedef int32 NAME_INDEX;
 /** These characters cannot be used in object names */
 #define INVALID_OBJECTNAME_CHARACTERS	TEXT("\"' ,/.:|&!~\n\r\t@#(){}[]=;^%$`")
 
-/** These characters cannot be used in textboxes which take group names (i.e. Group1.Group2) */
-#define INVALID_GROUPNAME_CHARACTERS	TEXT("\"' ,/:|&!~\n\r\t@#")
+/** These characters cannot be used in ObjectPaths, which includes both the package path and part after the first . */
+#define INVALID_OBJECTPATH_CHARACTERS	TEXT("\"' ,|&!~\n\r\t@#(){}[]=;^%$`")
 
 /** These characters cannot be used in long package names */
 #define INVALID_LONGPACKAGE_CHARACTERS	TEXT("\\:*?\"<>|' ,.&!~\n\r\t@#")
@@ -740,6 +740,23 @@ public:
 	 */
 	int32 Compare( const FName& Other ) const;
 
+	/**
+	 * Fast compares name to passed in one using indexes. Sort is allocation order ascending.
+	 *
+	 * @param	Other	Name to compare this against
+	 * @return	< 0 is this < Other, 0 if this == Other, > 0 if this > Other
+	 */
+	FORCEINLINE int32 CompareIndexes(const FName& Other) const
+	{
+		int32 ComparisonDiff = GetComparisonIndexFast() - Other.GetComparisonIndexFast();
+
+		if (ComparisonDiff == 0)
+		{
+			return GetNumber() - Other.GetNumber();
+		}
+		return ComparisonDiff;
+	}
+
 	/**
 	 * Create an FName with a hardcoded string index.
 	 *
@@ -1179,3 +1196,11 @@ inline bool operator!=(const CharType *LHS, const FName &RHS)
 template <> struct TIsPODType<FName> { enum { Value = true }; };
 
 
+/** Sort predicate to sort FName by index instead of alphabetically, pass to anything that wants TLess */
+struct FNameSortIndexes
+{
+	FORCEINLINE bool operator()(const FName& A, const FName& B) const
+	{
+		return A.CompareIndexes(B) < 0;
+	}
+};
\ No newline at end of file
diff --git a/Engine/Source/Runtime/CoreUObject/Private/Blueprint/BlueprintSupport.cpp b/Engine/Source/Runtime/CoreUObject/Private/Blueprint/BlueprintSupport.cpp
--- a/Engine/Source/Runtime/CoreUObject/Private/Blueprint/BlueprintSupport.cpp
+++ b/Engine/Source/Runtime/CoreUObject/Private/Blueprint/BlueprintSupport.cpp
@@ -319,7 +319,7 @@ FScopedClassDependencyGather::~FScopedClassDependencyGather()
 {
 	// If this gatherer was the initial gatherer for the current scope, process 
 	// dependencies (unless compiling on load is explicitly disabled)
-	if( bMasterClass && !GForceDisableBlueprintCompileOnLoad )
+	if( bMasterClass )
 	{
 		auto DependencyIter = BatchClassDependencies.CreateIterator();
 		// implemented as a lambda, to prevent duplicated code between 
@@ -1474,6 +1474,17 @@ void FLinkerLoad::FinalizeBlueprint(UClass* LoadClass)
 	// we need it ran for any super-classes before we regen
 	ResolveAllImports();
 
+	// Now that imports have been resolved we optionally flush the compilation
+	// queue. This is only done for level blueprints, which will have instances
+	// of actors in them that cannot reliably be reinstanced on load (see useage
+	// of Scene pointers in things like UActorComponent::ExecuteRegisterEvents)
+	// - on load the Scene may not yet be created, meaning this code cannot 
+	// correctly be run. We could address that, but avoiding reinstancings is
+	// also a performance win:
+#if WITH_EDITOR
+	LoadClass->FlushCompilationQueueForLevel();
+#endif
+
 	// interfaces can invalidate classes which implement them (when the  
 	// interface is regenerated), they essentially define the makeup of the  
 	// implementing class; so here, like we do with the parent class above, we  
diff --git a/Engine/Source/Runtime/CoreUObject/Private/Misc/PackageName.cpp b/Engine/Source/Runtime/CoreUObject/Private/Misc/PackageName.cpp
--- a/Engine/Source/Runtime/CoreUObject/Private/Misc/PackageName.cpp
+++ b/Engine/Source/Runtime/CoreUObject/Private/Misc/PackageName.cpp
@@ -565,6 +565,70 @@ bool FPackageName::IsValidLongPackageName(const FString& InLongPackageName, bool
 	return bValidRoot;
 }
 
+bool FPackageName::IsValidObjectPath(const FString& InObjectPath, FText* OutReason)
+{
+	FString PackageName;
+	FString RemainingObjectPath;
+
+	// Check for package delimiter
+	int32 ObjectDelimiterIdx;
+	if (InObjectPath.FindChar('.', ObjectDelimiterIdx))
+	{
+		if (ObjectDelimiterIdx == InObjectPath.Len() - 1)
+		{
+			if (OutReason)
+			{
+				*OutReason = NSLOCTEXT("Core", "ObjectPath_EndWithPeriod", "Object Path may not end with .");
+			}
+			return false;
+		}
+
+		PackageName = InObjectPath.Mid(0, ObjectDelimiterIdx);
+		RemainingObjectPath = InObjectPath.Mid(ObjectDelimiterIdx + 1);
+	}
+	else
+	{
+		PackageName = InObjectPath;
+	}
+
+	if (!IsValidLongPackageName(PackageName, true, OutReason))
+	{
+		return false;
+	}
+
+	if (RemainingObjectPath.Len() > 0)
+	{
+		FText PathContext = NSLOCTEXT("Core", "ObjectPathContext", "Object Path");
+		if (!FName::IsValidXName(RemainingObjectPath, INVALID_OBJECTPATH_CHARACTERS, OutReason, &PathContext))
+		{
+			return false;
+		}
+
+		TCHAR LastChar = RemainingObjectPath[RemainingObjectPath.Len() - 1];
+		if (LastChar == '.' || LastChar == ':')
+		{
+			if (OutReason)
+			{
+				*OutReason = NSLOCTEXT("Core", "ObjectPath_PathWithTrailingSeperator", "Object Path may not end with : or .");
+			}
+			return false;
+		}
+
+		int32 SlashIndex;
+		if (RemainingObjectPath.FindChar('/', SlashIndex))
+		{
+			if (OutReason)
+			{
+				*OutReason = NSLOCTEXT("Core", "ObjectPath_SlashAfterPeriod", "Object Path may not have / after first .");
+			}
+
+			return false;
+		}
+	}
+
+	return true;
+}
+
 void FPackageName::RegisterMountPoint(const FString& RootPath, const FString& ContentPath)
 {
 	FLongPackagePathsSingleton& Paths = FLongPackagePathsSingleton::Get();
diff --git a/Engine/Source/Runtime/CoreUObject/Private/Misc/RedirectCollector.cpp b/Engine/Source/Runtime/CoreUObject/Private/Misc/RedirectCollector.cpp
--- a/Engine/Source/Runtime/CoreUObject/Private/Misc/RedirectCollector.cpp
+++ b/Engine/Source/Runtime/CoreUObject/Private/Misc/RedirectCollector.cpp
@@ -1,7 +1,5 @@
 // Copyright 1998-2017 Epic Games, Inc. All Rights Reserved.
 
-
-// Core includes.
 #include "Misc/RedirectCollector.h"
 #include "Misc/CoreDelegates.h"
 #include "UObject/UObjectGlobals.h"
@@ -13,48 +11,13 @@
 #include "Misc/PackageName.h"
 #include "UObject/LinkerLoad.h"
 #include "UObject/UObjectThreadContext.h"
-//#include "CookStats.h"
-DEFINE_LOG_CATEGORY_STATIC(LogRedirectors, Log, All);
-
-static TArray<FString> StringAssetRefFilenameStack;
-
-void FRedirectCollector::OnRedirectorFollowed(const FString& InString, UObject* InObject)
-{
-	check(InObject);
 
-	// the object had better be a redir
-	UObjectRedirector& Redirector = dynamic_cast<UObjectRedirector&>(*InObject);
+#if WITH_EDITOR
 
-	// save the info if it matches the parameters we need to match on:
-	// if the string matches the package we were loading
-	// AND
-	// we aren't matching a particular package || the string matches the package
-	if (!FileToFixup.Len() || FileToFixup == FPackageName::FilenameToLongPackageName(Redirector.GetLinker()->Filename))
-	{
-		FRedirection Redir;
-		Redir.PackageFilename = InString;
-		Redir.RedirectorName = Redirector.GetFullName();
-		Redir.RedirectorPackageFilename = Redirector.GetLinker()->Filename;
-		Redir.DestinationObjectName = Redirector.DestinationObject->GetFullName();
-		// we only want one of each redirection reported
-		Redirections.AddUnique(Redir);
-
-		// and add a fake reference caused by a string redirectors
-		if (StringAssetRefFilenameStack.Num())
-		{
-			Redir.PackageFilename = StringAssetRefFilenameStack[StringAssetRefFilenameStack.Num() - 1];
-			UE_LOG(LogRedirectors, Verbose, TEXT("String Asset Reference fake redir %s(%s) -> %s"), *Redir.RedirectorName, *Redir.PackageFilename, *Redir.DestinationObjectName);
-			if (!Redir.PackageFilename.IsEmpty())
-			{
-				Redirections.AddUnique(Redir);
-			}
-		}
-	}
-}
+DEFINE_LOG_CATEGORY_STATIC(LogRedirectors, Log, All);
 
 void FRedirectCollector::OnStringAssetReferenceLoaded(const FString& InString)
 {
-	FUObjectThreadContext& ThreadContext = FUObjectThreadContext::Get();
 	FPackagePropertyPair ContainingPackageAndProperty;
 
 	if (InString.IsEmpty())
@@ -63,27 +26,39 @@ void FRedirectCollector::OnStringAssetReferenceLoaded(const FString& InString)
 		return;
 	}
 
-	if (ThreadContext.SerializedObject)
+	FStringAssetReferenceThreadContext& ThreadContext = FStringAssetReferenceThreadContext::Get();
+
+	FName PackageName, PropertyName;
+	EStringAssetReferenceCollectType CollectType = EStringAssetReferenceCollectType::AlwaysCollect;
+
+	ThreadContext.GetSerializationOptions(PackageName, PropertyName, CollectType);
+
+	if (CollectType == EStringAssetReferenceCollectType::NeverCollect)
 	{
-		FLinkerLoad* Linker = ThreadContext.SerializedObject->GetLinker();
-		if (Linker)
+		// Do not track
+		return;
+	}
+
+	if (PackageName != NAME_None)
+	{
+		ContainingPackageAndProperty.SetPackage(PackageName);
+		if (PropertyName != NAME_None)
 		{
-			ContainingPackageAndProperty.SetPackage(FName(*FPackageName::FilenameToLongPackageName(Linker->Filename)));
-			if (Linker->GetSerializedProperty())
-			{
-				ContainingPackageAndProperty.SetProperty( FName(*FString::Printf(TEXT("%s:%s"), *ThreadContext.SerializedObject->GetPathName(), *Linker->GetSerializedProperty()->GetName())));
-			}
-#if WITH_EDITORONLY_DATA
-			ContainingPackageAndProperty.SetReferencedByEditorOnlyProperty( Linker->IsEditorOnlyPropertyOnTheStack() );
-#endif
+			ContainingPackageAndProperty.SetProperty(PropertyName);
 		}
+		ContainingPackageAndProperty.SetReferencedByEditorOnlyProperty(CollectType == EStringAssetReferenceCollectType::EditorOnlyCollect);
 	}
+
+	FScopeLock ScopeLock(&CriticalSection);
+
 	StringAssetReferences.AddUnique(FName(*InString), ContainingPackageAndProperty);
 }
 
 FString FRedirectCollector::OnStringAssetReferenceSaved(const FString& InString)
 {
-	FString* Found = StringAssetRemap.Find(InString);
+	FScopeLock ScopeLock(&CriticalSection);
+
+	FString* Found = AssetPathRedirectionMap.Find(InString);
 	if (Found)
 	{
 		return *Found;
@@ -91,7 +66,6 @@ FString FRedirectCollector::OnStringAssetReferenceSaved(const FString& InString)
 	return InString;
 }
 
-
 #define REDIRECT_TIMERS 1
 #if REDIRECT_TIMERS
 struct FRedirectScopeTimer
@@ -121,15 +95,11 @@ double ResolveTimeDelegate;
 	LOG_REDIRECT_TIMER(ResolveTimeDelegate);\
 	LOG_REDIRECT_TIMER(ResolveTimeTotal);\
 
-	
-
-
 #else
 #define SCOPE_REDIRECT_TIMER(TimerName)
 #define LOG_REDIRECT_TIMERS()
 #endif
 
-
 void FRedirectCollector::LogTimers() const
 {
 	LOG_REDIRECT_TIMERS();
@@ -139,6 +109,8 @@ void FRedirectCollector::ResolveStringAssetReference(FName FilterPackageFName, b
 {
 	SCOPE_REDIRECT_TIMER(ResolveTimeTotal);
 	
+	FScopeLock ScopeLock(&CriticalSection);
+
 	TMultiMap<FName, FPackagePropertyPair> SkippedReferences;
 	SkippedReferences.Empty(StringAssetReferences.Num());
 	while ( StringAssetReferences.Num())
@@ -173,7 +145,6 @@ void FRedirectCollector::ResolveStringAssetReference(FName FilterPackageFName, b
 					continue;
 				}
 			}
-#if WITH_EDITOR
 			if ( (bProcessAlreadyResolvedPackages == false) && 
 				AlreadyRemapped.Contains(ToLoadFName) )
 			{
@@ -181,7 +152,6 @@ void FRedirectCollector::ResolveStringAssetReference(FName FilterPackageFName, b
 				// once is enough for somethings.
 				continue;
 			}
-#endif
 			if (ToLoad.Len() > 0 )
 			{
 				SCOPE_REDIRECT_TIMER(ResolveTimeLoad);
@@ -189,8 +159,6 @@ void FRedirectCollector::ResolveStringAssetReference(FName FilterPackageFName, b
 				UE_LOG(LogRedirectors, Verbose, TEXT("String Asset Reference '%s'"), *ToLoad);
 				UE_CLOG(RefFilenameAndProperty.GetProperty().ToString().Len(), LogRedirectors, Verbose, TEXT("    Referenced by '%s'"), *RefFilenameAndProperty.GetProperty().ToString());
 
-				StringAssetRefFilenameStack.Push(RefFilenameAndProperty.GetPackage().ToString());
-
 				// LoadObject will mark any failed loads as "KnownMissing", so since we want to print out the referencer if it was missing and not-known-missing
 				// then check here before attempting to load (TODO: Can we just not even do the load? Seems like we should be able to...)
 
@@ -204,21 +172,7 @@ void FRedirectCollector::ResolveStringAssetReference(FName FilterPackageFName, b
 				}
 
 				UObject *Loaded = LoadObject<UObject>(NULL, *ToLoad, NULL, RefFilenameAndProperty.GetReferencedByEditorOnlyProperty() ? LOAD_EditorOnly | LOAD_NoWarn : LOAD_NoWarn, NULL);
-				StringAssetRefFilenameStack.Pop();
 
-				UObjectRedirector* Redirector = dynamic_cast<UObjectRedirector*>(Loaded);
-				if (Redirector)
-				{
-					UE_LOG(LogRedirectors, Verbose, TEXT("    Found redir '%s'"), *Redirector->GetFullName());
-					FRedirection Redir;
-					Redir.PackageFilename = RefFilenameAndProperty.GetPackage().ToString();
-					Redir.RedirectorName = Redirector->GetFullName();
-					Redir.RedirectorPackageFilename = Redirector->GetLinker()->Filename;
-					CA_SUPPRESS(28182)
-						Redir.DestinationObjectName = Redirector->DestinationObject->GetFullName();
-					Redirections.AddUnique(Redir);
-					Loaded = Redirector->DestinationObject;
-				}
 				if (Loaded)
 				{
 					if (FCoreUObjectDelegates::PackageLoadedFromStringAssetReference.IsBound())
@@ -229,11 +183,9 @@ void FRedirectCollector::ResolveStringAssetReference(FName FilterPackageFName, b
 					UE_LOG(LogRedirectors, Verbose, TEXT("    Resolved to '%s'"), *Dest);
 					if (Dest != ToLoad)
 					{
-						StringAssetRemap.Add(ToLoad, Dest);
+						AssetPathRedirectionMap.Add(ToLoad, Dest);
 					}
-#if WITH_EDITOR
 					AlreadyRemapped.Add(ToLoadFName);
-#endif
 				}
 				else
 				{
@@ -252,37 +204,60 @@ void FRedirectCollector::ResolveStringAssetReference(FName FilterPackageFName, b
 	check((StringAssetReferences.Num() == 0) || (FilterPackageFName != NAME_None));
 }
 
-void FRedirectCollector::GetStringAssetReferencePackageList(FName FilterPackage, TSet<FName>& ReferencedPackages)
+void FRedirectCollector::ProcessStringAssetReferencePackageList(FName FilterPackage, bool bGetEditorOnly, TSet<FName>& OutReferencedPackages)
 {
-	for (const auto& CurrentReference : StringAssetReferences)
+	FScopeLock ScopeLock(&CriticalSection);
+
+	// Iterate map, remove all matching and potentially add to OutReferencedPackages
+	for (auto It = StringAssetReferences.CreateIterator(); It; ++It)
 	{
-		const FName& ToLoadFName = CurrentReference.Key;
-		const FPackagePropertyPair& RefFilenameAndProperty = CurrentReference.Value;
+		const FName& ToLoadFName = It->Key;
+		const FPackagePropertyPair& RefFilenameAndProperty = It->Value;
 
 		// Package name may be null, if so return the set of things not associated with a package
 		if (FilterPackage == RefFilenameAndProperty.GetCachedPackageName())
 		{
 			FString PackageNameString = FPackageName::ObjectPathToPackageName(ToLoadFName.ToString());
 
-			ReferencedPackages.Add(FName(*PackageNameString));
+			if (!RefFilenameAndProperty.GetReferencedByEditorOnlyProperty() || bGetEditorOnly)
+			{
+				OutReferencedPackages.Add(FName(*PackageNameString));
+			}
+
+			It.RemoveCurrent();
 		}
 	}
+
+	StringAssetReferences.Compact();
 }
 
-void FRedirectCollector::AddStringAssetReferenceRedirection(const FString& OriginalPath, const FString& RedirectedPath)
+void FRedirectCollector::AddAssetPathRedirection(const FString& OriginalPath, const FString& RedirectedPath)
 {
-	FString* Found = StringAssetRemap.Find(OriginalPath);
+	FScopeLock ScopeLock(&CriticalSection);
 
-	if (Found)
-	{
-		ensureMsgf(*Found == RedirectedPath, TEXT("Cannot remap %s to %s, it is already mapped to %s!"), *OriginalPath, *RedirectedPath, **Found);
-	}
-	else
+	// This replaces an existing mapping, can happen in the editor if things are renamed twice
+	AssetPathRedirectionMap.Add(OriginalPath, RedirectedPath);
+}
+
+void FRedirectCollector::RemoveAssetPathRedirection(const FString& OriginalPath)
+{
+	FScopeLock ScopeLock(&CriticalSection);
+
+	FString* Found = AssetPathRedirectionMap.Find(OriginalPath);
+
+	if (ensureMsgf(Found, TEXT("Cannot remove redirection from %s, it was not registered"), *OriginalPath))
 	{
-		StringAssetRemap.Add(OriginalPath, RedirectedPath);
+		AssetPathRedirectionMap.Remove(OriginalPath);
 	}
 }
 
-FRedirectCollector GRedirectCollector;
+FString* FRedirectCollector::GetAssetPathRedirection(const FString& OriginalPath)
+{
+	FScopeLock ScopeLock(&CriticalSection);
 
+	return AssetPathRedirectionMap.Find(OriginalPath);
+}
+
+FRedirectCollector GRedirectCollector;
 
+#endif // WITH_EDITOR
diff --git a/Engine/Source/Runtime/CoreUObject/Private/Misc/StringAssetReference.cpp b/Engine/Source/Runtime/CoreUObject/Private/Misc/StringAssetReference.cpp
--- a/Engine/Source/Runtime/CoreUObject/Private/Misc/StringAssetReference.cpp
+++ b/Engine/Source/Runtime/CoreUObject/Private/Misc/StringAssetReference.cpp
@@ -5,7 +5,9 @@
 #include "UObject/UnrealType.h"
 #include "UObject/ObjectRedirector.h"
 #include "Misc/PackageName.h"
-
+#include "UObject/LinkerLoad.h"
+#include "UObject/UObjectThreadContext.h"
+#include "Misc/RedirectCollector.h"
 
 FStringAssetReference::FStringAssetReference(const UObject* InObject)
 {
@@ -15,7 +17,6 @@ FStringAssetReference::FStringAssetReference(const UObject* InObject)
 	}
 }
 
-PRAGMA_DISABLE_DEPRECATION_WARNINGS
 void FStringAssetReference::SetPath(FString Path)
 {
 	if (Path.IsEmpty())
@@ -39,32 +40,72 @@ void FStringAssetReference::SetPath(FString Path)
 		AssetLongPathname = MoveTemp(Path);
 	}
 }
-PRAGMA_POP
 
-bool FStringAssetReference::Serialize(FArchive& Ar)
+void FStringAssetReference::PreSavePath()
 {
 #if WITH_EDITOR
-	if (Ar.IsSaving() && Ar.IsPersistent() && FCoreUObjectDelegates::StringAssetReferenceSaving.IsBound())
+	FString* FoundRedirection = GRedirectCollector.GetAssetPathRedirection(ToString());
+
+	if (FoundRedirection)
 	{
-		SetPath(FCoreUObjectDelegates::StringAssetReferenceSaving.Execute(ToString()));
+		SetPath(*FoundRedirection);
 	}
 #endif // WITH_EDITOR
+}
+
+void FStringAssetReference::PostLoadPath() const
+{
+#if WITH_EDITOR
+	GRedirectCollector.OnStringAssetReferenceLoaded(ToString());
+#endif // WITH_EDITOR
+}
+
+bool FStringAssetReference::Serialize(FArchive& Ar)
+{
+	// Archivers will call back into SerializePath for the various fixups
 	Ar << *this;
+
+	return true;
+}
+
+void FStringAssetReference::SerializePath(FArchive& Ar, bool bSkipSerializeIfArchiveHasSize)
+{
 #if WITH_EDITOR
-	if (Ar.IsLoading() && Ar.IsPersistent() && FCoreUObjectDelegates::StringAssetReferenceLoaded.IsBound())
+	if (Ar.IsSaving())
 	{
-		FCoreUObjectDelegates::StringAssetReferenceLoaded.Execute(ToString());
+		PreSavePath();
 	}
 #endif // WITH_EDITOR
 
-	if (Ar.IsLoading() && (Ar.GetPortFlags()&PPF_DuplicateForPIE))
+	FString Path = ToString();
+
+	if (!bSkipSerializeIfArchiveHasSize || Ar.IsObjectReferenceCollector() || Ar.Tell() < 0)
 	{
-		// Remap unique ID if necessary
-		// only for fixing up cross-level references, inter-level references handled in FDuplicateDataReader
-		FixupForPIE();
+		Ar << Path;
 	}
 
-	return true;
+	if (Ar.IsLoading())
+	{
+		if (Ar.UE4Ver() < VER_UE4_KEEP_ONLY_PACKAGE_NAMES_IN_STRING_ASSET_REFERENCES_MAP)
+		{
+			Path = FPackageName::GetNormalizedObjectPath(Path);
+		}
+		
+		SetPath(MoveTemp(Path));
+		
+#if WITH_EDITOR
+		if (Ar.IsPersistent())
+		{
+			PostLoadPath();
+		}
+		if (Ar.GetPortFlags()&PPF_DuplicateForPIE)
+		{
+			// Remap unique ID if necessary
+			// only for fixing up cross-level references, inter-level references handled in FDuplicateDataReader
+			FixupForPIE();
+		}
+#endif // WITH_EDITOR
+	}
 }
 
 bool FStringAssetReference::operator==(FStringAssetReference const& Other) const
@@ -87,7 +128,11 @@ bool FStringAssetReference::ExportTextItem(FString& ValueStr, FStringAssetRefere
 
 	if (IsValid())
 	{
-		ValueStr += ToString();
+		// Fixup any redirectors
+		FStringAssetReference Temp = *this;
+		Temp.PreSavePath();
+
+		ValueStr += Temp.ToString();
 	}
 	else
 	{
@@ -95,6 +140,7 @@ bool FStringAssetReference::ExportTextItem(FString& ValueStr, FStringAssetRefere
 	}
 	return true;
 }
+
 bool FStringAssetReference::ImportTextItem(const TCHAR*& Buffer, int32 PortFlags, UObject* Parent, FOutputDevice* ErrorText)
 {
 	FString ImportedPath = TEXT("");
@@ -127,13 +173,8 @@ bool FStringAssetReference::ImportTextItem(const TCHAR*& Buffer, int32 PortFlags
 
 	SetPath(MoveTemp(ImportedPath));
 
-#if WITH_EDITOR
 	// Consider this a load, so Config string asset references get cooked
-	if (FCoreUObjectDelegates::StringAssetReferenceLoaded.IsBound())
-	{
-		FCoreUObjectDelegates::StringAssetReferenceLoaded.Execute(ToString());
-	}
-#endif // WITH_EDITOR
+	PostLoadPath();
 
 	return true;
 }
@@ -155,6 +196,7 @@ bool FStringAssetReference::SerializeFromMismatchedTag(struct FPropertyTag const
 	if (Ar.IsLoading())
 	{
 		SetPath(MoveTemp(Path));
+		PostLoadPath();
 	}
 
 	return bReturn;
@@ -240,5 +282,77 @@ void FStringAssetReference::FixupForPIE()
 	}
 }
 
+bool FStringAssetReferenceThreadContext::GetSerializationOptions(FName& OutPackageName, FName& OutPropertyName, EStringAssetReferenceCollectType& OutCollectType) const
+{
+	FName CurrentPackageName, CurrentPropertyName;
+	EStringAssetReferenceCollectType CurrentCollectType = EStringAssetReferenceCollectType::AlwaysCollect;
+	bool bFoundAnything = false;
+	if (OptionStack.Num() > 0)
+	{
+		// Go from the top of the stack down
+		for (int32 i = OptionStack.Num() - 1; i >= 0; i--)
+		{
+			const FSerializationOptions& Options = OptionStack[i];
+			// Find first valid package/property names. They may not necessarily match
+			if (Options.PackageName != NAME_None && CurrentPackageName == NAME_None)
+			{
+				CurrentPackageName = Options.PackageName;
+			}
+			if (Options.PropertyName != NAME_None && CurrentPropertyName == NAME_None)
+			{
+				CurrentPropertyName = Options.PropertyName;
+			}
+
+			// Restrict based on lowest/most restrictive collect type
+			if (Options.CollectType < CurrentCollectType)
+			{
+				CurrentCollectType = Options.CollectType;
+			}
+		}
+
+		bFoundAnything = true;
+	}
+	
+	// Check UObject thread context as a backup
+	FUObjectThreadContext& ThreadContext = FUObjectThreadContext::Get();
+	if (ThreadContext.SerializedObject)
+	{
+		FLinkerLoad* Linker = ThreadContext.SerializedObject->GetLinker();
+		if (Linker)
+		{
+			if (CurrentPackageName == NAME_None)
+			{
+				CurrentPackageName = FName(*FPackageName::FilenameToLongPackageName(Linker->Filename));
+			}
+			if (Linker->GetSerializedProperty() && CurrentPropertyName == NAME_None)
+			{
+				CurrentPropertyName = Linker->GetSerializedProperty()->GetFName();
+			}
+#if WITH_EDITORONLY_DATA
+			bool bEditorOnly = Linker->IsEditorOnlyPropertyOnTheStack();
+#else
+			bool bEditorOnly = false;
+#endif
+			// If we were always collect before and not overridden by stack options, set to editor only
+			if (bEditorOnly && CurrentCollectType == EStringAssetReferenceCollectType::AlwaysCollect)
+			{
+				CurrentCollectType = EStringAssetReferenceCollectType::EditorOnlyCollect;
+			}
+
+			bFoundAnything = true;
+		}
+	}
+
+	if (bFoundAnything)
+	{
+		OutPackageName = CurrentPackageName;
+		OutPropertyName = CurrentPropertyName;
+		OutCollectType = CurrentCollectType;
+		return true;
+	}
+	
+	return bFoundAnything;
+}
+
 FThreadSafeCounter FStringAssetReference::CurrentTag(1);
 TArray<FString> FStringAssetReference::PackageNamesBeingDuplicatedForPIE;
diff --git a/Engine/Source/Runtime/CoreUObject/Private/Misc/StringClassReference.cpp b/Engine/Source/Runtime/CoreUObject/Private/Misc/StringClassReference.cpp
--- a/Engine/Source/Runtime/CoreUObject/Private/Misc/StringClassReference.cpp
+++ b/Engine/Source/Runtime/CoreUObject/Private/Misc/StringClassReference.cpp
@@ -21,6 +21,7 @@ bool FStringClassReference::SerializeFromMismatchedTag(struct FPropertyTag const
 	if (Ar.IsLoading())
 	{
 		SetPath(MoveTemp(Path));
+		PostLoadPath();
 	}
 
 	return bReturn;
diff --git a/Engine/Source/Runtime/CoreUObject/Private/Misc/StringReferenceTemplates.h b/Engine/Source/Runtime/CoreUObject/Private/Misc/StringReferenceTemplates.h
--- a/Engine/Source/Runtime/CoreUObject/Private/Misc/StringReferenceTemplates.h
+++ b/Engine/Source/Runtime/CoreUObject/Private/Misc/StringReferenceTemplates.h
@@ -30,12 +30,6 @@ bool SerializeFromMismatchedTagTemplate(FString& Output, const FPropertyTag& Tag
 		{
 			Output = FString();
 		}
-#if WITH_EDITOR
-		if (Ar.IsLoading() && Ar.IsPersistent() && FCoreUObjectDelegates::StringAssetReferenceLoaded.IsBound())
-		{
-			FCoreUObjectDelegates::StringAssetReferenceLoaded.Execute(Output);
-		}
-#endif // WITH_EDITOR
 		return true;
 	}
 	else if (Tag.Type == NAME_StrProperty)
@@ -44,12 +38,6 @@ bool SerializeFromMismatchedTagTemplate(FString& Output, const FPropertyTag& Tag
 		Ar << String;
 
 		Output = String;
-#if WITH_EDITOR
-		if (Ar.IsLoading() && Ar.IsPersistent() && FCoreUObjectDelegates::StringAssetReferenceLoaded.IsBound())
-		{
-			FCoreUObjectDelegates::StringAssetReferenceLoaded.Execute(Output);
-		}
-#endif // WITH_EDITOR
 		return true;
 	}
 	return false;
diff --git a/Engine/Source/Runtime/CoreUObject/Private/Serialization/ArchiveUObject.cpp b/Engine/Source/Runtime/CoreUObject/Private/Serialization/ArchiveUObject.cpp
--- a/Engine/Source/Runtime/CoreUObject/Private/Serialization/ArchiveUObject.cpp
+++ b/Engine/Source/Runtime/CoreUObject/Private/Serialization/ArchiveUObject.cpp
@@ -45,8 +45,7 @@ FArchive& FArchiveUObject::operator<<(class FLazyObjectPtr& LazyObjectPtr)
 		else
 #endif
 		{
-			// Downcast from UObjectBase to UObject
-			UObject* Object = static_cast<UObject*>(LazyObjectPtr.Get());
+			UObject* Object = LazyObjectPtr.Get();
 
 			Ar << Object;
 
@@ -59,73 +58,34 @@ FArchive& FArchiveUObject::operator<<(class FLazyObjectPtr& LazyObjectPtr)
 	return Ar;
 }
 
-/**
- * Asset pointer serialization.  Asset pointers only have weak references to objects and
- * won't serialize the object when gathering references for garbage collection.  So in many cases, you
- * don't need to bother serializing asset pointers.  However, serialization is required if you
- * want to load and save your object.
- */
 FArchive& FArchiveUObject::operator<<( class FAssetPtr& AssetPtr )
 {
 	FArchive& Ar = *this;
-	// We never serialize our reference while the garbage collector is harvesting references
-	// to objects, because we don't want weak object pointers to keep objects from being garbage
-	// collected.  That would defeat the whole purpose of a weak object pointer!
-	// However, when modifying both kinds of references we want to serialize and writeback the updated value.
-	// We only want to write the modified value during reference fixup if the data is loaded
-	if( !IsObjectReferenceCollector() || IsModifyingWeakAndStrongReferences() )
+	if (Ar.IsSaving() || Ar.IsLoading())
 	{
-#if WITH_EDITORONLY_DATA
-		// When transacting, just serialize as a string since the object may
-		// not be in memory and you don't want to save a nullptr in this case.
-		if (IsTransacting())
-		{
-			if (Ar.IsLoading())
-			{
-				// Reset before serializing to clear the internal weak pointer. 
-				AssetPtr.Reset();
-			}
-			Ar << AssetPtr.GetUniqueID();
-		}
-		else
-#endif
-		{
-			// Downcast from UObjectBase to UObject
-			UObject* Object = static_cast< UObject* >( AssetPtr.Get() );
+		// Reset before serializing to clear the internal weak pointer. 
+		AssetPtr.ResetWeakPtr();
+		Ar << AssetPtr.GetUniqueID();
+	}
+	else if (!IsObjectReferenceCollector() || IsModifyingWeakAndStrongReferences())
+	{
+		// Treat this like a weak pointer object, as we are doing something like replacing references in memory
+		UObject* Object = AssetPtr.Get();
 
-			Ar << Object;
+		Ar << Object;
 
-			if( IsLoading() || (Object && IsModifyingWeakAndStrongReferences()) )
-			{
-				AssetPtr = Object;
-			}
+		if (IsLoading() || (Object && IsModifyingWeakAndStrongReferences()))
+		{
+			AssetPtr = Object;
 		}
 	}
+
 	return Ar;
 }
 
 FArchive& FArchiveUObject::operator<<(struct FStringAssetReference& Value)
 {
-	FString Path = Value.ToString();
-
-	*this << Path;
-
-	if (IsLoading())
-	{
-		if (UE4Ver() < VER_UE4_KEEP_ONLY_PACKAGE_NAMES_IN_STRING_ASSET_REFERENCES_MAP)
-		{
-			FString NormalizedPath = FPackageName::GetNormalizedObjectPath(Path);
-			if (Value.ToString() != NormalizedPath)
-			{
-				Value.SetPath(NormalizedPath);
-			}
-		}
-		else
-		{
-			Value.SetPath(MoveTemp(Path));
-		}
-	}
-
+	Value.SerializePath(*this);
 	return *this;
 }
 
diff --git a/Engine/Source/Runtime/CoreUObject/Private/Serialization/DuplicateDataReader.cpp b/Engine/Source/Runtime/CoreUObject/Private/Serialization/DuplicateDataReader.cpp
--- a/Engine/Source/Runtime/CoreUObject/Private/Serialization/DuplicateDataReader.cpp
+++ b/Engine/Source/Runtime/CoreUObject/Private/Serialization/DuplicateDataReader.cpp
@@ -89,16 +89,6 @@ FArchive& FDuplicateDataReader::operator<<( FLazyObjectPtr& LazyObjectPtr)
 	return Ar;
 }
 
-FArchive& FDuplicateDataReader::operator<<( FAssetPtr& AssetPtr)
-{
-	FArchive& Ar = *this;
-	FStringAssetReference ID;
-	ID.Serialize(Ar);
-
-	AssetPtr = ID;
-	return Ar;
-}
-
 FArchive& FDuplicateDataReader::operator<<( FStringAssetReference& StringAssetReference)
 {
 	FArchiveUObject::operator<<(StringAssetReference);
diff --git a/Engine/Source/Runtime/CoreUObject/Private/Serialization/DuplicateDataWriter.cpp b/Engine/Source/Runtime/CoreUObject/Private/Serialization/DuplicateDataWriter.cpp
--- a/Engine/Source/Runtime/CoreUObject/Private/Serialization/DuplicateDataWriter.cpp
+++ b/Engine/Source/Runtime/CoreUObject/Private/Serialization/DuplicateDataWriter.cpp
@@ -90,13 +90,6 @@ FArchive& FDuplicateDataWriter::operator<<(FLazyObjectPtr& LazyObjectPtr)
 	return *this << ID;
 }
 
-FArchive& FDuplicateDataWriter::operator<<(FAssetPtr& AssetPtr)
-{
-	FStringAssetReference ID = AssetPtr.GetUniqueID();
-	ID.Serialize(*this);
-	return *this;
-}
-
 void FDuplicateDataWriter::AddDuplicate(UObject* SourceObject, UObject* DupObject)
 {
 	if ( DupObject && !DupObject->IsTemplate() )
diff --git a/Engine/Source/Runtime/CoreUObject/Private/Serialization/ObjectReader.cpp b/Engine/Source/Runtime/CoreUObject/Private/Serialization/ObjectReader.cpp
--- a/Engine/Source/Runtime/CoreUObject/Private/Serialization/ObjectReader.cpp
+++ b/Engine/Source/Runtime/CoreUObject/Private/Serialization/ObjectReader.cpp
@@ -39,25 +39,13 @@ FArchive& FObjectReader::operator<<( class FLazyObjectPtr& LazyObjectPtr )
 
 FArchive& FObjectReader::operator<<( class FAssetPtr& AssetPtr )
 {
-	FArchive& Ar = *this;
-	FStringAssetReference ID;
-	ID.Serialize(Ar);
-
-	AssetPtr = ID;
-	return Ar;
+	AssetPtr.ResetWeakPtr();
+	return *this << AssetPtr.GetUniqueID();
 }
 
 FArchive& FObjectReader::operator<<(FStringAssetReference& Value)
 {
-	FString Path = Value.ToString();
-
-	*this << Path;
-
-	if (IsLoading())
-	{
-		Value.SetPath(MoveTemp(Path));
-	}
-
+	Value.SerializePath(*this);
 	return *this;
 }
 
diff --git a/Engine/Source/Runtime/CoreUObject/Private/Serialization/ObjectWriter.cpp b/Engine/Source/Runtime/CoreUObject/Private/Serialization/ObjectWriter.cpp
--- a/Engine/Source/Runtime/CoreUObject/Private/Serialization/ObjectWriter.cpp
+++ b/Engine/Source/Runtime/CoreUObject/Private/Serialization/ObjectWriter.cpp
@@ -33,22 +33,13 @@ FArchive& FObjectWriter::operator<<( class FLazyObjectPtr& LazyObjectPtr )
 
 FArchive& FObjectWriter::operator<<( class FAssetPtr& AssetPtr )
 {
-	FStringAssetReference ID = AssetPtr.GetUniqueID();
-	ID.Serialize(*this);
-	return *this;
+	AssetPtr.ResetWeakPtr();
+	return *this << AssetPtr.GetUniqueID();
 }
 
 FArchive& FObjectWriter::operator<<(FStringAssetReference& Value)
 {
-	FString Path = Value.ToString();
-
-	*this << Path;
-
-	if (IsLoading())
-	{
-		Value.SetPath(MoveTemp(Path));
-	}
-
+	Value.SerializePath(*this);
 	return *this;
 }
 
diff --git a/Engine/Source/Runtime/CoreUObject/Private/UObject/Class.cpp b/Engine/Source/Runtime/CoreUObject/Private/UObject/Class.cpp
--- a/Engine/Source/Runtime/CoreUObject/Private/UObject/Class.cpp
+++ b/Engine/Source/Runtime/CoreUObject/Private/UObject/Class.cpp
@@ -50,6 +50,10 @@ DEFINE_LOG_CATEGORY(LogClass);
 	#endif
 #endif
 
+// If we end up pushing class flags out beyond a uint32, there are various places
+// casting it to uint32 that need to be fixed up (mostly printfs but also some serialization code)
+static_assert(sizeof(__underlying_type(EClassFlags)) == sizeof(uint32), "expecting ClassFlags enum to fit in a uint32");
+
 //////////////////////////////////////////////////////////////////////////
 
 FThreadSafeBool& InternalSafeGetTokenStreamDirtyFlag()
@@ -2559,7 +2563,7 @@ class FRestoreClassInfo: public FRestoreForUObjectOverwrite
 	/** Saved ClassDefaultObject **/
 	UObject*		DefaultObject;
 	/** Saved ClassFlags **/
-	uint32			Flags;
+	EClassFlags		Flags;
 	/** Saved ClassCastFlags **/
 	EClassCastFlags	CastFlags;
 	/** Saved ClassConstructor **/
@@ -3410,18 +3414,18 @@ void UClass::Serialize( FArchive& Ar )
 	// Class flags first.
 	if (Ar.IsSaving())
 	{
-		auto SavedClassFlags = ClassFlags;
+		uint32 SavedClassFlags = ClassFlags;
 		SavedClassFlags &= ~(CLASS_ShouldNeverBeLoaded | CLASS_TokenStreamAssembled);
 		Ar << SavedClassFlags;
 	}
 	else if (Ar.IsLoading())
 	{
-		Ar << ClassFlags;
+		Ar << (uint32&)ClassFlags;
 		ClassFlags &= ~(CLASS_ShouldNeverBeLoaded | CLASS_TokenStreamAssembled);
 	}
 	else 
 	{
-		Ar << ClassFlags;
+		Ar << (uint32&)ClassFlags;
 	}
 	if (Ar.UE4Ver() < VER_UE4_CLASS_NOTPLACEABLE_ADDED)
 	{
@@ -3696,7 +3700,7 @@ void UClass::PurgeClass(bool bRecompilingOnLoad)
 {
 	ClassConstructor = nullptr;
 	ClassVTableHelperCtorCaller = nullptr;
-	ClassFlags = 0;
+	ClassFlags = CLASS_None;
 	ClassCastFlags = 0;
 	ClassUnique = 0;
 	ClassReps.Empty();
@@ -3794,7 +3798,7 @@ bool UClass::HasProperty(UProperty* InProperty) const
 UClass::UClass(const FObjectInitializer& ObjectInitializer)
 :	UStruct( ObjectInitializer )
 ,	ClassUnique(0)
-,	ClassFlags(0)
+,	ClassFlags(CLASS_None)
 ,	ClassCastFlags(0)
 ,	ClassWithin( UObject::StaticClass() )
 ,	ClassGeneratedBy(nullptr)
@@ -3811,7 +3815,7 @@ UClass::UClass(const FObjectInitializer& ObjectInitializer)
 UClass::UClass(const FObjectInitializer& ObjectInitializer, UClass* InBaseClass)
 :	UStruct(ObjectInitializer, InBaseClass)
 ,	ClassUnique(0)
-,	ClassFlags(0)
+,	ClassFlags(CLASS_None)
 ,	ClassCastFlags(0)
 ,	ClassWithin(UObject::StaticClass())
 ,	ClassGeneratedBy(nullptr)
@@ -3849,7 +3853,7 @@ UClass::UClass
 	EStaticConstructor,
 	FName			InName,
 	uint32			InSize,
-	uint32			InClassFlags,
+	EClassFlags		InClassFlags,
 	EClassCastFlags	InClassCastFlags,
 	const TCHAR*    InConfigName,
 	EObjectFlags	InFlags,
@@ -3885,7 +3889,7 @@ UClass::UClass
 
 bool UClass::HotReloadPrivateStaticClass(
 	uint32			InSize,
-	uint32			InClassFlags,
+	EClassFlags		InClassFlags,
 	EClassCastFlags	InClassCastFlags,
 	const TCHAR*    InConfigName,
 	ClassConstructorType InClassConstructor,
@@ -4163,7 +4167,7 @@ const FString UClass::GetConfigName() const
 	}
 	else if( ClassConfigName == NAME_None )
 	{
-		UE_LOG(LogClass, Fatal,TEXT("UObject::GetConfigName() called on class with config name 'None'. Class flags = %d"), ClassFlags );
+		UE_LOG(LogClass, Fatal,TEXT("UObject::GetConfigName() called on class with config name 'None'. Class flags = 0x%08X"), (uint32)ClassFlags );
 		return TEXT("");
 	}
 	else
@@ -4288,7 +4292,7 @@ void GetPrivateStaticClassBody(
 	UClass*& ReturnClass,
 	void(*RegisterNativeFunc)(),
 	uint32 InSize,
-	uint32 InClassFlags,
+	EClassFlags InClassFlags,
 	EClassCastFlags InClassCastFlags,
 	const TCHAR* InConfigName,
 	UClass::ClassConstructorType InClassConstructor,
@@ -4869,7 +4873,7 @@ UDynamicClass::UDynamicClass(
 	EStaticConstructor,
 	FName			InName,
 	uint32			InSize,
-	uint32			InClassFlags,
+	EClassFlags		InClassFlags,
 	EClassCastFlags	InClassCastFlags,
 	const TCHAR*    InConfigName,
 	EObjectFlags	InFlags,
diff --git a/Engine/Source/Runtime/CoreUObject/Private/UObject/CoreNative.cpp b/Engine/Source/Runtime/CoreUObject/Private/UObject/CoreNative.cpp
--- a/Engine/Source/Runtime/CoreUObject/Private/UObject/CoreNative.cpp
+++ b/Engine/Source/Runtime/CoreUObject/Private/UObject/CoreNative.cpp
@@ -45,7 +45,7 @@ IMPLEMENT_CLASS(UObject, 0);
 COREUOBJECT_API class UClass* Z_Construct_UClass_UObject();
 UClass* Z_Construct_UClass_UObject()
 {
-	static UClass* OuterClass = NULL;
+	static UClass* OuterClass = nullptr;
 	if (!OuterClass)
 	{
 		OuterClass = UObject::StaticClass();
@@ -63,8 +63,8 @@ UClass* Z_Construct_UClass_UObject()
 -----------------------------------------------------------------------------*/
 
 FObjectInstancingGraph::FObjectInstancingGraph(bool bDisableInstancing)
-	: SourceRoot(NULL)
-	, DestinationRoot(NULL)
+	: SourceRoot(nullptr)
+	, DestinationRoot(nullptr)
 	, bCreatingArchetype(false)
 	, bEnableSubobjectInstancing(!bDisableInstancing)
 	, bLoadingObject(false)
@@ -72,8 +72,8 @@ FObjectInstancingGraph::FObjectInstancingGraph(bool bDisableInstancing)
 }
 
 FObjectInstancingGraph::FObjectInstancingGraph( UObject* DestinationSubobjectRoot )
-	: SourceRoot(NULL)
-	, DestinationRoot(NULL)
+	: SourceRoot(nullptr)
+	, DestinationRoot(nullptr)
 	, bCreatingArchetype(false)
 	, bEnableSubobjectInstancing(true)
 	, bLoadingObject(false)
@@ -107,7 +107,7 @@ UObject* FObjectInstancingGraph::GetInstancedSubobject( UObject* SourceSubobject
 
 	UObject* InstancedSubobject = INVALID_OBJECT;
 
-	if ( SourceSubobject != NULL && CurrentValue != NULL )
+	if ( SourceSubobject != nullptr && CurrentValue != nullptr )
 	{
 		bool bAllowedSelfReference = bAllowSelfReference && SourceSubobject == SourceRoot;
 
@@ -127,7 +127,7 @@ UObject* FObjectInstancingGraph::GetInstancedSubobject( UObject* SourceSubobject
 		{
 			// search for the unique component instance that corresponds to this component template
 			InstancedSubobject = GetDestinationObject(SourceSubobject);
-			if ( InstancedSubobject == NULL )
+			if ( InstancedSubobject == nullptr )
 			{
 				if (bDoNotCreateNewInstance)
 				{
@@ -177,20 +177,20 @@ UObject* FObjectInstancingGraph::GetInstancedSubobject( UObject* SourceSubobject
 							// final archetype archetype will be the archetype of the template
 							UObject* FinalSubobjectArchetype = CurrentValue->GetArchetype();
 
-							// Don't seach for the existing subobjects on Blueprint-generated classes. What we'll find is a subobject
+							// Don't search for the existing subobjects on Blueprint-generated classes. What we'll find is a subobject
 							// created by the constructor which may not have all of its fields initialized to the correct value (which
 							// should be coming from a blueprint).
 							// NOTE: Since this function is called ONLY for Blueprint-generated classes, we may as well delete this 'if'.
 							if (!SubobjectOuter->GetClass()->HasAnyClassFlags(CLASS_CompiledFromBlueprint))
 							{
-								InstancedSubobject = StaticFindObjectFast(NULL, SubobjectOuter, SubobjectName);
+								InstancedSubobject = StaticFindObjectFast(nullptr, SubobjectOuter, SubobjectName);
 							}
 
 							if (InstancedSubobject && IsCreatingArchetype())
 							{
 								// since we are updating an archetype, this needs to reconstruct as that is the mechanism used to copy properties
 								// it will destroy the existing object and overwrite it
-								InstancedSubobject = NULL;
+								InstancedSubobject = nullptr;
 							}
 
 							if (!InstancedSubobject)
@@ -253,16 +253,16 @@ UObject* FObjectInstancingGraph::InstancePropertyValue( class UObject* Component
 		return NewValue; // not instancing
 	}
 
-	// if the object we're instancing the components for (Owner) has the current component's outer in its archetype chain, and its archetype has a NULL value
+	// if the object we're instancing the components for (Owner) has the current component's outer in its archetype chain, and its archetype has a nullptr value
 	// for this component property it means that the archetype didn't instance its component, so we shouldn't either.
 
-	if (ComponentTemplate == NULL && CurrentValue != NULL && (Owner && Owner->IsBasedOnArchetype(CurrentValue->GetOuter())))
+	if (ComponentTemplate == nullptr && CurrentValue != nullptr && (Owner && Owner->IsBasedOnArchetype(CurrentValue->GetOuter())))
 	{
-		NewValue = NULL;
+		NewValue = nullptr;
 	}
 	else
 	{
-		if ( ComponentTemplate == NULL )
+		if ( ComponentTemplate == nullptr )
 		{
 			// should only be here if our archetype doesn't contain this component property
 			ComponentTemplate = CurrentValue;
@@ -307,7 +307,7 @@ void FObjectInstancingGraph::AddNewInstance(UObject* ObjectInstance, UObject* In
 
 void FObjectInstancingGraph::RetrieveObjectInstances( UObject* SearchOuter, TArray<UObject*>& out_Objects )
 {
-	if ( HasDestinationRoot() && SearchOuter != NULL && (SearchOuter == DestinationRoot || SearchOuter->IsIn(DestinationRoot)) )
+	if ( HasDestinationRoot() && SearchOuter != nullptr && (SearchOuter == DestinationRoot || SearchOuter->IsIn(DestinationRoot)) )
 	{
 		for ( TMap<UObject*,UObject*>::TIterator It(SourceToDestinationMap); It; ++It )
 		{
diff --git a/Engine/Source/Runtime/CoreUObject/Private/UObject/LinkerLoad.cpp b/Engine/Source/Runtime/CoreUObject/Private/UObject/LinkerLoad.cpp
--- a/Engine/Source/Runtime/CoreUObject/Private/UObject/LinkerLoad.cpp
+++ b/Engine/Source/Runtime/CoreUObject/Private/UObject/LinkerLoad.cpp
@@ -511,16 +511,20 @@ void FLinkerLoad::PRIVATE_PatchNewObjectIntoExport(UObject* OldObject, UObject*
 		FObjectExport& ObjExport = OldObjectLinker->ExportMap[CachedLinkerIndex];
 
 		// Detach the old object to make room for the new
+		const EObjectFlags OldObjectFlags = OldObject->GetFlags();
 		OldObject->ClearFlags(RF_NeedLoad|RF_NeedPostLoad);
-		OldObject->SetLinker(NULL, INDEX_NONE, true);
+		OldObject->SetLinker(nullptr, INDEX_NONE, true);
+
+		// Copy flags from the old CDO.
+		NewObject->SetFlags(OldObjectFlags);
 
 		// Move the new object into the old object's slot, so any references to this object will now reference the new
 		NewObject->SetLinker(OldObjectLinker, CachedLinkerIndex);
 		ObjExport.Object = NewObject;
 
-		auto& ObjLoaded = FUObjectThreadContext::Get().ObjLoaded;
+		TArray<UObject*>& ObjLoaded = FUObjectThreadContext::Get().ObjLoaded;
 		// If the object was in the ObjLoaded queue (exported, but not yet serialized), swap out for our new object
-		int32 ObjLoadedIdx = ObjLoaded.Find(OldObject);
+		const int32 ObjLoadedIdx = ObjLoaded.Find(OldObject);
 		if( ObjLoadedIdx != INDEX_NONE )
 		{
 			ObjLoaded[ObjLoadedIdx] = NewObject;
@@ -2196,9 +2200,6 @@ FLinkerLoad::EVerifyResult FLinkerLoad::VerifyImport(int32 ImportIndex)
 				{
 					Result = VERIFY_Redirected;
 
-					// send a callback saying we followed a redirector successfully
-					FCoreUObjectDelegates::RedirectorFollowed.Broadcast(Filename, Redir);
-
 					// now, fake our Import to be what the redirector pointed to
 					Import.XObject = Redir->DestinationObject;
 					FUObjectThreadContext::Get().ImportCount++;
@@ -3002,8 +3003,6 @@ UObject* FLinkerLoad::Create( UClass* ObjectClass, FName ObjectName, UObject* Ou
 			// if we found what it was point to, then return it
 			if (Redir->DestinationObject && Redir->DestinationObject->IsA(ObjectClass))
 			{
-				// send a callback saying we followed a redirector successfully
-				FCoreUObjectDelegates::RedirectorFollowed.Broadcast(Filename, Redir);
 				// and return the object we are being redirected to
 				return Redir->DestinationObject;
 			}
diff --git a/Engine/Source/Runtime/CoreUObject/Private/UObject/Obj.cpp b/Engine/Source/Runtime/CoreUObject/Private/UObject/Obj.cpp
--- a/Engine/Source/Runtime/CoreUObject/Private/UObject/Obj.cpp
+++ b/Engine/Source/Runtime/CoreUObject/Private/UObject/Obj.cpp
@@ -3303,7 +3303,7 @@ bool StaticExec( UWorld* InWorld, const TCHAR* Cmd, FOutputDevice& Ar )
 
 			if (Class != NULL)
 			{
-				Ar.Logf(TEXT("Listing functions introduced in class %s (class flags = %d)"), ClassName, Class->GetClassFlags());
+				Ar.Logf(TEXT("Listing functions introduced in class %s (class flags = 0x%08X)"), ClassName, (uint32)Class->GetClassFlags());
 				for (TFieldIterator<UFunction> It(Class); It; ++It)
 				{
 					UFunction* Function = *It;
@@ -4031,17 +4031,11 @@ void InitUObject()
 	
 
 
-#if WITH_EDITORONLY_DATA
-	const FString CommandLine = FCommandLine::Get();
-
-	// this is a hack to give fixup redirects insight into the startup packages
-	if (CommandLine.Contains(TEXT("fixupredirects")) )
-	{
-		FCoreUObjectDelegates::RedirectorFollowed.AddRaw(&GRedirectCollector, &FRedirectCollector::OnRedirectorFollowed);
-	}
-
+#if WITH_EDITOR
+PRAGMA_DISABLE_DEPRECATION_WARNINGS
 	FCoreUObjectDelegates::StringAssetReferenceLoaded.BindRaw(&GRedirectCollector, &FRedirectCollector::OnStringAssetReferenceLoaded);
 	FCoreUObjectDelegates::StringAssetReferenceSaving.BindRaw(&GRedirectCollector, &FRedirectCollector::OnStringAssetReferenceSaved);
+PRAGMA_ENABLE_DEPRECATION_WARNINGS
 #endif
 
 	// Object initialization.
diff --git a/Engine/Source/Runtime/CoreUObject/Private/UObject/ObjectRedirector.cpp b/Engine/Source/Runtime/CoreUObject/Private/UObject/ObjectRedirector.cpp
--- a/Engine/Source/Runtime/CoreUObject/Private/UObject/ObjectRedirector.cpp
+++ b/Engine/Source/Runtime/CoreUObject/Private/UObject/ObjectRedirector.cpp
@@ -91,48 +91,3 @@ IMPLEMENT_CORE_INTRINSIC_CLASS(UObjectRedirector, UObject,
 		Class->EmitObjectReference(STRUCT_OFFSET(UObjectRedirector, DestinationObject), TEXT("DestinationObject"));
 	}
 );
-
-/**
- * Constructor for the callback device, will register itself with the RedirectorFollowed delegate
- *
- * @param InObjectPathNameToMatch Full pathname of the object refrence that is being compiled by script
- */
-FScopedRedirectorCatcher::FScopedRedirectorCatcher(const FString& InObjectPathNameToMatch)
-: ObjectPathNameToMatch(InObjectPathNameToMatch)
-, bWasRedirectorFollowed(false)
-{
-	// register itself on construction to see if the object is a redirector 
-	FCoreUObjectDelegates::RedirectorFollowed.AddRaw(this, &FScopedRedirectorCatcher::OnRedirectorFollowed);
-}
-
-/**
- * Destructor. Will unregister the callback
- */
-FScopedRedirectorCatcher::~FScopedRedirectorCatcher()
-{
-	// register itself on construction to see if the object is a redirector 
-	FCoreUObjectDelegates::RedirectorFollowed.RemoveAll(this);
-}
-
-
-/**
- * Responds to FCoreDelegates::RedirectorFollowed. Records all followed redirections
- * so they can be cleaned later.
- *
- * @param InType Callback type (should only be FCoreDelegates::RedirectorFollowed)
- * @param InString Name of the package that pointed to the redirect
- * @param InObject The UObjectRedirector that was followed
- */
-void FScopedRedirectorCatcher::OnRedirectorFollowed( const FString& InString, UObject* InObject)
-{
-	// this needs to be the redirector
-	check(dynamic_cast<UObjectRedirector*>(InObject));
-
-	// if the path of the redirector was the same as the path to the object constant
-	// being compiled, then the script code has a text reference to a redirector, 
-	// which will cause FixupRedirects to break the reference
-	if ( InObject->GetPathName() == ObjectPathNameToMatch )
-	{
-		bWasRedirectorFollowed = true;
-	}
-}
diff --git a/Engine/Source/Runtime/CoreUObject/Private/UObject/SavePackage.cpp b/Engine/Source/Runtime/CoreUObject/Private/UObject/SavePackage.cpp
--- a/Engine/Source/Runtime/CoreUObject/Private/UObject/SavePackage.cpp
+++ b/Engine/Source/Runtime/CoreUObject/Private/UObject/SavePackage.cpp
@@ -988,6 +988,12 @@ FArchive& FArchiveSaveTagExports::operator<<(UObject*& Obj)
 				break;
 			}
 		}
+
+		if(!bNeedsLoadForEditorGame && Obj->HasAnyFlags(RF_ClassDefaultObject))
+		{
+			bNeedsLoadForEditorGame = Obj->GetClass()->NeedsLoadForEditorGame();
+		}
+
 		if (!bNeedsLoadForEditorGame)
 		{
 			Obj->Mark(OBJECTMARK_NotAlwaysLoadedForEditorGame);
@@ -1102,7 +1108,6 @@ public:
 	virtual FArchive& operator<<(UObject*& Obj) override;
 	virtual FArchive& operator<< (struct FWeakObjectPtr& Value) override;
 	virtual FArchive& operator<<(FLazyObjectPtr& LazyObjectPtr) override;
-	virtual FArchive& operator<<(FAssetPtr& AssetPtr) override;
 	virtual FArchive& operator<<(FStringAssetReference& Value) override;
 	virtual FArchive& operator<<(FName& Name) override;
 	
@@ -1243,33 +1248,12 @@ FArchive& FArchiveSaveTagImports::operator<<( FLazyObjectPtr& LazyObjectPtr)
 	return *this << ID;
 }
 
-FArchive& FArchiveSaveTagImports::operator<<( FAssetPtr& AssetPtr)
-{
-	FStringAssetReference ID;
-	UObject *Object = AssetPtr.Get();
-
-	if (Object)
-	{
-		// Use object in case name has changed. 
-		ID = FStringAssetReference(Object);
-	}
-	else
-	{
-		ID = AssetPtr.GetUniqueID();
-	}
-	return *this << ID;
-}
-
 FArchive& FArchiveSaveTagImports::operator<<(FStringAssetReference& Value)
 {
 	if (Value.IsValid())
 	{
+		Value.PreSavePath();
 		FString Path = Value.ToString();
-		if (FCoreUObjectDelegates::StringAssetReferenceSaving.IsBound())
-		{
-			// This picks up any redirectors
-			Path = FCoreUObjectDelegates::StringAssetReferenceSaving.Execute(Path);
-		}
 
 		if (GetIniFilenameFromObjectsReference(Path) != nullptr)
 		{
diff --git a/Engine/Source/Runtime/CoreUObject/Private/UObject/UObjectGlobals.cpp b/Engine/Source/Runtime/CoreUObject/Private/UObject/UObjectGlobals.cpp
--- a/Engine/Source/Runtime/CoreUObject/Private/UObject/UObjectGlobals.cpp
+++ b/Engine/Source/Runtime/CoreUObject/Private/UObject/UObjectGlobals.cpp
@@ -113,7 +113,7 @@ FCoreUObjectDelegates::FOnAssetLoaded FCoreUObjectDelegates::OnAssetLoaded;
 FCoreUObjectDelegates::FOnObjectSaved FCoreUObjectDelegates::OnObjectSaved;
 #endif // WITH_EDITOR
 
-FCoreUObjectDelegates::FOnRedirectorFollowed FCoreUObjectDelegates::RedirectorFollowed;
+
 
 FSimpleMulticastDelegate FCoreUObjectDelegates::PreGarbageCollect;
 FSimpleMulticastDelegate FCoreUObjectDelegates::PostGarbageCollect;
@@ -125,12 +125,13 @@ FCoreUObjectDelegates::FPreLoadMapDelegate FCoreUObjectDelegates::PreLoadMap;
 FCoreUObjectDelegates::FPostLoadMapDelegate FCoreUObjectDelegates::PostLoadMapWithWorld;
 PRAGMA_DISABLE_DEPRECATION_WARNINGS
 FSimpleMulticastDelegate FCoreUObjectDelegates::PostLoadMap;
+FCoreUObjectDelegates::FStringAssetReferenceLoaded FCoreUObjectDelegates::StringAssetReferenceLoaded;
+FCoreUObjectDelegates::FStringAssetReferenceSaving FCoreUObjectDelegates::StringAssetReferenceSaving;
+FCoreUObjectDelegates::FOnRedirectorFollowed FCoreUObjectDelegates::RedirectorFollowed;
 PRAGMA_ENABLE_DEPRECATION_WARNINGS
 FSimpleMulticastDelegate FCoreUObjectDelegates::PostDemoPlay;
 FCoreUObjectDelegates::FOnLoadObjectsOnTop FCoreUObjectDelegates::ShouldLoadOnTop;
 
-FCoreUObjectDelegates::FStringAssetReferenceLoaded FCoreUObjectDelegates::StringAssetReferenceLoaded;
-FCoreUObjectDelegates::FStringAssetReferenceSaving FCoreUObjectDelegates::StringAssetReferenceSaving;
 FCoreUObjectDelegates::FPackageCreatedForLoad FCoreUObjectDelegates::PackageCreatedForLoad;
 FCoreUObjectDelegates::FPackageLoadedFromStringAssetReference FCoreUObjectDelegates::PackageLoadedFromStringAssetReference;
 FCoreUObjectDelegates::FGetPrimaryAssetIdForObject FCoreUObjectDelegates::GetPrimaryAssetIdForObject;
diff --git a/Engine/Source/Runtime/CoreUObject/Private/UObject/WeakObjectPtr.cpp b/Engine/Source/Runtime/CoreUObject/Private/UObject/WeakObjectPtr.cpp
--- a/Engine/Source/Runtime/CoreUObject/Private/UObject/WeakObjectPtr.cpp
+++ b/Engine/Source/Runtime/CoreUObject/Private/UObject/WeakObjectPtr.cpp
@@ -107,8 +107,7 @@ void FWeakObjectPtr::Serialize(FArchive& Ar)
 	// However, when modifying both kinds of references we want to serialize and writeback the updated value.
 	if (!Ar.IsObjectReferenceCollector() || Ar.IsModifyingWeakAndStrongReferences())
 	{
-		// Downcast from UObjectBase to UObject
-		UObject* Object = static_cast<UObject*>(Get(true));
+		UObject* Object = Get(true);
 
 		Ar << Object;
 
diff --git a/Engine/Source/Runtime/CoreUObject/Public/Misc/PackageName.h b/Engine/Source/Runtime/CoreUObject/Public/Misc/PackageName.h
--- a/Engine/Source/Runtime/CoreUObject/Public/Misc/PackageName.h
+++ b/Engine/Source/Runtime/CoreUObject/Public/Misc/PackageName.h
@@ -109,7 +109,17 @@ public:
 	 * @param OutReason					When returning false, this will provide a description of what was wrong with the name.
 	 * @return							true if a valid long package name
 	 */
-	static bool IsValidLongPackageName(const FString& InLongPackageName, bool bIncludeReadOnlyRoots = false, FText* OutReason = NULL);
+	static bool IsValidLongPackageName(const FString& InLongPackageName, bool bIncludeReadOnlyRoots = false, FText* OutReason = nullptr);
+	/** 
+	 * Returns true if the path starts with a valid root (i.e. /Game/, /Engine/, etc) and contains no illegal characters.
+	 * This validates that the packagename is valid, and also makes sure the object after package name is also correct.
+	 * This will return false if passed a path starting with Classname'
+	 *
+	 * @param InObjectPath				The object path to test
+	 * @param OutReason					When returning false, this will provide a description of what was wrong with the name.
+	 * @return							true if a valid object path
+	 */
+	static bool IsValidObjectPath(const FString& InObjectPath, FText* OutReason = nullptr);
 	/**
 	 * Checks if the given string is a long package name or not.
 	 *
diff --git a/Engine/Source/Runtime/CoreUObject/Public/Misc/RedirectCollector.h b/Engine/Source/Runtime/CoreUObject/Public/Misc/RedirectCollector.h
--- a/Engine/Source/Runtime/CoreUObject/Public/Misc/RedirectCollector.h
+++ b/Engine/Source/Runtime/CoreUObject/Public/Misc/RedirectCollector.h
@@ -1,46 +1,16 @@
 // Copyright 1998-2017 Epic Games, Inc. All Rights Reserved.
 
 /*=============================================================================
-	CallbackDevice.h:  Allows the engine to do callbacks into the editor
+	RedirectCollector:  Editor-only global object that handles resolving redirectors and handling string asset cooking rules
 =============================================================================*/
 
 #pragma once
 
 #include "CoreMinimal.h"
+#include "Misc/ScopeLock.h"
 
-/**
- * Container struct that holds info about a redirector that was followed. 
- * Cannot contain a pointer to the actual UObjectRedirector as the object
- * will be GC'd before it is used
- */
-struct FRedirection
-{
-	FString PackageFilename;
-	FString RedirectorName;
-	FString RedirectorPackageFilename;
-	FString DestinationObjectName;
-
-	bool operator==(const FRedirection& Other) const
-	{
-		return PackageFilename == Other.PackageFilename &&
-			RedirectorName == Other.RedirectorName &&
-			RedirectorPackageFilename == Other.RedirectorPackageFilename;
-	}
-
-	friend FArchive& operator<<(FArchive& Ar, FRedirection& Redir)
-	{
-		Ar << Redir.PackageFilename;
-		Ar << Redir.RedirectorName;
-		Ar << Redir.RedirectorPackageFilename;
-		Ar << Redir.DestinationObjectName;
-		return Ar;
-	}
-};
+#if WITH_EDITOR
 
-/**
- * Callback structure to respond to redirect-followed callbacks and store
- * the information
- */
 class COREUOBJECT_API FRedirectCollector
 {
 private:
@@ -112,27 +82,11 @@ private:
 public:
 
 	/**
-	 * Responds to FCoreDelegates::RedirectorFollowed. Records all followed redirections
-	 * so they can be cleaned later.
-	 *
-	 * @param InString Name of the package that pointed to the redirect
-	 * @param InObject The UObjectRedirector that was followed
-	 */
-	void OnRedirectorFollowed( const FString& InString, UObject* InObject);
-	
-	/**
-	 * Responds to FCoreDelegates::StringAssetReferenceLoaded. Just saves it off for later use.
+	 * Called from FStringAssetReference::PostLoadPath, registers this for later querying
 	 * @param InString Name of the asset that was loaded
 	 */
 	void OnStringAssetReferenceLoaded(const FString& InString);
 
-	/**
-	 * Responds to FCoreDelegates::StringAssetReferenceSaved. Runs the string through the remap table.
-	 * @param InString Name of the asset that will be saved
-	 * @return actual value to save
-	 */
-	FString OnStringAssetReferenceSaved(const FString& InString);
-
 	/**
 	 * Load the string asset references to resolve them, add that to the remap table, and empty the array
 	 * @param FilterPackage If set, only fixup references that were created by FilterPackage. If empty, clear all of them
@@ -142,14 +96,21 @@ public:
 	void ResolveStringAssetReference(FName FilterPackage = NAME_None, bool bProcessAlreadyResolvedPackages = true);
 
 	/**
-	 * Returns the list of packages that would be loaded by ResolveStringAssetReference 
-	 * @param FilterPackage Return references made by loading this package. If passed null will return all references made before package loading
-	 * @param ReferencedPackages Return list of packages referenced by FilterPackage
+	 * Returns the list of packages referenced by string asset references loaded by FilterPackage, and remove them from the internal list
+	 * @param FilterPackage Return references made by loading this package. If passed null will return all references made with no explicit package
+	 * @param OutReferencedPackages Return list of packages referenced by FilterPackage
+	 * @param bGetEditorOnly If true will return references loaded by editor only objects, if false it will not
 	 */
-	void GetStringAssetReferencePackageList(FName FilterPackage, TSet<FName>& ReferencedPackages);
+	void ProcessStringAssetReferencePackageList(FName FilterPackage, bool bGetEditorOnly, TSet<FName>& OutReferencedPackages);
+
+	/** Adds a new mapping for redirector path to destination path, this is called from the Asset Registry to register all redirects it knows about */
+	void AddAssetPathRedirection(const FString& OriginalPath, const FString& RedirectedPath);
 
-	/** Adds a new mapping for redirector path to destination path, this is useful if you have external data about redirection */
-	void AddStringAssetReferenceRedirection(const FString& OriginalPath, const FString& RedirectedPath);
+	/** Removes an asset path redirection, call this when deleting redirectors */
+	void RemoveAssetPathRedirection(const FString& OriginalPath);
+
+	/** Returns a remapped asset path, if it returns null there is no relevant redirector */
+	FString* GetAssetPathRedirection(const FString& OriginalPath);
 
 	/**
 	 * Do we have any references to resolve.
@@ -160,19 +121,11 @@ public:
 		return StringAssetReferences.Num() > 0;
 	}
 
-	const TArray<FRedirection>& GetRedirections() const
-	{
-		return Redirections;
-	}
-
+	/** Writes out log info for how long string asset references took to resolve */
 	void LogTimers() const;
 
-public:
-	/** If not an empty string, only fixup redirects in this package */
-	FString FileToFixup;
-
-	/** A gathered list of all non-script referenced redirections */
-	TArray<FRedirection> Redirections;
+	DEPRECATED(4.17, "OnStringAssetReferenceSaved is deprecated, call GetAssetPathRedirection")
+	FString OnStringAssetReferenceSaved(const FString& InString);
 
 private:
 
@@ -180,12 +133,14 @@ private:
 	TMultiMap<FName, FPackagePropertyPair> StringAssetReferences;
 
 	/** When saving, apply this remapping to all string asset references */
-	TMap<FString, FString> StringAssetRemap;
-#if WITH_EDITOR
+	TMap<FString, FString> AssetPathRedirectionMap;
 	TSet<FName> AlreadyRemapped;
-#endif
+
+	/** For StringAssetReferences map */
+	FCriticalSection CriticalSection;
 };
 
 // global redirect collector callback structure
 COREUOBJECT_API extern FRedirectCollector GRedirectCollector;
 
+#endif // WITH_EDITOR
\ No newline at end of file
diff --git a/Engine/Source/Runtime/CoreUObject/Public/Misc/StringAssetReference.h b/Engine/Source/Runtime/CoreUObject/Public/Misc/StringAssetReference.h
--- a/Engine/Source/Runtime/CoreUObject/Public/Misc/StringAssetReference.h
+++ b/Engine/Source/Runtime/CoreUObject/Public/Misc/StringAssetReference.h
@@ -5,6 +5,7 @@
 #include "CoreMinimal.h"
 #include "HAL/ThreadSafeCounter.h"
 #include "UObject/Class.h"
+#include "HAL/ThreadSingleton.h"
 
 /**
  * A struct that contains a string reference to an asset on disk.
@@ -14,13 +15,7 @@
  */
 struct COREUOBJECT_API FStringAssetReference
 {
-	/** Asset path */
-	DEPRECATED(4.9, "Please don't use AssetLongPathname directly. Instead use SetPath and ToString methods.")
-	FString AssetLongPathname;
-	
-	FStringAssetReference()
-	{
-	}
+	FStringAssetReference() {}
 
 	/** Construct from another asset reference */
 	FStringAssetReference(const FStringAssetReference& Other)
@@ -37,9 +32,6 @@ struct COREUOBJECT_API FStringAssetReference
 	/** Construct from an existing object in memory */
 	FStringAssetReference(const UObject* InObject);
 
-PRAGMA_DISABLE_DEPRECATION_WARNINGS
-	
-	/** Destructor, in deprecation block to avoid warnings */
 	~FStringAssetReference() {}
 
 	/** Returns string representation of reference, in form /package/path.assetname[.objectname] */
@@ -48,8 +40,6 @@ PRAGMA_DISABLE_DEPRECATION_WARNINGS
 		return AssetLongPathname;		
 	}
 
-PRAGMA_ENABLE_DEPRECATION_WARNINGS
-
 	/** Returns /package/path, leaving off the asset name */
 	FString GetLongPackageName() const
 	{
@@ -113,6 +103,15 @@ PRAGMA_ENABLE_DEPRECATION_WARNINGS
 	bool ImportTextItem( const TCHAR*& Buffer, int32 PortFlags, UObject* Parent, FOutputDevice* ErrorText );
 	bool SerializeFromMismatchedTag(struct FPropertyTag const& Tag, FArchive& Ar);
 
+	/** Serializes the internal path and also handles save/PIE fixups. Call this from the archiver overrides */
+	void SerializePath(FArchive& Ar, bool bSkipSerializeIfArchiveHasSize = false);
+
+	/** Fixes up string asset reference for saving, call if saving with a method that skips SerializePath. This can modify the path */
+	void PreSavePath();
+
+	/** Handles when a string asset reference has been loaded, call if loading with a method that skips SerializePath. This does not modify path but might call callbacks */
+	void PostLoadPath() const;
+
 	FORCEINLINE friend uint32 GetTypeHash(FStringAssetReference const& This)
 	{
 		return GetTypeHash(This.ToString());
@@ -136,6 +135,9 @@ PRAGMA_ENABLE_DEPRECATION_WARNINGS
 	static void ClearPackageNamesBeingDuplicatedForPIE();
 
 private:
+	/** Asset path */
+	FString AssetLongPathname;
+
 	/** Fixes up this StringAssetReference to add or remove the PIE prefix depending on what is currently active */
 	void FixupForPIE();
 
@@ -144,5 +146,66 @@ private:
 
 	/** Package names currently being duplicated, needed by FixupForPIE */
 	static TArray<FString> PackageNamesBeingDuplicatedForPIE;
+
+	friend COREUOBJECT_API UScriptStruct* Z_Construct_UScriptStruct_FStringAssetReference();
+};
+
+enum class EStringAssetReferenceCollectType : uint8
+{
+	/** References is not tracked in any situation, transient reference */
+	NeverCollect,
+	/** Editor only reference, this is tracked for redirector fixup but not for cooking */
+	EditorOnlyCollect,
+	/** Game reference, this is gathered for both redirector fixup and cooking */
+	AlwaysCollect,
 };
 
+class COREUOBJECT_API FStringAssetReferenceThreadContext : public TThreadSingleton<FStringAssetReferenceThreadContext>
+{
+	friend TThreadSingleton<FStringAssetReferenceThreadContext>;
+	friend struct FStringAssetReferenceSerializationScope;
+
+	FStringAssetReferenceThreadContext() {}
+
+	struct FSerializationOptions
+	{
+		FName PackageName;
+		FName PropertyName;
+		EStringAssetReferenceCollectType CollectType;
+
+		FSerializationOptions() : CollectType(EStringAssetReferenceCollectType::AlwaysCollect) {}
+		FSerializationOptions(FName InPackageName, FName InPropertyName, EStringAssetReferenceCollectType InCollectType) : PackageName(InPackageName), PropertyName(InPropertyName), CollectType(InCollectType) {}
+	};
+
+	TArray<FSerializationOptions> OptionStack;
+public:
+	/** 
+	 * Returns the current serialization options that were added using SerializationScope or LinkerLoad
+	 *
+	 * @param OutPackageName Package that this string asset belongs to
+	 * @param OutPropertyName Property that this string asset reference belongs to
+	 * @param OutCollectType Type of collecting that should be done
+	 */
+	bool GetSerializationOptions(FName& OutPackageName, FName& OutPropertyName, EStringAssetReferenceCollectType& OutCollectType) const;
+};
+
+/** Helper class to set and restore serialization options for string asset references */
+struct FStringAssetReferenceSerializationScope
+{
+	/** 
+	 * Create a new serialization scope, which affects the way that string asset references are saved
+	 *
+	 * @param SerializingPackageName Package that this string asset belongs to
+	 * @param SerializingPropertyName Property that this string asset reference belongs to
+	 * @param CollectType Set type of collecting that should be done, can be used to disable tracking entirely
+	 */
+	FStringAssetReferenceSerializationScope(FName SerializingPackageName, FName SerializingPropertyName, EStringAssetReferenceCollectType CollectType)
+	{
+		FStringAssetReferenceThreadContext::Get().OptionStack.Emplace(SerializingPackageName, SerializingPropertyName, CollectType);
+	}
+
+	~FStringAssetReferenceSerializationScope()
+	{
+		FStringAssetReferenceThreadContext::Get().OptionStack.Pop();
+	}
+};
\ No newline at end of file
diff --git a/Engine/Source/Runtime/CoreUObject/Public/Serialization/DuplicatedDataReader.h b/Engine/Source/Runtime/CoreUObject/Public/Serialization/DuplicatedDataReader.h
--- a/Engine/Source/Runtime/CoreUObject/Public/Serialization/DuplicatedDataReader.h
+++ b/Engine/Source/Runtime/CoreUObject/Public/Serialization/DuplicatedDataReader.h
@@ -7,7 +7,6 @@
 #include "UObject/UObjectAnnotation.h"
 #include "Serialization/DuplicatedObject.h"
 
-class FAssetPtr;
 class FLazyObjectPtr;
 struct FStringAssetReference;
 
@@ -31,7 +30,6 @@ private:
 	virtual FArchive& operator<<(FName& N);
 	virtual FArchive& operator<<(UObject*& Object);
 	virtual FArchive& operator<<(FLazyObjectPtr& LazyObjectPtr);
-	virtual FArchive& operator<<(FAssetPtr& AssetPtr);
 	virtual FArchive& operator<<(FStringAssetReference& StringAssetReference);
 	
 	void SerializeFail();
