UE4 - Fixed critical bug with  FTaskGraphInterface::BroadcastSlow_OnlyUseForSpecialPurposes.

diff --git a/Engine/Source/Runtime/Core/Private/Async/TaskGraph.cpp b/Engine/Source/Runtime/Core/Private/Async/TaskGraph.cpp
--- a/Engine/Source/Runtime/Core/Private/Async/TaskGraph.cpp
+++ b/Engine/Source/Runtime/Core/Private/Async/TaskGraph.cpp
@@ -1610,6 +1610,7 @@ void FTaskGraphInterface::BroadcastSlow_OnlyUseForSpecialPurposes(bool bDoTaskTh
 
 	FEvent* MyEvent = nullptr;
 	FGraphEventArray TaskThreadTasks;
+	FThreadSafeCounter StallForTaskThread;
 	if (bDoTaskThreads)
 	{
 		MyEvent = FPlatformProcess::GetSynchEventFromPool(false);
@@ -1674,11 +1675,14 @@ void FTaskGraphInterface::BroadcastSlow_OnlyUseForSpecialPurposes(bool bDoTaskTh
 			TaskEvent->Trigger();
 		}
 		FTaskGraphInterface::Get().WaitUntilTasksComplete(TaskThreadTasks, ENamedThreads::GameThread_Local);
-		for (FEvent* TaskEvent : TaskEvents)
-		{
-			FPlatformProcess::ReturnSynchEventToPool(TaskEvent);
-		}
-		FTaskGraphInterface::Get().WaitUntilTasksComplete(Tasks, ENamedThreads::GameThread_Local);
+	}
+	FTaskGraphInterface::Get().WaitUntilTasksComplete(Tasks, ENamedThreads::GameThread_Local);
+	for (FEvent* TaskEvent : TaskEvents)
+	{
+		FPlatformProcess::ReturnSynchEventToPool(TaskEvent);
+	}
+	if (MyEvent)
+	{
 		FPlatformProcess::ReturnSynchEventToPool(MyEvent);
 	}
 }
