UE4 - Repaired botched merge of BroadcastSlow_OnlyUseForSpecialPurposes.

diff --git a/Engine/Source/Runtime/Core/Private/Async/TaskGraph.cpp b/Engine/Source/Runtime/Core/Private/Async/TaskGraph.cpp
--- a/Engine/Source/Runtime/Core/Private/Async/TaskGraph.cpp
+++ b/Engine/Source/Runtime/Core/Private/Async/TaskGraph.cpp
@@ -1614,10 +1614,8 @@ void FTaskGraphInterface::BroadcastSlow_OnlyUseForSpecialPurposes(bool bDoTaskTh
 	if (bDoTaskThreads)
 	{
 		MyEvent = FPlatformProcess::GetSynchEventFromPool(false);
-		FThreadSafeCounter StallForTaskThread;
 
 		int32 Workers = FTaskGraphInterface::Get().GetNumWorkerThreads();
-		StallForTaskThread.Reset();
 		StallForTaskThread.Add(Workers * (1 + (bDoBackgroundThreads && ENamedThreads::bHasBackgroundThreads) + !!(ENamedThreads::bHasHighPriorityThreads)));
 
 		TaskEvents.Reserve(StallForTaskThread.GetValue());
@@ -1688,7 +1686,6 @@ void FTaskGraphInterface::BroadcastSlow_OnlyUseForSpecialPurposes(bool bDoTaskTh
 }
 
 
-
 // Benchmark
 
 #include "Async/ParallelFor.h"
