UE4 - Fix bogus static analysis warning.

diff --git a/Engine/Source/Runtime/Core/Private/Async/TaskGraph.cpp b/Engine/Source/Runtime/Core/Private/Async/TaskGraph.cpp
--- a/Engine/Source/Runtime/Core/Private/Async/TaskGraph.cpp
+++ b/Engine/Source/Runtime/Core/Private/Async/TaskGraph.cpp
@@ -1664,7 +1664,8 @@ void FTaskGraphInterface::BroadcastSlow_OnlyUseForSpecialPurposes(bool bDoTaskTh
 	Tasks.Add(TGraphTask<FBroadcastTask>::CreateTask().ConstructAndDispatchWhenReady(Callback, ENamedThreads::GameThread_Local, nullptr, nullptr, nullptr));
 	if (bDoTaskThreads)
 	{
-		if (!MyEvent->Wait(3000))
+		check(MyEvent);
+		if (MyEvent && !MyEvent->Wait(3000))
 		{
 			UE_LOG(LogTaskGraph, Log, TEXT("FTaskGraphInterface::BroadcastSlow_OnlyUseForSpecialPurposes Broadcast failed after three seconds. Ok during automated tests."));
 		}
