UE4 - Suppressed warning for unaligned reads from loose files.

diff --git a/Engine/Source/Runtime/Core/Private/Windows/WindowsAsyncIO.h b/Engine/Source/Runtime/Core/Private/Windows/WindowsAsyncIO.h
--- a/Engine/Source/Runtime/Core/Private/Windows/WindowsAsyncIO.h
+++ b/Engine/Source/Runtime/Core/Private/Windows/WindowsAsyncIO.h
@@ -92,12 +92,12 @@ public:
 				if (NumMessages < 10)
 				{
 					NumMessages++;
-					UE_LOG(LogTemp, Warning, TEXT("FWindowsReadRequest request was not aligned. This is expected with loose files, but not a pak file."));
+					UE_LOG(LogTemp, Log, TEXT("FWindowsReadRequest request was not aligned. This is expected with loose files, but not a pak file."));
 				}
 				else if (NumMessages == 10)
 				{
 					NumMessages++;
-					UE_LOG(LogTemp, Warning, TEXT("LAST NOTIFICATION THIS RUN: FWindowsReadRequest request was not aligned."));
+					UE_LOG(LogTemp, Log, TEXT("LAST NOTIFICATION THIS RUN: FWindowsReadRequest request was not aligned."));
 				}
 				TempMemory = (uint8*)FMemory::Malloc(AlignedBytesToRead);
 				INC_MEMORY_STAT_BY(STAT_AsyncFileMemory, AlignedBytesToRead);
