UE4 - Critical fix: Stats were non functional.

diff --git a/Engine/Source/Runtime/Core/Private/Async/TaskGraph.cpp b/Engine/Source/Runtime/Core/Private/Async/TaskGraph.cpp
--- a/Engine/Source/Runtime/Core/Private/Async/TaskGraph.cpp
+++ b/Engine/Source/Runtime/Core/Private/Async/TaskGraph.cpp
@@ -656,6 +656,13 @@ public:
 				TestRandomizedThreads();
 			}
 		}
+#if STATS
+		if (bTasksOpen)
+		{
+			ProcessingTasks.Stop();
+			bTasksOpen = false;
+		}
+#endif
 	}
 	virtual void EnqueueFromThisThread(int32 QueueIndex, FBaseGraphTask* Task) override
 	{
@@ -887,7 +894,7 @@ private:
 	struct FThreadTaskQueue
 	{
 		/** We need to disallow reentry of the processing loop **/
-		uint32												RecursionGuard;
+		uint32 RecursionGuard;
 		/** Indicates we executed a return task, so break out of the processing loop. **/
 		bool QuitForShutdown;
 		/** Event that this thread blocks on when it runs out of work. **/
diff --git a/Engine/Source/Runtime/Core/Private/Stats/Stats2.cpp b/Engine/Source/Runtime/Core/Private/Stats/Stats2.cpp
--- a/Engine/Source/Runtime/Core/Private/Stats/Stats2.cpp
+++ b/Engine/Source/Runtime/Core/Private/Stats/Stats2.cpp
@@ -825,6 +825,7 @@ FName FStatNameAndInfo::GetGroupCategoryFrom(FName InLongName)
 
 static TAutoConsoleVariable<int32> CVarDumpStatPackets(	TEXT("DumpStatPackets"),0,	TEXT("If true, dump stat packets."));
 
+
 /** The rendering thread runnable object. */
 class FStatsThread : public FRunnable, FSingleThreadRunnable
 {
@@ -857,6 +858,7 @@ public:
 	/** Attaches to the task graph stats thread, all processing will be handled by the task graph. */
 	virtual uint32 Run() override
 	{
+		FThreadStats::GetThreadStats()->bIsStatsThread = true;
 		FMemory::SetupTLSCachesOnCurrentThread();
 		FTaskGraphInterface::Get().AttachToThread(ENamedThreads::StatsThread);
 		FTaskGraphInterface::Get().ProcessThreadUntilRequestReturn(ENamedThreads::StatsThread);
@@ -960,6 +962,17 @@ public:
 		Tick();
 	}
 
+	void SelfStatMessage(FStatPacket* Packet)
+	{
+		if (CVarDumpStatPackets.GetValueOnAnyThread())
+		{
+			UE_LOG(LogStats, Log, TEXT("Self Packet from %x with %d messages"), Packet->ThreadId, Packet->StatMessages.Num());
+		}
+
+		IncomingData.Packets.Add(Packet);
+		State.NumStatMessages.Add(Packet->StatMessages.Num());
+	}
+
 	/** Start a stats runnable thread. */
 	void Start()
 	{
@@ -1043,7 +1056,8 @@ FThreadStats::FThreadStats():
 	bWaitForExplicitFlush(0),
 	MemoryMessageScope(0),
 	bReentranceGuard(false),
-	bSawExplicitFlush(false)
+	bSawExplicitFlush(false),
+	bIsStatsThread(false)
 {
 	Packet.SetThreadProperties();
 
@@ -1057,7 +1071,8 @@ FThreadStats::FThreadStats( EConstructor ):
 	bWaitForExplicitFlush(0),
 	MemoryMessageScope(0),
 	bReentranceGuard(false),
-	bSawExplicitFlush(false)
+	bSawExplicitFlush(false),
+	bIsStatsThread(false)
 {}
 
 void FThreadStats::CheckEnable()
@@ -1147,8 +1162,14 @@ void FThreadStats::FlushRegularStats( bool bHasBrokenCallstacks, bool bForceFlus
 			}
 			Packet.StatMessages.Empty(MaxPresize);
 		}
-
-		TGraphTask<FStatMessagesTask>::CreateTask().ConstructAndDispatchWhenReady(ToSend);
+		if (bIsStatsThread)
+		{
+			FStatsThread::Get().SelfStatMessage(ToSend);
+		}
+		else
+		{
+			TGraphTask<FStatMessagesTask>::CreateTask().ConstructAndDispatchWhenReady(ToSend);
+		}
 		UpdateExplicitFlush();
 	}
 }
@@ -1184,7 +1205,14 @@ void FThreadStats::FlushRawStats( bool bHasBrokenCallstacks /*= false*/, bool bF
 
 		check(!Packet.StatMessages.Num());
 
-		TGraphTask<FStatMessagesTask>::CreateTask().ConstructAndDispatchWhenReady(ToSend);
+		if (bIsStatsThread)
+		{
+			FStatsThread::Get().SelfStatMessage(ToSend);
+		}
+		else
+		{
+			TGraphTask<FStatMessagesTask>::CreateTask().ConstructAndDispatchWhenReady(ToSend);
+		}
 		UpdateExplicitFlush();
 
 		const float NumMessagesAsMB = NumMessages*sizeof(FStatMessage) / 1024.0f / 1024.0f;
diff --git a/Engine/Source/Runtime/Core/Public/Stats/Stats2.h b/Engine/Source/Runtime/Core/Public/Stats/Stats2.h
--- a/Engine/Source/Runtime/Core/Public/Stats/Stats2.h
+++ b/Engine/Source/Runtime/Core/Public/Stats/Stats2.h
@@ -1252,6 +1252,9 @@ class FThreadStats : FNoncopyable
 	/** Tracks current stack depth for cycle counters. **/
 	bool bSawExplicitFlush;
 
+	/** True if this is the stats thread, which needs special handling. **/
+	bool bIsStatsThread;
+
 	/** Gathers information about the current thread and sets up the TLS value. **/
 	CORE_API FThreadStats();
 
