UE4 - Critical Fix: Fixed stats crash.

diff --git a/Engine/Source/Runtime/Core/Private/Containers/LockFreeList.cpp b/Engine/Source/Runtime/Core/Private/Containers/LockFreeList.cpp
--- a/Engine/Source/Runtime/Core/Private/Containers/LockFreeList.cpp
+++ b/Engine/Source/Runtime/Core/Private/Containers/LockFreeList.cpp
@@ -48,14 +48,24 @@ void LockFreeLinksExhausted(uint32 TotalNum)
 	UE_LOG(LogTemp, Fatal, TEXT("Consumed %d lock free links; there are no more."), TotalNum);
 }
 
+static void ChangeMem(int32 Delta)
+{
+	static FThreadSafeCounter LockFreeListMem;
+	LockFreeListMem.Add(Delta);
+	if (GIsRunning)
+	{
+		SET_MEMORY_STAT(STAT_LockFreeListLinks, LockFreeListMem.GetValue());
+	}
+}
+
 void* LockFreeAllocLinks(SIZE_T AllocSize)
 {
-	INC_MEMORY_STAT_BY(STAT_LockFreeListLinks, AllocSize);
+	ChangeMem(AllocSize);
 	return FMemory::Malloc(AllocSize);
 }
 void LockFreeFreeLinks(SIZE_T AllocSize, void* Ptr)
 {
-	DEC_MEMORY_STAT_BY(STAT_LockFreeListLinks, AllocSize);
+	ChangeMem(-int32(AllocSize));
 	FMemory::Free(Ptr);
 }
 
diff --git a/Engine/Source/Runtime/Core/Private/Stats/Stats2.cpp b/Engine/Source/Runtime/Core/Private/Stats/Stats2.cpp
--- a/Engine/Source/Runtime/Core/Private/Stats/Stats2.cpp
+++ b/Engine/Source/Runtime/Core/Private/Stats/Stats2.cpp
@@ -1061,7 +1061,7 @@ FThreadStats::FThreadStats():
 {
 	Packet.SetThreadProperties();
 
-	check(FPlatformTLS::IsValidTlsSlot(TlsSlot));
+	check(TlsSlot && FPlatformTLS::IsValidTlsSlot(TlsSlot));
 	FPlatformTLS::SetTlsValue(TlsSlot, this);
 }
 
@@ -1311,18 +1311,18 @@ void FThreadStats::StartThread()
 	FThreadStats::FrameDataIsIncomplete(); // make this non-zero
 	check(IsInGameThread());
 	check(!IsThreadingReady());
-	FStatsThreadState::GetLocalState(); // start up the state
-	FStatsThread::Get();
-	FStatsThread::Get().Start();
-
 	// Preallocate a bunch of FThreadStats to avoid dynamic memory allocation.
 	// (Must do this before we expose ourselves to other threads via tls).
 	FThreadStatsPool::Get();
-
+	FStatsThreadState::GetLocalState(); // start up the state
 	if (!TlsSlot)
 	{
 		TlsSlot = FPlatformTLS::AllocTlsSlot();
+                check(TlsSlot);
 	}
+	FStatsThread::Get();
+	FStatsThread::Get().Start();
+
 	check(IsThreadingReady());
 	CheckEnable();
 	
