UE4 - Added more lock free links to support huge cooks.

diff --git a/Engine/Source/Runtime/Core/Private/Async/TaskGraph.cpp b/Engine/Source/Runtime/Core/Private/Async/TaskGraph.cpp
--- a/Engine/Source/Runtime/Core/Private/Async/TaskGraph.cpp
+++ b/Engine/Source/Runtime/Core/Private/Async/TaskGraph.cpp
@@ -2107,14 +2107,14 @@ struct FTestRigFIFO
 {
 	FLockFreePointerFIFOBase<FTestStruct, PLATFORM_CACHE_LINE_SIZE> Test1;
 	FLockFreePointerFIFOBase<FTestStruct, 1> Test2;
-	FLockFreePointerFIFOBase<FTestStruct, 1, 1 << 16> Test3;
+	FLockFreePointerFIFOBase<FTestStruct, 1, 1 << 4> Test3;
 };
 
 struct FTestRigLIFO
 {
 	FLockFreePointerListLIFOBase<FTestStruct, PLATFORM_CACHE_LINE_SIZE> Test1;
 	FLockFreePointerListLIFOBase<FTestStruct, 1> Test2;
-	FLockFreePointerListLIFOBase<FTestStruct, 1, 1 << 16> Test3;
+	FLockFreePointerListLIFOBase<FTestStruct, 1, 1 << 4> Test3;
 };
 
 static void TestLockFree(int32 OuterIters = 3)
diff --git a/Engine/Source/Runtime/Core/Public/Containers/LockFreeList.h b/Engine/Source/Runtime/Core/Public/Containers/LockFreeList.h
--- a/Engine/Source/Runtime/Core/Public/Containers/LockFreeList.h
+++ b/Engine/Source/Runtime/Core/Public/Containers/LockFreeList.h
@@ -40,8 +40,8 @@ CORE_API void LockFreeLinksExhausted(uint32 TotalNum);
 CORE_API void* LockFreeAllocLinks(SIZE_T AllocSize);
 CORE_API void LockFreeFreeLinks(SIZE_T AllocSize, void* Ptr);
 
-#define MAX_LOCK_FREE_LINKS_AS_BITS (22)
-#define MAX_LOCK_FREE_LINKS (1 << 22)
+#define MAX_LOCK_FREE_LINKS_AS_BITS (26)
+#define MAX_LOCK_FREE_LINKS (1 << 26)
 
 template<int TPaddingForCacheContention>
 struct FPaddingForCacheContention
