UE4 - Fixed debug config problems with new low level windows async IO layer.

diff --git a/Engine/Source/Runtime/Core/Private/Windows/WindowsAsyncIO.h b/Engine/Source/Runtime/Core/Private/Windows/WindowsAsyncIO.h
--- a/Engine/Source/Runtime/Core/Private/Windows/WindowsAsyncIO.h
+++ b/Engine/Source/Runtime/Core/Private/Windows/WindowsAsyncIO.h
@@ -36,6 +36,7 @@ HANDLE GetIOPooledEvent()
 
 void FreeIOPooledEvent(HANDLE ToFree)
 {
+	check(ToFree != INVALID_HANDLE_VALUE && (void*)(UPTRINT)ToFree); //awkwardly using void* to store handles, hope we don't ever have a zero handle :)
 	ResetEvent(ToFree);
 	WindowsAsyncIOEventPool.Push((void*)(UPTRINT)ToFree);
 }
@@ -67,6 +68,7 @@ public:
 		, TempMemory(nullptr)
 	{
 		FMemory::Memzero(OverlappedIO);
+		OverlappedIO.hEvent = INVALID_HANDLE_VALUE;
 		check(Offset >= 0 && BytesToRead > 0);
 		if (BytesToRead == MAX_int64)
 		{
