HTML5 - TaskGraph crash fix & compiler fix when STATS disabled

#rnx

diff --git a/Engine/Source/Runtime/Core/Private/Async/TaskGraph.cpp b/Engine/Source/Runtime/Core/Private/Async/TaskGraph.cpp
--- a/Engine/Source/Runtime/Core/Private/Async/TaskGraph.cpp
+++ b/Engine/Source/Runtime/Core/Private/Async/TaskGraph.cpp
@@ -785,6 +785,18 @@ public:
 		} while (!Queue.QuitForShutdown && FPlatformProcess::SupportsMultithreading()); // @Hack - quit now when running with only one thread.
 	}
 
+	virtual void ProcessTasksUntilIdle(int32 QueueIndex) override
+	{
+		if (!FPlatformProcess::SupportsMultithreading())
+		{
+			ProcessTasks();
+		}
+		else
+		{
+			check(0);
+		}
+	}
+
 	// Calls meant to be called from any thread.
 
 	/**
@@ -870,13 +882,14 @@ private:
 #endif
 
 				TestRandomizedThreads();
+				if (FPlatformProcess::SupportsMultithreading())
 				{
 					FScopeCycleCounter Scope(StallStatId);
 					Queue.StallRestartEvent->Wait(MAX_uint32, bCountAsStall);
 				}
-				if (Queue.QuitForShutdown)
+				if (Queue.QuitForShutdown || !FPlatformProcess::SupportsMultithreading())
 				{
-					return;
+					break;
 				}
 
 				TestRandomizedThreads();
diff --git a/Engine/Source/Runtime/Launch/Private/LaunchEngineLoop.cpp b/Engine/Source/Runtime/Launch/Private/LaunchEngineLoop.cpp
--- a/Engine/Source/Runtime/Launch/Private/LaunchEngineLoop.cpp
+++ b/Engine/Source/Runtime/Launch/Private/LaunchEngineLoop.cpp
@@ -2867,7 +2867,7 @@ bool FEngineLoop::ShouldUseIdleMode() const
 	return bIdleMode;
 }
 
-#if !UE_BUILD_SHIPPING && !UE_BUILD_TEST
+#if !UE_BUILD_SHIPPING && !UE_BUILD_TEST && MALLOC_GT_HOOKS
 
 #include "Containers/StackTracker.h"
 static TAutoConsoleVariable<int32> CVarLogGameThreadMallocChurn(
@@ -2991,7 +2991,7 @@ uint64 FScopedSampleMallocChurn::DumpFrame = 0;
 
 void FEngineLoop::Tick()
 {
-#if !UE_BUILD_SHIPPING && !UE_BUILD_TEST
+#if !UE_BUILD_SHIPPING && !UE_BUILD_TEST && MALLOC_GT_HOOKS
 	FScopedSampleMallocChurn ChurnTracker;
 #endif
 
