Code marked "Do not merge to //UE4/Main" was merged to //UE4/Main. Was disabling logging in test builds.

diff --git a/Engine/Source/Runtime/Core/Public/Misc/Build.h b/Engine/Source/Runtime/Core/Public/Misc/Build.h
--- a/Engine/Source/Runtime/Core/Public/Misc/Build.h
+++ b/Engine/Source/Runtime/Core/Public/Misc/Build.h
@@ -201,7 +201,7 @@
 	#define STATS											(USE_MALLOC_PROFILER)
 	#define ALLOW_DEBUG_FILES								1
 	#define ALLOW_CONSOLE									1
-	#define NO_LOGGING										0 // NOTE: do not merge to UE4/Main
+	#define NO_LOGGING										!USE_LOGGING_IN_SHIPPING
 #elif UE_BUILD_SHIPPING
 	#if WITH_EDITOR
 		#define DO_GUARD_SLOW								0
