Attempting to fix case of iOS directory (pt2: re-adding files)

diff --git a/Engine/Source/Runtime/Core/Public/iOS/IOSAppDelegate.h b/Engine/Source/Runtime/Core/Public/iOS/IOSAppDelegate.h
new file mode 100644
--- /dev/null
+++ b/Engine/Source/Runtime/Core/Public/iOS/IOSAppDelegate.h
@@ -0,0 +1,142 @@
+// Copyright 1998-2017 Epic Games, Inc. All Rights Reserved.
+
+#pragma once
+
+#import <UIKit/UIKit.h>
+#import <AVFoundation/AVAudioSession.h>
+#import <GameKit/GKGameCenterViewController.h>
+#include "Delegates/Delegate.h"
+#include "Logging/LogMacros.h"
+#include "Containers/UnrealString.h"
+
+class CORE_API FIOSCoreDelegates
+{
+public:
+	// Broadcast when this application is opened from an external source.
+	DECLARE_MULTICAST_DELEGATE_FourParams(FOnOpenURL, UIApplication*, NSURL*, NSString*, id);
+	static FOnOpenURL OnOpenURL;
+	
+};
+
+@class FIOSView;
+@class IOSViewController;
+@class SlateOpenGLESViewController;
+@class IOSAppDelegate;
+
+DECLARE_LOG_CATEGORY_EXTERN(LogIOSAudioSession, Log, All);
+
+namespace FAppEntry
+{
+	void PlatformInit();
+	void PreInit(IOSAppDelegate* AppDelegate, UIApplication* Application);
+	void Init();
+	void Tick();
+    void SuspendTick();
+	void Shutdown();
+    void Suspend();
+    void Resume();
+	
+	extern bool	gAppLaunchedWithLocalNotification;
+	extern FString	gLaunchLocalNotificationActivationEvent;
+	extern int32	gLaunchLocalNotificationFireDate;
+}
+
+@interface IOSAppDelegate : UIResponder <UIApplicationDelegate,
+#if !UE_BUILD_SHIPPING
+	UIGestureRecognizerDelegate,
+#endif
+	GKGameCenterControllerDelegate,
+UITextFieldDelegate>
+
+/** Window object */
+@property (strong, retain, nonatomic) UIWindow *Window;
+
+/** Main GL View */
+@property (retain) FIOSView* IOSView;
+
+@property class FIOSApplication* IOSApplication;
+
+/** The controller to handle rotation of the view */
+@property (retain) IOSViewController* IOSController;
+
+/** The view controlled by the auto-rotating controller */
+@property (retain) UIView* RootView;
+
+/** The controller to handle rotation of the view */
+@property (retain) SlateOpenGLESViewController* SlateController;
+
+/** The value of the alert response (atomically set since main thread and game thread use it */
+@property (assign) int AlertResponse;
+
+/** Version of the OS we are running on (NOT compiled with) */
+@property (readonly) float OSVersion;
+
+@property bool bDeviceInPortraitMode;
+
+@property (retain) NSTimer* timer;
+
+@property (retain) NSTimer* PeakMemoryTimer;
+
+/** Timer used for re-enabling the idle timer */
+@property (retain) NSTimer* IdleTimerEnableTimer;
+
+/** The time delay (in seconds) between idle timer enable requests and actually enabling the idle timer */
+@property (readonly) float IdleTimerEnablePeriod;
+
+#if !UE_BUILD_SHIPPING && !PLATFORM_TVOS
+	/** Properties for managing the console */
+#if __IPHONE_OS_VERSION_MIN_REQUIRED < __IPHONE_9_0
+	@property (nonatomic, retain) UIAlertView*		ConsoleAlert;
+#endif
+#ifdef __IPHONE_8_0
+	@property (nonatomic, assign) UIAlertController* ConsoleAlertController;
+#endif
+	@property (nonatomic, retain) NSMutableArray*	ConsoleHistoryValues;
+	@property (nonatomic, assign) int				ConsoleHistoryValuesIndex;
+#endif
+
+/** True if the engine has been initialized */
+@property (readonly) bool bEngineInit;
+
+/** Delays game initialization slightly in case we have a URL launch to handle */
+@property (retain) NSTimer* CommandLineParseTimer;
+@property (atomic) bool bCommandLineReady;
+
+/** initial launch options */
+@property(retain) NSDictionary* launchOptions;
+
+/**
+ * @return the single app delegate object
+ */
++ (IOSAppDelegate*)GetDelegate;
+
+-(void)EnableIdleTimer:(bool)bEnable;
+
+-(void) ParseCommandLineOverrides;
+
+-(int)GetAudioVolume;
+-(bool)AreHeadphonesPluggedIn;
+-(int)GetBatteryLevel;
+-(bool)IsRunningOnBattery;
+
+/** TRUE if the device is playing background music and we want to allow that */
+@property (assign) bool bUsingBackgroundMusic;
+
+- (void)InitializeAudioSession;
+- (void)ToggleAudioSession:(bool)bActive;
+- (bool)IsBackgroundAudioPlaying;
+
+@property (atomic) bool bAudioActive;
+
+@property (atomic) bool bIsSuspended;
+@property (atomic) bool bHasSuspended;
+@property (atomic) bool bHasStarted;
+- (void)ToggleSuspend:(bool)bSuspend;
+
+static void interruptionListener(void* ClientData, UInt32 Interruption);
+
+-(UIWindow*)window;
+
+@end
+
+void InstallSignalHandlers();
diff --git a/Engine/Source/Runtime/Core/Public/iOS/IOSApplication.h b/Engine/Source/Runtime/Core/Public/iOS/IOSApplication.h
new file mode 100644
--- /dev/null
+++ b/Engine/Source/Runtime/Core/Public/iOS/IOSApplication.h
@@ -0,0 +1,56 @@
+// Copyright 1998-2017 Epic Games, Inc. All Rights Reserved.
+
+#pragma once
+
+#include "GenericApplication.h"
+
+class FIOSWindow;
+
+class FIOSApplication : public GenericApplication
+{
+public:
+
+	static FIOSApplication* CreateIOSApplication();
+
+
+public:
+
+	virtual ~FIOSApplication() {}
+
+	void SetMessageHandler( const TSharedRef< FGenericApplicationMessageHandler >& InMessageHandler );
+
+	virtual void PollGameDeviceState( const float TimeDelta ) override;
+
+	virtual FPlatformRect GetWorkArea( const FPlatformRect& CurrentWindow ) const override;
+
+	virtual TSharedRef< FGenericWindow > MakeWindow() override;
+
+	virtual void AddExternalInputDevice(TSharedPtr<class IInputDevice> InputDevice);
+
+#if !PLATFORM_TVOS
+	static void OrientationChanged(UIDeviceOrientation orientation);
+#endif
+
+	virtual IInputInterface* GetInputInterface() override { return (IInputInterface*)InputInterface.Get(); }
+
+protected:
+	virtual void InitializeWindow( const TSharedRef< FGenericWindow >& Window, const TSharedRef< FGenericWindowDefinition >& InDefinition, const TSharedPtr< FGenericWindow >& InParent, const bool bShowImmediately ) override;
+
+private:
+
+	FIOSApplication();
+
+
+private:
+
+	TSharedPtr< class FIOSInputInterface > InputInterface;
+
+	/** List of input devices implemented in external modules. */
+	TArray< TSharedPtr<class IInputDevice> > ExternalInputDevices;
+	bool bHasLoadedInputPlugins;
+
+	TArray< TSharedRef< FIOSWindow > > Windows;
+
+	static FCriticalSection CriticalSection;
+	static bool bOrientationChanged;
+};
diff --git a/Engine/Source/Runtime/Core/Public/iOS/IOSCommandLineHelper.h b/Engine/Source/Runtime/Core/Public/iOS/IOSCommandLineHelper.h
new file mode 100644
--- /dev/null
+++ b/Engine/Source/Runtime/Core/Public/iOS/IOSCommandLineHelper.h
@@ -0,0 +1,92 @@
+// Copyright 1998-2017 Epic Games, Inc. All Rights Reserved.
+
+#pragma once
+
+#include "CoreMinimal.h"
+#include "Misc/CommandLine.h"
+
+#define IOS_MAX_PATH 1024
+#define CMD_LINE_MAX 16384
+
+extern FString GSavedCommandLine;
+
+class FIOSCommandLineHelper
+{
+	public:
+		/**
+		 * Merge the given commandline with GSavedCommandLinePortion, which may start with ?
+		 * options that need to come after first token
+		 */
+
+		static void MergeCommandlineWithSaved(TCHAR CommandLine[16384])
+		{
+			// the saved commandline may be in the format ?opt?opt -opt -opt, so we need to insert it 
+			// after the first token on the commandline unless the first token starts with a -, in which 
+			// case use it at the start of the command line
+			if (CommandLine[0] == '-' || CommandLine[0] == 0)
+			{
+				// handle the easy - case, just use the saved command line part as the start, in case it
+				// started with a ?
+				FString CombinedCommandLine = GSavedCommandLine + CommandLine;
+				FCString::Strcpy(CommandLine, CMD_LINE_MAX, *CombinedCommandLine);
+			}
+			else
+			{
+				// otherwise, we need to get the first token from the command line and insert after
+				TCHAR* Space = FCString::Strchr(CommandLine, ' ');
+				if (Space == NULL)
+				{
+					// if there is only one token (no spaces), just append us after it
+					FCString::Strcat(CommandLine, CMD_LINE_MAX, *GSavedCommandLine);
+				}
+				else
+				{
+					// save off what's after the space (include the space for pasting later)
+					FString AfterSpace(Space);
+					// copy the save part where the space was
+					FCString::Strcpy(Space, CMD_LINE_MAX - (Space - CommandLine), *GSavedCommandLine);
+					// now put back the 2nd and so on token
+					FCString::Strcat(CommandLine, CMD_LINE_MAX, *AfterSpace);
+				}
+			}
+		}
+
+		static void InitCommandArgs(FString AdditionalCommandArgs)
+		{
+			// initialize the commandline
+			FCommandLine::Set(TEXT(""));
+
+			FString CommandLineFilePath = FString([[NSBundle mainBundle] bundlePath]) + TEXT("/ue4commandline.txt");
+
+			// read in the command line text file (coming from UnrealFrontend) if it exists
+			FILE* CommandLineFile = fopen(TCHAR_TO_UTF8(*CommandLineFilePath), "r");
+			if(CommandLineFile)
+			{
+				FPlatformMisc::LowLevelOutputDebugStringf(TEXT("Found ue4commandline.txt file") LINE_TERMINATOR);
+
+				char CommandLine[CMD_LINE_MAX];
+				fgets(CommandLine, ARRAY_COUNT(CommandLine) - 1, CommandLineFile);
+
+				// chop off trailing spaces
+				while (*CommandLine && isspace(CommandLine[strlen(CommandLine) - 1]))
+				{
+					CommandLine[strlen(CommandLine) - 1] = 0;
+				}
+
+				FCommandLine::Append(UTF8_TO_TCHAR(CommandLine));
+			}
+			else
+			{
+				FPlatformMisc::LowLevelOutputDebugStringf(TEXT("No ue4commandline.txt [%s] found") LINE_TERMINATOR, *CommandLineFilePath);
+			}
+
+			FCommandLine::Append(*AdditionalCommandArgs);
+
+			// now merge the GSavedCommandLine with the rest
+			FCommandLine::Append(*GSavedCommandLine);
+
+			FPlatformMisc::LowLevelOutputDebugStringf(TEXT("Combined iOS Commandline: %s") LINE_TERMINATOR, FCommandLine::Get());
+		}
+
+};
+
diff --git a/Engine/Source/Runtime/Core/Public/iOS/IOSCriticalSection.h b/Engine/Source/Runtime/Core/Public/iOS/IOSCriticalSection.h
new file mode 100644
--- /dev/null
+++ b/Engine/Source/Runtime/Core/Public/iOS/IOSCriticalSection.h
@@ -0,0 +1,11 @@
+// Copyright 1998-2017 Epic Games, Inc. All Rights Reserved.
+
+#pragma once
+
+#include "GenericPlatform/GenericPlatformCriticalSection.h"
+#include "PThreadCriticalSection.h"
+#include "PThreadRWLock.h"
+
+typedef FPThreadsCriticalSection FCriticalSection;
+typedef FSystemWideCriticalSectionNotImplemented FSystemWideCriticalSection;
+typedef FPThreadsRWLock FRWLock;
diff --git a/Engine/Source/Runtime/Core/Public/iOS/IOSInputInterface.h b/Engine/Source/Runtime/Core/Public/iOS/IOSInputInterface.h
new file mode 100644
--- /dev/null
+++ b/Engine/Source/Runtime/Core/Public/iOS/IOSInputInterface.h
@@ -0,0 +1,171 @@
+// Copyright 1998-2017 Epic Games, Inc. All Rights Reserved.
+
+#pragma once
+
+#include "CoreMinimal.h"
+#include "IInputInterface.h"
+#include "IForceFeedbackSystem.h"
+#if !PLATFORM_TVOS
+#import <CoreMotion/CoreMotion.h>
+#endif
+#import <GameController/GameController.h>
+#include "Misc/CoreMisc.h"
+#include "Math/Quat.h"
+#include "Math/Vector.h"
+#include "Math/Color.h"
+#include "GenericPlatform/GenericApplicationMessageHandler.h"
+
+
+#define KEYCODE_ENTER 1000
+#define KEYCODE_BACKSPACE 1001
+#define KEYCODE_ESCAPE 1002
+
+
+enum TouchType
+{
+	TouchBegan,
+	TouchMoved,
+	TouchEnded,
+};
+
+struct TouchInput
+{
+	int Handle;
+	TouchType Type;
+	FVector2D LastPosition;
+	FVector2D Position;
+};
+
+/**
+ * Interface class for IOS input devices
+ */
+class FIOSInputInterface : public IForceFeedbackSystem, FSelfRegisteringExec
+{
+public:
+
+	static TSharedRef< FIOSInputInterface > Create(  const TSharedRef< FGenericApplicationMessageHandler >& InMessageHandler );
+	static TSharedPtr< FIOSInputInterface > Get();
+
+public:
+
+	virtual ~FIOSInputInterface() {}
+
+	void SetMessageHandler( const TSharedRef< FGenericApplicationMessageHandler >& InMessageHandler );
+
+	/** Tick the interface (i.e check for new controllers) */
+	void Tick( float DeltaTime );
+
+	/**
+	 * Poll for controller state and send events if needed
+	 */
+	void SendControllerEvents();
+
+	static void QueueTouchInput(const TArray<TouchInput>& InTouchEvents);
+	static void QueueKeyInput(int32 Key, int32 Char);
+
+	//~ Begin Exec Interface
+	virtual bool Exec(UWorld* InWorld, const TCHAR* Cmd, FOutputDevice& Ar) override;
+	//~ End Exec Interface
+
+	/**
+	 * IForceFeedbackSystem implementation
+	 */
+	virtual void SetForceFeedbackChannelValue(int32 ControllerId, FForceFeedbackChannelType ChannelType, float Value) override;
+	virtual void SetForceFeedbackChannelValues(int32 ControllerId, const FForceFeedbackValues &values) override;
+	virtual void SetLightColor(int32 ControllerId, FColor Color) override {}
+
+	bool IsControllerAssignedToGamepad(int32 ControllerId);
+
+private:
+
+	FIOSInputInterface( const TSharedRef< FGenericApplicationMessageHandler >& InMessageHandler );
+
+	// handle disconnect and connect events
+	void HandleConnection(GCController* Controller);
+	void HandleDisconnect(GCController* Controller);
+	
+	/**
+	 * Get the current Movement data from the device
+	 *
+	 * @param Attitude The current Roll/Pitch/Yaw of the device
+	 * @param RotationRate The current rate of change of the attitude
+	 * @param Gravity A vector that describes the direction of gravity for the device
+	 * @param Acceleration returns the current acceleration of the device
+	 */
+	void GetMovementData(FVector& Attitude, FVector& RotationRate, FVector& Gravity, FVector& Acceleration);
+
+	/**
+	 * Calibrate the devices motion
+	 */
+	void CalibrateMotion(uint32 PlayerIndex);
+
+private:
+	void ProcessTouchesAndKeys(uint32 ControllerId);
+
+
+	static TArray<TouchInput> TouchInputStack;
+	static TArray<int32> KeyInputStack;
+
+	// protects the input stack used on 2 threads
+	static FCriticalSection CriticalSection;
+
+	TSharedRef< FGenericApplicationMessageHandler > MessageHandler;
+
+
+	/** Game controller objects (per user)*/
+	struct FUserController
+	{
+		GCGamepadSnapshot* PreviousGamepad;
+		GCExtendedGamepadSnapshot* PreviousExtendedGamepad;
+#if PLATFORM_TVOS
+		GCMicroGamepadSnapshot* PreviousMicroGamepad;
+#endif
+		FQuat ReferenceAttitude;
+		bool bNeedsReferenceAttitude;
+		bool bHasReferenceAttitude;
+		bool bIsGamepadConnected;
+		bool bIsRemoteConnected;
+		bool bPauseWasPressed;
+	};
+	// there is a hardcoded limit of 4 controllers in the API
+	FUserController Controllers[4];
+
+	// can the remote be rotated to landscape
+	bool bAllowRemoteRotation;
+
+	// is the remote treated as a separate controller?
+	bool bTreatRemoteAsSeparateController;
+
+	// should the remote be used as virtual joystick vs touch events
+	bool bUseRemoteAsVirtualJoystick;
+
+	// should the tracking use the pad center as the virtual joystick center?
+	bool bUseRemoteAbsoluteDpadValues;
+	
+#if !PLATFORM_TVOS
+	/** Access to the ios devices motion */
+	CMMotionManager* MotionManager;
+
+	/** Access to the ios devices tilt information */
+	CMAttitude* ReferenceAttitude;
+#endif
+	
+
+	/** Last frames roll, for calculating rate */
+	float LastRoll;
+
+	/** Last frames pitch, for calculating rate */
+	float LastPitch;
+
+	/** True if a calibration is requested */
+	bool bIsCalibrationRequested;
+
+	/** The center roll value for tilt calibration */
+	float CenterRoll;
+
+	/** The center pitch value for tilt calibration */
+	float CenterPitch;
+
+	/** When using just acceleration (without full motion) we store a frame of accel data to filter by */
+	FVector FilteredAccelerometer;
+};
diff --git a/Engine/Source/Runtime/Core/Public/iOS/IOSMallocZone.h b/Engine/Source/Runtime/Core/Public/iOS/IOSMallocZone.h
new file mode 100644
--- /dev/null
+++ b/Engine/Source/Runtime/Core/Public/iOS/IOSMallocZone.h
@@ -0,0 +1,68 @@
+// Copyright 1998-2017 Epic Games, Inc. All Rights Reserved.
+
+#pragma once
+
+#include "MemoryBase.h"
+
+typedef struct _malloc_zone_t malloc_zone_t;
+struct FIOSCrashContext;
+
+/*
+ * A malloc interface using a unique malloc zone
+ */
+class FIOSMallocZone : public FMalloc
+{
+public:
+	FIOSMallocZone( uint64 const InitialSize );
+	virtual ~FIOSMallocZone();
+	
+	// FMalloc interface.
+	virtual void* Malloc( SIZE_T Size, uint32 Alignment ) override;
+	
+	virtual void* Realloc( void* Ptr, SIZE_T NewSize, uint32 Alignment ) override;
+	
+	virtual void Free( void* Ptr ) override;
+	
+	virtual bool GetAllocationSize(void *Original, SIZE_T &SizeOut) override;
+	
+	virtual bool IsInternallyThreadSafe() const override;
+	
+	virtual bool ValidateHeap() override;
+	
+	virtual const TCHAR* GetDescriptiveName() override;
+	
+protected:
+	malloc_zone_t* MemoryZone;
+};
+
+/*
+ * Specific FIOSMallocZone for the crash handler malloc override so that we avoid malloc reentrancy problems
+ */
+class FIOSMallocCrashHandler : public FIOSMallocZone
+{
+public:
+	FIOSMallocCrashHandler( uint64 const InitialSize );
+	
+	virtual ~FIOSMallocCrashHandler();
+	
+	void Enable( FIOSCrashContext* Context, uint32 CrashedThreadId );
+	
+	// FMalloc interface.
+	virtual void* Malloc( SIZE_T Size, uint32 Alignment ) override;
+	
+	virtual void* Realloc( void* Ptr, SIZE_T NewSize, uint32 Alignment ) override;
+	
+	virtual void Free( void* Ptr ) override;
+	
+	virtual bool GetAllocationSize( void *Original, SIZE_T &SizeOut ) override;
+	
+	virtual const TCHAR* GetDescriptiveName() override;
+
+private:	
+	bool IsOnCrashedThread(void);
+	
+private:
+	FMalloc* OriginalHeap;
+	FIOSCrashContext* CrashContext;
+	int32 ThreadId;
+};
diff --git a/Engine/Source/Runtime/Core/Public/iOS/IOSPlatformAffinity.h b/Engine/Source/Runtime/Core/Public/iOS/IOSPlatformAffinity.h
new file mode 100644
--- /dev/null
+++ b/Engine/Source/Runtime/Core/Public/iOS/IOSPlatformAffinity.h
@@ -0,0 +1,19 @@
+// Copyright 1998-2017 Epic Games, Inc. All Rights Reserved.
+
+#pragma once
+
+#include "GenericPlatform/GenericPlatformAffinity.h"
+
+class CORE_API FIOSPlatformAffinity : public FGenericPlatformAffinity
+{
+public:
+	static const uint64 GetMainGameMask();
+	static const uint64 GetRenderingThreadMask();
+	static const uint64 GetRTHeartBeatMask();
+	static const uint64 GetPoolThreadMask();
+	static const uint64 GetTaskGraphThreadMask();
+	static const uint64 GetStatsThreadMask();
+	static const uint64 GetNoAffinityMask();
+};
+
+typedef FIOSPlatformAffinity FPlatformAffinity;
diff --git a/Engine/Source/Runtime/Core/Public/iOS/IOSPlatformCrashContext.h b/Engine/Source/Runtime/Core/Public/iOS/IOSPlatformCrashContext.h
new file mode 100644
--- /dev/null
+++ b/Engine/Source/Runtime/Core/Public/iOS/IOSPlatformCrashContext.h
@@ -0,0 +1,25 @@
+// Copyright 1998-2017 Epic Games, Inc. All Rights Reserved.
+
+#pragma once
+
+#include "Apple/ApplePlatformCrashContext.h"
+
+struct CORE_API FIOSCrashContext : public FApplePlatformCrashContext
+{
+	/** Mimics Windows WER format */
+	void GenerateWindowsErrorReport(char const* WERPath, bool bIsEnsure = false) const;
+
+	/** Copies the PLCrashReporter minidump */
+	void CopyMinidump(char const* OutputPath, char const* InputPath) const;
+
+	/** Generates the ensure/crash info into the given folder */
+	void GenerateInfoInFolder(char const* const InfoFolder, bool bIsEnsure = false) const;
+	
+	/** Generates information for crash reporter */
+	void GenerateCrashInfo() const;
+	
+	/** Generates information for ensures sent via the CrashReporter */
+	void GenerateEnsureInfo() const;
+};
+
+typedef FIOSCrashContext FPlatformCrashContext;
diff --git a/Engine/Source/Runtime/Core/Public/iOS/IOSPlatformString.h b/Engine/Source/Runtime/Core/Public/iOS/IOSPlatformString.h
new file mode 100644
--- /dev/null
+++ b/Engine/Source/Runtime/Core/Public/iOS/IOSPlatformString.h
@@ -0,0 +1,25 @@
+// Copyright 1998-2017 Epic Games, Inc. All Rights Reserved.
+
+/*================================================================================
+	IOSPlatformString.h: Additional NSString methods
+==================================================================================*/
+
+#pragma once
+#include "CoreTypes.h"
+#import <Foundation/NSString.h>
+
+class FString;
+
+@interface NSString (FString_Extensions)
+
+/**
+* Converts an TCHAR string to an NSString
+*/
++ (NSString*) stringWithTCHARString:(const TCHAR*)MyTCHARString;
+
+/**
+ * Converts an FString to an NSString
+ */
++ (NSString*) stringWithFString:(const FString&)MyFString;
+
+@end
diff --git a/Engine/Source/Runtime/Core/Public/iOS/IOSSystemIncludes.h b/Engine/Source/Runtime/Core/Public/iOS/IOSSystemIncludes.h
new file mode 100644
--- /dev/null
+++ b/Engine/Source/Runtime/Core/Public/iOS/IOSSystemIncludes.h
@@ -0,0 +1,81 @@
+// Copyright 1998-2017 Epic Games, Inc. All Rights Reserved.
+#pragma once
+
+// Set up compiler pragmas, etc
+#include "CoreTypes.h"
+#include "IOS/IOSPlatformCompilerSetup.h"
+
+#define FVector FVectorWorkaround
+#ifdef __OBJC__
+#import <UIKit/UIKit.h>
+#import <CoreData/CoreData.h>
+#else //@todo - zombie - iOS should always have Obj-C to function. This definition guard may be unnecessary. -astarkey
+#error "Objective C is undefined."
+#endif
+#undef check
+#undef verify
+#undef FVector
+
+#undef TYPE_BOOL
+
+#include <CoreFoundation/CoreFoundation.h>
+
+#include <string.h>
+#include <alloca.h>
+#include <stdlib.h>
+#include <stdio.h>
+#include <float.h>
+#include <ctype.h>
+#include <stdarg.h>
+#include <stdint.h>
+#include <limits.h>
+#include <sys/time.h>
+#include <math.h>
+#include <mach/mach_time.h>
+#include <wchar.h>
+#include <wctype.h>
+#include <errno.h>
+#include <unistd.h>
+#include <pthread.h>
+#include <libkern/OSAtomic.h>
+#include <sys/types.h>
+#include <sys/stat.h>
+#include <pwd.h>
+#include <fcntl.h>
+#include <dirent.h>
+#include <dlfcn.h>
+#include <copyfile.h>
+#include <utime.h>
+#include <mach/mach.h>
+#include <mach/mach_host.h>
+#include <execinfo.h>
+#include <signal.h>
+#include <sys/sysctl.h>
+#include <malloc/malloc.h>
+
+// SIMD intrinsics
+#if WITH_SIMULATOR
+#include <xmmintrin.h>
+#else
+#include <arm_neon.h>
+#endif
+
+struct tagRECT
+{
+	int32 left;
+	int32 top;
+	int32 right;
+	int32 bottom;
+};
+typedef struct tagRECT RECT;
+
+#define OUT
+#define IN
+
+/*----------------------------------------------------------------------------
+ Memory. On Mac OS X malloc allocates memory aligned to 16 bytes.
+ ----------------------------------------------------------------------------*/
+#define _aligned_malloc(Size,Align) malloc(Size)
+#define _aligned_realloc(Ptr,Size,Align) realloc(Ptr,Size)
+#define _aligned_free(Ptr) free(Ptr)
+#define _aligned_msize(Ptr,Align,Offset) malloc_size(Ptr)
