Copying //UE4/Dev-Automation to //UE4/Dev-Main (Source: //UE4/Dev-Automation @ 3502919)

=====================================
 MAJOR FEATURES + CHANGES
=====================================

Change 3487931 by [REDACTED_PERSON]

	Add Negative test filter

	Filter out tests that test negative cases and are expected to fail

	#tests ran affected tests in editor, confirmed UI changes

Change 3491768 by [REDACTED_PERSON]

	Add FieldOfView to Screenshot FTest options

	#tests used in WIP test map

Change 3492809 by [REDACTED_PERSON]

	Automation - Fixing some delay screenshot code in the Screenshot Functional Test Actor.

Change 3493557 by [REDACTED_PERSON]

	Add automated postprocessing test map

Change 3494836 by [REDACTED_PERSON]

	Fix CurvePlayTest timing out

	Time-out tolerances were set too tightly. Extended preparation time limit to 7 seconds.

	#tests PIE

Change 3494962 by [REDACTED_PERSON]

	Fix CurveStopTest failure

	Added OnTestPrepare event for FTests using LevelBP for logic.
	Removed CurveStopTest dependency on CurveStartTest.
	Updated ground truth screenshot.

	#tests PIE

Change 3495166 by [REDACTED_PERSON]

	Fix empty engine version content warnings

	#tests none

Change 3495412 by [REDACTED_PERSON]

	Add Shield asset from ContentExamples project

Change 3495747 by [REDACTED_PERSON]

	Automate static light tests

Change 3495914 by [REDACTED_PERSON]

	Adding HDR environment asset

Change 3496341 by [REDACTED_PERSON]

	Fix NotifyTest timing failures

	Rewrote tests to use blueprint events to verify notifications instead of fragile screenshots.

	#tests UFE

Change 3496613 by [REDACTED_PERSON]

	Add TranslucencyLightingModes test map

	Automates test cases from QAGame's TM-TranslucencyLightingModes.
	Also resave shield assets to fix empty engine version warnings.

	#tests UFE

Change 3496703 by [REDACTED_PERSON]

	Fix occasional ClothBallTest failures

	Extend screenshot delay to let cloth to fully settle.
	Consolidated screenshot delay values and removed redundant Delays.

	#tests UFE, 20+ iterations

Change 3498413 by [REDACTED_PERSON]

	Fix occasional CurveTest failures

	Merged tests and set to use specific anim positions instead of relying on fragile test timing.

	#tests UFE

Change 3498976 by [REDACTED_PERSON]

	Add automated ShaderModels test map

Change 3499079 by [REDACTED_PERSON]

	Delete unused test and related assets

	#tests UFE

Change 3500869 by [REDACTED_PERSON]

	Add horizontal scrollbar to automation results view

	#tests UFE

diff --git a/Engine/Plugins/Tests/RuntimeTests/Source/RuntimeTests/Private/AutomationExpectedError.spec.cpp b/Engine/Plugins/Tests/RuntimeTests/Source/RuntimeTests/Private/AutomationExpectedError.spec.cpp
--- a/Engine/Plugins/Tests/RuntimeTests/Source/RuntimeTests/Private/AutomationExpectedError.spec.cpp
+++ b/Engine/Plugins/Tests/RuntimeTests/Source/RuntimeTests/Private/AutomationExpectedError.spec.cpp
@@ -120,7 +120,7 @@ void FAutomationExpectedErrorTest::Define()
 // IMPORTANT: The pass condition for these tests is that they FAIL. To prevent
 // the expected failures from interfering with regular test runs, these tests
 // must be run manually.
-BEGIN_DEFINE_SPEC(FAutomationExpectedErrorFailureTest, "System.Automation.ExpectedError", EAutomationTestFlags::EngineFilter | EAutomationTestFlags::ApplicationContextMask | EAutomationTestFlags::RequiresUser)
+BEGIN_DEFINE_SPEC(FAutomationExpectedErrorFailureTest, "System.Automation.ExpectedError", EAutomationTestFlags::NegativeFilter | EAutomationTestFlags::ApplicationContextMask | EAutomationTestFlags::RequiresUser)
 END_DEFINE_SPEC(FAutomationExpectedErrorFailureTest)
 
 void FAutomationExpectedErrorFailureTest::Define()
diff --git a/Engine/Source/Developer/AutomationWindow/Private/SAutomationWindow.cpp b/Engine/Source/Developer/AutomationWindow/Private/SAutomationWindow.cpp
--- a/Engine/Source/Developer/AutomationWindow/Private/SAutomationWindow.cpp
+++ b/Engine/Source/Developer/AutomationWindow/Private/SAutomationWindow.cpp
@@ -262,6 +262,7 @@ void SAutomationWindow::Construct( const FArguments& InArgs, const IAutomationCo
 	RequestedFilterComboList.Add(MakeShareable(new FString(TEXT("Performance Tests"))));
 	RequestedFilterComboList.Add(MakeShareable(new FString(TEXT("Stress Tests"))));
 	RequestedFilterComboList.Add(MakeShareable(new FString(TEXT("Standard Tests"))));
+	RequestedFilterComboList.Add(MakeShareable(new FString(TEXT("Negative Tests"))));
 
 	TSharedRef<SNotificationList> NotificationList = SNew(SNotificationList) .Visibility( EVisibility::HitTestInvisible );
 
@@ -471,12 +472,17 @@ void SAutomationWindow::Construct( const FArguments& InArgs, const IAutomationCo
 									SNew(SBorder)
 									.BorderImage(FEditorStyle::GetBrush("MessageLog.ListBorder"))
 									[
-										SAssignNew(LogListView, SListView<TSharedPtr<FAutomationOutputMessage> >)
-										.ItemHeight(18)
-										.ListItemsSource(&LogMessages)
-										.SelectionMode(ESelectionMode::Multi)
-										.OnGenerateRow(this, &SAutomationWindow::OnGenerateWidgetForLog)
-										.OnSelectionChanged(this, &SAutomationWindow::HandleLogListSelectionChanged)
+										SNew(SScrollBox)
+										.Orientation(EOrientation::Orient_Horizontal)
+										+SScrollBox::Slot()
+										[
+											SAssignNew(LogListView, SListView<TSharedPtr<FAutomationOutputMessage> >)
+											.ItemHeight(18)
+											.ListItemsSource(&LogMessages)
+											.SelectionMode(ESelectionMode::Multi)
+											.OnGenerateRow(this, &SAutomationWindow::OnGenerateWidgetForLog)
+											.OnSelectionChanged(this, &SAutomationWindow::HandleLogListSelectionChanged)
+										]
 									]
 								]
 
@@ -963,6 +969,9 @@ void SAutomationWindow::HandleRequesteFilterChanged(TSharedPtr<FString> Item, ES
 		case 6:	//	"Standard Tests"
 			NewRequestedFlags = EAutomationTestFlags::SmokeFilter | EAutomationTestFlags::EngineFilter | EAutomationTestFlags::ProductFilter | EAutomationTestFlags::PerfFilter;
 			break;
+		case 7: //  "Negative Tests"
+			NewRequestedFlags = EAutomationTestFlags::NegativeFilter;
+			break;
 	}
 	AutomationController->SetRequestedTestFlags(NewRequestedFlags);
 }
diff --git a/Engine/Source/Developer/FunctionalTesting/Classes/FunctionalTest.h b/Engine/Source/Developer/FunctionalTesting/Classes/FunctionalTest.h
--- a/Engine/Source/Developer/FunctionalTesting/Classes/FunctionalTest.h
+++ b/Engine/Source/Developer/FunctionalTesting/Classes/FunctionalTest.h
@@ -332,6 +332,10 @@ public:
 
 public:
 
+	/** Called when the test is ready to prepare */
+	UPROPERTY(BlueprintAssignable)
+	FFunctionalTestEventSignature OnTestPrepare;
+
 	/** Called when the test is started */
 	UPROPERTY(BlueprintAssignable)
 	FFunctionalTestEventSignature OnTestStart;
@@ -637,6 +641,7 @@ public:
 	float TotalTime;
 
 	uint32 RunFrame;
+	float RunTime;
 
 	uint32 StartFrame;
 	float StartTime;
diff --git a/Engine/Source/Developer/FunctionalTesting/Classes/ScreenshotFunctionalTest.h b/Engine/Source/Developer/FunctionalTesting/Classes/ScreenshotFunctionalTest.h
--- a/Engine/Source/Developer/FunctionalTesting/Classes/ScreenshotFunctionalTest.h
+++ b/Engine/Source/Developer/FunctionalTesting/Classes/ScreenshotFunctionalTest.h
@@ -33,6 +33,10 @@ protected:
 	void OnScreenshotTakenAndCompared();
 	
 protected:
+
+	/** The horizontal field of view (in degrees) for the screenshot camera */
+	UPROPERTY(EditANywhere, Category="Screenshot", meta=(UIMin = "5.0", UIMax = "170", ClampMin = "0.001", ClampMax = "360.0", Units = deg))
+	float FieldOfView;
 	
 	UPROPERTY()
 	class UCameraComponent* ScreenshotCamera;
diff --git a/Engine/Source/Developer/FunctionalTesting/Private/FunctionalTest.cpp b/Engine/Source/Developer/FunctionalTesting/Private/FunctionalTest.cpp
--- a/Engine/Source/Developer/FunctionalTesting/Private/FunctionalTest.cpp
+++ b/Engine/Source/Developer/FunctionalTesting/Private/FunctionalTest.cpp
@@ -94,6 +94,7 @@ AFunctionalTest::AFunctionalTest( const FObjectInitializer& ObjectInitializer )
 	, bIsRunning(false)
 	, TotalTime(0.f)
 	, RunFrame(0)
+	, RunTime(0.0f)
 	, StartFrame(0)
 	, StartTime(0.0f)
 	, bIsReady(false)
@@ -200,6 +201,8 @@ bool AFunctionalTest::RunTest(const TArray<FString>& Params)
 	GetWorld()->DelayGarbageCollection();
 
 	RunFrame = GFrameNumber;
+	RunTime = GetWorld()->GetTimeSeconds();
+
 	TotalTime = 0.f;
 	if (TimeLimit >= 0)
 	{
@@ -212,6 +215,7 @@ bool AFunctionalTest::RunTest(const TArray<FString>& Params)
 	GoToObservationPoint();
 
 	PrepareTest();
+	OnTestPrepare.Broadcast();
 
 	return true;
 }
diff --git a/Engine/Source/Developer/FunctionalTesting/Private/ScreenshotFunctionalTest.cpp b/Engine/Source/Developer/FunctionalTesting/Private/ScreenshotFunctionalTest.cpp
--- a/Engine/Source/Developer/FunctionalTesting/Private/ScreenshotFunctionalTest.cpp
+++ b/Engine/Source/Developer/FunctionalTesting/Private/ScreenshotFunctionalTest.cpp
@@ -17,6 +17,7 @@ AScreenshotFunctionalTest::AScreenshotFunctionalTest( const FObjectInitializer&
 {
 	ScreenshotCamera = CreateDefaultSubobject<UCameraComponent>(TEXT("Camera"));
 	ScreenshotCamera->SetupAttachment(RootComponent);
+	FieldOfView = ScreenshotCamera->FieldOfView;
 }
 
 void AScreenshotFunctionalTest::PrepareTest()
@@ -32,7 +33,7 @@ void AScreenshotFunctionalTest::PrepareTest()
 
 bool AScreenshotFunctionalTest::IsReady_Implementation()
 {
-	if ( (GetWorld()->GetTimeSeconds() - StartTime) > ScreenshotOptions.Delay )
+	if ( (GetWorld()->GetTimeSeconds() - RunTime) > ScreenshotOptions.Delay )
 	{
 		return ( GFrameNumber - RunFrame ) > 5;
 	}
@@ -91,6 +92,10 @@ void AScreenshotFunctionalTest::PostEditChangeProperty(FPropertyChangedEvent& Pr
 		{
 			ScreenshotOptions.SetToleranceAmounts(ScreenshotOptions.Tolerance);
 		}
+		else if (PropertyName == GET_MEMBER_NAME_CHECKED(AScreenshotFunctionalTest, FieldOfView))
+		{
+			ScreenshotCamera->SetFieldOfView(FieldOfView);
+		}
 	}
 }
 
diff --git a/Engine/Source/Developer/FunctionalTesting/Public/AutomationScreenshotOptions.h b/Engine/Source/Developer/FunctionalTesting/Public/AutomationScreenshotOptions.h
--- a/Engine/Source/Developer/FunctionalTesting/Public/AutomationScreenshotOptions.h
+++ b/Engine/Source/Developer/FunctionalTesting/Public/AutomationScreenshotOptions.h
@@ -140,7 +140,7 @@ public:
 	 * essentially the same.  Generally this is necessary as modern rendering techniques tend to introduce
 	 * noise constantly to hide aliasing.
 	 */
-	UPROPERTY(EditAnywhere, Category="Comparison")
+	UPROPERTY(EditAnywhere, BlueprintReadWrite, Category="Comparison")
 	FComparisonToleranceAmount ToleranceAmount;
 
 	/**
diff --git a/Engine/Source/Runtime/Core/Public/Misc/AutomationTest.h b/Engine/Source/Runtime/Core/Public/Misc/AutomationTest.h
--- a/Engine/Source/Runtime/Core/Public/Misc/AutomationTest.h
+++ b/Engine/Source/Runtime/Core/Public/Misc/AutomationTest.h
@@ -81,7 +81,9 @@ namespace EAutomationTestFlags
 		PerfFilter					= 0x08000000,
 		//Stress Test
 		StressFilter				= 0x10000000,
-		FilterMask = SmokeFilter | EngineFilter | ProductFilter | PerfFilter | StressFilter
+		//Negative Test. For tests whose correct expected outcome is failure.
+		NegativeFilter				= 0x20000000,
+		FilterMask = SmokeFilter | EngineFilter | ProductFilter | PerfFilter | StressFilter | NegativeFilter
 	};
 };
 
@@ -2367,7 +2369,8 @@ public: \
 							(((TFlags)&EAutomationTestFlags::FilterMask) == EAutomationTestFlags::EngineFilter) || \
 							(((TFlags)&EAutomationTestFlags::FilterMask) == EAutomationTestFlags::ProductFilter) || \
 							(((TFlags)&EAutomationTestFlags::FilterMask) == EAutomationTestFlags::PerfFilter) || \
-							(((TFlags)&EAutomationTestFlags::FilterMask) == EAutomationTestFlags::StressFilter), \
+							(((TFlags)&EAutomationTestFlags::FilterMask) == EAutomationTestFlags::StressFilter) || \
+							(((TFlags)&EAutomationTestFlags::FilterMask) == EAutomationTestFlags::NegativeFilter), \
 							"All AutomationTests must have exactly 1 filter type specified.  See AutomationTest.h."); \
 		} \
 		virtual uint32 GetTestFlags() const override { return TFlags; } \
@@ -2396,7 +2399,8 @@ public: \
 							(((TFlags)&EAutomationTestFlags::FilterMask) == EAutomationTestFlags::EngineFilter) || \
 							(((TFlags)&EAutomationTestFlags::FilterMask) == EAutomationTestFlags::ProductFilter) || \
 							(((TFlags)&EAutomationTestFlags::FilterMask) == EAutomationTestFlags::PerfFilter) || \
-							(((TFlags)&EAutomationTestFlags::FilterMask) == EAutomationTestFlags::StressFilter), \
+							(((TFlags)&EAutomationTestFlags::FilterMask) == EAutomationTestFlags::StressFilter) || \
+							(((TFlags)&EAutomationTestFlags::FilterMask) == EAutomationTestFlags::NegativeFilter), \
 							"All AutomationTests must have exactly 1 filter type specified.  See AutomationTest.h."); \
 		} \
 		virtual uint32 GetTestFlags() const override { return ((TFlags) & ~(EAutomationTestFlags::SmokeFilter)); } \
@@ -2421,7 +2425,8 @@ public: \
 							(((TFlags)&EAutomationTestFlags::FilterMask) == EAutomationTestFlags::EngineFilter) || \
 							(((TFlags)&EAutomationTestFlags::FilterMask) == EAutomationTestFlags::ProductFilter) || \
 							(((TFlags)&EAutomationTestFlags::FilterMask) == EAutomationTestFlags::PerfFilter) || \
-							(((TFlags)&EAutomationTestFlags::FilterMask) == EAutomationTestFlags::StressFilter), \
+							(((TFlags)&EAutomationTestFlags::FilterMask) == EAutomationTestFlags::StressFilter) || \
+							(((TFlags)&EAutomationTestFlags::FilterMask) == EAutomationTestFlags::NegativeFilter), \
 							"All AutomationTests must have exactly 1 filter type specified.  See AutomationTest.h."); \
 		} \
 		virtual uint32 GetTestFlags() const override { return ((TFlags) & ~(EAutomationTestFlags::EditorContext | EAutomationTestFlags::CommandletContext | EAutomationTestFlags::SmokeFilter)); } \
@@ -2449,7 +2454,8 @@ public: \
 							(((TFlags)&EAutomationTestFlags::FilterMask) == EAutomationTestFlags::EngineFilter) || \
 							(((TFlags)&EAutomationTestFlags::FilterMask) == EAutomationTestFlags::ProductFilter) || \
 							(((TFlags)&EAutomationTestFlags::FilterMask) == EAutomationTestFlags::PerfFilter) || \
-							(((TFlags)&EAutomationTestFlags::FilterMask) == EAutomationTestFlags::StressFilter), \
+							(((TFlags)&EAutomationTestFlags::FilterMask) == EAutomationTestFlags::StressFilter) || \
+							(((TFlags)&EAutomationTestFlags::FilterMask) == EAutomationTestFlags::NegativeFilter), \
 							"All AutomationTests must have exactly 1 filter type specified.  See AutomationTest.h."); \
 		} \
 		virtual uint32 GetTestFlags() const override { return TFlags; } \
@@ -2475,7 +2481,8 @@ public: \
 							(((TFlags)&EAutomationTestFlags::FilterMask) == EAutomationTestFlags::EngineFilter) || \
 							(((TFlags)&EAutomationTestFlags::FilterMask) == EAutomationTestFlags::ProductFilter) || \
 							(((TFlags)&EAutomationTestFlags::FilterMask) == EAutomationTestFlags::PerfFilter) || \
-							(((TFlags)&EAutomationTestFlags::FilterMask) == EAutomationTestFlags::StressFilter), \
+							(((TFlags)&EAutomationTestFlags::FilterMask) == EAutomationTestFlags::StressFilter) || \
+							(((TFlags)&EAutomationTestFlags::FilterMask) == EAutomationTestFlags::NegativeFilter), \
 							"All AutomationTests must have exactly 1 filter type specified.  See AutomationTest.h."); \
 		} \
 		virtual uint32 GetTestFlags() const override { return TFlags; } \
@@ -2501,7 +2508,8 @@ public: \
 							(((TFlags)&EAutomationTestFlags::FilterMask) == EAutomationTestFlags::EngineFilter) || \
 							(((TFlags)&EAutomationTestFlags::FilterMask) == EAutomationTestFlags::ProductFilter) || \
 							(((TFlags)&EAutomationTestFlags::FilterMask) == EAutomationTestFlags::PerfFilter) || \
-							(((TFlags)&EAutomationTestFlags::FilterMask) == EAutomationTestFlags::StressFilter), \
+							(((TFlags)&EAutomationTestFlags::FilterMask) == EAutomationTestFlags::StressFilter) || \
+							(((TFlags)&EAutomationTestFlags::FilterMask) == EAutomationTestFlags::NegativeFilter), \
 							"All AutomationTests must have exactly 1 filter type specified.  See AutomationTest.h."); \
 		} \
 		virtual uint32 GetTestFlags() const override { return TFlags; } \
