Mac: Fix double-printing of log entries for console apps (eg. cooking) on Mac, and suppress output when running UHT.

diff --git a/Engine/Source/Runtime/Core/Private/Mac/MacPlatformMisc.cpp b/Engine/Source/Runtime/Core/Private/Mac/MacPlatformMisc.cpp
--- a/Engine/Source/Runtime/Core/Private/Mac/MacPlatformMisc.cpp
+++ b/Engine/Source/Runtime/Core/Private/Mac/MacPlatformMisc.cpp
@@ -1653,7 +1653,7 @@ bool FMacPlatformMisc::GetDiskTotalAndFreeSpace(const FString& InPath, uint64& T
 
 bool FMacPlatformMisc::HasSeparateChannelForDebugOutput()
 {
-	return true;
+	return false;
 }
 
 #include "ModuleManager.h"
