Re-adding IOS files in directory with the correct case.

diff --git a/Engine/Source/Runtime/Core/Private/IOS/IOSAppDelegate.cpp b/Engine/Source/Runtime/Core/Private/IOS/IOSAppDelegate.cpp
new file mode 100644
--- /dev/null
+++ b/Engine/Source/Runtime/Core/Private/IOS/IOSAppDelegate.cpp
@@ -0,0 +1,1123 @@
+// Copyright 1998-2017 Epic Games, Inc. All Rights Reserved.
+
+#include "IOSAppDelegate.h"
+#include "IOSCommandLineHelper.h"
+#include "ExceptionHandling.h"
+#include "ModuleBoilerplate.h"
+#include "CallbackDevice.h"
+#include "IOSView.h"
+#include "TaskGraphInterfaces.h"
+#include "GenericPlatformChunkInstall.h"
+#include "IOSPlatformMisc.h"
+#include "HAL/PlatformStackWalk.h"
+#include "HAL/PlatformProcess.h"
+#include "Misc/OutputDeviceError.h"
+#include "Misc/CommandLine.h"
+#include "IOS/IOSPlatformFramePacer.h"
+#include "IOS/IOSAsyncTask.h"
+#include "Misc/ConfigCacheIni.h"
+#include "IOSPlatformCrashContext.h"
+#include "Misc/OutputDeviceError.h"
+#include "Misc/OutputDeviceRedirector.h"
+#include "Misc/FeedbackContext.h"
+
+#include <AudioToolbox/AudioToolbox.h>
+#include <AVFoundation/AVAudioSession.h>
+
+// this is the size of the game thread stack, it must be a multiple of 4k
+#if (UE_BUILD_SHIPPING || UE_BUILD_TEST)
+#define GAME_THREAD_STACK_SIZE 1024 * 1024
+#else
+#define GAME_THREAD_STACK_SIZE 8 * 1024 * 1024
+#endif
+
+DEFINE_LOG_CATEGORY(LogIOSAudioSession);
+DECLARE_LOG_CATEGORY_EXTERN(LogEngine, Log, All);
+
+extern bool GShowSplashScreen;
+
+FIOSCoreDelegates::FOnOpenURL FIOSCoreDelegates::OnOpenURL;
+
+/*
+	From: https://developer.apple.com/library/ios/documentation/UIKit/Reference/UIApplicationDelegate_Protocol/#//apple_ref/occ/intfm/UIApplicationDelegate/applicationDidEnterBackground:
+	"In practice, you should return from applicationDidEnterBackground: as quickly as possible. If the method does not return before time runs out your app is terminated and purged from memory."
+*/
+const double cMaxThreadWaitTime = 2.0;	// Setting this to be 2 seconds since this wait has to be done twice (once for sending the enter background event to the game thread, and another for waiting on the suspend msg
+										// I could not find a reference for this but in the past I believe the timeout was 5 seconds
+
+static void SignalHandler(int32 Signal, struct __siginfo* Info, void* Context)
+{
+	static int32 bHasEntered = 0;
+	if (FPlatformAtomics::InterlockedCompareExchange(&bHasEntered, 1, 0) == 0)
+	{
+		const SIZE_T StackTraceSize = 65535;
+		ANSICHAR* StackTrace = (ANSICHAR*)FMemory::Malloc(StackTraceSize);
+		StackTrace[0] = 0;
+		
+		// Walk the stack and dump it to the allocated memory.
+		FPlatformStackWalk::StackWalkAndDump(StackTrace, StackTraceSize, 0, (ucontext_t*)Context);
+		UE_LOG(LogEngine, Error, TEXT("%s"), ANSI_TO_TCHAR(StackTrace));
+		FMemory::Free(StackTrace);
+		
+		GError->HandleError();
+		FPlatformMisc::RequestExit(true);
+	}
+}
+
+void InstallSignalHandlers()
+{
+	struct sigaction Action;
+	FMemory::Memzero(&Action, sizeof(struct sigaction));
+	
+	Action.sa_sigaction = SignalHandler;
+	sigemptyset(&Action.sa_mask);
+	Action.sa_flags = SA_SIGINFO | SA_RESTART | SA_ONSTACK;
+	
+	sigaction(SIGQUIT, &Action, NULL);
+	sigaction(SIGILL, &Action, NULL);
+	sigaction(SIGEMT, &Action, NULL);
+	sigaction(SIGFPE, &Action, NULL);
+	sigaction(SIGBUS, &Action, NULL);
+	sigaction(SIGSEGV, &Action, NULL);
+	sigaction(SIGSYS, &Action, NULL);
+}
+
+void EngineCrashHandler(const FGenericCrashContext& GenericContext)
+{
+    const FIOSCrashContext& Context = static_cast<const FIOSCrashContext&>(GenericContext);
+    
+    Context.ReportCrash();
+    if (GLog)
+    {
+        GLog->SetCurrentThreadAsMasterThread();
+        GLog->Flush();
+    }
+    if (GWarn)
+    {
+        GWarn->Flush();
+    }
+    if (GError)
+    {
+        GError->Flush();
+        GError->HandleError();
+    }
+    return Context.GenerateCrashInfo();
+}
+
+@implementation IOSAppDelegate
+
+#if !UE_BUILD_SHIPPING && !PLATFORM_TVOS
+#if __IPHONE_OS_VERSION_MIN_REQUIRED < __IPHONE_9_0
+	@synthesize ConsoleAlert;
+#endif
+#ifdef __IPHONE_8_0
+    @synthesize ConsoleAlertController;
+#endif
+	@synthesize ConsoleHistoryValues;
+	@synthesize ConsoleHistoryValuesIndex;
+#endif
+
+@synthesize AlertResponse;
+@synthesize bDeviceInPortraitMode;
+@synthesize bEngineInit;
+@synthesize OSVersion;
+
+@synthesize Window;
+@synthesize IOSView;
+@synthesize IOSController;
+@synthesize SlateController;
+@synthesize timer;
+@synthesize IdleTimerEnableTimer;
+@synthesize IdleTimerEnablePeriod;
+
+-(void) ParseCommandLineOverrides
+{
+	//// Check for device type override
+	//FString IOSDeviceName;
+	//if ( Parse( appCmdLine(), TEXT("-IOSDevice="), IOSDeviceName) )
+	//{
+	//	for ( int32 DeviceTypeIndex = 0; DeviceTypeIndex < IOS_Unknown; ++DeviceTypeIndex )
+	//	{
+	//		if ( IOSDeviceName == IPhoneGetDeviceTypeString( (EIOSDevice) DeviceTypeIndex) )
+	//		{
+	//			GOverrideIOSDevice = (EIOSDevice) DeviceTypeIndex;
+	//		}
+	//	}
+	//}
+
+	//// Check for iOS version override
+	//FLOAT IOSVersion = 0.0f;
+	//if ( Parse( appCmdLine(), TEXT("-IOSVersion="), IOSVersion) )
+	//{
+	//	GOverrideIOSVersion = IOSVersion;
+	//}
+
+	// check to see if we are using the network file system, if so, disable the idle timer
+	FString HostIP;
+	if (FParse::Value(FCommandLine::Get(), TEXT("-FileHostIP="), HostIP))
+	{
+		[UIApplication sharedApplication].idleTimerDisabled = YES;
+	}
+}
+
+-(void)MainAppThread:(NSDictionary*)launchOptions
+{
+    self.bHasStarted = true;
+	GIsGuarded = false;
+	GStartTime = FPlatformTime::Seconds();
+
+	// make sure this thread has an auto release pool setup 
+	NSAutoreleasePool* AutoreleasePool = [[NSAutoreleasePool alloc] init];
+
+	while(!self.bCommandLineReady)
+	{
+		usleep(100);
+	}
+
+
+
+	// Look for overrides specified on the command-line
+	[self ParseCommandLineOverrides];
+
+	FAppEntry::Init();
+	
+	[self InitIdleTimerSettings];
+
+	bEngineInit = true;
+	// @PJS - need a better way to allow the game to turn off the splash screen
+//    GShowSplashScreen = false;
+
+	while( !GIsRequestingExit )
+	{
+        if (self.bIsSuspended)
+        {
+            FAppEntry::SuspendTick();
+            
+            self.bHasSuspended = true;
+        }
+        else
+        {
+            FAppEntry::Tick();
+        
+            // free any autoreleased objects every once in awhile to keep memory use down (strings, splash screens, etc)
+            if (((GFrameCounter) & 31) == 0)
+            {
+                [AutoreleasePool release];
+                AutoreleasePool = [[NSAutoreleasePool alloc] init];
+            }
+        }
+
+        // drain the async task queue from the game thread
+        [FIOSAsyncTask ProcessAsyncTasks];
+	}
+
+	[UIApplication sharedApplication].idleTimerDisabled = NO;
+
+	[AutoreleasePool release];
+	FAppEntry::Shutdown();
+    
+    self.bHasStarted = false;
+}
+
+-(void)timerForSplashScreen
+{
+    if (!GShowSplashScreen)
+    {
+        if ([self.Window viewWithTag:2] != nil)
+        {
+            [[self.Window viewWithTag:2] removeFromSuperview];
+        }
+        [timer invalidate];
+    }
+}
+
+-(void)RecordPeakMemory
+{
+    FIOSPlatformMemory::GetStats();
+}
+
+-(void)InitIdleTimerSettings
+{
+	float TimerDuration = 0.0F;
+	GConfig->GetFloat(TEXT("/Script/IOSRuntimeSettings.IOSRuntimeSettings"), TEXT("IdleTimerEnablePeriod"), TimerDuration, GEngineIni);
+	IdleTimerEnablePeriod = TimerDuration;
+	
+	self.IdleTimerEnableTimer = nil;
+}
+
+-(void)DeferredEnableIdleTimer
+{
+	[UIApplication sharedApplication].idleTimerDisabled = NO;
+	self.IdleTimerEnableTimer = nil;
+}
+
+-(void)EnableIdleTimer:(bool)bEnabled
+{
+	dispatch_async(dispatch_get_main_queue(),^
+	{
+		if (bEnabled)
+		{
+			// Nothing needs to be done, if the enable timer is already running.
+			if (self.IdleTimerEnableTimer == nil)
+			{
+				self.IdleTimerEnableTimer = [NSTimer scheduledTimerWithTimeInterval:IdleTimerEnablePeriod target:self selector:@selector(DeferredEnableIdleTimer) userInfo:nil repeats:NO];
+			}
+		}
+		else
+		{
+			// Ensure pending attempts to enable the idle timer are cancelled.
+			if (self.IdleTimerEnableTimer != nil)
+			{
+				[self.IdleTimerEnableTimer invalidate];
+				self.IdleTimerEnableTimer = nil;
+			}
+
+			[UIApplication sharedApplication].idleTimerDisabled = NO;
+			[UIApplication sharedApplication].idleTimerDisabled = YES;
+		}
+	});
+}
+
+-(void)NoUrlCommandLine
+{
+	//Since it is non-repeating, the timer should kill itself.
+	self.bCommandLineReady = true;
+}
+
+-(void)AudioInterrupted:(NSNotification*)notification
+{
+	// pull the type of notification out
+	NSDictionary* Info = [notification userInfo];
+	NSNumber* Type = (NSNumber*)[Info valueForKey:AVAudioSessionInterruptionTypeKey];
+
+	NSLog(@"AUDIO INTERRUPTION NOTIFICATION: %d, (began = %d, ended = %d)", (int)[Type integerValue], (int)AVAudioSessionInterruptionTypeBegan, (int)AVAudioSessionInterruptionTypeEnded);
+	if ([Type integerValue] == AVAudioSessionInterruptionTypeBegan)
+	{
+		FAppEntry::Suspend();
+		[self ToggleAudioSession:false];
+	}
+	else if ([Type integerValue] == AVAudioSessionInterruptionTypeEnded)
+	{
+		[self ToggleAudioSession:true];
+	    FAppEntry::Resume();
+	}
+}
+
+- (void)InitializeAudioSession
+{
+	// get notified about interruptions
+	// @todo ios8: Test this (I can't make a notification happen so far)
+	[[NSNotificationCenter defaultCenter] addObserver:self selector:@selector(AudioInterrupted:) name:AVAudioSessionInterruptionNotification object:[AVAudioSession sharedInstance]];
+	
+	self.bUsingBackgroundMusic = [self IsBackgroundAudioPlaying];
+	if (!self.bUsingBackgroundMusic)
+	{
+		NSError* ActiveError = nil;
+		[[AVAudioSession sharedInstance] setActive:YES error:&ActiveError];
+		if (ActiveError)
+		{
+			UE_LOG(LogIOSAudioSession, Error, TEXT("Failed to set audio session as active!"));
+		}
+		ActiveError = nil;
+		
+		[[AVAudioSession sharedInstance] setCategory:AVAudioSessionCategorySoloAmbient error:&ActiveError];
+		if (ActiveError)
+		{
+			UE_LOG(LogIOSAudioSession, Error, TEXT("Failed to set audio session category to AVAudioSessionCategorySoloAmbient!"));
+		}
+		ActiveError = nil;
+	}
+	else
+	{
+		// Allow iPod music to continue playing in the background
+		NSError* ActiveError = nil;
+		[[AVAudioSession sharedInstance] setCategory:AVAudioSessionCategoryAmbient error:&ActiveError];
+		if (ActiveError)
+		{
+			UE_LOG(LogIOSAudioSession, Error, TEXT("Failed to set audio session category to AVAudioSessionCategoryAmbient!"));
+		}
+		ActiveError = nil;
+	}
+
+    self.bAudioActive = true;
+}
+
+- (void)ToggleAudioSession:(bool)bActive
+{
+	if (bActive)
+	{
+        if (!self.bAudioActive)
+        {
+			bool bWasUsingBackgroundMusic = self.bUsingBackgroundMusic;
+			self.bUsingBackgroundMusic = [self IsBackgroundAudioPlaying];
+	
+			if (bWasUsingBackgroundMusic != self.bUsingBackgroundMusic)
+			{
+				if (!self.bUsingBackgroundMusic)
+				{
+					NSError* ActiveError = nil;
+					[[AVAudioSession sharedInstance] setActive:YES error:&ActiveError];
+					if (ActiveError)
+					{
+						UE_LOG(LogIOSAudioSession, Error, TEXT("Failed to set audio session as active! [Error = %s]"), *FString([ActiveError description]));
+					}
+					ActiveError = nil;
+	
+					[[AVAudioSession sharedInstance] setCategory:AVAudioSessionCategorySoloAmbient error:&ActiveError];
+					if (ActiveError)
+					{
+						UE_LOG(LogIOSAudioSession, Error, TEXT("Failed to set audio session category to AVAudioSessionCategorySoloAmbient! [Error = %s]"), *FString([ActiveError description]));
+					}
+					ActiveError = nil;
+	
+					/* TODO::JTM - Jan 16, 2013 05:05PM - Music player support */
+				}
+				else
+				{
+					/* TODO::JTM - Jan 16, 2013 05:05PM - Music player support */
+	
+					// Allow iPod music to continue playing in the background
+					NSError* ActiveError = nil;
+					[[AVAudioSession sharedInstance] setCategory:AVAudioSessionCategoryAmbient error:&ActiveError];
+					if (ActiveError)
+					{
+						UE_LOG(LogIOSAudioSession, Error, TEXT("Failed to set audio session category to AVAudioSessionCategoryAmbient! [Error = %s]"), *FString([ActiveError description]));
+					}
+					ActiveError = nil;
+				}
+			}
+			else if (!self.bUsingBackgroundMusic)
+			{
+				NSError* ActiveError = nil;
+				[[AVAudioSession sharedInstance] setActive:YES error:&ActiveError];
+				if (ActiveError)
+				{
+					UE_LOG(LogIOSAudioSession, Error, TEXT("Failed to set audio session as active! [Error = %s]"), *FString([ActiveError description]));
+				}
+				ActiveError = nil;
+				
+				[[AVAudioSession sharedInstance] setCategory:AVAudioSessionCategorySoloAmbient error:&ActiveError];
+				if (ActiveError)
+				{
+					UE_LOG(LogIOSAudioSession, Error, TEXT("Failed to set audio session category to AVAudioSessionCategorySoloAmbient! [Error = %s]"), *FString([ActiveError description]));
+				}
+				ActiveError = nil;
+			}
+        }
+	}
+	else if (self.bAudioActive && !self.bUsingBackgroundMusic)
+	{
+		NSError* ActiveError = nil;
+		[[AVAudioSession sharedInstance] setActive:NO error:&ActiveError];
+		if (ActiveError)
+		{
+			UE_LOG(LogIOSAudioSession, Error, TEXT("Failed to set audio session as inactive! [Error = %s]"), *FString([ActiveError description]));
+		}
+		ActiveError = nil;
+        
+		// Necessary to prevent audio from getting killing when setup for background iPod audio playback
+		[[AVAudioSession sharedInstance] setCategory:AVAudioSessionCategoryAmbient error:&ActiveError];
+		if (ActiveError)
+		{
+			UE_LOG(LogIOSAudioSession, Error, TEXT("Failed to set audio session category to AVAudioSessionCategoryAmbient! [Error = %s]"), *FString([ActiveError description]));
+		}
+		ActiveError = nil;
+ 	}
+    self.bAudioActive = bActive;
+}
+
+- (bool)IsBackgroundAudioPlaying
+{
+	AVAudioSession* Session = [AVAudioSession sharedInstance];
+	return Session.otherAudioPlaying;
+}
+
+- (int)GetAudioVolume
+{
+	float vol = [[AVAudioSession sharedInstance] outputVolume];
+	int roundedVol = (int)((vol * 100.0f) + 0.5f);
+	return roundedVol;
+}
+
+- (bool)AreHeadphonesPluggedIn
+{
+	AVAudioSessionRouteDescription *route = [[AVAudioSession sharedInstance] currentRoute];
+
+	bool headphonesFound = false;
+	for (AVAudioSessionPortDescription *portDescription in route.outputs)
+	{
+		//compare to the iOS constant for headphones
+		if ([portDescription.portType isEqualToString : AVAudioSessionPortHeadphones])
+		{
+			headphonesFound = true;
+			break;
+		}
+	}
+	return headphonesFound;
+}
+
+- (int)GetBatteryLevel
+{
+#if PLATFORM_TVOS
+	// TVOS does not have a battery, return fully charged
+	return 100;
+#else
+	UIDevice* Device = [UIDevice currentDevice];
+	Device.batteryMonitoringEnabled = YES;
+	
+	// Battery level is from 0.0 to 1.0, get it in terms of 0-100
+	int Level = ((int)([Device batteryLevel] * 100));
+	
+	Device.batteryMonitoringEnabled = NO;
+	return Level;
+#endif
+}
+
+- (bool)IsRunningOnBattery
+{
+#if PLATFORM_TVOS
+	// TVOS does not have a battery, return plugged in
+	return false;
+#else
+	UIDevice* Device = [UIDevice currentDevice];
+	Device.batteryMonitoringEnabled = YES;
+
+	UIDeviceBatteryState State = Device.batteryState;
+
+	Device.batteryMonitoringEnabled = NO;
+	return State == UIDeviceBatteryStateUnplugged || State == UIDeviceBatteryStateUnknown;
+#endif
+}
+
+/**
+ * @return the single app delegate object
+ */
++ (IOSAppDelegate*)GetDelegate
+{
+	return (IOSAppDelegate*)[UIApplication sharedApplication].delegate;
+}
+
+- (void)ToggleSuspend:(bool)bSuspend
+{
+    self.bHasSuspended = !bSuspend;
+    self.bIsSuspended = bSuspend;
+
+	if (bSuspend)
+	{
+		FAppEntry::Suspend();
+	}
+	else
+	{
+        FIOSPlatformRHIFramePacer::Resume();
+		FAppEntry::Resume();
+	}
+    
+	if (IOSView && IOSView->bIsInitialized)
+	{
+		// Don't deadlock here because a msg box may appear super early blocking the game thread and then the app may go into the background
+		double	startTime = FPlatformTime::Seconds();
+		while(!self.bHasSuspended && (FPlatformTime::Seconds() - startTime) < cMaxThreadWaitTime)
+		{
+            FIOSPlatformRHIFramePacer::Suspend();
+			FPlatformProcess::Sleep(0.05f);
+		}
+	}
+}
+
+- (BOOL)application:(UIApplication *)application willFinishLaunchingWithOptions:(NSDictionary *)launchOptions
+{
+	self.bDeviceInPortraitMode = false;
+	bEngineInit = false;
+
+	return YES;
+}
+
+- (BOOL)application:(UIApplication *)application didFinishLaunchingWithOptions:(NSDictionary*)launchOptions
+{
+	// save launch options
+	self.launchOptions = launchOptions;
+
+#if PLATFORM_TVOS
+	self.bDeviceInPortraitMode = false;
+#else
+	// use the status bar orientation to properly determine landscape vs portrait
+	self.bDeviceInPortraitMode = UIInterfaceOrientationIsPortrait([[UIApplication sharedApplication] statusBarOrientation]);
+	printf("========= This app is in %s mode\n", self.bDeviceInPortraitMode ? "PORTRAIT" : "LANDSCAPE");
+#endif
+
+		// check OS version to make sure we have the API
+	OSVersion = [[[UIDevice currentDevice] systemVersion] floatValue];
+	if (!FPlatformMisc::IsDebuggerPresent() || GAlwaysReportCrash)
+	{
+        FPlatformMisc::SetCrashHandler(EngineCrashHandler);
+	}
+
+	// create the main landscape window object
+	CGRect MainFrame = [[UIScreen mainScreen] bounds];
+	self.Window = [[UIWindow alloc] initWithFrame:MainFrame];
+	self.Window.screen = [UIScreen mainScreen];
+    
+    // get the native scale
+    const float NativeScale = [[UIScreen mainScreen] scale];
+    
+	//Make this the primary window, and show it.
+	[self.Window makeKeyAndVisible];
+
+	FAppEntry::PreInit(self, application);
+    
+    // add the default image as a subview
+    NSMutableString* path = [[NSMutableString alloc]init];
+    [path setString: [[NSBundle mainBundle] resourcePath]];
+    UIImageOrientation orient = UIImageOrientationUp;
+    NSMutableString* ImageString = [[NSMutableString alloc]init];
+	NSMutableString* PngString = [[NSMutableString alloc] init];
+    [ImageString appendString:@"Default"];
+
+	FPlatformMisc::EIOSDevice Device = FPlatformMisc::GetIOSDeviceType();
+
+	// iphone6 has specially named files
+    if (Device == FPlatformMisc::IOS_IPhone6 || Device == FPlatformMisc::IOS_IPhone6S || Device == FPlatformMisc::IOS_IPhone7)
+	{
+		[ImageString appendString:@"-IPhone6"];
+		if (!self.bDeviceInPortraitMode)
+		{
+			[ImageString appendString : @"-Landscape"];
+		}
+	}
+    else if (Device == FPlatformMisc::IOS_IPhone6Plus || Device == FPlatformMisc::IOS_IPhone6SPlus || Device == FPlatformMisc::IOS_IPhone7Plus)
+	{
+		[ImageString appendString : @"-IPhone6Plus"];
+		if (!self.bDeviceInPortraitMode)
+		{
+			[ImageString appendString : @"-Landscape"];
+		}
+		else
+		{
+			[ImageString appendString : @"-Portrait"];
+		}
+	}
+	else if (Device == FPlatformMisc::IOS_AppleTV)
+	{
+		// @todo tvos: Make an AppleTV one?
+		// use IPhone6 image for now
+		[ImageString appendString : @"-IPhone6Plus-Landscape"];
+	}
+	else if (Device == FPlatformMisc::IOS_IPadPro)
+	{
+		if (!self.bDeviceInPortraitMode)
+		{
+			[ImageString appendString : @"-Landscape-1336"];
+		}
+		else
+		{
+			[ImageString appendString : @"-Portrait-1336"];
+		}
+        
+        if (NativeScale > 1.0f)
+        {
+            [ImageString appendString:@"@2x"];
+        }
+	}
+	else
+	{
+		if (MainFrame.size.height == 320 && MainFrame.size.width != 480 && !self.bDeviceInPortraitMode)
+		{
+			[ImageString appendString:@"-568h"];
+			orient = UIImageOrientationRight;
+		}
+		else if (MainFrame.size.height == 320 && MainFrame.size.width == 480 && !self.bDeviceInPortraitMode)
+		{
+			orient = UIImageOrientationRight;
+		}
+        else if (MainFrame.size.height == 568 || Device == FPlatformMisc::IOS_IPodTouch6)
+		{
+			[ImageString appendString:@"-568h"];
+		}
+		else if (MainFrame.size.height == 1024 && !self.bDeviceInPortraitMode)
+		{
+			[ImageString appendString:@"-Landscape"];
+			orient = UIImageOrientationRight;
+		}
+		else if (MainFrame.size.height == 1024)
+		{
+			[ImageString appendString:@"-Portrait"];
+		}
+		else if (MainFrame.size.height == 768 && !self.bDeviceInPortraitMode)
+		{
+			[ImageString appendString:@"-Landscape"];
+		}
+        
+        if (NativeScale > 1.0f)
+        {
+            [ImageString appendString:@"@2x"];
+        }
+	}
+
+	[PngString appendString : ImageString];
+	[PngString appendString : @".png"];
+	[ImageString appendString : @".jpg"];
+    [path setString: [path stringByAppendingPathComponent:ImageString]];
+    UIImage* image = [[UIImage alloc] initWithContentsOfFile: path];
+	if (image == nil)
+	{
+        [path setString: [[NSBundle mainBundle] resourcePath]];
+		[path setString : [path stringByAppendingPathComponent : PngString]];
+		image = [[UIImage alloc] initWithContentsOfFile:path];
+	}
+	[path release];
+	
+    UIImage* imageToDisplay = [UIImage imageWithCGImage: [image CGImage] scale: 1.0 orientation: orient];
+    UIImageView* imageView = [[UIImageView alloc] initWithImage: imageToDisplay];
+    imageView.frame = MainFrame;
+    imageView.tag = 2;
+    [self.Window addSubview: imageView];
+    GShowSplashScreen = true;
+	
+#if !PLATFORM_TVOS
+	// Save launch local notification so the app can check for it when it is ready
+	UILocalNotification *notification = [launchOptions objectForKey:UIApplicationLaunchOptionsLocalNotificationKey];
+	if ( notification != nullptr )
+	{
+		NSDictionary*	userInfo = [notification userInfo];
+		if(userInfo != nullptr)
+		{
+			NSString*	activationEvent = (NSString*)[notification.userInfo objectForKey: @"ActivationEvent"];
+			
+			if(activationEvent != nullptr)
+			{
+				FAppEntry::gAppLaunchedWithLocalNotification = true;
+				FAppEntry::gLaunchLocalNotificationActivationEvent = FString(activationEvent);
+				FAppEntry::gLaunchLocalNotificationFireDate = [notification.fireDate timeIntervalSince1970];
+			}
+		}
+	}
+#endif
+	
+    timer = [NSTimer scheduledTimerWithTimeInterval: 0.05f target:self selector:@selector(timerForSplashScreen) userInfo:nil repeats:YES];
+    
+    self.PeakMemoryTimer = [NSTimer scheduledTimerWithTimeInterval:0.1f target:self selector:@selector(RecordPeakMemory) userInfo:nil repeats:YES];
+    
+	// create a new thread (the pointer will be retained forever)
+	NSThread* GameThread = [[NSThread alloc] initWithTarget:self selector:@selector(MainAppThread:) object:launchOptions];
+	[GameThread setStackSize:GAME_THREAD_STACK_SIZE];
+	[GameThread start];
+
+	self.CommandLineParseTimer = [NSTimer scheduledTimerWithTimeInterval:0.01f target:self selector:@selector(NoUrlCommandLine) userInfo:nil repeats:NO];
+#if !UE_BUILD_SHIPPING && !PLATFORM_TVOS
+	// make a history buffer
+	self.ConsoleHistoryValues = [[NSMutableArray alloc] init];
+
+	// load saved history from disk
+	NSArray* SavedHistory = [[NSUserDefaults standardUserDefaults] objectForKey:@"ConsoleHistory"];
+	if (SavedHistory != nil)
+	{
+		[self.ConsoleHistoryValues addObjectsFromArray:SavedHistory];
+	}
+	self.ConsoleHistoryValuesIndex = -1;
+#endif
+
+	[self InitializeAudioSession];
+	
+#if !PLATFORM_TVOS
+	// Register for device orientation changes
+	[[UIDevice currentDevice] beginGeneratingDeviceOrientationNotifications];
+	[[NSNotificationCenter defaultCenter] addObserver:self selector:@selector(didRotate:) name:UIDeviceOrientationDidChangeNotification object:nil];
+#endif
+	
+	return YES;
+}
+
+- (void) didRotate:(NSNotification *)notification
+{   
+#if !PLATFORM_TVOS
+	UIDeviceOrientation orientation = [[UIDevice currentDevice] orientation];
+
+    if (bEngineInit)
+    {
+		FFunctionGraphTask::CreateAndDispatchWhenReady([orientation]()
+		{
+			FCoreDelegates::ApplicationReceivedScreenOrientationChangedNotificationDelegate.Broadcast((int32)orientation);
+		}, TStatId(), NULL, ENamedThreads::GameThread);
+	}
+#endif
+}
+
+- (BOOL)application:(UIApplication *)application openURL:(NSURL *)url sourceApplication:(NSString *)sourceApplication annotation:(id)annotation
+{
+#if !NO_LOGGING
+	NSLog(@"%s", "IOSAppDelegate openURL\n");
+#endif
+
+	NSString* EncdodedURLString = [url absoluteString];
+	NSString* URLString = [EncdodedURLString stringByRemovingPercentEncoding];
+	FString CommandLineParameters(URLString);
+
+	// Strip the "URL" part of the URL before treating this like args. It comes in looking like so:
+	// "MyGame://arg1 arg2 arg3 ..."
+	// So, we're going to make it look like:
+	// "arg1 arg2 arg3 ..."
+	int32 URLTerminator = CommandLineParameters.Find( TEXT("://"));
+	if ( URLTerminator > -1 )
+	{
+		CommandLineParameters = CommandLineParameters.RightChop(URLTerminator + 3);
+	}
+
+	FIOSCommandLineHelper::InitCommandArgs(CommandLineParameters);
+	self.bCommandLineReady = true;
+	[self.CommandLineParseTimer invalidate];
+	self.CommandLineParseTimer = nil;
+	
+	FIOSCoreDelegates::OnOpenURL.Broadcast(application, url, sourceApplication, annotation);
+	
+	return YES;
+}
+
+- (void)applicationWillResignActive:(UIApplication *)application
+{
+	/*
+	 Sent when the application is about to move from active to inactive state. This can occur for certain types of temporary interruptions (such as an incoming phone call or SMS message) or when the user quits the application and it begins the transition to the background state.
+	 Use this method to pause ongoing tasks, disable timers, and throttle down OpenGL ES frame rates. Games should use this method to pause the game.
+	 */
+    if (bEngineInit)
+    {
+        FGraphEventRef ResignTask = FFunctionGraphTask::CreateAndDispatchWhenReady([]()
+        {
+            FCoreDelegates::ApplicationWillDeactivateDelegate.Broadcast();
+        }, TStatId(), NULL, ENamedThreads::GameThread);
+		
+		// Do not wait forever for this task to complete since the game thread may be stuck on waiting for user input from a modal dialog box
+		double	startTime = FPlatformTime::Seconds();
+		while((FPlatformTime::Seconds() - startTime) < cMaxThreadWaitTime)
+		{
+			FPlatformProcess::Sleep(0.05f);
+			if(ResignTask->IsComplete())
+			{
+				break;
+			}
+		}
+    }
+    
+	[self ToggleSuspend:true];
+	[self ToggleAudioSession:false];
+}
+
+- (void)applicationDidEnterBackground:(UIApplication *)application
+{
+	/*
+	 Use this method to release shared resources, save user data, invalidate timers, and store enough application state information to restore your application to its current state in case it is terminated later. 
+	 If your application supports background execution, this method is called instead of applicationWillTerminate: when the user quits.
+	 */
+    FCoreDelegates::ApplicationWillEnterBackgroundDelegate.Broadcast();
+}
+
+- (void)applicationWillEnterForeground:(UIApplication *)application
+{
+	/*
+	 Called as part of the transition from the background to the inactive state; here you can undo many of the changes made on entering the background.
+	 */
+    FCoreDelegates::ApplicationHasEnteredForegroundDelegate.Broadcast();
+}
+
+- (void)applicationDidBecomeActive:(UIApplication *)application
+{
+	/*
+	 Restart any tasks that were paused (or not yet started) while the application was inactive. If the application was previously in the background, optionally refresh the user interface.
+	 */
+    [self ToggleSuspend:false];
+	[self ToggleAudioSession:true];
+
+    if (bEngineInit)
+    {
+        FGraphEventRef ResignTask = FFunctionGraphTask::CreateAndDispatchWhenReady([]()
+        {
+            FCoreDelegates::ApplicationHasReactivatedDelegate.Broadcast();
+        }, TStatId(), NULL, ENamedThreads::GameThread);
+
+		// Do not wait forever for this task to complete since the game thread may be stuck on waiting for user input from a modal dialog box
+		double	startTime = FPlatformTime::Seconds();
+ 		while((FPlatformTime::Seconds() - startTime) < cMaxThreadWaitTime)
+		{
+			FPlatformProcess::Sleep(0.05f);
+			if(ResignTask->IsComplete())
+			{
+				break;
+			}
+		}
+    }
+}
+
+- (void)applicationWillTerminate:(UIApplication *)application
+{
+	/*
+	 Called when the application is about to terminate.
+	 Save data if appropriate.
+	 See also applicationDidEnterBackground:.
+	 */
+	FCoreDelegates::ApplicationWillTerminateDelegate.Broadcast();
+    
+    // note that we are shutting down
+    GIsRequestingExit = true;
+    
+    // wait for the game thread to shut down
+    while (self.bHasStarted == true)
+    {
+        usleep(3);
+    }
+}
+
+- (void)applicationDidReceiveMemoryWarning:(UIApplication *)application
+{
+	/*
+	Tells the delegate when the application receives a memory warning from the system
+	*/
+	FPlatformMisc::HandleLowMemoryWarning();
+}
+
+#if !PLATFORM_TVOS && NOTIFICATIONS_ENABLED
+
+#ifdef __IPHONE_8_0
+- (void)application:(UIApplication *)application didRegisterUserNotificationSettings:(UIUserNotificationSettings *)notificationSettings
+{
+	[application registerForRemoteNotifications];
+	int32 types = (int32)[notificationSettings types];
+    FFunctionGraphTask::CreateAndDispatchWhenReady([types]()
+    {
+		FCoreDelegates::ApplicationRegisteredForUserNotificationsDelegate.Broadcast(types);
+    }, TStatId(), NULL, ENamedThreads::GameThread);
+}
+#endif
+
+- (void)application:(UIApplication *)application didRegisterForRemoteNotificationsWithDeviceToken:(NSData *)deviceToken
+{
+	TArray<uint8> Token;
+	Token.AddUninitialized([deviceToken length]);
+	memcpy(Token.GetData(), [deviceToken bytes], [deviceToken length]);
+
+	const char *data = (const char*)([deviceToken bytes]);
+	NSMutableString *token = [NSMutableString string];
+
+	for (NSUInteger i = 0; i < [deviceToken length]; i++) {
+		[token appendFormat : @"%02.2hhX", data[i]];
+	}
+
+	UE_LOG(LogTemp, Display, TEXT("Device Token: %s"), *FString(token));
+
+    FFunctionGraphTask::CreateAndDispatchWhenReady([Token]()
+    {
+		FCoreDelegates::ApplicationRegisteredForRemoteNotificationsDelegate.Broadcast(Token);
+    }, TStatId(), NULL, ENamedThreads::GameThread);
+}
+
+-(void)application:(UIApplication *)application didFailtoRegisterForRemoteNotificationsWithError:(NSError *)error
+{
+	FString errorDescription([error description]);
+	
+    FFunctionGraphTask::CreateAndDispatchWhenReady([errorDescription]()
+    {
+		FCoreDelegates::ApplicationFailedToRegisterForRemoteNotificationsDelegate.Broadcast(errorDescription);
+    }, TStatId(), NULL, ENamedThreads::GameThread);
+}
+
+-(void)application : (UIApplication *)application didReceiveRemoteNotification:(NSDictionary *)userInfo fetchCompletionHandler:(void(^)(UIBackgroundFetchResult result))handler
+{
+	NSString* JsonString = @"{}";
+	NSError* JsonError;
+	NSData* JsonData = [NSJSONSerialization dataWithJSONObject : userInfo
+					options : 0
+					error : &JsonError];
+
+	if (JsonData)
+	{
+		JsonString = [[[NSString alloc] initWithData:JsonData encoding : NSUTF8StringEncoding] autorelease];
+	}
+	
+	FString	jsonFString(JsonString);
+	int AppState;
+	if (application.applicationState == UIApplicationStateInactive)
+	{
+		AppState = 1; // EApplicationState::Inactive;
+	}
+	else if (application.applicationState == UIApplicationStateBackground)
+	{
+		AppState = 2; // EApplicationState::Background;
+	}
+	else
+	{
+		AppState = 3; // EApplicationState::Active;
+	}
+
+    FFunctionGraphTask::CreateAndDispatchWhenReady([jsonFString, AppState]()
+    {
+		FCoreDelegates::ApplicationReceivedRemoteNotificationDelegate.Broadcast(jsonFString, AppState);
+    }, TStatId(), NULL, ENamedThreads::GameThread);
+}
+
+#endif
+
+#if !PLATFORM_TVOS
+// the below code breaks all of the plist specifying what to support
+/*- (NSUInteger)application:(UIApplication*)application supportedInterfaceOrientationsForWindow:(UIWindow*)window
+{
+ 	// UIImagePickerController or GameCenter might have portrait-only variant and will throw exception if portrait is not supported here
+ 	return UIInterfaceOrientationMaskAll;
+}*/
+
+- (void)application:(UIApplication *)application didReceiveLocalNotification:(UILocalNotification *)notification
+{
+	NSString*	activationEvent = (NSString*)[notification.userInfo objectForKey: @"ActivationEvent"];
+	
+	if(activationEvent != nullptr)
+	{
+		FString	activationEventFString(activationEvent);
+		int32	fireDate = [notification.fireDate timeIntervalSince1970];
+		
+		int AppState;
+		if (application.applicationState == UIApplicationStateInactive)
+		{
+			AppState = 1;// EApplicationState::Inactive;
+		}
+		else if (application.applicationState == UIApplicationStateBackground)
+		{
+			AppState = 2; // EApplicationState::Background;
+		}
+		else
+		{
+			AppState = 3; // EApplicationState::Active;
+		}
+
+		FFunctionGraphTask::CreateAndDispatchWhenReady([activationEventFString, fireDate, AppState]()
+		{
+			FCoreDelegates::ApplicationReceivedLocalNotificationDelegate.Broadcast(activationEventFString, fireDate, AppState);
+		}, TStatId(), NULL, ENamedThreads::GameThread);
+	}
+	else
+	{
+		NSLog(@"Warning: Missing local notification activation event");
+	}
+}
+
+#endif
+
+/**
+ * Shows the given Game Center supplied controller on the screen
+ *
+ * @param Controller The Controller object to animate onto the screen
+ */
+-(void)ShowController:(UIViewController*)Controller
+{
+	// slide it onto the screen
+	[[IOSAppDelegate GetDelegate].IOSController presentViewController : Controller animated : YES completion : nil];
+	
+	// stop drawing the 3D world for faster UI speed
+	//FViewport::SetGameRenderingEnabled(false);
+}
+
+/**
+ * Hides the given Game Center supplied controller from the screen, optionally controller
+ * animation (sliding off)
+ *
+ * @param Controller The Controller object to animate off the screen
+ * @param bShouldAnimate YES to slide down, NO to hide immediately
+ */
+-(void)HideController:(UIViewController*)Controller Animated:(BOOL)bShouldAnimate
+{
+    // slide it off
+    [Controller dismissViewControllerAnimated : bShouldAnimate completion : nil];
+
+	// stop drawing the 3D world for faster UI speed
+	//FViewport::SetGameRenderingEnabled(true);
+}
+
+/**
+ * Hides the given Game Center supplied controller from the screen
+ *
+ * @param Controller The Controller object to animate off the screen
+ */
+-(void)HideController:(UIViewController*)Controller
+{
+	// call the other version with default animation of YES
+	[self HideController : Controller Animated : YES];
+}
+
+-(void)gameCenterViewControllerDidFinish:(GKGameCenterViewController*)GameCenterDisplay
+{
+	// close the view 
+	[self HideController : GameCenterDisplay];
+}
+
+/**
+ * Show the leaderboard interface (call from iOS main thread)
+ */
+-(void)ShowLeaderboard:(NSString*)Category
+{
+	// create the leaderboard display object 
+	GKGameCenterViewController* GameCenterDisplay = [[[GKGameCenterViewController alloc] init] autorelease];
+#if !PLATFORM_TVOS
+	GameCenterDisplay.viewState = GKGameCenterViewControllerStateLeaderboards;
+#endif
+#ifdef __IPHONE_7_0
+    if ([GameCenterDisplay respondsToSelector:@selector(leaderboardIdentifier)] == YES)
+    {
+#if !PLATFORM_TVOS // @todo tvos: Why not??
+        GameCenterDisplay.leaderboardIdentifier = Category;
+#endif
+    }
+    else
+#endif
+    {
+#if __IPHONE_OS_VERSION_MIN_REQUIRED < __IPHONE_7_0
+        GameCenterDisplay.leaderboardCategory = Category;
+#endif
+    }
+	GameCenterDisplay.gameCenterDelegate = self;
+
+	// show it 
+	[self ShowController : GameCenterDisplay];
+}
+
+/**
+ * Show the achievements interface (call from iOS main thread)
+ */
+-(void)ShowAchievements
+{
+	// create the leaderboard display object 
+	GKGameCenterViewController* GameCenterDisplay = [[[GKGameCenterViewController alloc] init] autorelease];
+#if !PLATFORM_TVOS
+	GameCenterDisplay.viewState = GKGameCenterViewControllerStateAchievements;
+#endif
+	GameCenterDisplay.gameCenterDelegate = self;
+
+	// show it 
+	[self ShowController : GameCenterDisplay];
+}
+
+/**
+ * Show the leaderboard interface (call from game thread)
+ */
+CORE_API bool IOSShowLeaderboardUI(const FString& CategoryName)
+{
+	// route the function to iOS thread, passing the category string along as the object
+	NSString* CategoryToShow = [NSString stringWithFString:CategoryName];
+	[[IOSAppDelegate GetDelegate] performSelectorOnMainThread:@selector(ShowLeaderboard:) withObject:CategoryToShow waitUntilDone : NO];
+
+	return true;
+}
+
+/**
+* Show the achievements interface (call from game thread)
+*/
+CORE_API bool IOSShowAchievementsUI()
+{
+	
+	// route the function to iOS thread
+	[[IOSAppDelegate GetDelegate] performSelectorOnMainThread:@selector(ShowAchievements) withObject:nil waitUntilDone : NO];
+
+	return true;
+}
+
+-(UIWindow*)window
+{
+	return Window;
+}
+
+@end
diff --git a/Engine/Source/Runtime/Core/Private/IOS/IOSApplication.cpp b/Engine/Source/Runtime/Core/Private/IOS/IOSApplication.cpp
new file mode 100644
--- /dev/null
+++ b/Engine/Source/Runtime/Core/Private/IOS/IOSApplication.cpp
@@ -0,0 +1,129 @@
+// Copyright 1998-2017 Epic Games, Inc. All Rights Reserved.
+
+#include "IOSApplication.h"
+#include "IOSInputInterface.h"
+#include "IOSWindow.h"
+#include "IOSAppDelegate.h"
+#include "IInputDeviceModule.h"
+#include "IInputInterface.h"
+#include "IInputDevice.h"
+#include "Misc/ScopeLock.h"
+
+
+FCriticalSection FIOSApplication::CriticalSection;
+bool FIOSApplication::bOrientationChanged = false;
+
+FIOSApplication* FIOSApplication::CreateIOSApplication()
+{
+	return new FIOSApplication();
+}
+
+FIOSApplication::FIOSApplication()
+	: GenericApplication( NULL )
+	, InputInterface( FIOSInputInterface::Create( MessageHandler ) )
+	, bHasLoadedInputPlugins(false)
+{
+	[IOSAppDelegate GetDelegate].IOSApplication = this;
+}
+
+void FIOSApplication::InitializeWindow( const TSharedRef< FGenericWindow >& InWindow, const TSharedRef< FGenericWindowDefinition >& InDefinition, const TSharedPtr< FGenericWindow >& InParent, const bool bShowImmediately )
+{
+	const TSharedRef< FIOSWindow > Window = StaticCastSharedRef< FIOSWindow >( InWindow );
+	const TSharedPtr< FIOSWindow > ParentWindow = StaticCastSharedPtr< FIOSWindow >( InParent );
+
+	Windows.Add( Window );
+	Window->Initialize( this, InDefinition, ParentWindow, bShowImmediately );
+}
+
+void FIOSApplication::SetMessageHandler( const TSharedRef< FGenericApplicationMessageHandler >& InMessageHandler )
+{
+	GenericApplication::SetMessageHandler(InMessageHandler);
+	InputInterface->SetMessageHandler(InMessageHandler);
+
+	TArray<IInputDeviceModule*> PluginImplementations = IModularFeatures::Get().GetModularFeatureImplementations<IInputDeviceModule>(IInputDeviceModule::GetModularFeatureName());
+	for (auto DeviceIt = ExternalInputDevices.CreateIterator(); DeviceIt; ++DeviceIt)
+	{
+		(*DeviceIt)->SetMessageHandler(InMessageHandler);
+	}
+}
+
+void FIOSApplication::AddExternalInputDevice(TSharedPtr<IInputDevice> InputDevice)
+{
+	if (InputDevice.IsValid())
+	{
+		ExternalInputDevices.Add(InputDevice);
+	}
+}
+
+void FIOSApplication::PollGameDeviceState( const float TimeDelta )
+{
+	// initialize any externally-implemented input devices (we delay load initialize the array so any plugins have had time to load)
+	if (!bHasLoadedInputPlugins)
+	{
+		TArray<IInputDeviceModule*> PluginImplementations = IModularFeatures::Get().GetModularFeatureImplementations<IInputDeviceModule>(IInputDeviceModule::GetModularFeatureName());
+		for (auto InputPluginIt = PluginImplementations.CreateIterator(); InputPluginIt; ++InputPluginIt)
+		{
+			TSharedPtr<IInputDevice> Device = (*InputPluginIt)->CreateInputDevice(MessageHandler);
+			AddExternalInputDevice(Device);
+		}
+
+		bHasLoadedInputPlugins = true;
+	}
+
+	// Poll game device state and send new events
+	InputInterface->Tick(TimeDelta);
+	InputInterface->SendControllerEvents();
+
+	// Poll externally-implemented devices
+	for (auto DeviceIt = ExternalInputDevices.CreateIterator(); DeviceIt; ++DeviceIt)
+	{
+		(*DeviceIt)->Tick(TimeDelta);
+		(*DeviceIt)->SendControllerEvents();
+	}
+
+	FScopeLock Lock(&CriticalSection);
+	if(bOrientationChanged && Windows.Num() > 0)
+	{
+		int32 WindowX,WindowY, WindowWidth,WindowHeight;
+		Windows[0]->GetFullScreenInfo(WindowX, WindowY, WindowWidth, WindowHeight);
+
+		GenericApplication::GetMessageHandler()->OnSizeChanged(Windows[0],WindowWidth,WindowHeight, false);
+		GenericApplication::GetMessageHandler()->OnResizingWindow(Windows[0]);
+		FDisplayMetrics DisplayMetrics;
+		FDisplayMetrics::GetDisplayMetrics(DisplayMetrics);
+		BroadcastDisplayMetricsChanged(DisplayMetrics);
+		bOrientationChanged = false;
+	}
+}
+
+FPlatformRect FIOSApplication::GetWorkArea( const FPlatformRect& CurrentWindow ) const
+{
+	return FIOSWindow::GetScreenRect();
+}
+
+void FDisplayMetrics::GetDisplayMetrics(FDisplayMetrics& OutDisplayMetrics)
+{
+	// Get screen rect
+	OutDisplayMetrics.PrimaryDisplayWorkAreaRect = FIOSWindow::GetScreenRect();
+	OutDisplayMetrics.VirtualDisplayRect = OutDisplayMetrics.PrimaryDisplayWorkAreaRect;
+
+	// Total screen size of the primary monitor
+	OutDisplayMetrics.PrimaryDisplayWidth = OutDisplayMetrics.PrimaryDisplayWorkAreaRect.Right - OutDisplayMetrics.PrimaryDisplayWorkAreaRect.Left;
+	OutDisplayMetrics.PrimaryDisplayHeight = OutDisplayMetrics.PrimaryDisplayWorkAreaRect.Bottom - OutDisplayMetrics.PrimaryDisplayWorkAreaRect.Top;
+
+	// Apply the debug safe zones
+	OutDisplayMetrics.ApplyDefaultSafeZones();
+}
+
+TSharedRef< FGenericWindow > FIOSApplication::MakeWindow()
+{
+	return FIOSWindow::Make();
+}
+
+#if !PLATFORM_TVOS
+void FIOSApplication::OrientationChanged(UIDeviceOrientation orientation)
+{
+	FScopeLock Lock(&CriticalSection);
+	bOrientationChanged = true;
+}
+#endif
diff --git a/Engine/Source/Runtime/Core/Private/IOS/IOSAsyncTask.cpp b/Engine/Source/Runtime/Core/Private/IOS/IOSAsyncTask.cpp
new file mode 100644
--- /dev/null
+++ b/Engine/Source/Runtime/Core/Private/IOS/IOSAsyncTask.cpp
@@ -0,0 +1,127 @@
+// Copyright 1998-2017 Epic Games, Inc. All Rights Reserved.
+
+#include "IOSAsyncTask.h"
+#include "HAL/PlatformAtomics.h"
+
+@implementation FIOSAsyncTask
+
+@synthesize UserData;
+@synthesize GameThreadCallback;
+
+/** All currently running tasks (which can be created on iOS thread or main thread) */
+NSMutableArray* RunningTasks;
+
+/**
+ * Static class constructor, called before any instances are created
+ */
++ (void)initialize
+{
+	// create the RunningTasks object one time
+	RunningTasks = [[NSMutableArray arrayWithCapacity:4] retain];
+}
+
+/**
+ * Initialize the async task
+ */
+- (id)init
+{
+	self = [super init];
+
+	// add ourself to the list of tasks
+	@synchronized(RunningTasks)
+	{
+		[RunningTasks addObject:self];
+	}
+
+	// return ourself, the constructed object
+	return self;
+}
+
+
++ (void)CreateTaskWithBlock:(bool (^)(void))Block
+{
+	// create a task, and add it to the array
+	FIOSAsyncTask* Task = [[FIOSAsyncTask alloc] init];
+	// set the callback
+	Task.GameThreadCallback = Block;
+	// safely tell the game thread we are ready to go
+	[Task FinishedTask];
+}
+
+- (void)FinishedTask
+{
+	FPlatformAtomics::InterlockedIncrement(&bIsReadyForGameThread);
+}
+
+/**
+ * Check for completion
+ *
+ * @return TRUE if we succeeded (the completion block will have been called)
+ */
+- (bool)CheckForCompletion
+{
+	// handle completion
+	if (bIsReadyForGameThread)
+	{
+		// call the game thread block
+		if (GameThreadCallback)
+		{
+			if (GameThreadCallback())
+			{
+				// only return true if the callback says it's complete
+				return true;
+			}
+		}
+		else
+		{
+			// if there isn't a callback, then just return TRUE to remove the 
+			// task from the queue
+			return true;
+		}
+		
+	}
+
+	// all other cases, we are not complete
+	return false;
+}
+
+/**
+ * Tick all currently running tasks
+ */
++ (void)ProcessAsyncTasks
+{
+	@synchronized(RunningTasks)
+	{
+		// check all outstanding tasks to see if they are done
+		// (we go backwards so we can remove as we go)
+		for (int32 TaskIndex = 0; TaskIndex < [RunningTasks count]; TaskIndex++)
+		{
+			FIOSAsyncTask* Task = [RunningTasks objectAtIndex:TaskIndex];
+			// if it finished, remove it from the list of running tasks
+			if ([Task CheckForCompletion])
+			{
+				// release the object count
+				[Task release];
+
+				// and remove from the list
+				[RunningTasks removeObjectAtIndex:TaskIndex];
+				TaskIndex--;
+			}
+		}
+	}
+}
+
+
+/** 
+ * Application destructor
+ */
+- (void)dealloc 
+{
+	[GameThreadCallback release];
+	self.UserData = nil;
+
+	[super dealloc];
+}
+
+
+@end
\ No newline at end of file
diff --git a/Engine/Source/Runtime/Core/Private/IOS/IOSChunkInstaller.cpp b/Engine/Source/Runtime/Core/Private/IOS/IOSChunkInstaller.cpp
new file mode 100644
--- /dev/null
+++ b/Engine/Source/Runtime/Core/Private/IOS/IOSChunkInstaller.cpp
@@ -0,0 +1,183 @@
+// Copyright 1998-2017 Epic Games, Inc. All Rights Reserved.
+
+#include "IOSChunkInstaller.h"
+#include "Misc/CallbackDevice.h"
+
+
+
+FIOSChunkInstall::FIOSChunkInstall()
+{
+    // TODO: call conditionallyBeginAccess on all of the possible chunks
+    GetChunkLocation(1);
+}
+
+
+FIOSChunkInstall::~FIOSChunkInstall()
+{
+    ShutDown();
+}
+
+
+void FIOSChunkInstall::ShutDown()
+{
+    // call endAccess on all requests
+    for ( TMap<uint32, FIOSChunkStatus>::TIterator It(ChunkStatus); It; ++It )
+    {
+        FIOSChunkStatus Status = It.Value();
+        [Status.Request endAccessingResources];
+    }
+}
+
+
+EChunkLocation::Type FIOSChunkInstall::GetChunkLocation( uint32 ChunkID )
+{
+    return GetChunkLocation(ChunkID, false);
+}
+
+
+EChunkLocation::Type FIOSChunkInstall::GetChunkLocation( uint32 ChunkID, bool bSkipMount )
+{
+    // check for cached location
+    FIOSChunkStatus* CurStatus = ChunkStatus.Find(ChunkID);
+    
+    // early out if we already know it's local (should be the most common case)
+    if (CurStatus && (CurStatus->Request.progress.fractionCompleted == 1.0f) && MountedChunks.Find(ChunkID))
+    {
+        return EChunkLocation::LocalFast;
+    }
+    
+    // update the cached chunk status
+    if (!CurStatus)
+    {
+        PrioritizeChunk(ChunkID, EChunkPriority::Immediate);
+    }
+    
+    return EChunkLocation::NotAvailable;
+}
+
+
+bool FIOSChunkInstall::GetProgressReportingTypeSupported( EChunkProgressReportingType::Type ReportType )
+{
+    bool bSupported = false;
+    switch (ReportType)
+    {
+        case EChunkProgressReportingType::PercentageComplete:
+            bSupported = true;
+            break;
+        default:
+            bSupported = false;
+            break;
+    }
+    return bSupported;
+}
+
+
+float FIOSChunkInstall::GetChunkProgress( uint32 ChunkID, EChunkProgressReportingType::Type ReportType )
+{
+    float Val;
+    switch(ReportType)
+    {
+        case EChunkProgressReportingType::PercentageComplete:
+            Val = GetChunkPercentComplete(ChunkID);
+            break;
+        default:
+            UE_LOG(LogChunkInstaller, Error, TEXT("Unsupported ProgressReportType: %i"), (int)ReportType);
+            Val = 0.0f;
+            break;
+    }
+    return Val;
+}
+
+
+float FIOSChunkInstall::GetChunkPercentComplete(uint32 ChunkID)
+{
+    // early out if we already know it's local
+    const FIOSChunkStatus* CurStatus = ChunkStatus.Find(ChunkID);
+    
+    if (CurStatus)
+    {
+        return ((float)CurStatus->Request.progress.fractionCompleted) * 100.0f;
+    }
+    
+    return 0.0f;
+}
+
+
+EChunkInstallSpeed::Type FIOSChunkInstall::GetInstallSpeed()
+{
+    // TODO: utilize the NSProgress of all requests to determine if we are paused, slow, or fast
+    
+    return EChunkInstallSpeed::Fast;
+}
+
+
+bool FIOSChunkInstall::SetInstallSpeed( EChunkInstallSpeed::Type InstallSpeed )
+{
+    // TODO: speed doesn't matter for IOS, but use the NSProgress to pause if not already complete
+    
+    return true;
+}
+
+
+bool FIOSChunkInstall::PrioritizeChunk( uint32 ChunkID, EChunkPriority::Type Priority )
+{
+    // create the request
+    FString RequestIDs = FString::Printf(TEXT("Chunk%d"), ChunkID);
+    NSString* ID = RequestIDs.GetNSString();
+    NSSet<NSString* > * tags = [[NSSet alloc] initWithObjects:&ID count:1];
+    NSBundleResourceRequest* ChunkRequest = [[NSBundleResourceRequest alloc] initWithTags:tags];
+    
+    // set the priority
+    ChunkRequest.loadingPriority = 1.0f - (float)Priority/2.0f;
+    
+    // create the chunk status
+    FIOSChunkStatus Status;
+    Status.Request = ChunkRequest;
+    ChunkStatus.Add(ChunkID, Status);
+    
+    // start the download
+    __block uint32 BlockID = ChunkID;
+    [ChunkRequest beginAccessingResourcesWithCompletionHandler:^(NSError* error)
+     {
+        if (error)
+        {
+            // TODO: handle error
+            NSLog(@"Error: %@", error);
+        }
+        else
+        {
+            // mount the pak file in the NSBundle
+            if (!MountedChunks.Find(BlockID))
+            {
+                FIOSChunkStatus* FoundStatus = ChunkStatus.Find(BlockID);
+                NSBundle* Bundle = FoundStatus->Request.bundle;
+                FString PakFile = FString::Printf(TEXT("pakchunk%i-ios.pak"), BlockID);
+                NSString* ResourcePath = [Bundle pathForResource:PakFile.GetNSString() ofType:nil];
+                NSLog(@"ResourcePath: %@", ResourcePath);
+                bool bMounted = FCoreDelegates::OnMountPak.Execute(FString(ResourcePath), 0, nullptr);
+                if (bMounted)
+                {
+                    MountedChunks.Add(BlockID);
+                }
+                else
+                {
+                    UE_LOG(LogChunkInstaller, Warning, TEXT("NSBundle Chunk %i couldn't be mounted."), ID);
+                }
+            }
+        }
+     }];
+    return true;
+}
+
+
+bool FIOSChunkInstall::DebugStartNextChunk()
+{
+#if !UE_BUILD_SHIPPING
+    static uint32 NextChunkID = 0;
+    
+    //cannot call this function in a submission build.  Will fail cert.
+    // create a request and ask for the bundle
+    return true;
+#endif
+    return true;
+}
diff --git a/Engine/Source/Runtime/Core/Private/IOS/IOSChunkInstaller.h b/Engine/Source/Runtime/Core/Private/IOS/IOSChunkInstaller.h
new file mode 100644
--- /dev/null
+++ b/Engine/Source/Runtime/Core/Private/IOS/IOSChunkInstaller.h
@@ -0,0 +1,130 @@
+// Copyright 1998-2017 Epic Games, Inc. All Rights Reserved.
+
+#pragma once
+
+#include <Foundation/NSBundle.h>
+#include "GenericPlatformChunkInstall.h"
+#include "Containers/Map.h"
+
+struct FIOSChunkStatus
+{
+    NSBundleResourceRequest* Request;
+};
+
+/**
+ * IOS implementation of FGenericPlatformChunkInstall.
+ */
+class CORE_API FIOSChunkInstall
+: public IPlatformChunkInstall
+{
+public:
+    
+    FIOSChunkInstall();
+    virtual ~FIOSChunkInstall();
+    
+    /**
+     * Get the current location of a chunk.
+     *
+     * @param ChunkID The id of the chunk to check.
+     * @return Enum specifying whether the chunk is available to use, waiting to install, or does not exist.
+     */
+    virtual EChunkLocation::Type GetChunkLocation( uint32 ChunkID ) override;
+    
+    /**
+     * Check if a given reporting type is supported.
+     *
+     * @param ReportType Enum specifying how progress is reported.
+     * @return true if reporting type is supported on the current platform.
+     */
+    virtual bool GetProgressReportingTypeSupported(EChunkProgressReportingType::Type ReportType) override;
+    
+    /**
+     * Get the current install progress of a chunk.  Let the user specify report type for platforms that support more than one.
+     *
+     * @param ChunkID The id of the chunk to check.
+     * @param ReportType The type of progress report you want.
+     * @return A value whose meaning is dependent on the ReportType param.
+     */
+    virtual float GetChunkProgress( uint32 ChunkID, EChunkProgressReportingType::Type ReportType ) override;
+    
+    /**
+     * Inquire about the priority of chunk installation vs. game IO.
+     *
+     * @return Paused, low or high priority.
+     */
+    virtual EChunkInstallSpeed::Type GetInstallSpeed() override;
+    
+    /**
+     * Specify the priority of chunk installation vs. game IO.
+     *
+     * @param InstallSpeed Pause, low or high priority.
+     * @return false if the operation is not allowed, otherwise true.
+     */
+    virtual bool SetInstallSpeed( EChunkInstallSpeed::Type InstallSpeed ) override;
+    
+    /**
+     * Hint to the installer that we would like to prioritize a specific chunk
+     *
+     * @param ChunkID The id of the chunk to prioritize.
+     * @param Priority The priority for the chunk.
+     * @return false if the operation is not allowed or the chunk doesn't exist, otherwise true.
+     */
+    virtual bool PrioritizeChunk( uint32 ChunkID, EChunkPriority::Type Priority ) override;
+    
+    /**
+     * Debug function to start transfer of next chunk in the transfer list.  When in PlayGo HostFS emulation, this is the only
+     * way moving out of network-local will happen for a chunk.	 Does nothing in a shipping build.
+     *
+     * @return true if the operation succeeds.
+     */
+    virtual bool DebugStartNextChunk( ) override;
+    
+    /**
+     * Request a delegate callback on chunk install completion. Request may not be respected.
+     * @param ChunkID		The id of the chunk of interest.
+     * @param Delegate		The delegate when the chunk is installed.
+     * @return				False if the delegate was not registered or the chunk is already installed. True on success.
+     */
+    virtual FDelegateHandle SetChunkInstallDelgate(uint32 ChunkID, FPlatformChunkInstallCompleteDelegate Delegate) override
+    {
+		return FDelegateHandle();
+    }
+    
+    /**
+     * Remove a delegate callback on chunk install completion.
+     * @param ChunkID		The id of the chunk of interest.
+     * @param Delegate		The delegate to remove.
+     * @return				False if the delegate was not registered with a call to SetChunkInstallDelgate. True on success.
+     */
+	virtual void RemoveChunkInstallDelgate(uint32 ChunkID, FDelegateHandle Delegate) override
+    {
+        return;
+    }
+    
+private:
+    
+    /**
+     * Get the current location of a chunk.
+     *
+     * @param ChunkID The id of the chunk to check.
+     * @return Enum specifying whether the chunk is available to use, waiting to install, or does not exist.
+     */
+    virtual EChunkLocation::Type GetChunkLocation( uint32 ChunkID, bool bSkipMount );
+    
+    /**
+     * Get the current percent downloaded of a chunk.
+     *
+     * @param ChunkID The id of the chunk to check.
+     * @return Percent downloaded in 99.99 format.
+     */
+    float GetChunkPercentComplete(uint32 ChunkID);
+    
+    /** Unload PlayGo and free everything up */
+    void ShutDown();
+    
+private:
+    
+    /** map of chunk IDs to the cached status of that chunk */
+    TMap<uint32, FIOSChunkStatus> ChunkStatus;
+    TSet<int32> MountedChunks;
+};
diff --git a/Engine/Source/Runtime/Core/Private/IOS/IOSInputInterface.cpp b/Engine/Source/Runtime/Core/Private/IOS/IOSInputInterface.cpp
new file mode 100644
--- /dev/null
+++ b/Engine/Source/Runtime/Core/Private/IOS/IOSInputInterface.cpp
@@ -0,0 +1,632 @@
+// Copyright 1998-2017 Epic Games, Inc. All Rights Reserved.
+
+#include "IOSInputInterface.h"
+#include "IOSAppDelegate.h"
+#include "Misc/ConfigCacheIni.h"
+#include "Misc/ScopeLock.h"
+
+#import <AudioToolbox/AudioToolbox.h>
+
+DECLARE_LOG_CATEGORY_EXTERN(LogIOSInput, Log, All);
+
+
+//@interface FControllerHelper : NSObject
+//{
+//	FIOSInputInterface
+//}
+//
+//@end
+
+TArray<TouchInput> FIOSInputInterface::TouchInputStack = TArray<TouchInput>();
+TArray<int32> FIOSInputInterface::KeyInputStack;
+FCriticalSection FIOSInputInterface::CriticalSection;
+
+TSharedRef< FIOSInputInterface > FIOSInputInterface::Create(  const TSharedRef< FGenericApplicationMessageHandler >& InMessageHandler )
+{
+	return MakeShareable( new FIOSInputInterface( InMessageHandler ) );
+}
+
+FIOSInputInterface::FIOSInputInterface( const TSharedRef< FGenericApplicationMessageHandler >& InMessageHandler )
+	: MessageHandler( InMessageHandler )
+	, bAllowRemoteRotation(false)
+	, bTreatRemoteAsSeparateController(false)
+	, bUseRemoteAsVirtualJoystick(true)
+	, bUseRemoteAbsoluteDpadValues(false)
+{
+#if !PLATFORM_TVOS
+	MotionManager = nil;
+	ReferenceAttitude = nil;
+#endif
+
+	GConfig->GetBool(TEXT("/Script/IOSRuntimeSettings.IOSRuntimeSettings"), TEXT("bTreatRemoteAsSeparateController"), bTreatRemoteAsSeparateController, GEngineIni);
+	GConfig->GetBool(TEXT("/Script/IOSRuntimeSettings.IOSRuntimeSettings"), TEXT("bAllowRemoteRotation"), bAllowRemoteRotation, GEngineIni);
+	GConfig->GetBool(TEXT("/Script/IOSRuntimeSettings.IOSRuntimeSettings"), TEXT("bUseRemoteAsVirtualJoystick"), bUseRemoteAsVirtualJoystick, GEngineIni);
+	GConfig->GetBool(TEXT("/Script/IOSRuntimeSettings.IOSRuntimeSettings"), TEXT("bUseRemoteAbsoluteDpadValues"), bUseRemoteAbsoluteDpadValues, GEngineIni);
+
+	
+	[[NSNotificationCenter defaultCenter] addObserverForName:GCControllerDidConnectNotification object:nil queue:[NSOperationQueue currentQueue] usingBlock:^(NSNotification* Notification)
+	 {
+		HandleConnection(Notification.object);
+	 }];
+
+	[[NSNotificationCenter defaultCenter] addObserverForName:GCControllerDidDisconnectNotification object:nil queue:[NSOperationQueue currentQueue] usingBlock:^(NSNotification* Notification)
+	 {
+		HandleDisconnect(Notification.object);
+	 }];
+	
+
+	[GCController startWirelessControllerDiscoveryWithCompletionHandler:^{
+	 }];
+//	[GCController stopWirelessControllerDiscovery];
+	
+	
+	FMemory::Memzero(Controllers, sizeof(Controllers));
+}
+
+void FIOSInputInterface::SetMessageHandler( const TSharedRef< FGenericApplicationMessageHandler >& InMessageHandler )
+{
+	MessageHandler = InMessageHandler;
+}
+
+void FIOSInputInterface::Tick( float DeltaTime )
+{
+}
+
+void FIOSInputInterface::HandleConnection(GCController* Controller)
+{
+	static_assert(GCControllerPlayerIndex1 == 0 && GCControllerPlayerIndex4 == 3, "Apple changed the player index enums");
+
+	// is this guy a gamepad (i.e., not the Remote)
+	bool bIsGamepadType = (Controller.gamepad != nil);
+	// if we want to use the Remote as a separate player, then we treat it as a Gamepad for player assignment
+	bool bIsTreatedAsGamepad = bIsGamepadType || bTreatRemoteAsSeparateController;
+
+	// find a good controller index to use
+	bool bFoundSlot = false;
+	for (int32 ControllerIndex = 0; ControllerIndex < ARRAY_COUNT(Controllers); ControllerIndex++)
+	{
+		// is this one already connected for this type of controller?
+		if ((!bIsTreatedAsGamepad && Controllers[ControllerIndex].bIsRemoteConnected == false) ||
+			( bIsTreatedAsGamepad && Controllers[ControllerIndex].bIsGamepadConnected == false))
+		{
+			Controller.playerIndex = (GCControllerPlayerIndex)ControllerIndex;
+#if PLATFORM_TVOS
+			if (Controller.microGamepad != nil)
+			{
+				Controller.microGamepad.allowsRotation = bAllowRemoteRotation;
+				Controller.microGamepad.reportsAbsoluteDpadValues = bUseRemoteAbsoluteDpadValues;
+			}
+#endif
+			
+			// update the appropriate flag
+			if (bIsTreatedAsGamepad)
+			{
+				Controllers[ControllerIndex].bIsGamepadConnected = true;
+			}
+			else
+			{
+				Controllers[ControllerIndex].bIsRemoteConnected = true;
+			}
+			
+			Controllers[ControllerIndex].bPauseWasPressed = false;
+			Controller.controllerPausedHandler = ^(GCController* Cont)
+			{
+				Controllers[ControllerIndex].bPauseWasPressed = true;
+			};
+			
+			
+			bFoundSlot = true;
+
+			UE_LOG(LogIOS, Log, TEXT("New %s controller inserted, assigned to playerIndex %d"), bIsTreatedAsGamepad ? TEXT("Gamepad") : TEXT("Remote"), Controller.playerIndex);
+			break;
+		}
+	}
+	checkf(bFoundSlot, TEXT("Used a fifth controller somehow!"));
+	
+}
+
+void FIOSInputInterface::HandleDisconnect(GCController* Controller)
+{
+	UE_LOG(LogIOS, Log, TEXT("Controller for playerIndex %d was removed"), Controller.playerIndex);
+	
+	// mark this controller as disconnected, and reset the state
+	FMemory::Memzero(&Controllers[Controller.playerIndex], sizeof(Controllers[Controller.playerIndex]));
+}
+
+
+#if !PLATFORM_TVOS
+void ModifyVectorByOrientation(FVector& Vec, bool bIsRotation)
+{
+    UIInterfaceOrientation Orientation = UIInterfaceOrientationPortrait;
+#ifdef __IPHONE_8_0
+    if ([[IOSAppDelegate GetDelegate].IOSController respondsToSelector:@selector(interfaceOrientation)] == NO)
+    {
+        Orientation = [[UIApplication sharedApplication] statusBarOrientation];
+    }
+    else
+#endif
+    {
+#if __IPHONE_OS_VERSION_MIN_REQUIRED < __IPHONE_8_0
+        Orientation = [[IOSAppDelegate GetDelegate].IOSController interfaceOrientation];
+#endif
+    }
+
+	switch (Orientation)
+	{
+	case UIInterfaceOrientationPortrait:
+		// this is the base orientation, so nothing to do
+		break;
+
+	case UIInterfaceOrientationPortraitUpsideDown:
+		if (bIsRotation)
+		{
+			// negate roll and pitch
+			Vec.X = -Vec.X;
+			Vec.Z = -Vec.Z;
+		}
+		else
+		{
+			// negate x/y
+			Vec.X = -Vec.X;
+			Vec.Y = -Vec.Y;
+		}
+		break;
+
+	case UIInterfaceOrientationLandscapeRight:
+		if (bIsRotation)
+		{
+			// swap and negate (as needed) roll and pitch
+			float Temp = Vec.X;
+			Vec.X = -Vec.Z;
+			Vec.Z = Temp;
+		}
+		else
+		{
+			// swap and negate (as needed) x and y
+			float Temp = Vec.X;
+			Vec.X = -Vec.Y;
+			Vec.Y = Temp;
+		}
+		break;
+
+	case UIInterfaceOrientationLandscapeLeft:
+		if (bIsRotation)
+		{
+			// swap and negate (as needed) roll and pitch
+			float Temp = Vec.X;
+			Vec.X = Vec.Z;
+			Vec.Z = -Temp;
+		}
+		else
+		{
+			// swap and negate (as needed) x and y
+			float Temp = Vec.X;
+			Vec.X = Vec.Y;
+			Vec.Y = -Temp;
+		}
+		break;
+	}
+}
+#endif
+
+void FIOSInputInterface::ProcessTouchesAndKeys(uint32 ControllerId)
+{
+	for(int i = 0; i < TouchInputStack.Num(); ++i)
+	{
+		const TouchInput& Touch = TouchInputStack[i];
+		
+		// send input to handler
+		if (Touch.Type == TouchBegan)
+		{
+			MessageHandler->OnTouchStarted( NULL, Touch.Position, Touch.Handle, ControllerId);
+		}
+		else if (Touch.Type == TouchEnded)
+		{
+			MessageHandler->OnTouchEnded(Touch.Position, Touch.Handle, ControllerId);
+		}
+		else
+		{
+			MessageHandler->OnTouchMoved(Touch.Position, Touch.Handle, ControllerId);
+		}
+	}
+	
+	TouchInputStack.Empty(0);
+	
+	
+	// these come in pairs
+	for(int32 KeyIndex = 0; KeyIndex < KeyInputStack.Num(); KeyIndex+=2)
+	{
+		int32 KeyCode = KeyInputStack[KeyIndex];
+		int32 CharCode = KeyInputStack[KeyIndex + 1];
+		MessageHandler->OnKeyDown(KeyCode, CharCode, false);
+		MessageHandler->OnKeyChar(CharCode,  false);
+		MessageHandler->OnKeyUp  (KeyCode, CharCode, false);
+	}
+	KeyInputStack.Empty(0);
+}
+
+void FIOSInputInterface::SendControllerEvents()
+{
+	FScopeLock Lock(&CriticalSection);
+
+	int32 ControllerIndex = -1;
+	
+#if !PLATFORM_TVOS
+	// on ios, touches always go go player 0
+	ProcessTouchesAndKeys(0);
+#endif
+
+	
+#if !PLATFORM_TVOS // @todo tvos: This needs to come from the Microcontroller rotation
+	// Update motion controls.
+	FVector Attitude;
+	FVector RotationRate;
+	FVector Gravity;
+	FVector Acceleration;
+
+	GetMovementData(Attitude, RotationRate, Gravity, Acceleration);
+
+	// Fix-up yaw to match directions
+	Attitude.Y = -Attitude.Y;
+	RotationRate.Y = -RotationRate.Y;
+
+	// munge the vectors based on the orientation
+	ModifyVectorByOrientation(Attitude, true);
+	ModifyVectorByOrientation(RotationRate, true);
+	ModifyVectorByOrientation(Gravity, false);
+	ModifyVectorByOrientation(Acceleration, false);
+
+	MessageHandler->OnMotionDetected(Attitude, RotationRate, Gravity, Acceleration, 0);
+#endif
+	
+	
+	
+	for (GCController* Cont in [GCController controllers])
+ 	{
+		// make sure the connection handler has run on this guy
+		if (Cont.playerIndex == GCControllerPlayerIndexUnset)
+		{
+			HandleConnection(Cont);
+		}
+		
+		GCGamepad* Gamepad = Cont.gamepad;
+		GCExtendedGamepad* ExtendedGamepad = Cont.extendedGamepad;
+#if PLATFORM_TVOS
+		GCMicroGamepad* MicroGamepad = Cont.microGamepad;
+#endif
+		GCMotion* Motion = Cont.motion;
+
+		FUserController& Controller = Controllers[Cont.playerIndex];
+		
+		if (Controller.bPauseWasPressed)
+		{
+			MessageHandler->OnControllerButtonPressed(FGamepadKeyNames::SpecialRight, Cont.playerIndex, false);
+			MessageHandler->OnControllerButtonReleased(FGamepadKeyNames::SpecialRight, Cont.playerIndex, false);
+			
+			Controller.bPauseWasPressed = false;
+		}
+		
+		// @todo tvos: Handle repeated buttons?
+		
+#define HANDLE_BUTTON(Gamepad, GCButton, UEButton) \
+if ((Controller.Previous##Gamepad == nil && Gamepad.GCButton.pressed) || Controller.Previous##Gamepad.GCButton.pressed != Gamepad.GCButton.pressed) \
+{ \
+	NSLog(@"%@ button %s on controller %d", (ExtendedGamepad.GCButton.pressed) ? @"Pressed" : @"Released", TCHAR_TO_ANSI(*UEButton.ToString()), (int32)Cont.playerIndex); \
+	(Gamepad.GCButton.pressed) ? MessageHandler->OnControllerButtonPressed(UEButton, Cont.playerIndex, false) : MessageHandler->OnControllerButtonReleased(UEButton, Cont.playerIndex, false); \
+}
+
+#define HANDLE_ANALOG(Gamepad, GCAxis, UEAxis) \
+if (Controller.Previous##Gamepad != nil && Gamepad.GCAxis.value != Controller.Previous##Gamepad.GCAxis.value) \
+{ \
+	NSLog(@"Axis %s is %f", TCHAR_TO_ANSI(*UEAxis.ToString()), Gamepad.GCAxis.value); \
+	MessageHandler->OnControllerAnalog(UEAxis, Cont.playerIndex, Gamepad.GCAxis.value); \
+}
+		
+		if (ExtendedGamepad != nil)
+		{
+			HANDLE_BUTTON(ExtendedGamepad, buttonA,			FGamepadKeyNames::FaceButtonBottom);
+			HANDLE_BUTTON(ExtendedGamepad, buttonB,			FGamepadKeyNames::FaceButtonRight);
+			HANDLE_BUTTON(ExtendedGamepad, buttonX,			FGamepadKeyNames::FaceButtonLeft);
+			HANDLE_BUTTON(ExtendedGamepad, buttonY,			FGamepadKeyNames::FaceButtonTop);
+			HANDLE_BUTTON(ExtendedGamepad, leftShoulder,	FGamepadKeyNames::LeftShoulder);
+			HANDLE_BUTTON(ExtendedGamepad, rightShoulder,	FGamepadKeyNames::RightShoulder);
+			HANDLE_BUTTON(ExtendedGamepad, leftTrigger,		FGamepadKeyNames::LeftTriggerThreshold);
+			HANDLE_BUTTON(ExtendedGamepad, rightTrigger,	FGamepadKeyNames::RightTriggerThreshold);
+			HANDLE_BUTTON(ExtendedGamepad, dpad.up,			FGamepadKeyNames::DPadUp);
+			HANDLE_BUTTON(ExtendedGamepad, dpad.down,		FGamepadKeyNames::DPadDown);
+			HANDLE_BUTTON(ExtendedGamepad, dpad.right,		FGamepadKeyNames::DPadRight);
+			HANDLE_BUTTON(ExtendedGamepad, dpad.left,		FGamepadKeyNames::DPadLeft);
+			
+			HANDLE_ANALOG(ExtendedGamepad, leftThumbstick.xAxis,	FGamepadKeyNames::LeftAnalogX);
+			HANDLE_ANALOG(ExtendedGamepad, leftThumbstick.yAxis,	FGamepadKeyNames::LeftAnalogY);
+			HANDLE_ANALOG(ExtendedGamepad, rightThumbstick.xAxis,	FGamepadKeyNames::RightAnalogX);
+			HANDLE_ANALOG(ExtendedGamepad, rightThumbstick.yAxis,	FGamepadKeyNames::RightAnalogY);
+			HANDLE_ANALOG(ExtendedGamepad, leftTrigger,				FGamepadKeyNames::LeftTriggerAnalog);
+			HANDLE_ANALOG(ExtendedGamepad, rightTrigger,			FGamepadKeyNames::RightTriggerAnalog);
+
+			[Controller.PreviousExtendedGamepad release];
+			Controller.PreviousExtendedGamepad = [ExtendedGamepad saveSnapshot];
+			[Controller.PreviousExtendedGamepad retain];
+		}
+		// get basic input (extended is a superset, don't do both)
+		else if (Gamepad != nil)
+		{
+			HANDLE_BUTTON(Gamepad, buttonA,			FGamepadKeyNames::FaceButtonBottom);
+			HANDLE_BUTTON(Gamepad, buttonA,			FGamepadKeyNames::FaceButtonBottom);
+			HANDLE_BUTTON(Gamepad, buttonB,			FGamepadKeyNames::FaceButtonRight);
+			HANDLE_BUTTON(Gamepad, buttonX,			FGamepadKeyNames::FaceButtonLeft);
+			HANDLE_BUTTON(Gamepad, buttonY,			FGamepadKeyNames::FaceButtonTop);
+			HANDLE_BUTTON(Gamepad, leftShoulder,	FGamepadKeyNames::LeftShoulder);
+			HANDLE_BUTTON(Gamepad, rightShoulder,	FGamepadKeyNames::RightShoulder);
+			HANDLE_BUTTON(Gamepad, dpad.up,			FGamepadKeyNames::DPadUp);
+			HANDLE_BUTTON(Gamepad, dpad.down,		FGamepadKeyNames::DPadDown);
+			HANDLE_BUTTON(Gamepad, dpad.right,		FGamepadKeyNames::DPadRight);
+			HANDLE_BUTTON(Gamepad, dpad.left,		FGamepadKeyNames::DPadLeft);
+			
+
+			HANDLE_ANALOG(ExtendedGamepad, dpad.xAxis,	FGamepadKeyNames::LeftAnalogX);
+			HANDLE_ANALOG(ExtendedGamepad, dpad.yAxis,	FGamepadKeyNames::LeftAnalogY);
+			
+			[Controller.PreviousGamepad release];
+			Controller.PreviousGamepad = [Gamepad saveSnapshot];
+			[Controller.PreviousGamepad retain];
+        }
+#if PLATFORM_TVOS
+        // get micro input (shouldn't have the other two)
+        else if (MicroGamepad != nil)
+        {
+			// if we want virtual joysticks, then use the dpad values (and drain the touch queue to not leak memory)
+			if (bUseRemoteAsVirtualJoystick)
+			{
+				HANDLE_ANALOG(MicroGamepad, dpad.xAxis, FGamepadKeyNames::LeftAnalogX);
+				HANDLE_ANALOG(MicroGamepad, dpad.yAxis, FGamepadKeyNames::LeftAnalogY);
+				
+				// @todo tvos: Do these need a deadzone? I need a good test case for these bindings
+				HANDLE_BUTTON(MicroGamepad, dpad.up,	FGamepadKeyNames::LeftStickUp);
+				HANDLE_BUTTON(MicroGamepad, dpad.down,	FGamepadKeyNames::LeftStickDown);
+				HANDLE_BUTTON(MicroGamepad, dpad.right,	FGamepadKeyNames::LeftStickRight);
+				HANDLE_BUTTON(MicroGamepad, dpad.left,	FGamepadKeyNames::LeftStickLeft);
+
+				TouchInputStack.Empty(0);
+			}
+			// otherwise, process touches like ios for the remote's index
+			else
+			{
+				ProcessTouchesAndKeys(Cont.playerIndex);
+			}
+			
+			HANDLE_BUTTON(MicroGamepad, buttonA,	FGamepadKeyNames::FaceButtonBottom);
+			HANDLE_BUTTON(MicroGamepad, buttonX,	FGamepadKeyNames::FaceButtonRight);
+		
+			
+			[Controller.PreviousMicroGamepad release];
+			Controller.PreviousMicroGamepad = [MicroGamepad saveSnapshot];
+			[Controller.PreviousMicroGamepad retain];
+        }
+		
+		// motion is orthogonal to buttons
+// @todo tvos: handle motion without attitude or rotation rate
+#if 0
+		if (Motion != nil)
+		{
+			FVector Attitude;
+			FVector RotationRate;
+			FVector Gravity;
+			FVector Acceleration;
+
+		
+			FQuat CurrentAttitude(Motion.attitude.x, Motion.attitude.y, Motion.attitude.z, Motion.attitude.w);
+			
+			// save off current as reference when we calibrate
+			if (Controller.bNeedsReferenceAttitude)
+			{
+				Controller.ReferenceAttitude = CurrentAttitude;
+				Controller.bHasReferenceAttitude = true;
+				Controller.bNeedsReferenceAttitude = false;
+			}
+			
+			// take away the reference if we have it
+			if (Controller.bHasReferenceAttitude)
+			{
+				CurrentAttitude *= Controller.ReferenceAttitude.Inverse();
+			}
+			
+			// convert from GCMotion to unrealism
+			Attitude = CurrentAttitude.Euler();
+			RotationRate = FVector(float(Motion.rotationRate.x), float(Motion.rotationRate.y), float(Motion.rotationRate.z));
+			Gravity = FVector(float(Motion.gravity.x), float(Motion.gravity.y), float(Motion.gravity.z));
+			Acceleration = FVector(float(Motion.userAcceleration.x), float(Motion.userAcceleration.y), float(Motion.userAcceleration.z));
+			
+			
+			// Fix-up yaw to match directions we expect
+			Attitude.Y = -Attitude.Y;
+			RotationRate.Y = -RotationRate.Y;
+			
+			// @todo tvos: Handle rotated controller? Not sure that we can get if it's in landscape mode or not
+			MessageHandler->OnMotionDetected(Attitude, RotationRate, Gravity, Acceleration, 0);
+			
+//			NSLog(@"Atti %.2f, %.2f, %.2f\n", Attitude.X, Attitude.Y, Attitude.Z);
+//			NSLog(@"SrcA %.2f, %.2f, %.2f, %.2f\n", Motion.attitude.x, Motion.attitude.y, Motion.attitude.z, Motion.attitude.w);
+//			NSLog(@"Rate %.2f, %.2f, %.2f\n", RotationRate.X, RotationRate.Y, RotationRate.Z);
+//			NSLog(@"Grav %.2f, %.2f, %.2f\n", Gravity.X, Gravity.Y, Gravity.Z);
+//			NSLog(@"Acce %.2f, %.2f, %.2f\n", Acceleration.X, Acceleration.Y, Acceleration.Z);
+		}
+#endif
+		
+#endif
+	}
+}
+
+void FIOSInputInterface::QueueTouchInput(const TArray<TouchInput>& InTouchEvents)
+{
+	FScopeLock Lock(&CriticalSection);
+
+	FIOSInputInterface::TouchInputStack.Append(InTouchEvents);
+}
+
+void FIOSInputInterface::QueueKeyInput(int32 Key, int32 Char)
+{
+	FScopeLock Lock(&CriticalSection);
+
+	// put the key and char into the array
+	FIOSInputInterface::KeyInputStack.Add(Key);
+	FIOSInputInterface::KeyInputStack.Add(Char);
+}
+
+
+void FIOSInputInterface::GetMovementData(FVector& Attitude, FVector& RotationRate, FVector& Gravity, FVector& Acceleration)
+{
+#if !PLATFORM_TVOS
+	// initialize on first use
+	if (MotionManager == nil)
+	{
+		// Look to see if we can create the motion manager
+		MotionManager = [[CMMotionManager alloc] init];
+
+		// Check to see if the device supports full motion (gyro + accelerometer)
+		if (MotionManager.deviceMotionAvailable)
+		{
+			MotionManager.deviceMotionUpdateInterval = 0.02;
+
+			// Start the Device updating motion
+			[MotionManager startDeviceMotionUpdates];
+		}
+		else
+		{
+			[MotionManager startAccelerometerUpdates];
+			CenterPitch = CenterPitch = 0;
+			bIsCalibrationRequested = false;
+		}
+	}
+
+	// do we have full motion data?
+	if (MotionManager.deviceMotionActive)
+	{
+		// Grab the values
+		CMAttitude* CurrentAttitude = MotionManager.deviceMotion.attitude;
+		CMRotationRate CurrentRotationRate = MotionManager.deviceMotion.rotationRate;
+		CMAcceleration CurrentGravity = MotionManager.deviceMotion.gravity;
+		CMAcceleration CurrentUserAcceleration = MotionManager.deviceMotion.userAcceleration;
+
+		// apply a reference attitude if we have been calibrated away from default
+		if (ReferenceAttitude)
+		{
+			[CurrentAttitude multiplyByInverseOfAttitude : ReferenceAttitude];
+		}
+
+		// convert to UE3
+		Attitude = FVector(float(CurrentAttitude.pitch), float(CurrentAttitude.yaw), float(CurrentAttitude.roll));
+		RotationRate = FVector(float(CurrentRotationRate.x), float(CurrentRotationRate.y), float(CurrentRotationRate.z));
+		Gravity = FVector(float(CurrentGravity.x), float(CurrentGravity.y), float(CurrentGravity.z));
+		Acceleration = FVector(float(CurrentUserAcceleration.x), float(CurrentUserAcceleration.y), float(CurrentUserAcceleration.z));
+	}
+	else
+	{
+		// get the plain accleration
+		CMAcceleration RawAcceleration = [MotionManager accelerometerData].acceleration;
+		FVector NewAcceleration(RawAcceleration.x, RawAcceleration.y, RawAcceleration.z);
+
+		// storage for keeping the accelerometer values over time (for filtering)
+		static bool bFirstAccel = true;
+
+		// how much of the previous frame's acceleration to keep
+		const float VectorFilter = bFirstAccel ? 0.0f : 0.85f;
+		bFirstAccel = false;
+
+		// apply new accelerometer values to last frames
+		FilteredAccelerometer = FilteredAccelerometer * VectorFilter + (1.0f - VectorFilter) * NewAcceleration;
+
+		// create an normalized acceleration vector
+		FVector FinalAcceleration = -FilteredAccelerometer.GetSafeNormal();
+
+		// calculate Roll/Pitch
+		float CurrentPitch = FMath::Atan2(FinalAcceleration.Y, FinalAcceleration.Z);
+		float CurrentRoll = -FMath::Atan2(FinalAcceleration.X, FinalAcceleration.Z);
+
+		// if we want to calibrate, use the current values as center
+		if (bIsCalibrationRequested)
+		{
+			CenterPitch = CurrentPitch;
+			CenterRoll = CurrentRoll;
+			bIsCalibrationRequested = false;
+		}
+
+		CurrentPitch -= CenterPitch;
+		CurrentRoll -= CenterRoll;
+
+		Attitude = FVector(CurrentPitch, 0, CurrentRoll);
+		RotationRate = FVector(LastPitch - CurrentPitch, 0, LastRoll - CurrentRoll);
+		Gravity = FVector(0, 0, 0);
+
+		// use the raw acceleration for acceleration
+		Acceleration = NewAcceleration;
+
+		// remember for next time (for rotation rate)
+		LastPitch = CurrentPitch;
+		LastRoll = CurrentRoll;
+	}
+#endif
+}
+
+void FIOSInputInterface::CalibrateMotion(uint32 PlayerIndex)
+{
+#if !PLATFORM_TVOS
+	// If we are using the motion manager, grab a reference frame.  Note, once you set the Attitude Reference frame
+	// all additional reference information will come from it
+	if (MotionManager.deviceMotionActive)
+	{
+		ReferenceAttitude = [MotionManager.deviceMotion.attitude retain];
+	}
+	else
+	{
+		bIsCalibrationRequested = true;
+	}
+#endif
+
+	if (PlayerIndex >= 0 && PlayerIndex < ARRAY_COUNT(Controllers))
+	{
+		Controllers[PlayerIndex].bNeedsReferenceAttitude = true;
+	}
+}
+
+bool FIOSInputInterface::Exec(UWorld* InWorld, const TCHAR* Cmd, FOutputDevice& Ar)
+{
+	// Keep track whether the command was handled or not.
+	bool bHandledCommand = false;
+
+	if (FParse::Command(&Cmd, TEXT("CALIBRATEMOTION")))
+	{
+		uint32 PlayerIndex = FCString::Atoi(Cmd);
+		CalibrateMotion(PlayerIndex);
+		bHandledCommand = true;
+	}
+
+	return bHandledCommand;
+}
+bool FIOSInputInterface::IsControllerAssignedToGamepad(int32 ControllerId)
+{
+	return ControllerId < ARRAY_COUNT(Controllers) &&
+		(Controllers[ControllerId].bIsGamepadConnected ||
+		 Controllers[ControllerId].bIsRemoteConnected);
+}
+
+void FIOSInputInterface::SetForceFeedbackChannelValue(int32 ControllerId, FForceFeedbackChannelType ChannelType, float Value)
+{
+	if (Value >= 0.3f)
+	{
+		AudioServicesPlaySystemSound(kSystemSoundID_Vibrate);
+	}
+}
+
+void FIOSInputInterface::SetForceFeedbackChannelValues(int32 ControllerId, const FForceFeedbackValues &Values)
+{
+	// Use largest vibration state as value
+	float MaxLeft = Values.LeftLarge > Values.LeftSmall ? Values.LeftLarge : Values.LeftSmall;
+	float MaxRight = Values.RightLarge > Values.RightSmall ? Values.RightLarge : Values.RightSmall;
+	float Value = MaxLeft > MaxRight ? MaxLeft : MaxRight;
+
+	if (Value >= 0.3f)
+	{
+		AudioServicesPlaySystemSound(kSystemSoundID_Vibrate);
+	}
+}
diff --git a/Engine/Source/Runtime/Core/Private/IOS/IOSMallocZone.cpp b/Engine/Source/Runtime/Core/Private/IOS/IOSMallocZone.cpp
new file mode 100644
--- /dev/null
+++ b/Engine/Source/Runtime/Core/Private/IOS/IOSMallocZone.cpp
@@ -0,0 +1,192 @@
+// Copyright 1998-2017 Epic Games, Inc. All Rights Reserved.
+
+#include "IOSMallocZone.h"
+#include "IOSPlatformMisc.h"
+#include "IOSPlatformCrashContext.h"
+#include "HAL/PlatformTLS.h"
+#include "HAL/PlatformProcess.h"
+#include <malloc/malloc.h>
+
+FIOSMallocZone::FIOSMallocZone( uint64 const InitialSize )
+: MemoryZone( malloc_create_zone( InitialSize, 0 ) )
+{
+}
+
+FIOSMallocZone::~FIOSMallocZone()
+{
+	if ( MemoryZone )
+	{
+		malloc_destroy_zone( MemoryZone );
+	}
+}
+
+void* FIOSMallocZone::Malloc( SIZE_T Size, uint32 Alignment )
+{
+	check( MemoryZone );
+	void* Result = malloc_zone_malloc( MemoryZone, Size );
+	return Result;
+}
+
+void* FIOSMallocZone::Realloc( void* Ptr, SIZE_T NewSize, uint32 Alignment )
+{
+	check( MemoryZone );
+	void* Result = malloc_zone_realloc( MemoryZone, Ptr, NewSize );
+	return Result;
+}
+
+void FIOSMallocZone::Free( void* Ptr )
+{
+	check( MemoryZone );
+	malloc_zone_free( MemoryZone, Ptr );
+}
+
+bool FIOSMallocZone::GetAllocationSize(void *Original, SIZE_T &SizeOut)
+{
+	SizeOut = malloc_size( Original );
+	return SizeOut > 0;
+}
+
+bool FIOSMallocZone::IsInternallyThreadSafe() const
+{
+	return true;
+}
+
+bool FIOSMallocZone::ValidateHeap()
+{
+	check( MemoryZone );
+	return malloc_zone_check( MemoryZone ) ;
+}
+
+const TCHAR* FIOSMallocZone::GetDescriptiveName()
+{
+	return TEXT("MallocZone");
+}
+
+FIOSMallocCrashHandler::FIOSMallocCrashHandler( uint64 const InitialSize )
+: FIOSMallocZone( InitialSize )
+, OriginalHeap( GMalloc )
+, CrashContext( nullptr )
+, ThreadId( -1 )
+{
+	check( GMalloc );
+}
+
+FIOSMallocCrashHandler::~FIOSMallocCrashHandler()
+{
+	// We crashed, so don't try and tidy the malloc zone
+	if ( ThreadId != -1 )
+	{
+		MemoryZone = nullptr;
+	}
+}
+	
+void FIOSMallocCrashHandler::Enable( FIOSCrashContext* Context, uint32 CrashedThreadId )
+{
+	check(Context);
+	CrashContext = Context;
+	ThreadId = CrashedThreadId;
+	OriginalHeap = GMalloc;
+	if (PLATFORM_USES_FIXED_GMalloc_CLASS && GFixedMallocLocationPtr)
+	{
+		*GFixedMallocLocationPtr = nullptr; // this disables any fast-path inline allocators
+	}
+	GMalloc = this;
+}
+	
+void* FIOSMallocCrashHandler::Malloc( SIZE_T Size, uint32 Alignment )
+{
+	void* Result = nullptr;
+	if ( IsOnCrashedThread() )
+	{
+		Result = FIOSMallocZone::Malloc( Size, Alignment );
+		if( !Result )
+		{
+			check(CrashContext);
+			CrashContext->GenerateCrashInfo();
+		}
+	}
+	return Result;
+}
+
+void* FIOSMallocCrashHandler::Realloc( void* Ptr, SIZE_T NewSize, uint32 Alignment )
+{
+	void* Result = nullptr;
+	if ( IsOnCrashedThread() )
+	{
+		if ( Ptr == nullptr || MemoryZone->size( MemoryZone, Ptr ) > 0 )
+		{
+			Result = FIOSMallocZone::Realloc( Ptr, NewSize, Alignment );
+			if( NewSize && !Result )
+			{
+				check(CrashContext);
+				CrashContext->GenerateCrashInfo();
+			}
+		}
+		else if ( NewSize )
+		{
+			if ( FCStringWide::Strcmp( OriginalHeap->GetDescriptiveName(), TEXT("ANSI") ) != 0 )
+			{
+				SIZE_T OldSize = 0;
+				if ( OriginalHeap->GetAllocationSize( Ptr, OldSize ) )
+				{
+					Result = Malloc( NewSize, Alignment );
+					if ( Result )
+					{
+						FMemory::Memcpy( Result, Ptr, FMath::Min( NewSize, OldSize ) );
+					}
+					else
+					{
+						check(CrashContext);
+						CrashContext->GenerateCrashInfo();
+					}
+				}
+			}
+			else // Can't safely handle this, so just report & exit
+			{
+				check(CrashContext);
+				CrashContext->GenerateCrashInfo();
+			}
+		}
+	}
+	return Result;
+}
+
+void FIOSMallocCrashHandler::Free( void* Ptr )
+{
+	if ( IsOnCrashedThread() && Ptr && MemoryZone->size( MemoryZone, Ptr ) > 0 )
+	{
+		FIOSMallocZone::Free( Ptr );
+	}
+}
+
+bool FIOSMallocCrashHandler::GetAllocationSize( void *Original, SIZE_T &SizeOut )
+{
+	SizeOut = 0;
+	if ( IsOnCrashedThread() && Original )
+	{
+		SizeOut = MemoryZone->size( MemoryZone, Original );
+	}
+	return SizeOut > 0;
+}
+
+const TCHAR* FIOSMallocCrashHandler::GetDescriptiveName()
+{
+	return TEXT("MallocCrashHandler");
+}
+
+bool FIOSMallocCrashHandler::IsOnCrashedThread( void )
+{
+	// Suspend threads other than the crashed one to prevent serious memory errors.
+	// Only the crashed thread can do anything meaningful from here anyway.
+	if ( ThreadId == FPlatformTLS::GetCurrentThreadId() )
+	{
+		return true;
+	}
+	else
+	{
+		do {
+			FPlatformProcess::SleepInfinite();
+		} while (true);
+		return false;
+	}
+}
diff --git a/Engine/Source/Runtime/Core/Private/IOS/IOSPlatformFile.cpp b/Engine/Source/Runtime/Core/Private/IOS/IOSPlatformFile.cpp
new file mode 100644
--- /dev/null
+++ b/Engine/Source/Runtime/Core/Private/IOS/IOSPlatformFile.cpp
@@ -0,0 +1,819 @@
+// Copyright 1998-2017 Epic Games, Inc. All Rights Reserved.
+
+#include "IOSPlatformFile.h"
+#include "HAL/PlatformTLS.h"
+#include "Containers/StringConv.h"
+#include "HAL/PlatformTime.h"
+#include "Templates/Function.h"
+#include "HAL/PlatformProcess.h"
+#include "Misc/Parse.h"
+#include "Misc/CoreMisc.h"
+#include "Misc/CommandLine.h"
+
+
+//#if PLATFORM_IOS
+//#include "IPlatformFileSandboxWrapper.h"
+//#endif
+
+// make an FTimeSpan object that represents the "epoch" for time_t (from a stat struct)
+const FDateTime IOSEpoch(1970, 1, 1);
+
+namespace
+{
+	FFileStatData IOSStatToUEFileData(struct stat& FileInfo)
+	{
+		const bool bIsDirectory = S_ISDIR(FileInfo.st_mode);
+
+		int64 FileSize = -1;
+		if (!bIsDirectory)
+		{
+			FileSize = FileInfo.st_size;
+		}
+
+		return FFileStatData(
+			IOSEpoch + FTimespan(0, 0, FileInfo.st_ctime), 
+			IOSEpoch + FTimespan(0, 0, FileInfo.st_atime), 
+			IOSEpoch + FTimespan(0, 0, FileInfo.st_mtime), 
+			FileSize,
+			bIsDirectory,
+			!!(FileInfo.st_mode & S_IWUSR)
+			);
+	}
+}
+
+
+
+/* FIOSFileHandle class
+ *****************************************************************************/
+
+/** 
+ * Managed IOS file handle implementation which limits number of open files. This
+ * is to prevent running out of system file handles (700). Should not be neccessary when 
+ * using pak file (e.g., SHIPPING?) so not particularly optimized. Only manages 
+ * files which are opened READ_ONLY.
+ **/
+// @todo: Merge all of the managed file handles into one class!
+#define MANAGE_FILE_HANDLES_IOS 1 // !UE_BUILD_SHIPPING
+
+struct FManagedFile
+{
+	int32 Handle;
+	uint32 ID;
+	double AccessTime;
+};
+
+class FIOSFileHandle : public IFileHandle
+{
+	static const int32 READWRITE_SIZE = 1024 * 1024;
+	static const int32 ACTIVE_HANDLE_COUNT_PER_THREAD = 100;
+
+public:
+
+	FIOSFileHandle( int32 InFileHandle, const FString& InFilename, bool bIsForRead )
+		: FileHandle(InFileHandle)
+#if !UE_BUILD_SHIPPING || MANAGE_FILE_HANDLES_IOS
+		, Filename(InFilename)
+#endif
+#if MANAGE_FILE_HANDLES_IOS
+        , HandleSlot(-1)
+        , FileSize(0)
+#endif
+	{
+		check(FileHandle != 0);
+
+#if MANAGE_FILE_HANDLES_IOS
+
+		static uint32 NextID = 1;
+		FileID = NextID++;
+
+		// get the per-thread buffers
+		ManagedFiles = (FManagedFile*)FPlatformTLS::GetTlsValue(ManagedFilesTlsSlot);
+
+		// not made yet on this thread, make the buffers now
+		if (ManagedFiles == NULL)
+		{
+			ManagedFiles = new FManagedFile[ACTIVE_HANDLE_COUNT_PER_THREAD];
+			FMemory::Memzero(ManagedFiles, sizeof(FManagedFile) * ACTIVE_HANDLE_COUNT_PER_THREAD);
+			FPlatformTLS::SetTlsValue(ManagedFilesTlsSlot, ManagedFiles);
+		}
+
+        // Only files opened for read will be managed
+        if( bIsForRead )
+        {
+            ReserveSlot();
+            ManagedFiles[HandleSlot].Handle = FileHandle;
+
+			struct stat FileInfo;
+			FileInfo.st_size = -1;
+			// check the read path
+			fstat(FileHandle, &FileInfo);
+			FileSize = FileInfo.st_size;
+        }
+#endif
+
+		Seek(0);
+	}
+
+	/**
+	 * Destructor.
+	 */
+	virtual ~FIOSFileHandle( )
+	{
+#if MANAGE_FILE_HANDLES_IOS
+        if( IsManaged() )
+        {
+            if( ManagedFiles[HandleSlot].ID == FileID )
+            {
+                close(FileHandle);
+				ManagedFiles[HandleSlot].ID = 0;
+            }
+        }
+        else
+#endif
+        {
+		    close(FileHandle);
+        }
+
+		FileHandle = -1;
+	}
+
+	bool InternalRead(uint8* Destination, int64 BytesToRead)
+	{
+		while (BytesToRead)
+		{
+			check(BytesToRead >= 0);
+			int64 ThisSize = FMath::Min<int64>(READWRITE_SIZE, BytesToRead);
+			check(Destination);
+			if (read(FileHandle, Destination, ThisSize) != ThisSize)
+			{
+				return false;
+			}
+			Destination += ThisSize;
+			BytesToRead -= ThisSize;
+		}
+		return true;
+	}
+
+public:
+
+	// Begin IFileHandle interface
+
+	virtual bool Read( uint8* Destination, int64 BytesToRead ) override
+	{
+ #if MANAGE_FILE_HANDLES_IOS
+       if( IsManaged() )
+        {
+            ActivateSlot();
+			lseek(FileHandle, FileOffset, SEEK_SET);
+			// read into the buffer, and make sure it worked
+			if (InternalRead(Destination, BytesToRead))
+			{
+				FileOffset += BytesToRead;
+				return true;
+			}
+			return false;
+        }
+        else
+#endif
+        {
+		    return InternalRead(Destination, BytesToRead);
+        }
+	}
+
+	virtual bool Seek( int64 NewPosition ) override
+	{
+		check(NewPosition >= 0);
+
+#if MANAGE_FILE_HANDLES_IOS
+        if( IsManaged() )
+        {
+            FileOffset = NewPosition >= FileSize ? FileSize - 1 : NewPosition;
+            return true;
+        }
+        else
+#endif
+        {
+		    return (lseek(FileHandle, NewPosition, SEEK_SET) != -1);
+        }
+	}
+
+	virtual bool SeekFromEnd( int64 NewPositionRelativeToEnd = 0 ) override
+	{
+		check(NewPositionRelativeToEnd <= 0);
+
+#if MANAGE_FILE_HANDLES_IOS
+        if( IsManaged() )
+        {
+            FileOffset = (NewPositionRelativeToEnd >= FileSize) ? 0 : ( FileSize + NewPositionRelativeToEnd - 1 );
+            return true;
+        }
+        else
+#endif
+        {
+		    return lseek(FileHandle, NewPositionRelativeToEnd, SEEK_END) != -1;
+        }
+	}
+
+	virtual int64 Size( ) override
+	{
+#if MANAGE_FILE_HANDLES_IOS
+        if( IsManaged() )
+        {
+            return FileSize;
+        }
+        else
+#endif
+        {
+		    return IFileHandle::Size();
+        }
+	}
+
+	virtual int64 Tell( ) override
+	{
+#if MANAGE_FILE_HANDLES_IOS
+        if( IsManaged() )
+        {
+            return FileOffset;
+        }
+        else
+#endif
+        {
+		    return lseek(FileHandle, 0, SEEK_CUR);
+        }
+	}
+
+	virtual bool Write( const uint8* Source, int64 BytesToWrite ) override
+	{
+		while (BytesToWrite)
+		{
+			check(BytesToWrite >= 0);
+			int64 ThisSize = FMath::Min<int64>(READWRITE_SIZE, BytesToWrite);
+			check(Source);
+			if (write(FileHandle, Source, ThisSize) != ThisSize)
+			{
+				return false;
+			}
+			Source += ThisSize;
+			BytesToWrite -= ThisSize;
+		}
+		return true;
+	}
+
+	// End IFileHandle interface
+
+private:
+
+#if MANAGE_FILE_HANDLES_IOS
+    FORCEINLINE bool IsManaged()
+    {
+        return HandleSlot != -1;
+    }
+
+    void ActivateSlot()
+    {
+        if( IsManaged() )
+        {
+            if( ManagedFiles[HandleSlot].ID != FileID )
+            {
+				ReserveSlot();
+                
+				FileHandle = open(TCHAR_TO_UTF8(*Filename), O_RDONLY);
+
+				if (FileHandle != -1)
+                {
+                    ManagedFiles[HandleSlot].Handle = FileHandle;
+                }
+            }
+            else
+            {
+                ManagedFiles[HandleSlot].AccessTime = FPlatformTime::Seconds();
+            }
+        }
+    }
+
+    void ReserveSlot()
+    {
+        HandleSlot = -1;
+
+        // Look for non-reserved slot
+        for( int32 i = 0; i < ACTIVE_HANDLE_COUNT_PER_THREAD; ++i )
+        {
+            if( ManagedFiles[i].ID == 0 )
+            {
+                HandleSlot = i;
+                break;
+            }
+        }
+
+        // Take the oldest handle
+        if( HandleSlot == -1 )
+        {
+            int32 Oldest = 0;
+            for( int32 i = 1; i < ACTIVE_HANDLE_COUNT_PER_THREAD; ++i )
+            {
+                if( ManagedFiles[Oldest].AccessTime > ManagedFiles[i].AccessTime )
+                {
+                    Oldest = i;
+                }
+            }
+
+            close( ManagedFiles[Oldest].Handle );
+            HandleSlot = Oldest;
+        }
+
+        ManagedFiles[HandleSlot].ID = FileID;
+        ManagedFiles[HandleSlot].AccessTime = FPlatformTime::Seconds();
+    }
+#endif
+
+private:
+
+	// Holds the internal file handle.
+	int32 FileHandle;
+
+#if !UE_BUILD_SHIPPING || MANAGE_FILE_HANDLES_IOS
+	// Holds the name of the file that this handle represents. Kept around for possible reopen of file.
+	FString Filename;
+#endif
+
+#if MANAGE_FILE_HANDLES_IOS
+    // Most recent valid slot index for this handle; >=0 for handles which are managed.
+    int32 HandleSlot;
+
+    // Current file offset; valid iff a managed handle.
+    int64 FileOffset;
+
+    // Cached file size; valid iff a managed handle.
+    int64 FileSize;
+
+    // Each thread keeps a collection of active handles with access times.
+    FManagedFile* ManagedFiles;
+
+	// Unique FileID for this file (since handles aren't unique)
+	uint32 FileID;
+
+	static int32 ManagedFilesTlsSlot;
+#endif
+};
+
+#if MANAGE_FILE_HANDLES_IOS
+int32 FIOSFileHandle::ManagedFilesTlsSlot = FPlatformTLS::AllocTlsSlot();
+#endif
+
+
+
+
+
+/**
+ * iOS File I/O implementation
+**/
+bool Initialize(IPlatformFile* Inner, const TCHAR* CommandLineParam)
+{
+	return true;
+}
+
+FString FIOSPlatformFile::NormalizeFilename(const TCHAR* Filename)
+{
+	FString Result(Filename);
+	Result.ReplaceInline(TEXT("\\"), TEXT("/"));
+	return Result;
+}
+
+FString FIOSPlatformFile::NormalizeDirectory(const TCHAR* Directory)
+{
+	FString Result(Directory);
+	Result.ReplaceInline(TEXT("\\"), TEXT("/"));
+	if (Result.EndsWith(TEXT("/")))
+	{
+		Result.LeftChop(1);
+	}
+	return Result;
+}
+
+
+FString FIOSPlatformFile::ConvertToAbsolutePathForExternalAppForRead( const TCHAR* Filename )
+{
+    struct stat FileInfo;
+    FString NormalizedFilename = NormalizeFilename(Filename);
+    if (stat(TCHAR_TO_UTF8(*ConvertToIOSPath(NormalizedFilename, false)), &FileInfo) == -1)
+    {
+        return ConvertToAbsolutePathForExternalAppForWrite(Filename);
+    }
+    else
+    {
+        return ConvertToIOSPath(NormalizedFilename, false);
+    }
+}
+
+FString FIOSPlatformFile::ConvertToAbsolutePathForExternalAppForWrite( const TCHAR* Filename )
+{
+    FString NormalizedFilename = NormalizeFilename(Filename);
+    return ConvertToIOSPath(NormalizedFilename, true);
+}
+
+bool FIOSPlatformFile::FileExists(const TCHAR* Filename)
+{
+	struct stat FileInfo;
+	FString NormalizedFilename = NormalizeFilename(Filename);
+	// check the read path
+	if (stat(TCHAR_TO_UTF8(*ConvertToIOSPath(NormalizedFilename, false)), &FileInfo) == -1)
+	{
+		// if not in read path, check the write path
+		if (stat(TCHAR_TO_UTF8(*ConvertToIOSPath(NormalizedFilename, true)), &FileInfo) == -1)
+		{
+			return false;
+		}
+	}
+
+	return S_ISREG(FileInfo.st_mode);
+}
+
+int64 FIOSPlatformFile::FileSize(const TCHAR* Filename)
+{
+	struct stat FileInfo;
+	FileInfo.st_size = -1;
+	FString NormalizedFilename = NormalizeFilename(Filename);
+	// check the read path
+	if(stat(TCHAR_TO_UTF8(*ConvertToIOSPath(NormalizedFilename, false)), &FileInfo) == -1)
+	{
+		// if not in read path, check the write path
+		stat(TCHAR_TO_UTF8(*ConvertToIOSPath(NormalizedFilename, true)), &FileInfo);
+	}
+
+	// make sure to return -1 for directories
+	if (S_ISDIR(FileInfo.st_mode))
+	{
+		FileInfo.st_size = -1;
+	}
+	return FileInfo.st_size;
+}
+
+bool FIOSPlatformFile::DeleteFile(const TCHAR* Filename)
+{
+	// only delete from write path
+	FString IOSFilename = ConvertToIOSPath(NormalizeFilename(Filename), true);
+	return unlink(TCHAR_TO_UTF8(*IOSFilename)) == 0;
+}
+
+bool FIOSPlatformFile::IsReadOnly(const TCHAR* Filename)
+{
+	FString NormalizedFilename = NormalizeFilename(Filename);
+	FString Filepath = ConvertToIOSPath(NormalizedFilename, false);
+	// check read path
+	if (access(TCHAR_TO_UTF8(*Filepath), F_OK) == -1)
+	{
+		// if not in read path, check write path
+		Filepath = ConvertToIOSPath(NormalizedFilename, true);
+		if (access(TCHAR_TO_UTF8(*Filepath), F_OK) == -1)
+		{
+			return false; // file doesn't exist
+		}
+	}
+
+	if (access(TCHAR_TO_UTF8(*Filepath), W_OK) == -1)
+	{
+		return errno == EACCES;
+	}
+	return false;
+}
+
+bool FIOSPlatformFile::MoveFile(const TCHAR* To, const TCHAR* From)
+{
+	// move to the write path
+	FString ToIOSFilename = ConvertToIOSPath(NormalizeFilename(To), true);
+	FString FromIOSFilename = ConvertToIOSPath(NormalizeFilename(From), false);
+	return rename(TCHAR_TO_UTF8(*FromIOSFilename), TCHAR_TO_UTF8(*ToIOSFilename)) != -1;
+}
+
+bool FIOSPlatformFile::SetReadOnly(const TCHAR* Filename, bool bNewReadOnlyValue)
+{
+	struct stat FileInfo;
+	FString IOSFilename = ConvertToIOSPath(NormalizeFilename(Filename), false);
+	if (stat(TCHAR_TO_UTF8(*IOSFilename), &FileInfo) != -1)
+	{
+		if (bNewReadOnlyValue)
+		{
+			FileInfo.st_mode &= ~S_IWUSR;
+		}
+		else
+		{
+			FileInfo.st_mode |= S_IWUSR;
+		}
+		return chmod(TCHAR_TO_UTF8(*IOSFilename), FileInfo.st_mode) == 0;
+	}
+	return false;
+}
+
+
+FDateTime FIOSPlatformFile::GetTimeStamp(const TCHAR* Filename)
+{
+	// get file times
+	struct stat FileInfo;
+	FString NormalizedFilename = NormalizeFilename(Filename);
+	// check the read path
+	if(stat(TCHAR_TO_UTF8(*ConvertToIOSPath(NormalizedFilename, false)), &FileInfo) == -1)
+	{
+		// if not in the read path, check the write path
+		if(stat(TCHAR_TO_UTF8(*ConvertToIOSPath(NormalizedFilename, true)), &FileInfo) == -1)
+		{
+			return FDateTime::MinValue();
+		}
+	}
+
+	// convert _stat time to FDateTime
+	FTimespan TimeSinceEpoch(0, 0, FileInfo.st_mtime);
+	return IOSEpoch + TimeSinceEpoch;
+
+}
+
+void FIOSPlatformFile::SetTimeStamp(const TCHAR* Filename, const FDateTime DateTime)
+{
+	// get file times
+	struct stat FileInfo;
+	FString IOSFilename = ConvertToIOSPath(NormalizeFilename(Filename), true);
+	if(stat(TCHAR_TO_UTF8(*IOSFilename), &FileInfo) == -1)
+	{
+		return;
+	}
+
+	// change the modification time only
+	struct utimbuf Times;
+	Times.actime = FileInfo.st_atime;
+	Times.modtime = (DateTime - IOSEpoch).GetTotalSeconds();
+	utime(TCHAR_TO_UTF8(*IOSFilename), &Times);
+}
+
+FDateTime FIOSPlatformFile::GetAccessTimeStamp(const TCHAR* Filename)
+{
+	// get file times
+	struct stat FileInfo;
+	FString NormalizedFilename = NormalizeFilename(Filename);
+	// check the read path
+	if(stat(TCHAR_TO_UTF8(*ConvertToIOSPath(NormalizedFilename, false)), &FileInfo) == -1)
+	{
+		// if not in the read path, check the write path
+		if(stat(TCHAR_TO_UTF8(*ConvertToIOSPath(NormalizedFilename, true)), &FileInfo) == -1)
+		{
+			return FDateTime::MinValue();
+		}
+	}
+
+	// convert _stat time to FDateTime
+	FTimespan TimeSinceEpoch(0, 0, FileInfo.st_atime);
+	return IOSEpoch + TimeSinceEpoch;
+}
+
+FString FIOSPlatformFile::GetFilenameOnDisk(const TCHAR* Filename)
+{
+	return Filename;
+}
+
+FFileStatData FIOSPlatformFile::GetStatData(const TCHAR* FilenameOrDirectory)
+{
+	struct stat FileInfo;
+	FString NormalizedFilename = NormalizeFilename(FilenameOrDirectory);
+
+	// check the read path
+	if(stat(TCHAR_TO_UTF8(*ConvertToIOSPath(NormalizedFilename, false)), &FileInfo) == -1)
+	{
+		// if not in the read path, check the write path
+		if(stat(TCHAR_TO_UTF8(*ConvertToIOSPath(NormalizedFilename, true)), &FileInfo) == -1)
+		{
+			return FFileStatData();
+		}
+	}
+
+	return IOSStatToUEFileData(FileInfo);
+}
+
+IFileHandle* FIOSPlatformFile::OpenRead(const TCHAR* Filename, bool bAllowWrite)
+{
+	FString NormalizedFilename = NormalizeFilename(Filename);
+
+	// check the read path
+	FString FinalPath = ConvertToIOSPath(NormalizedFilename, false);
+	int32 Handle = open(TCHAR_TO_UTF8(*FinalPath), O_RDONLY);
+	if(Handle == -1)
+	{
+		// if not in the read path, check the write path
+		FinalPath = ConvertToIOSPath(NormalizedFilename, true);
+		Handle = open(TCHAR_TO_UTF8(*FinalPath), O_RDONLY);
+	}
+
+	if (Handle != -1)
+	{
+		return new FIOSFileHandle(Handle, FinalPath, true);
+	}
+	return NULL;
+}
+
+IFileHandle* FIOSPlatformFile::OpenWrite(const TCHAR* Filename, bool bAppend, bool bAllowRead)
+{
+	int Flags = O_CREAT;
+	if (!bAppend)
+	{
+		Flags |= O_TRUNC;
+	}
+	if (bAllowRead)
+	{
+		Flags |= O_RDWR;
+	}
+	else
+	{
+		Flags |= O_WRONLY;
+	}
+	FString IOSFilename = ConvertToIOSPath(NormalizeFilename(Filename), true);
+	int32 Handle = open(TCHAR_TO_UTF8(*IOSFilename), Flags, S_IRUSR | S_IWUSR | S_IRGRP | S_IWGRP | S_IROTH | S_IWOTH);
+	if (Handle != -1)
+	{
+		FIOSFileHandle* FileHandleIOS = new FIOSFileHandle(Handle, IOSFilename, false);
+		if (bAppend)
+		{
+			FileHandleIOS->SeekFromEnd(0);
+		}
+		return FileHandleIOS;
+	}
+	return NULL;
+}
+
+bool FIOSPlatformFile::DirectoryExists(const TCHAR* Directory)
+{
+	struct stat FileInfo;
+	FString NormalizedDirectory = NormalizeFilename(Directory);
+	if (stat(TCHAR_TO_UTF8(*ConvertToIOSPath(NormalizedDirectory, false)), &FileInfo)== -1)
+	{
+		if (stat(TCHAR_TO_UTF8(*ConvertToIOSPath(NormalizedDirectory, true)), &FileInfo)== -1)
+		{
+			return false;
+		}
+	}
+	return S_ISDIR(FileInfo.st_mode);
+}
+
+bool FIOSPlatformFile::CreateDirectory(const TCHAR* Directory)
+{
+	FString IOSDirectory = ConvertToIOSPath(NormalizeFilename(Directory), true);
+	CFStringRef CFDirectory = FPlatformString::TCHARToCFString(*IOSDirectory);
+	bool Result = [[NSFileManager defaultManager] createDirectoryAtPath:(NSString*)CFDirectory withIntermediateDirectories:true attributes:nil error:nil];
+	CFRelease(CFDirectory);
+	return Result;
+}
+
+bool FIOSPlatformFile::DeleteDirectory(const TCHAR* Directory)
+{
+	FString IOSDirectory = ConvertToIOSPath(NormalizeFilename(Directory), true);
+	return rmdir(TCHAR_TO_UTF8(*IOSDirectory));
+}
+
+bool FIOSPlatformFile::IterateDirectory(const TCHAR* Directory, FDirectoryVisitor& Visitor)
+{
+	const FString DirectoryStr = Directory;
+	return IterateDirectoryCommon(Directory, [&](struct dirent* InEntry) -> bool
+	{
+		// Normalize any unicode forms so we match correctly
+		const FString NormalizedFilename = UTF8_TO_TCHAR(([[[NSString stringWithUTF8String:InEntry->d_name] precomposedStringWithCanonicalMapping] cStringUsingEncoding:NSUTF8StringEncoding]));
+		const FString FullPath = DirectoryStr / NormalizedFilename;
+
+		return Visitor.Visit(*FullPath, InEntry->d_type == DT_DIR);
+	});
+}
+
+bool FIOSPlatformFile::IterateDirectoryStat(const TCHAR* Directory, FDirectoryStatVisitor& Visitor)
+{
+	const FString DirectoryStr = Directory;
+	const FString NormalizedDirectoryStr = NormalizeFilename(Directory);
+
+	return IterateDirectoryCommon(Directory, [&](struct dirent* InEntry) -> bool
+	{
+		// Normalize any unicode forms so we match correctly
+		const FString NormalizedFilename = UTF8_TO_TCHAR(([[[NSString stringWithUTF8String:InEntry->d_name] precomposedStringWithCanonicalMapping] cStringUsingEncoding:NSUTF8StringEncoding]));
+		const FString FullPath = DirectoryStr / NormalizedFilename;
+		const FString FullNormalizedPath = NormalizedDirectoryStr / NormalizedFilename;
+
+		struct stat FileInfo;
+
+		// check the read path
+		if(stat(TCHAR_TO_UTF8(*ConvertToIOSPath(FullNormalizedPath, false)), &FileInfo) == -1)
+		{
+			// if not in the read path, check the write path
+			if(stat(TCHAR_TO_UTF8(*ConvertToIOSPath(FullNormalizedPath, true)), &FileInfo) == -1)
+			{
+				return true;
+			}
+		}
+
+		return Visitor.Visit(*FullPath, IOSStatToUEFileData(FileInfo));
+	});
+}
+
+bool FIOSPlatformFile::IterateDirectoryCommon(const TCHAR* Directory, const TFunctionRef<bool(struct dirent*)>& Visitor)
+{
+	bool Result = false;
+	const ANSICHAR* FrameworksPath;
+	if (Directory[0] == 0)
+	{
+		if ([[[[NSBundle mainBundle] bundlePath] pathExtension] isEqual: @"app"])
+		{
+			FrameworksPath = [[[NSBundle mainBundle] privateFrameworksPath] fileSystemRepresentation];
+		}
+		else
+		{
+			FrameworksPath = [[[NSBundle mainBundle] bundlePath] fileSystemRepresentation];
+		}
+	}
+
+	FString NormalizedDirectory = NormalizeFilename(Directory);
+	// If Directory is an empty string, assume that we want to iterate Binaries/Mac (current dir), but because we're an app bundle, iterate bundle's Contents/Frameworks instead
+	DIR* Handle = opendir(Directory[0] ? TCHAR_TO_UTF8(*ConvertToIOSPath(NormalizedDirectory, false)) : FrameworksPath);
+	if(!Handle)
+	{
+		// look in the write file path if it's not in the read file path
+		Handle = opendir(Directory[0] ? TCHAR_TO_UTF8(*ConvertToIOSPath(NormalizedDirectory, true)) : FrameworksPath);
+	}
+	if (Handle)
+	{
+		Result = true;
+		struct dirent *Entry;
+		while ((Entry = readdir(Handle)) != NULL)
+		{
+			if (FCStringAnsi::Strcmp(Entry->d_name, ".") && FCStringAnsi::Strcmp(Entry->d_name, ".."))
+			{
+				Result = Visitor(Entry);
+			}
+		}
+		closedir(Handle);
+	}
+	return Result;
+}
+
+FString FIOSPlatformFile::ConvertToIOSPath(const FString& Filename, bool bForWrite)
+{
+	FString Result = Filename;
+    if (Result.Contains(TEXT("/OnDemandResources/")))
+    {
+        return Result;
+    }
+    
+	Result.ReplaceInline(TEXT("../"), TEXT(""));
+	Result.ReplaceInline(TEXT(".."), TEXT(""));
+	Result.ReplaceInline(FPlatformProcess::BaseDir(), TEXT(""));
+
+	if(bForWrite)
+	{
+		static FString WritePathBase = FString([NSSearchPathForDirectoriesInDomains(NSCachesDirectory, NSUserDomainMask, YES) objectAtIndex:0]) + TEXT("/");
+		return WritePathBase + Result;
+	}
+	else
+	{
+		// if filehostip exists in the command line, cook on the fly read path should be used
+		FString Value;
+		// Cache this value as the command line doesn't change...
+		static bool bHasHostIP = FParse::Value(FCommandLine::Get(), TEXT("filehostip"), Value) || FParse::Value(FCommandLine::Get(), TEXT("streaminghostip"), Value);
+		static bool bIsIterative = FParse::Value(FCommandLine::Get(), TEXT("iterative"), Value);
+		if (bHasHostIP)
+		{
+			static FString ReadPathBase = FString([NSSearchPathForDirectoriesInDomains(NSCachesDirectory, NSUserDomainMask, YES) objectAtIndex:0]) + TEXT("/");
+			return ReadPathBase + Result;
+		}
+		else if (bIsIterative)
+		{
+			static FString ReadPathBase = FString([NSSearchPathForDirectoriesInDomains(NSCachesDirectory, NSUserDomainMask, YES) objectAtIndex:0]) + TEXT("/");
+			return ReadPathBase + Result.ToLower();
+		}
+		else
+		{
+			FString ReadPathBase = FString([[NSBundle mainBundle] bundlePath]) + TEXT("/cookeddata/");
+            FString OutString = ReadPathBase + Result.ToLower();
+            return OutString;
+		}
+	}
+
+	return Result;
+}
+
+
+/**
+ * iOS platform file declaration
+**/
+IPlatformFile& IPlatformFile::GetPlatformPhysical()
+{
+	static FIOSPlatformFile IOSPlatformSingleton;
+	//static FSandboxPlatformFile CookOnTheFlySandboxSingleton(false);
+	//static FSandboxPlatformFile CookByTheBookSandboxSingleton(false);
+	//static bool bInitialized = false;
+	//if(!bInitialized)
+	//{
+	//	NSString *DocumentsPath = [NSSearchPathForDirectoriesInDomains(NSDocumentDirectory, NSUserDomainMask, YES) objectAtIndex:0];
+	//	CookOnTheFlySandboxSingleton.Initialize(&IOSPlatformSingleton, ANSI_TO_TCHAR([DocumentsPath cStringUsingEncoding:NSASCIIStringEncoding]));
+
+	//	DocumentsPath = [[NSBundle mainBundle] bundlePath];
+	//	DocumentsPath = [DocumentsPath stringByAppendingPathComponent:@"/CookedContent/"];
+	//	CookByTheBookSandboxSingleton.Initialize(&CookOnTheFlySandboxSingleton, ANSI_TO_TCHAR([DocumentsPath cStringUsingEncoding:NSASCIIStringEncoding]));
+
+	//	bInitialized = true;
+	//}
+
+	//return CookByTheBookSandboxSingleton;
+
+	return IOSPlatformSingleton;
+}
diff --git a/Engine/Source/Runtime/Core/Private/IOS/IOSPlatformFile.h b/Engine/Source/Runtime/Core/Private/IOS/IOSPlatformFile.h
new file mode 100644
--- /dev/null
+++ b/Engine/Source/Runtime/Core/Private/IOS/IOSPlatformFile.h
@@ -0,0 +1,58 @@
+// Copyright 1998-2017 Epic Games, Inc. All Rights Reserved.
+
+#pragma once
+
+#include "Apple/ApplePlatformFile.h"
+#include "UnrealString.h"
+
+class IFileHandle;
+class FApplePlatformFile;
+
+
+/**
+ * iOS File I/O implementation
+**/
+class CORE_API FIOSPlatformFile : public FApplePlatformFile
+{
+protected:
+	virtual FString NormalizeFilename(const TCHAR* Filename);
+	virtual FString NormalizeDirectory(const TCHAR* Directory);
+
+public:
+	//virtual bool Initialize(IPlatformFile* Inner, const TCHAR* CommandLineParam) override;
+
+	virtual bool FileExists(const TCHAR* Filename) override;
+	virtual int64 FileSize(const TCHAR* Filename) override;
+	virtual bool DeleteFile(const TCHAR* Filename) override;
+	virtual bool IsReadOnly(const TCHAR* Filename) override;
+	virtual bool MoveFile(const TCHAR* To, const TCHAR* From) override;
+	virtual bool SetReadOnly(const TCHAR* Filename, bool bNewReadOnlyValue) override;
+
+	virtual FDateTime GetTimeStamp(const TCHAR* Filename) override;
+	virtual void SetTimeStamp(const TCHAR* Filename, const FDateTime DateTime) override;
+	virtual FDateTime GetAccessTimeStamp(const TCHAR* Filename) override;
+	virtual FString GetFilenameOnDisk(const TCHAR* Filename) override;
+    
+    virtual FString ConvertToAbsolutePathForExternalAppForRead( const TCHAR* Filename ) override;
+    virtual FString ConvertToAbsolutePathForExternalAppForWrite( const TCHAR* Filename ) override;
+
+	virtual FFileStatData GetStatData(const TCHAR* FilenameOrDirectory) override;
+
+	virtual IFileHandle* OpenRead(const TCHAR* Filename, bool bAllowWrite = false) override;
+	virtual IFileHandle* OpenWrite(const TCHAR* Filename, bool bAppend = false, bool bAllowRead = false) override;
+
+	virtual bool DirectoryExists(const TCHAR* Directory) override;
+	virtual bool CreateDirectory(const TCHAR* Directory) override;
+	virtual bool DeleteDirectory(const TCHAR* Directory) override;
+
+	virtual bool IterateDirectory(const TCHAR* Directory, FDirectoryVisitor& Visitor) override;
+	virtual bool IterateDirectoryStat(const TCHAR* Directory, FDirectoryStatVisitor& Visitor) override;
+
+private:
+	bool IterateDirectoryCommon(const TCHAR* Directory, const TFunctionRef<bool(struct dirent*)>& Visitor);
+
+	FString ConvertToIOSPath(const FString& Filename, bool bForWrite);
+
+	FString ReadDirectory;
+	FString WriteDirectory;
+};
diff --git a/Engine/Source/Runtime/Core/Private/IOS/IOSPlatformFramePacer.cpp b/Engine/Source/Runtime/Core/Private/IOS/IOSPlatformFramePacer.cpp
new file mode 100644
--- /dev/null
+++ b/Engine/Source/Runtime/Core/Private/IOS/IOSPlatformFramePacer.cpp
@@ -0,0 +1,164 @@
+// Copyright 1998-2017 Epic Games, Inc. All Rights Reserved.
+
+#include "IOSPlatformFramePacer.h"
+#include "Containers/Array.h"
+#include "ThreadingBase.h"
+#include "Misc/ConfigCacheIni.h"
+#include "Misc/Parse.h"
+
+// Collection of events listening for this trigger.
+static TArray<FEvent*> ListeningEvents;
+
+
+/*******************************************************************
+ * FIOSFramePacer implementation
+ *******************************************************************/
+
+@interface FIOSFramePacer : NSObject
+{
+    @public
+	FEvent *FramePacerEvent;
+}
+
+-(void)run:(id)param;
+-(void)signal:(id)param;
+
+@end
+
+// @todo ios: Move these up into some shared header
+// __TV_OS_VERSION_MAX_ALLOWED is only defined when building for tvos, so we can use that to determine
+#if PLATFORM_TVOS
+
+#define UE4_HAS_IOS10 (defined(__TVOS_10_0) && __TV_OS_VERSION_MAX_ALLOWED >= __TVOS_10_0)
+#define UE4_HAS_IOS9 (defined(__TVOS_9_0) && __TV_OS_VERSION_MAX_ALLOWED >= __TVOS_9_0)
+#define UE4_TARGET_PRE_IOS10 (!defined(__TVOS_10_0_) || __TV_OS_VERSION_MIN_REQUIRED < __TVOS_10_0)
+#define UE4_TARGET_PRE_IOS9 (!defined(__TVOS_9_0_) || __TV_OS_VERSION_MIN_REQUIRED < __TVOS_9_0)
+
+#else
+
+#define UE4_HAS_IOS10 (defined(__IPHONE_10_0) && __IPHONE_OS_VERSION_MAX_ALLOWED >= __IPHONE_10_0)
+#define UE4_HAS_IOS9 (defined(__IPHONE_9_0) && __IPHONE_OS_VERSION_MAX_ALLOWED >= __IPHONE_9_0)
+#define UE4_TARGET_PRE_IOS10 (!defined(__IPHONE_10_0_) || __IPHONE_OS_VERSION_MIN_REQUIRED < __IPHONE_10_0)
+#define UE4_TARGET_PRE_IOS9 (!defined(__IPHONE_9_0_) || __IPHONE_OS_VERSION_MIN_REQUIRED < __IPHONE_9_0)
+
+#endif
+
+@implementation FIOSFramePacer
+
+#ifdef __GNUC__
+#pragma GCC diagnostic push
+#pragma GCC diagnostic ignored "-Wdeprecated-declarations"
+#endif
+-(void)run:(id)param
+{
+	NSRunLoop *runloop = [NSRunLoop currentRunLoop];
+	CADisplayLink *displayLink = [CADisplayLink displayLinkWithTarget:self selector:@selector(signal:)];
+#if UE4_HAS_IOS10
+	if ([displayLink respondsToSelector : @selector(preferredFramesPerSecond)] == YES)
+	{
+		displayLink.preferredFramesPerSecond = 60 / FIOSPlatformRHIFramePacer::FrameInterval;
+	}
+	else
+#endif
+	{
+#if UE4_TARGET_PRE_IOS10
+		displayLink.frameInterval = FIOSPlatformRHIFramePacer::FrameInterval;
+#endif
+	}
+
+	[displayLink addToRunLoop:runloop forMode:NSDefaultRunLoopMode];
+	[runloop run];
+}
+#ifdef __GNUC__
+#pragma GCC diagnostic pop
+#endif
+
+
+-(void)signal:(id)param
+{
+    for( auto& NextEvent : ListeningEvents )
+    {
+        NextEvent->Trigger();
+    }
+}
+
+@end
+
+
+
+/*******************************************************************
+ * FIOSPlatformRHIFramePacer implementation
+ *******************************************************************/
+
+
+namespace IOSDisplayConstants
+{
+    const uint32 MaxRefreshRate = 60;
+}
+
+uint32 FIOSPlatformRHIFramePacer::FrameInterval = 1;
+FIOSFramePacer* FIOSPlatformRHIFramePacer::FramePacer = nil;
+
+
+bool FIOSPlatformRHIFramePacer::IsEnabled()
+{
+    static bool bIsRHIFramePacerEnabled = false;
+	static bool bInitialized = false;
+
+	if (!bInitialized)
+	{
+		FString FrameRateLockAsEnum;
+		GConfig->GetString(TEXT("/Script/IOSRuntimeSettings.IOSRuntimeSettings"), TEXT("FrameRateLock"), FrameRateLockAsEnum, GEngineIni);
+
+		uint32 FrameRateLock = 60;
+		FParse::Value(*FrameRateLockAsEnum, TEXT("PUFRL_"), FrameRateLock);
+		if (FrameRateLock == 0)
+		{
+			FrameRateLock = 60;
+		}
+
+		if (!bIsRHIFramePacerEnabled)
+		{
+			check((IOSDisplayConstants::MaxRefreshRate % FrameRateLock) == 0);
+			FrameInterval = IOSDisplayConstants::MaxRefreshRate / FrameRateLock;
+
+			bIsRHIFramePacerEnabled = (FrameInterval > 0);
+		}
+		bInitialized = true;
+	}
+	
+	return bIsRHIFramePacerEnabled;
+}
+
+void FIOSPlatformRHIFramePacer::InitWithEvent(FEvent* TriggeredEvent)
+{
+    // Create display link thread
+    FramePacer = [[FIOSFramePacer alloc] init];
+    [NSThread detachNewThreadSelector:@selector(run:) toTarget:FramePacer withObject:nil];
+        
+    // Only one supported for now, we may want more eventually.
+    ListeningEvents.Add( TriggeredEvent );
+}
+
+void FIOSPlatformRHIFramePacer::Suspend()
+{
+    // send a signal to the events if we are enabled
+    if (IsEnabled())
+    {
+        [FramePacer signal:0];
+    }
+}
+
+void FIOSPlatformRHIFramePacer::Resume()
+{
+    
+}
+
+void FIOSPlatformRHIFramePacer::Destroy()
+{
+    if( FramePacer != nil )
+    {
+        [FramePacer release];
+        FramePacer = nil;
+    }
+}
\ No newline at end of file
diff --git a/Engine/Source/Runtime/Core/Private/IOS/IOSPlatformMemory.cpp b/Engine/Source/Runtime/Core/Private/IOS/IOSPlatformMemory.cpp
new file mode 100644
--- /dev/null
+++ b/Engine/Source/Runtime/Core/Private/IOS/IOSPlatformMemory.cpp
@@ -0,0 +1,161 @@
+// Copyright 1998-2017 Epic Games, Inc. All Rights Reserved.
+
+/*=============================================================================
+	IOSPlatformMemory.cpp: IOS platform memory functions
+=============================================================================*/
+
+#include "IOSPlatformMemory.h"
+#include "Misc/CoreStats.h"
+#include "MallocBinned.h"
+#include "MallocAnsi.h"
+#include "MemStack.h"
+
+#include <sys/mman.h>
+
+void FIOSPlatformMemory::Init()
+{
+	FGenericPlatformMemory::Init();
+
+	const FPlatformMemoryConstants& MemoryConstants = FPlatformMemory::GetConstants();
+	UE_LOG(LogInit, Log, TEXT("Memory total: Physical=%.1fGB (%dGB approx) Pagefile=%.1fGB Virtual=%.1fGB"), 
+		float(MemoryConstants.TotalPhysical/1024.0/1024.0/1024.0),
+		MemoryConstants.TotalPhysicalGB, 
+		float((MemoryConstants.TotalVirtual-MemoryConstants.TotalPhysical)/1024.0/1024.0/1024.0), 
+		float(MemoryConstants.TotalVirtual/1024.0/1024.0/1024.0) );
+}
+
+FPlatformMemoryStats FIOSPlatformMemory::GetStats()
+{
+	const FPlatformMemoryConstants& MemoryConstants = FPlatformMemory::GetConstants();
+
+    static SIZE_T PeakUsedPhysical = 0;
+    static SIZE_T PeakUsedVirtual = 0;
+
+	FPlatformMemoryStats MemoryStats;
+
+	// Gather system-wide memory stats.
+	vm_statistics Stats;
+	mach_msg_type_number_t StatsSize = sizeof(Stats);
+	host_statistics(mach_host_self(), HOST_VM_INFO, (host_info_t)&Stats, &StatsSize);
+	uint64_t FreeMem = Stats.free_count * MemoryConstants.PageSize;
+	uint64_t SysUsedMem = (Stats.active_count + Stats.inactive_count + Stats.wire_count) * MemoryConstants.PageSize;
+
+	// Gather process memory stats.
+	task_basic_info Info;
+	mach_msg_type_number_t InfoSize = sizeof(Info);
+	task_info(mach_task_self(), TASK_BASIC_INFO, (task_info_t)&Info, &InfoSize);
+
+    if (Info.resident_size > PeakUsedPhysical)
+    {
+        PeakUsedPhysical = Info.resident_size;
+    }
+    if (Info.virtual_size > PeakUsedVirtual)
+    {
+        PeakUsedVirtual = Info.virtual_size;
+    }
+    
+	MemoryStats.AvailablePhysical = FreeMem;
+	MemoryStats.AvailableVirtual = 0;
+	MemoryStats.UsedPhysical = Info.resident_size;
+	MemoryStats.UsedVirtual = Info.virtual_size;
+    MemoryStats.PeakUsedPhysical = PeakUsedPhysical;
+    MemoryStats.PeakUsedVirtual = PeakUsedVirtual;
+    
+	return MemoryStats;
+}
+
+const FPlatformMemoryConstants& FIOSPlatformMemory::GetConstants()
+{
+	static FPlatformMemoryConstants MemoryConstants;
+
+	if( MemoryConstants.TotalPhysical == 0 )
+	{
+		// Gather platform memory constants.
+		
+		// Get page size.
+		vm_size_t PageSize;
+		host_page_size(mach_host_self(), &PageSize);
+
+		// Get swap file info
+		xsw_usage SwapUsage;
+		SIZE_T Size = sizeof(SwapUsage);
+		sysctlbyname("vm.swapusage", &SwapUsage, &Size, NULL, 0);
+
+		// Get memory.
+		vm_statistics Stats;
+		mach_msg_type_number_t StatsSize = sizeof(Stats);
+		host_statistics(mach_host_self(), HOST_VM_INFO, (host_info_t)&Stats, &StatsSize);
+		uint64_t FreeMem = Stats.free_count * PageSize;
+		uint64_t UsedMem = (Stats.active_count + Stats.inactive_count + Stats.wire_count) * PageSize;
+		uint64_t TotalPhys = FreeMem + UsedMem;
+		uint64_t TotalPageFile = SwapUsage.xsu_total;
+		uint64_t TotalVirtual = TotalPhys + TotalPageFile;
+	
+		MemoryConstants.TotalPhysical = TotalPhys;
+		MemoryConstants.TotalVirtual = TotalVirtual;
+		MemoryConstants.PageSize = (uint32)PageSize;
+		MemoryConstants.OsAllocationGranularity = (uint32)PageSize;
+		MemoryConstants.BinnedPageSize = FMath::Max((SIZE_T)65536, (SIZE_T)PageSize);
+
+		MemoryConstants.TotalPhysicalGB = (MemoryConstants.TotalPhysical + 1024 * 1024 * 1024 - 1) / 1024 / 1024 / 1024;
+	}
+
+	return MemoryConstants;	
+}
+
+FMalloc* FIOSPlatformMemory::BaseAllocator()
+{
+	// get free memory
+	vm_statistics Stats;
+	mach_msg_type_number_t StatsSize = sizeof(Stats);
+	host_statistics(mach_host_self(), HOST_VM_INFO, (host_info_t)&Stats, &StatsSize);
+	// 1 << FMath::CeilLogTwo(MemoryConstants.TotalPhysical) should really be FMath::RoundUpToPowerOfTwo,
+	// but that overflows to 0 when MemoryConstants.TotalPhysical is close to 4GB, since CeilLogTwo returns 32
+	// this then causes the MemoryLimit to be 0 and crashing the app
+	uint64 MemoryLimit = FMath::Min<uint64>( uint64(1) << FMath::CeilLogTwo(Stats.free_count * GetConstants().PageSize), 0x100000000);
+
+	//return new FMallocAnsi();
+	// [RCL] 2017-03-06 FIXME: perhaps BinnedPageSize should be used here, but leaving this change to the iOS platform owner.
+	return new FMallocBinned(GetConstants().PageSize, MemoryLimit);
+}
+
+bool FIOSPlatformMemory::PageProtect(void* const Ptr, const SIZE_T Size, const bool bCanRead, const bool bCanWrite)
+{
+	int32 ProtectMode;
+	if (bCanRead && bCanWrite)
+	{
+		ProtectMode = PROT_READ | PROT_WRITE;
+	}
+	else if (bCanRead)
+	{
+		ProtectMode = PROT_READ;
+	}
+	else if (bCanWrite)
+	{
+		ProtectMode = PROT_WRITE;
+	}
+	else
+	{
+		ProtectMode = PROT_NONE;
+	}
+	return mprotect(Ptr, Size, ProtectMode) == 0;
+}
+
+void* FIOSPlatformMemory::BinnedAllocFromOS( SIZE_T Size )
+{
+	return mmap(nullptr, Size, PROT_READ | PROT_WRITE, MAP_PRIVATE | MAP_ANON, -1, 0);
+}
+
+void FIOSPlatformMemory::BinnedFreeToOS( void* Ptr, SIZE_T Size )
+{
+	if (munmap(Ptr, Size) != 0)
+	{
+		const int ErrNo = errno;
+		UE_LOG(LogHAL, Fatal, TEXT("munmap(addr=%p, len=%llu) failed with errno = %d (%s)"), Ptr, Size,
+			   ErrNo, StringCast< TCHAR >(strerror(ErrNo)).Get());
+	}
+}
+
+
+
+
diff --git a/Engine/Source/Runtime/Core/Private/IOS/IOSPlatformMisc.cpp b/Engine/Source/Runtime/Core/Private/IOS/IOSPlatformMisc.cpp
new file mode 100644
--- /dev/null
+++ b/Engine/Source/Runtime/Core/Private/IOS/IOSPlatformMisc.cpp
@@ -0,0 +1,2022 @@
+// Copyright 1998-2017 Epic Games, Inc. All Rights Reserved.
+
+/*=============================================================================
+	IOSPlatformMisc.mm: iOS implementations of misc functions
+=============================================================================*/
+
+#include "IOSPlatformMisc.h"
+#include "Misc/App.h"
+#include "ExceptionHandling.h"
+#include "SecureHash.h"
+#include "EngineVersion.h"
+#include "IOSMallocZone.h"
+#include "IOSApplication.h"
+#include "IOSAppDelegate.h"
+#include "IOSView.h"
+#include "IOSChunkInstaller.h"
+#include "IOSInputInterface.h"
+#include "Misc/CommandLine.h"
+#include "Misc/ConfigCacheIni.h"
+#include "Apple/ApplePlatformDebugEvents.h"
+#include "HAL/FileManager.h"
+#include "Misc/FileHelper.h"
+
+#include "Apple/ApplePlatformCrashContext.h"
+#include "IOSPlatformCrashContext.h"
+#if !PLATFORM_TVOS
+#include "PLCrashReporter.h"
+#include "PLCrashReport.h"
+#include "PLCrashReportTextFormatter.h"
+#endif
+#include "HAL/PlatformOutputDevices.h"
+#include "Misc/OutputDeviceError.h"
+#include "Misc/OutputDeviceRedirector.h"
+#include "Misc/FeedbackContext.h"
+
+#include "Internationalization/Internationalization.h"
+#include "Internationalization/Culture.h"
+
+#if !PLATFORM_TVOS
+#include <AdSupport/ASIdentifierManager.h> 
+#endif
+
+#include <SystemConfiguration/SystemConfiguration.h>
+#include <netinet/in.h>
+
+//#include <libproc.h>
+#include <mach-o/dyld.h>
+
+/** Amount of free memory in MB reported by the system at startup */
+CORE_API int32 GStartupFreeMemoryMB;
+
+extern CORE_API bool GIsGPUCrashed;
+
+/** Global pointer to memory warning handler */
+void (* GMemoryWarningHandler)(const FGenericMemoryWarningContext& Context) = NULL;
+
+/** global for showing the splash screen */
+bool GShowSplashScreen = true;
+
+static int32 GetFreeMemoryMB()
+{
+	// get free memory
+	vm_size_t PageSize;
+	host_page_size(mach_host_self(), &PageSize);
+
+	// get memory stats
+	vm_statistics Stats;
+	mach_msg_type_number_t StatsSize = sizeof(Stats);
+	host_statistics(mach_host_self(), HOST_VM_INFO, (host_info_t)&Stats, &StatsSize);
+	return (Stats.free_count * PageSize) / 1024 / 1024;
+}
+
+FIOSApplication* FPlatformMisc::CachedApplication = nullptr;
+
+void FIOSPlatformMisc::PlatformInit()
+{
+	FAppEntry::PlatformInit();
+
+	// Increase the maximum number of simultaneously open files
+	struct rlimit Limit;
+	Limit.rlim_cur = OPEN_MAX;
+	Limit.rlim_max = RLIM_INFINITY;
+	int32 Result = setrlimit(RLIMIT_NOFILE, &Limit);
+	check(Result == 0);
+
+	// Identity.
+	UE_LOG(LogInit, Log, TEXT("Computer: %s"), FPlatformProcess::ComputerName() );
+	UE_LOG(LogInit, Log, TEXT("User: %s"), FPlatformProcess::UserName() );
+
+	
+	const FPlatformMemoryConstants& MemoryConstants = FPlatformMemory::GetConstants();
+	UE_LOG(LogInit, Log, TEXT("CPU Page size=%i, Cores=%i"), MemoryConstants.PageSize, FPlatformMisc::NumberOfCores() );
+
+	// Timer resolution.
+	UE_LOG(LogInit, Log, TEXT("High frequency timer resolution =%f MHz"), 0.000001 / FPlatformTime::GetSecondsPerCycle() );
+	GStartupFreeMemoryMB = GetFreeMemoryMB();
+	UE_LOG(LogInit, Log, TEXT("Free Memory at startup: %d MB"), GStartupFreeMemoryMB);
+}
+
+void FIOSPlatformMisc::PlatformHandleSplashScreen(bool ShowSplashScreen)
+{
+//    GShowSplashScreen = ShowSplashScreen;
+}
+
+GenericApplication* FIOSPlatformMisc::CreateApplication()
+{
+	CachedApplication = FIOSApplication::CreateIOSApplication();
+	return CachedApplication;
+}
+
+void FIOSPlatformMisc::GetEnvironmentVariable(const TCHAR* VariableName, TCHAR* Result, int32 ResultLength)
+{
+	ANSICHAR *AnsiResult = getenv(TCHAR_TO_ANSI(VariableName));
+	if (AnsiResult)
+	{
+		wcsncpy(Result, ANSI_TO_TCHAR(AnsiResult), ResultLength);
+	}
+	else
+	{
+		*Result = 0;
+	}
+}
+
+// Make sure that SetStoredValue and GetStoredValue generate the same key
+static NSString* MakeStoredValueKeyName(const FString& SectionName, const FString& KeyName)
+{
+	return [NSString stringWithFString:(SectionName + "/" + KeyName)];
+}
+
+bool FIOSPlatformMisc::SetStoredValue(const FString& InStoreId, const FString& InSectionName, const FString& InKeyName, const FString& InValue)
+{
+	NSUserDefaults* UserSettings = [NSUserDefaults standardUserDefaults];
+
+	// convert input to an NSString
+	NSString* StoredValue = [NSString stringWithFString:InValue];
+
+	// store it
+	[UserSettings setObject:StoredValue forKey:MakeStoredValueKeyName(InSectionName, InKeyName)];
+
+	return true;
+}
+
+bool FIOSPlatformMisc::GetStoredValue(const FString& InStoreId, const FString& InSectionName, const FString& InKeyName, FString& OutValue)
+{
+	NSUserDefaults* UserSettings = [NSUserDefaults standardUserDefaults];
+	
+	// get the stored NSString
+	NSString* StoredValue = [UserSettings objectForKey:MakeStoredValueKeyName(InSectionName, InKeyName)];
+
+	// if it was there, convert back to FString
+	if (StoredValue != nil)
+	{
+		OutValue = StoredValue;
+		return true;
+	}
+
+	return false;
+}
+
+bool FIOSPlatformMisc::DeleteStoredValue(const FString& InStoreId, const FString& InSectionName, const FString& InKeyName)
+{
+	// No Implementation (currently only used by editor code so not needed on iOS)
+	return false;
+}
+
+//void FIOSPlatformMisc::LowLevelOutputDebugStringf(const TCHAR *Fmt, ... )
+//{
+//
+//}
+
+void FIOSPlatformMisc::LowLevelOutputDebugString( const TCHAR *Message )
+{
+	//NsLog will out to all iOS output consoles, instead of just the Xcode console.
+	NSLog(@"%@", [NSString stringWithUTF8String:TCHAR_TO_UTF8(Message)]);
+}
+
+const TCHAR* FIOSPlatformMisc::GetSystemErrorMessage(TCHAR* OutBuffer, int32 BufferCount, int32 Error)
+{
+	// There's no iOS equivalent for GetLastError()
+	check(OutBuffer && BufferCount);
+	*OutBuffer = TEXT('\0');
+	return OutBuffer;
+}
+
+void FIOSPlatformMisc::ClipboardCopy(const TCHAR* Str)
+{
+#if !PLATFORM_TVOS
+	CFStringRef CocoaString = FPlatformString::TCHARToCFString(Str);
+	UIPasteboard* Pasteboard = [UIPasteboard generalPasteboard];
+	[Pasteboard setString:(NSString*)CocoaString];
+#endif
+}
+
+void FIOSPlatformMisc::ClipboardPaste(class FString& Result)
+{
+#if !PLATFORM_TVOS
+	UIPasteboard* Pasteboard = [UIPasteboard generalPasteboard];
+	NSString* CocoaString = [Pasteboard string];
+	if(CocoaString)
+	{
+		TArray<TCHAR> Ch;
+		Ch.AddUninitialized([CocoaString length] + 1);
+		FPlatformString::CFStringToTCHAR((CFStringRef)CocoaString, Ch.GetData());
+		Result = Ch.GetData();
+	}
+	else
+	{
+		Result = TEXT("");
+	}
+#endif
+}
+
+FString FIOSPlatformMisc::GetDefaultLanguage()
+{
+	CFArrayRef Languages = CFLocaleCopyPreferredLanguages();
+	CFStringRef LangCodeStr = (CFStringRef)CFArrayGetValueAtIndex(Languages, 0);
+	FString LangCode((__bridge NSString*)LangCodeStr);
+	CFRelease(Languages);
+
+	return LangCode;
+}
+
+FString FIOSPlatformMisc::GetDefaultLocale()
+{
+	CFLocaleRef Locale = CFLocaleCopyCurrent();
+	CFStringRef LangCodeStr = (CFStringRef)CFLocaleGetValue(Locale, kCFLocaleLanguageCode);
+	FString LangCode((__bridge NSString*)LangCodeStr);
+	CFStringRef CountryCodeStr = (CFStringRef)CFLocaleGetValue(Locale, kCFLocaleCountryCode);
+	FString CountryCode((__bridge NSString*)CountryCodeStr);
+	CFRelease(Locale);
+
+	return CountryCode.IsEmpty() ? LangCode : FString::Printf(TEXT("%s-%s"), *LangCode, *CountryCode);
+}
+
+EAppReturnType::Type FIOSPlatformMisc::MessageBoxExt( EAppMsgType::Type MsgType, const TCHAR* Text, const TCHAR* Caption )
+{
+#if UE_BUILD_SHIPPING || PLATFORM_TVOS
+	return FGenericPlatformMisc::MessageBoxExt(MsgType, Text, Caption);
+#else
+	NSString* CocoaText = (NSString*)FPlatformString::TCHARToCFString(Text);
+	NSString* CocoaCaption = (NSString*)FPlatformString::TCHARToCFString(Caption);
+
+	NSMutableArray* StringArray = [NSMutableArray arrayWithCapacity:7];
+
+	[StringArray addObject:CocoaCaption];
+	[StringArray addObject:CocoaText];
+
+	// Figured that the order of all of these should be the same as their enum name.
+	switch (MsgType)
+	{
+		case EAppMsgType::YesNo:
+			[StringArray addObject:@"Yes"];
+			[StringArray addObject:@"No"];
+			break;
+		case EAppMsgType::OkCancel:
+			[StringArray addObject:@"Ok"];
+			[StringArray addObject:@"Cancel"];
+			break;
+		case EAppMsgType::YesNoCancel:
+			[StringArray addObject:@"Yes"];
+			[StringArray addObject:@"No"];
+			[StringArray addObject:@"Cancel"];
+			break;
+		case EAppMsgType::CancelRetryContinue:
+			[StringArray addObject:@"Cancel"];
+			[StringArray addObject:@"Retry"];
+			[StringArray addObject:@"Continue"];
+			break;
+		case EAppMsgType::YesNoYesAllNoAll:
+			[StringArray addObject:@"Yes"];
+			[StringArray addObject:@"No"];
+			[StringArray addObject:@"Yes To All"];
+			[StringArray addObject:@"No To All"];
+			break;
+		case EAppMsgType::YesNoYesAllNoAllCancel:
+			[StringArray addObject:@"Yes"];
+			[StringArray addObject:@"No"];
+			[StringArray addObject:@"Yes To All"];
+			[StringArray addObject:@"No To All"];
+			[StringArray addObject:@"Cancel"];
+			break;
+		case EAppMsgType::YesNoYesAll:
+			[StringArray addObject : @"Yes"];
+			[StringArray addObject : @"No"];
+			[StringArray addObject : @"Yes To All"];
+			break;
+		default:
+			[StringArray addObject:@"Ok"];
+			break;
+	}
+
+	IOSAppDelegate* AppDelegate = [IOSAppDelegate GetDelegate];
+	
+	// reset our response to unset
+	AppDelegate.AlertResponse = -1;
+
+	[AppDelegate performSelectorOnMainThread:@selector(ShowAlert:) withObject:StringArray waitUntilDone:NO];
+
+	while (AppDelegate.AlertResponse == -1)
+	{
+		FPlatformProcess::Sleep(.1);
+	}
+
+	EAppReturnType::Type Result = (EAppReturnType::Type)AppDelegate.AlertResponse;
+
+	// Need to remap the return type to the correct one, since AlertResponse actually returns a button index.
+	switch (MsgType)
+	{
+	case EAppMsgType::YesNo:
+		Result = (Result == EAppReturnType::No ? EAppReturnType::Yes : EAppReturnType::No);
+		break;
+	case EAppMsgType::OkCancel:
+		// return 1 for Ok, 0 for Cancel
+		Result = (Result == EAppReturnType::No ? EAppReturnType::Ok : EAppReturnType::Cancel);
+		break;
+	case EAppMsgType::YesNoCancel:
+		// return 0 for Yes, 1 for No, 2 for Cancel
+		if(Result == EAppReturnType::No)
+		{
+			Result = EAppReturnType::Yes;
+		}
+		else if(Result == EAppReturnType::Yes)
+		{
+			Result = EAppReturnType::No;
+		}
+		else
+		{
+			Result = EAppReturnType::Cancel;
+		}
+		break;
+	case EAppMsgType::CancelRetryContinue:
+		// return 0 for Cancel, 1 for Retry, 2 for Continue
+		if(Result == EAppReturnType::No)
+		{
+			Result = EAppReturnType::Cancel;
+		}
+		else if(Result == EAppReturnType::Yes)
+		{
+			Result = EAppReturnType::Retry;
+		}
+		else
+		{
+			Result = EAppReturnType::Continue;
+		}
+		break;
+	case EAppMsgType::YesNoYesAllNoAll:
+		// return 0 for No, 1 for Yes, 2 for YesToAll, 3 for NoToAll
+		break;
+	case EAppMsgType::YesNoYesAllNoAllCancel:
+		// return 0 for No, 1 for Yes, 2 for YesToAll, 3 for NoToAll, 4 for Cancel
+		break;
+	case EAppMsgType::YesNoYesAll:
+		// return 0 for No, 1 for Yes, 2 for YesToAll
+		break;
+	default:
+		Result = EAppReturnType::Ok;
+		break;
+	}
+
+	CFRelease((CFStringRef)CocoaCaption);
+	CFRelease((CFStringRef)CocoaText);
+
+	return Result;
+#endif
+}
+
+uint32 FIOSPlatformMisc::GetCharKeyMap(uint32* KeyCodes, FString* KeyNames, uint32 MaxMappings)
+{
+	return FGenericPlatformMisc::GetStandardPrintableKeyMap(KeyCodes, KeyNames, MaxMappings, true, true);
+}
+
+uint32 FIOSPlatformMisc::GetKeyMap( uint32* KeyCodes, FString* KeyNames, uint32 MaxMappings )
+{
+#define ADDKEYMAP(KeyCode, KeyName)		if (NumMappings<MaxMappings) { KeyCodes[NumMappings]=KeyCode; KeyNames[NumMappings]=KeyName; ++NumMappings; };
+	
+	uint32 NumMappings = 0;
+	
+	// we only handle a few "fake" keys from the IOS keyboard delegate stuff in IOSView.cpp
+	if (KeyCodes && KeyNames && (MaxMappings > 0))
+	{
+		ADDKEYMAP(KEYCODE_ENTER, TEXT("Enter"));
+		ADDKEYMAP(KEYCODE_BACKSPACE, TEXT("BackSpace"));
+		ADDKEYMAP(KEYCODE_ESCAPE, TEXT("Escape"));
+	}
+	return NumMappings;
+}
+
+bool FIOSPlatformMisc::ControlScreensaver(EScreenSaverAction Action)
+{
+	IOSAppDelegate* AppDelegate = [IOSAppDelegate GetDelegate];
+	[AppDelegate EnableIdleTimer : (Action == FGenericPlatformMisc::Enable)];
+	return true;
+}
+
+int FIOSPlatformMisc::GetAudioVolume()
+{
+	return [[IOSAppDelegate GetDelegate] GetAudioVolume];
+}
+
+bool FIOSPlatformMisc::AreHeadphonesPluggedIn()
+{
+	return [[IOSAppDelegate GetDelegate] AreHeadphonesPluggedIn];
+}
+
+int FIOSPlatformMisc::GetBatteryLevel()
+{
+	return [[IOSAppDelegate GetDelegate] GetBatteryLevel];
+}
+
+bool FIOSPlatformMisc::IsRunningOnBattery()
+{
+	return [[IOSAppDelegate GetDelegate] IsRunningOnBattery];
+}
+
+int32 FIOSPlatformMisc::NumberOfCores()
+{
+	// cache the number of cores
+	static int32 NumberOfCores = -1;
+	if (NumberOfCores == -1)
+	{
+		SIZE_T Size = sizeof(int32);
+		if (sysctlbyname("hw.ncpu", &NumberOfCores, &Size, NULL, 0) != 0)
+		{
+			NumberOfCores = 1;
+		}
+	}
+	return NumberOfCores;
+}
+
+#include "ModuleManager.h"
+
+void FIOSPlatformMisc::LoadPreInitModules()
+{
+	FModuleManager::Get().LoadModule(TEXT("OpenGLDrv"));
+	FModuleManager::Get().LoadModule(TEXT("IOSAudio"));
+	FModuleManager::Get().LoadModule(TEXT("AudioMixerAudioUnit"));
+}
+
+void* FIOSPlatformMisc::CreateAutoreleasePool()
+{
+	return [[NSAutoreleasePool alloc] init];
+}
+
+void FIOSPlatformMisc::ReleaseAutoreleasePool(void *Pool)
+{
+	[(NSAutoreleasePool*)Pool release];
+}
+
+void* FIOSPlatformMisc::GetHardwareWindow()
+{
+	IOSAppDelegate* AppDelegate = [IOSAppDelegate GetDelegate];
+    return AppDelegate.IOSView; 
+}
+
+bool FIOSPlatformMisc::HasPlatformFeature(const TCHAR* FeatureName)
+{
+	if (FCString::Stricmp(FeatureName, TEXT("Metal")) == 0)
+	{
+		return [IOSAppDelegate GetDelegate].IOSView->bIsUsingMetal;
+	}
+
+	return FGenericPlatformMisc::HasPlatformFeature(FeatureName);
+}
+
+FString GetIOSDeviceIDString()
+{
+	static FString CachedResult;
+	static bool bCached = false;
+	if (!bCached)
+	{
+		// get the device hardware type string length
+		size_t DeviceIDLen;
+		sysctlbyname("hw.machine", NULL, &DeviceIDLen, NULL, 0);
+
+		// get the device hardware type
+		char* DeviceID = (char*)malloc(DeviceIDLen);
+		sysctlbyname("hw.machine", DeviceID, &DeviceIDLen, NULL, 0);
+
+		CachedResult = ANSI_TO_TCHAR(DeviceID);
+		bCached = true;
+
+		free(DeviceID);
+	}
+
+	return CachedResult;
+}
+
+FIOSPlatformMisc::EIOSDevice FIOSPlatformMisc::GetIOSDeviceType()
+{
+	// default to unknown
+	static EIOSDevice DeviceType = IOS_Unknown;
+
+	// if we've already figured it out, return it
+	if (DeviceType != IOS_Unknown)
+	{
+		return DeviceType;
+	}
+
+	const FString DeviceIDString = GetIOSDeviceIDString();
+
+	// iPods
+	if (DeviceIDString.StartsWith(TEXT("iPod")))
+	{
+		// get major revision number
+        int Major = FCString::Atoi(&DeviceIDString[4]);
+
+		if (Major == 5)
+		{
+			DeviceType = IOS_IPodTouch5;
+		}
+		else if (Major >= 7)
+		{
+			DeviceType = IOS_IPodTouch6;
+		}
+	}
+	// iPads
+	else if (DeviceIDString.StartsWith(TEXT("iPad")))
+	{
+		// get major revision number
+		const int Major = FCString::Atoi(&DeviceIDString[4]);
+		const int CommaIndex = DeviceIDString.Find(TEXT(","), ESearchCase::CaseSensitive, ESearchDir::FromStart, 4);
+		const int Minor = FCString::Atoi(&DeviceIDString[CommaIndex + 1]);
+
+		// iPad2,[1|2|3] is iPad 2 (1 - wifi, 2 - gsm, 3 - cdma)
+		if (Major == 2)
+		{
+			// iPad2,5+ is the new iPadMini, anything higher will use these settings until released
+			if (Minor >= 5)
+			{
+				DeviceType = IOS_IPadMini;
+			}
+			else
+			{
+				DeviceType = IOS_IPad2;
+			}
+		}
+		// iPad3,[1|2|3] is iPad 3 and iPad3,4+ is iPad (4th generation)
+		else if (Major == 3)
+		{
+			if (Minor <= 3)
+			{
+				DeviceType = IOS_IPad3;
+			}
+			// iPad3,4+ is the new iPad, anything higher will use these settings until released
+			else if (Minor >= 4)
+			{
+				DeviceType = IOS_IPad4;
+			}
+		}
+		// iPadAir and iPad Mini 2nd Generation
+		else if (Major == 4)
+		{
+			if (Minor >= 4)
+			{
+				DeviceType = IOS_IPadMini2;
+			}
+			else
+			{
+				DeviceType = IOS_IPadAir;
+			}
+		}
+		// iPad Air 2 and iPadMini 4
+		else if (Major == 5)
+		{
+			if (Minor == 1 || Minor == 2)
+			{
+				DeviceType = IOS_IPadMini4;
+			}
+			else
+			{
+				DeviceType = IOS_IPadAir2;
+			}
+		}
+		else if (Major == 6)
+		{
+			if (Minor == 3 || Minor == 4)
+			{
+				DeviceType = IOS_IPadPro_97;
+			}
+			else
+			{
+				DeviceType = IOS_IPadPro_129;
+			}
+		}
+		// Default to highest settings currently available for any future device
+		else if (Major > 6)
+		{
+			DeviceType = IOS_IPadPro;
+		}
+	}
+	// iPhones
+	else if (DeviceIDString.StartsWith(TEXT("iPhone")))
+	{
+        const int Major = FCString::Atoi(&DeviceIDString[6]);
+		const int CommaIndex = DeviceIDString.Find(TEXT(","), ESearchCase::CaseSensitive, ESearchDir::FromStart, 6);
+		const int Minor = FCString::Atoi(&DeviceIDString[CommaIndex + 1]);
+
+		if (Major == 3)
+		{
+			DeviceType = IOS_IPhone4;
+		}
+		else if (Major == 4)
+		{
+			DeviceType = IOS_IPhone4S;
+		}
+		else if (Major == 5)
+		{
+			DeviceType = IOS_IPhone5;
+		}
+		else if (Major == 6)
+		{
+			DeviceType = IOS_IPhone5S;
+		}
+		else if (Major == 7)
+		{
+			if (Minor == 1)
+			{
+				DeviceType = IOS_IPhone6Plus;
+			}
+			else if (Minor == 2)
+			{
+				DeviceType = IOS_IPhone6;
+			}
+		}
+		else if (Major == 8)
+		{
+			// note that Apple switched the minor order around between 6 and 6S (gotta keep us on our toes!)
+			if (Minor == 1)
+			{
+				DeviceType = IOS_IPhone6S;
+			}
+			else if (Minor == 2)
+			{
+				DeviceType = IOS_IPhone6SPlus;
+			}
+			else if (Minor == 4)
+			{
+				DeviceType = IOS_IPhoneSE;
+			}
+		}
+		else if (Major == 9)
+		{
+            if (Minor == 1 || Minor == 3)
+            {
+                DeviceType = IOS_IPhone7;
+            }
+            else if (Minor == 2 || Minor == 4)
+            {
+                DeviceType = IOS_IPhone7Plus;
+            }
+		}
+        else if (Major >= 10)
+        {
+            // for going forward into unknown devices (like 8/8+?), we can't use Minor,
+            // so treat devices with a scale > 2.5 to be 6SPlus type devices, < 2.5 to be 6S type devices
+            if ([UIScreen mainScreen].scale > 2.5f)
+            {
+                DeviceType = IOS_IPhone7Plus;
+            }
+            else
+            {
+                DeviceType = IOS_IPhone7;
+            }
+        }
+	}
+	// tvOS
+	else if (DeviceIDString.StartsWith(TEXT("AppleTV")))
+	{
+		DeviceType = IOS_AppleTV;
+	}
+	// simulator
+	else if (DeviceIDString.StartsWith(TEXT("x86")))
+	{
+		// iphone
+		if (UI_USER_INTERFACE_IDIOM() == UIUserInterfaceIdiomPhone)
+		{
+			CGSize result = [[UIScreen mainScreen] bounds].size;
+			if(result.height >= 586)
+			{
+				DeviceType = IOS_IPhone5;
+			}
+			else
+			{
+				DeviceType = IOS_IPhone4S;
+			}
+		}
+		else
+		{
+			if ([[UIScreen mainScreen] scale] > 1.0f)
+			{
+				DeviceType = IOS_IPad4;
+			}
+			else
+			{
+				DeviceType = IOS_IPad2;
+			}
+		}
+	}
+
+	// if this is unknown at this point, we have a problem
+	if (DeviceType == IOS_Unknown)
+	{
+		UE_LOG(LogInit, Fatal, TEXT("This IOS device type is not supported by UE4 [%s]\n"), *FString(DeviceIDString));
+	}
+
+	return DeviceType;
+}
+
+void FIOSPlatformMisc::SetMemoryWarningHandler(void (* InHandler)(const FGenericMemoryWarningContext& Context))
+{
+	GMemoryWarningHandler = InHandler;
+}
+
+void FIOSPlatformMisc::HandleLowMemoryWarning()
+{
+	UE_LOG(LogInit, Log, TEXT("Free Memory at Startup: %d MB"), GStartupFreeMemoryMB);
+	UE_LOG(LogInit, Log, TEXT("Free Memory Now       : %d MB"), GetFreeMemoryMB());
+
+	if(GMemoryWarningHandler != NULL)
+	{
+		FGenericMemoryWarningContext Context;
+		GMemoryWarningHandler(Context);
+	}
+}
+
+bool FIOSPlatformMisc::IsPackagedForDistribution()
+{
+#if !UE_BUILD_SHIPPING
+	static bool PackagingModeCmdLine = FParse::Param(FCommandLine::Get(), TEXT("PACKAGED_FOR_DISTRIBUTION"));
+	if (PackagingModeCmdLine)
+	{
+		return true;
+	}
+#endif
+	NSString* PackagingMode = [[[NSBundle mainBundle] infoDictionary] objectForKey:@"EpicPackagingMode"];
+	return PackagingMode != nil && [PackagingMode isEqualToString : @"Distribution"];
+}
+
+/**
+ * Returns a unique string for device identification
+ * 
+ * @return the unique string generated by this platform for this device
+ */
+FString FIOSPlatformMisc::GetUniqueDeviceId()
+{
+	// Check to see if this OS has this function
+	if ([[UIDevice currentDevice] respondsToSelector:@selector(identifierForVendor)])
+	{
+		NSUUID* Id = [[UIDevice currentDevice] identifierForVendor];
+		if (Id != nil)
+		{
+			return FString([[Id UUIDString] autorelease]);
+		}
+	}
+
+	PRAGMA_DISABLE_DEPRECATION_WARNINGS
+	return FPlatformMisc::GetHashedMacAddressString();
+	PRAGMA_ENABLE_DEPRECATION_WARNINGS
+}
+
+FString FIOSPlatformMisc::GetDeviceId()
+{
+	// Check to see if this OS has this function
+
+	if ([[UIDevice currentDevice] respondsToSelector:@selector(identifierForVendor)])
+	{
+	    NSUUID* Id = [[UIDevice currentDevice] identifierForVendor];
+	    if (Id != nil)
+	    {
+		    NSString* IdfvString = [Id UUIDString];
+		    FString IDFV(IdfvString);
+		    return IDFV;
+	    }
+	}
+	return FString();
+}
+
+FString FIOSPlatformMisc::GetOSVersion()
+{
+	return FString([[UIDevice currentDevice] systemVersion]);
+}
+
+
+/**
+* Returns a unique string for advertising identification
+*
+* @return the unique string generated by this platform for this device
+*/
+FString FIOSPlatformMisc::GetUniqueAdvertisingId()
+{
+#if !PLATFORM_TVOS
+	// Check to see if this OS has this function
+	if ([[ASIdentifierManager sharedManager] respondsToSelector:@selector(advertisingIdentifier)])
+	{
+		NSString* IdfaString = [[[ASIdentifierManager sharedManager] advertisingIdentifier] UUIDString];
+		FString IDFA(IdfaString);
+		return IDFA;
+	}
+#endif
+	return FString();
+}
+
+class IPlatformChunkInstall* FIOSPlatformMisc::GetPlatformChunkInstall()
+{
+	static IPlatformChunkInstall* ChunkInstall = nullptr;
+	static bool bIniChecked = false;
+	if (!ChunkInstall || !bIniChecked)
+	{
+		FString ProviderName;
+		IPlatformChunkInstallModule* PlatformChunkInstallModule = nullptr;
+		if (!GEngineIni.IsEmpty())
+		{
+			FString InstallModule;
+			GConfig->GetString(TEXT("StreamingInstall"), TEXT("DefaultProviderName"), InstallModule, GEngineIni);
+			FModuleStatus Status;
+			if (FModuleManager::Get().QueryModule(*InstallModule, Status))
+			{
+				PlatformChunkInstallModule = FModuleManager::LoadModulePtr<IPlatformChunkInstallModule>(*InstallModule);
+				if (PlatformChunkInstallModule != nullptr)
+				{
+					// Attempt to grab the platform installer
+					ChunkInstall = PlatformChunkInstallModule->GetPlatformChunkInstall();
+				}
+			}
+			else if (ProviderName == TEXT("IOSChunkInstaller"))
+			{
+				static FIOSChunkInstall Singleton;
+				ChunkInstall = &Singleton;
+			}
+			bIniChecked = true;
+		}
+		if (!ChunkInstall)
+		{
+			// Placeholder instance
+			ChunkInstall = FGenericPlatformMisc::GetPlatformChunkInstall();
+		}
+	}
+
+	return ChunkInstall;
+}
+
+
+struct FFontHeader
+{
+	int32 Version;
+	uint16 NumTables;
+	uint16 SearchRange;
+	uint16 EntrySelector;
+	uint16 RangeShift;
+};
+
+struct FFontTableEntry
+{
+	uint32 Tag;
+	uint32 CheckSum;
+	uint32 Offset;
+	uint32 Length;
+};
+
+
+static uint32 CalcTableCheckSum(const uint32 *Table, uint32 NumberOfBytesInTable)
+{
+	uint32 Sum = 0;
+	uint32 NumLongs = (NumberOfBytesInTable + 3) / 4;
+	while (NumLongs-- > 0)
+	{
+		Sum += CFSwapInt32HostToBig(*Table++);
+	}
+	return Sum;
+}
+
+static uint32 CalcTableDataRefCheckSum(CFDataRef DataRef)
+{
+	const uint32 *DataBuff = (const uint32 *)CFDataGetBytePtr(DataRef);
+	uint32 DataLength = (uint32)CFDataGetLength(DataRef);
+	return CalcTableCheckSum(DataBuff, DataLength);
+}
+
+/**
+ * In order to get a system font from IOS we need to build one from the data we can gather from a CGFontRef
+ * @param InFontName - The name of the system font we are seeking to load.
+ * @param OutBytes - The data we have built for the font.
+ */
+void GetBytesForFont(const NSString* InFontName, OUT TArray<uint8>& OutBytes)
+{
+	CGFontRef cgFont = CGFontCreateWithFontName((CFStringRef)InFontName);
+
+	if (cgFont)
+	{
+		CFRetain(cgFont);
+
+		// Gather information on the font tags
+		CFArrayRef Tags = CGFontCopyTableTags(cgFont);
+		int TableCount = CFArrayGetCount(Tags);
+
+		// Collate the table sizes
+		TArray<size_t> TableSizes;
+
+		bool bContainsCFFTable = false;
+
+		size_t TotalSize = sizeof(FFontHeader)+sizeof(FFontTableEntry)* TableCount;
+		for (int TableIndex = 0; TableIndex < TableCount; ++TableIndex)
+		{
+			size_t TableSize = 0;
+			
+			uint64 aTag = (uint64)CFArrayGetValueAtIndex(Tags, TableIndex);
+			if (aTag == 'CFF ' && !bContainsCFFTable)
+			{
+				bContainsCFFTable = true;
+			}
+
+			CFDataRef TableDataRef = CGFontCopyTableForTag(cgFont, aTag);
+			if (TableDataRef != NULL)
+			{
+				TableSize = CFDataGetLength(TableDataRef);
+				CFRelease(TableDataRef);
+			}
+
+			TotalSize += (TableSize + 3) & ~3;
+			TableSizes.Add( TableSize );
+		}
+
+		OutBytes.Reserve( TotalSize );
+		OutBytes.AddZeroed( TotalSize );
+
+		// Start copying the table data into our buffer
+		uint8* DataStart = OutBytes.GetData();
+		uint8* DataPtr = DataStart;
+
+		// Compute font header entries
+		uint16 EntrySelector = 0;
+		uint16 SearchRange = 1;
+		while (SearchRange < TableCount >> 1)
+		{
+			EntrySelector++;
+			SearchRange <<= 1;
+		}
+		SearchRange <<= 4;
+
+		uint16 RangeShift = (TableCount << 4) - SearchRange;
+
+		// Write font header (also called sfnt header, offset subtable)
+		FFontHeader* OffsetTable = (FFontHeader*)DataPtr;
+
+		// OpenType Font contains CFF Table use 'OTTO' as version, and with .otf extension
+		// otherwise 0001 0000
+		OffsetTable->Version = bContainsCFFTable ? 'OTTO' : CFSwapInt16HostToBig(1);
+		OffsetTable->NumTables = CFSwapInt16HostToBig((uint16)TableCount);
+		OffsetTable->SearchRange = CFSwapInt16HostToBig((uint16)SearchRange);
+		OffsetTable->EntrySelector = CFSwapInt16HostToBig((uint16)EntrySelector);
+		OffsetTable->RangeShift = CFSwapInt16HostToBig((uint16)RangeShift);
+
+		DataPtr += sizeof(FFontHeader);
+
+		// Write tables
+		FFontTableEntry* CurrentTableEntry = (FFontTableEntry*)DataPtr;
+		DataPtr += sizeof(FFontTableEntry) * TableCount;
+
+		for (int TableIndex = 0; TableIndex < TableCount; ++TableIndex)
+		{
+			uint64 aTag = (uint64)CFArrayGetValueAtIndex(Tags, TableIndex);
+			CFDataRef TableDataRef = CGFontCopyTableForTag(cgFont, aTag);
+			uint32 TableSize = CFDataGetLength(TableDataRef);
+
+			FMemory::Memcpy(DataPtr, CFDataGetBytePtr(TableDataRef), TableSize);
+
+			CurrentTableEntry->Tag = CFSwapInt32HostToBig((uint32_t)aTag);
+			CurrentTableEntry->CheckSum = CFSwapInt32HostToBig(CalcTableCheckSum((uint32 *)DataPtr, TableSize));
+
+			uint32 Offset = DataPtr - DataStart;
+			CurrentTableEntry->Offset = CFSwapInt32HostToBig((uint32)Offset);
+			CurrentTableEntry->Length = CFSwapInt32HostToBig((uint32)TableSize);
+
+			DataPtr += (TableSize + 3) & ~3;
+			++CurrentTableEntry;
+
+			CFRelease(TableDataRef);
+		}
+
+		CFRelease(cgFont);
+	}
+}
+
+
+TArray<uint8> FIOSPlatformMisc::GetSystemFontBytes()
+{
+#if PLATFORM_TVOS
+	NSString* SystemFontName = [UIFont preferredFontForTextStyle:UIFontTextStyleBody].fontName;
+#else
+	// Gather some details about the system font
+	uint32 SystemFontSize = [UIFont systemFontSize];
+	NSString* SystemFontName = [UIFont systemFontOfSize:SystemFontSize].fontName;
+#endif
+
+	TArray<uint8> FontBytes;
+	GetBytesForFont(SystemFontName, FontBytes);
+
+	return FontBytes;
+}
+
+TArray<FString> FIOSPlatformMisc::GetPreferredLanguages()
+{
+	TArray<FString> Results;
+
+	NSArray* Languages = [NSLocale preferredLanguages];
+	for (NSString* Language in Languages)
+	{
+		Results.Add(FString(Language));
+	}
+	return Results;
+}
+
+FString FIOSPlatformMisc::GetLocalCurrencyCode()
+{
+	return FString([[NSLocale currentLocale] objectForKey:NSLocaleCurrencyCode]);
+}
+
+FString FIOSPlatformMisc::GetLocalCurrencySymbol()
+{
+	return FString([[NSLocale currentLocale] objectForKey:NSLocaleCurrencySymbol]);
+}
+
+void FIOSPlatformMisc::RegisterForRemoteNotifications()
+{
+#if !PLATFORM_TVOS && NOTIFICATIONS_ENABLED
+	UIApplication* application = [UIApplication sharedApplication];
+	if ([application respondsToSelector : @selector(registerUserNotificationSettings:)])
+	{
+#ifdef __IPHONE_8_0
+		UIUserNotificationSettings * settings = [UIUserNotificationSettings settingsForTypes : (UIUserNotificationTypeBadge | UIUserNotificationTypeSound | UIUserNotificationTypeAlert) categories:nil];
+		[application registerUserNotificationSettings : settings];
+#endif
+	}
+	else
+	{
+        
+#if __IPHONE_OS_VERSION_MIN_REQUIRED < __IPHONE_8_0
+		UIRemoteNotificationType myTypes = UIRemoteNotificationTypeBadge | UIRemoteNotificationTypeAlert | UIRemoteNotificationTypeSound;
+		[application registerForRemoteNotificationTypes : myTypes];
+#endif
+	}
+#endif
+}
+
+void FIOSPlatformMisc::GetValidTargetPlatforms(TArray<FString>& TargetPlatformNames)
+{
+	// this is only used to cook with the proper TargetPlatform with COTF, it's not the runtime platform (which is just IOS for both)
+#if PLATFORM_TVOS
+	TargetPlatformNames.Add(TEXT("TVOS"));
+#else
+	TargetPlatformNames.Add(FIOSPlatformProperties::PlatformName());
+#endif
+}
+
+bool FIOSPlatformMisc::HasActiveWiFiConnection()
+{
+	struct sockaddr_in ZeroAddress;
+	FMemory::Memzero(&ZeroAddress, sizeof(ZeroAddress));
+	ZeroAddress.sin_len = sizeof(ZeroAddress);
+	ZeroAddress.sin_family = AF_INET;
+
+	SCNetworkReachabilityRef ReachabilityRef = SCNetworkReachabilityCreateWithAddress(kCFAllocatorDefault, (const struct sockaddr*)&ZeroAddress);
+	SCNetworkReachabilityFlags ReachabilityFlags;
+	bool bFlagsAvailable = SCNetworkReachabilityGetFlags(ReachabilityRef, &ReachabilityFlags);
+	CFRelease(ReachabilityRef);
+	
+	bool bHasActiveWiFiConnection = false;
+	if (bFlagsAvailable)
+	{
+        bool bReachable =	(ReachabilityFlags & kSCNetworkReachabilityFlagsReachable) != 0 && 
+							(ReachabilityFlags & kSCNetworkReachabilityFlagsConnectionRequired) == 0 &&
+							// in case kSCNetworkReachabilityFlagsConnectionOnDemand  || kSCNetworkReachabilityFlagsConnectionOnTraffic
+							(ReachabilityFlags & kSCNetworkReachabilityFlagsInterventionRequired) == 0; 
+					
+		bHasActiveWiFiConnection = bReachable && (ReachabilityFlags & kSCNetworkReachabilityFlagsIsWWAN) == 0;
+	}
+	
+	return bHasActiveWiFiConnection; 
+}
+
+void FIOSPlatformMisc::ResetGamepadAssignments()
+{
+	UE_LOG(LogIOS, Warning, TEXT("Restting gamepad assignments is not allowed in IOS"))
+}
+
+void FIOSPlatformMisc::ResetGamepadAssignmentToController(int32 ControllerId)
+{
+	
+}
+
+bool FIOSPlatformMisc::IsControllerAssignedToGamepad(int32 ControllerId)
+{
+	FIOSInputInterface* InputInterface = (FIOSInputInterface*)CachedApplication->GetInputInterface();
+	return InputInterface->IsControllerAssignedToGamepad(ControllerId);
+}
+
+#if IOS_PROFILING_ENABLED
+void FIOSPlatformMisc::BeginNamedEvent(const struct FColor& Color,const TCHAR* Text)
+{
+	FApplePlatformDebugEvents::BeginNamedEvent(Color, Text);
+}
+
+void FIOSPlatformMisc::BeginNamedEvent(const struct FColor& Color,const ANSICHAR* Text)
+{
+	FApplePlatformDebugEvents::BeginNamedEvent(Color, Text);
+}
+
+void FIOSPlatformMisc::EndNamedEvent()
+{
+	FApplePlatformDebugEvents::EndNamedEvent();
+}
+#endif
+
+FString FIOSPlatformMisc::GetCPUVendor()
+{
+	return TEXT("Apple");
+}
+
+FString FIOSPlatformMisc::GetCPUBrand()
+{
+	return GetIOSDeviceIDString();
+}
+
+void FIOSPlatformMisc::GetOSVersions(FString& out_OSVersionLabel, FString& out_OSSubVersionLabel)
+{
+#if PLATFORM_TVOS
+	out_OSVersionLabel = TEXT("TVOS");
+#else
+	out_OSVersionLabel = TEXT("IOS");
+#endif
+	NSOperatingSystemVersion IOSVersion;
+	IOSVersion = [[NSProcessInfo processInfo] operatingSystemVersion];
+	out_OSSubVersionLabel = FString::Printf(TEXT("%ld.%ld.%ld"), IOSVersion.majorVersion, IOSVersion.minorVersion, IOSVersion.patchVersion);
+}
+
+int32 FIOSPlatformMisc::IOSVersionCompare(uint8 Major, uint8 Minor, uint8 Revision)
+{
+	NSOperatingSystemVersion IOSVersion;
+	IOSVersion = [[NSProcessInfo processInfo] operatingSystemVersion];
+	uint8 TargetValues[3] = { Major, Minor, Revision };
+	NSInteger ComponentValues[3] = { IOSVersion.majorVersion, IOSVersion.minorVersion, IOSVersion.patchVersion };
+
+	for (uint32 i = 0; i < 3; i++)
+	{
+		if (ComponentValues[i] < TargetValues[i])
+		{
+			return -1;
+		}
+		else if (ComponentValues[i] > TargetValues[i])
+		{
+			return 1;
+		}
+	}
+
+	return 0;
+}
+
+EScreenPhysicalAccuracy FIOSPlatformMisc::ComputePhysicalScreenDensity(int32& ScreenDensity)
+{
+	EIOSDevice Device = GetIOSDeviceType();
+	static_assert( EIOSDevice::IOS_Unknown == 25, "Every device needs to be handled here." );
+
+	ScreenDensity = 0;
+	EScreenPhysicalAccuracy Accuracy = EScreenPhysicalAccuracy::Unknown;
+
+	// look up what the device can support
+	const float NativeScale =[[UIScreen mainScreen] scale];
+
+	switch ( Device )
+	{
+	case IOS_IPhoneSE:
+	case IOS_IPhone4:
+	case IOS_IPhone4S:
+	case IOS_IPhone5:
+	case IOS_IPhone5S:
+	case IOS_IPodTouch5:
+	case IOS_IPodTouch6:
+	case IOS_IPhone6:
+	case IOS_IPhone6S:
+	case IOS_IPhone7:
+		ScreenDensity = 326;
+		Accuracy = EScreenPhysicalAccuracy::Truth;
+		break;
+	case IOS_IPhone6Plus:
+	case IOS_IPhone6SPlus:
+	case IOS_IPhone7Plus:
+		ScreenDensity = 401;
+		Accuracy = EScreenPhysicalAccuracy::Truth;
+		break;
+	case IOS_IPadMini:
+	case IOS_IPadMini2: // also the iPadMini3
+	case IOS_IPadMini4:
+		ScreenDensity = 401;
+		Accuracy = EScreenPhysicalAccuracy::Truth;
+		break;
+	case IOS_IPad2:
+	case IOS_IPad3:
+	case IOS_IPad4:
+	case IOS_IPadAir:
+	case IOS_IPadAir2:
+	case IOS_IPadPro_97:
+		ScreenDensity = 264;
+		Accuracy = EScreenPhysicalAccuracy::Truth;
+		break;
+	case IOS_IPadPro:
+	case IOS_IPadPro_129:
+		ScreenDensity = 264;
+		Accuracy = EScreenPhysicalAccuracy::Truth;
+		break;
+	case IOS_AppleTV:
+		Accuracy = EScreenPhysicalAccuracy::Unknown;
+		break;
+	default:
+		// If we don't know, assume that the density is a multiple of the 
+		// native Content Scaling Factor.  Won't be exact, but should be close enough.
+		ScreenDensity = 163 * NativeScale;
+		Accuracy = EScreenPhysicalAccuracy::Approximation;
+		break;
+	}
+
+	// look up the current scale factor
+	UIView* View = [IOSAppDelegate GetDelegate].IOSView;
+	const float ContentScaleFactor = View.contentScaleFactor;
+
+	if ( ContentScaleFactor != 0 )
+	{
+		ScreenDensity = ScreenDensity * ( ContentScaleFactor / NativeScale );
+	}
+
+	return Accuracy;
+}
+
+/*------------------------------------------------------------------------------
+ FIOSApplicationInfo - class to contain all state for crash reporting that is unsafe to acquire in a signal.
+ ------------------------------------------------------------------------------*/
+
+/**
+ * Information that cannot be obtained during a signal-handler is initialised here.
+ * This ensures that we only call safe functions within the crash reporting handler.
+ */
+struct FIOSApplicationInfo
+{
+    void Init()
+    {
+        SCOPED_AUTORELEASE_POOL;
+        
+        AppName = FApp::GetGameName();
+        FCStringAnsi::Strcpy(AppNameUTF8, PATH_MAX+1, TCHAR_TO_UTF8(*AppName));
+        
+        ExecutableName = FPlatformProcess::ExecutableName();
+        
+        AppPath = FString([[NSBundle mainBundle] executablePath]);
+        
+        AppBundleID = FString([[NSBundle mainBundle] bundleIdentifier]);
+        
+        NumCores = FPlatformMisc::NumberOfCores();
+        
+        LCID = FString::Printf(TEXT("%d"), FInternationalization::Get().GetCurrentCulture()->GetLCID());
+        
+        PrimaryGPU = FPlatformMisc::GetPrimaryGPUBrand();
+        
+        RunUUID = RunGUID();
+        
+        OSXVersion = [[NSProcessInfo processInfo] operatingSystemVersion];
+        OSVersion = FString::Printf(TEXT("%ld.%ld.%ld"), OSXVersion.majorVersion, OSXVersion.minorVersion, OSXVersion.patchVersion);
+        FCStringAnsi::Strcpy(OSVersionUTF8, PATH_MAX+1, TCHAR_TO_UTF8(*OSVersion));
+        
+        // macOS build number is only accessible on non-sandboxed applications as it resides outside the accessible sandbox
+        if(!bIsSandboxed)
+        {
+            NSDictionary* SystemVersion = [NSDictionary dictionaryWithContentsOfFile: @"/System/Library/CoreServices/SystemVersion.plist"];
+            OSBuild = FString((NSString*)[SystemVersion objectForKey: @"ProductBuildVersion"]);
+        }
+        
+        char TempSysCtlBuffer[PATH_MAX] = {};
+        size_t TempSysCtlBufferSize = PATH_MAX;
+        
+        sysctlbyname("kern.osrelease", TempSysCtlBuffer, &TempSysCtlBufferSize, NULL, 0);
+        BiosRelease = TempSysCtlBuffer;
+        uint32 KernelRevision = 0;
+        TempSysCtlBufferSize = 4;
+        sysctlbyname("kern.osrevision", &KernelRevision, &TempSysCtlBufferSize, NULL, 0);
+        BiosRevision = FString::Printf(TEXT("%d"), KernelRevision);
+        TempSysCtlBufferSize = PATH_MAX;
+        sysctlbyname("kern.uuid", TempSysCtlBuffer, &TempSysCtlBufferSize, NULL, 0);
+        BiosUUID = TempSysCtlBuffer;
+        TempSysCtlBufferSize = PATH_MAX;
+        sysctlbyname("hw.model", TempSysCtlBuffer, &TempSysCtlBufferSize, NULL, 0);
+        MachineModel = TempSysCtlBuffer;
+        TempSysCtlBufferSize = PATH_MAX+1;
+        sysctlbyname("machdep.cpu.brand_string", MachineCPUString, &TempSysCtlBufferSize, NULL, 0);
+        
+        gethostname(MachineName, ARRAY_COUNT(MachineName));
+        
+       BranchBaseDir = FString::Printf( TEXT( "%s!%s!%s!%d" ), *FApp::GetBranchName(), FPlatformProcess::BaseDir(), FPlatformMisc::GetEngineMode(), FEngineVersion::Current().GetChangelist() );
+        
+        // Get the paths that the files will actually have been saved to
+        FString LogDirectory = FPaths::GameLogDir();
+        TCHAR CommandlineLogFile[MAX_SPRINTF]=TEXT("");
+        
+        // Use the log file specified on the commandline if there is one
+        CommandLine = FCommandLine::Get();
+        FString LogPath = FGenericPlatformOutputDevices::GetAbsoluteLogFilename();
+        LogPath = IFileManager::Get().ConvertToAbsolutePathForExternalAppForWrite(*LogPath);
+        FCStringAnsi::Strcpy(AppLogPath, PATH_MAX + 1, TCHAR_TO_UTF8(*LogPath));
+        
+        // Cache & create the crash report folder.
+        FString ReportPath = FPaths::ConvertRelativePathToFull(FString::Printf(TEXT("%s"), *(FPaths::GameAgnosticSavedDir() / TEXT("Crashes"))));
+        IFileManager::Get().MakeDirectory(*ReportPath, true);
+        ReportPath = IFileManager::Get().ConvertToAbsolutePathForExternalAppForWrite(*ReportPath);
+        FCStringAnsi::Strcpy(CrashReportPath, PATH_MAX+1, TCHAR_TO_UTF8(*ReportPath));
+        
+        NSString* PLCrashReportFile = [TemporaryCrashReportFolder().GetNSString() stringByAppendingPathComponent:TemporaryCrashReportName().GetNSString()];
+        [PLCrashReportFile getCString:PLCrashReportPath maxLength:PATH_MAX encoding:NSUTF8StringEncoding];
+    }
+    
+    ~FIOSApplicationInfo()
+    {
+#if !PLATFORM_TVOS
+        if (CrashReporter)
+        {
+            [CrashReporter release];
+            CrashReporter = nil;
+        }
+#endif
+    }
+    
+    static FGuid RunGUID()
+    {
+        static FGuid Guid;
+        if (!Guid.IsValid())
+        {
+            FPlatformMisc::CreateGuid(Guid);
+        }
+        return Guid;
+    }
+    
+    static FString TemporaryCrashReportFolder()
+    {
+        static FString PLCrashReportFolder;
+        if (PLCrashReportFolder.IsEmpty())
+        {
+            SCOPED_AUTORELEASE_POOL;
+            
+            NSArray* Paths = NSSearchPathForDirectoriesInDomains(NSCachesDirectory, NSUserDomainMask, YES);
+            NSString* CacheDir = [Paths objectAtIndex: 0];
+            
+            NSString* BundleID = [[NSBundle mainBundle] bundleIdentifier];
+            if (!BundleID)
+            {
+                BundleID = [[NSProcessInfo processInfo] processName];
+            }
+            check(BundleID);
+            
+            NSString* PLCrashReportFolderPath = [CacheDir stringByAppendingPathComponent: BundleID];
+            PLCrashReportFolder = FString(PLCrashReportFolderPath);
+        }
+        return PLCrashReportFolder;
+    }
+    
+    static FString TemporaryCrashReportName()
+    {
+        static FString PLCrashReportFileName(RunGUID().ToString() + TEXT(".plcrash"));
+        return PLCrashReportFileName;
+    }
+    
+    bool bIsSandboxed;
+    int32 NumCores;
+    char AppNameUTF8[PATH_MAX+1];
+    char AppLogPath[PATH_MAX+1];
+    char CrashReportPath[PATH_MAX+1];
+    char PLCrashReportPath[PATH_MAX+1];
+    char OSVersionUTF8[PATH_MAX+1];
+    char MachineName[PATH_MAX+1];
+    char MachineCPUString[PATH_MAX+1];
+    FString AppPath;
+    FString AppName;
+    FString AppBundleID;
+    FString OSVersion;
+    FString OSBuild;
+    FString MachineUUID;
+    FString MachineModel;
+    FString BiosRelease;
+    FString BiosRevision;
+    FString BiosUUID;
+    FString ParentProcess;
+    FString LCID;
+    FString CommandLine;
+    FString BranchBaseDir;
+    FString PrimaryGPU;
+    FString ExecutableName;
+    NSOperatingSystemVersion OSXVersion;
+    FGuid RunUUID;
+    FString XcodePath;
+#if !PLATFORM_TVOS
+    static PLCrashReporter* CrashReporter;
+#endif
+    static FIOSMallocCrashHandler* CrashMalloc;
+};
+
+static FIOSApplicationInfo GIOSAppInfo;
+#if !PLATFORM_TVOS
+PLCrashReporter* FIOSApplicationInfo::CrashReporter = nullptr;
+#endif
+FIOSMallocCrashHandler* FIOSApplicationInfo::CrashMalloc = nullptr;
+
+void (*GCrashHandlerPointer)(const FGenericCrashContext& Context) = NULL;
+
+// good enough default crash reporter
+static void DefaultCrashHandler(FIOSCrashContext const& Context)
+{
+    Context.ReportCrash();
+    if (GLog)
+    {
+        GLog->SetCurrentThreadAsMasterThread();
+        GLog->Flush();
+    }
+    if (GWarn)
+    {
+        GWarn->Flush();
+    }
+    if (GError)
+    {
+        GError->Flush();
+        GError->HandleError();
+    }
+    return Context.GenerateCrashInfo();
+}
+
+// number of stack entries to ignore in backtrace
+static uint32 GIOSStackIgnoreDepth = 6;
+
+// true system specific crash handler that gets called first
+static void PlatformCrashHandler(int32 Signal, siginfo_t* Info, void* Context)
+{
+    FIOSCrashContext CrashContext;
+    CrashContext.IgnoreDepth = GIOSStackIgnoreDepth;
+    CrashContext.InitFromSignal(Signal, Info, Context);
+    
+    // switch to crash handler malloc to avoid malloc reentrancy
+    check(FIOSApplicationInfo::CrashMalloc);
+    FIOSApplicationInfo::CrashMalloc->Enable(&CrashContext, FPlatformTLS::GetCurrentThreadId());
+    
+    if (GCrashHandlerPointer)
+    {
+        GCrashHandlerPointer(CrashContext);
+    }
+    else
+    {
+        // call default one
+        DefaultCrashHandler(CrashContext);
+    }
+}
+
+static void PLCrashReporterHandler(siginfo_t* Info, ucontext_t* Uap, void* Context)
+{
+    PlatformCrashHandler((int32)Info->si_signo, Info, Uap);
+}
+
+// handles graceful termination.
+static void GracefulTerminationHandler(int32 Signal, siginfo_t* Info, void* Context)
+{
+    // make sure we write out as much as possible
+    if (GLog)
+    {
+        GLog->Flush();
+    }
+    if (GWarn)
+    {
+        GWarn->Flush();
+    }
+    if (GError)
+    {
+        GError->Flush();
+    }
+    
+    if (!GIsRequestingExit)
+    {
+        GIsRequestingExit = 1;
+    }
+    else
+    {
+        _Exit(0);
+    }
+}
+
+void FIOSPlatformMisc::PlatformPreInit()
+{
+    FGenericPlatformMisc::PlatformPreInit();
+    
+    GIOSAppInfo.Init();
+    
+    // turn off SIGPIPE crashes
+    signal(SIGPIPE, SIG_IGN);
+}
+
+void FIOSPlatformMisc::SetGracefulTerminationHandler()
+{
+    struct sigaction Action;
+    FMemory::Memzero(&Action, sizeof(struct sigaction));
+    Action.sa_sigaction = GracefulTerminationHandler;
+    sigemptyset(&Action.sa_mask);
+    Action.sa_flags = SA_SIGINFO | SA_RESTART | SA_ONSTACK;
+    sigaction(SIGINT, &Action, NULL);
+    sigaction(SIGTERM, &Action, NULL);
+    sigaction(SIGHUP, &Action, NULL);
+}
+
+void FIOSPlatformMisc::SetCrashHandler(void (* CrashHandler)(const FGenericCrashContext& Context))
+{
+    SCOPED_AUTORELEASE_POOL;
+    
+    GCrashHandlerPointer = CrashHandler;
+    
+#if !PLATFORM_TVOS
+    if (!FIOSApplicationInfo::CrashReporter && !FIOSApplicationInfo::CrashMalloc)
+    {
+        // configure the crash handler malloc zone to reserve a little memory for itself
+        FIOSApplicationInfo::CrashMalloc = new FIOSMallocCrashHandler(128*1024);
+        
+        PLCrashReporterConfig* Config = [[[PLCrashReporterConfig alloc] initWithSignalHandlerType: PLCrashReporterSignalHandlerTypeBSD symbolicationStrategy: PLCrashReporterSymbolicationStrategyNone crashReportFolder: FIOSApplicationInfo::TemporaryCrashReportFolder().GetNSString() crashReportName: FIOSApplicationInfo::TemporaryCrashReportName().GetNSString()] autorelease];
+        FIOSApplicationInfo::CrashReporter = [[PLCrashReporter alloc] initWithConfiguration: Config];
+        
+        PLCrashReporterCallbacks CrashReportCallback = {
+            .version = 0,
+            .context = nullptr,
+            .handleSignal = PLCrashReporterHandler
+        };
+        
+        [FIOSApplicationInfo::CrashReporter setCrashCallbacks: &CrashReportCallback];
+        
+        NSError* Error = nil;
+        if ([FIOSApplicationInfo::CrashReporter enableCrashReporterAndReturnError: &Error])
+        {
+            GIOSStackIgnoreDepth = 0;
+        }
+        else
+        {
+            UE_LOG(LogIOS, Log, TEXT("Failed to enable PLCrashReporter: %s"), *FString([Error localizedDescription]));
+            UE_LOG(LogIOS, Log, TEXT("Falling back to native signal handlers"));
+ 
+            struct sigaction Action;
+            FMemory::Memzero(&Action, sizeof(struct sigaction));
+            Action.sa_sigaction = PlatformCrashHandler;
+            sigemptyset(&Action.sa_mask);
+            Action.sa_flags = SA_SIGINFO | SA_RESTART | SA_ONSTACK;
+            sigaction(SIGQUIT, &Action, NULL);
+            sigaction(SIGILL, &Action, NULL);
+            sigaction(SIGEMT, &Action, NULL);
+            sigaction(SIGFPE, &Action, NULL);
+            sigaction(SIGBUS, &Action, NULL);
+            sigaction(SIGSEGV, &Action, NULL);
+            sigaction(SIGSYS, &Action, NULL);
+            sigaction(SIGABRT, &Action, NULL);
+        }
+    }
+#endif
+}
+
+void FIOSCrashContext::GenerateWindowsErrorReport(char const* WERPath, bool bIsEnsure) const
+{
+    int ReportFile = open(WERPath, O_CREAT|O_WRONLY, 0766);
+    if (ReportFile != -1)
+    {
+        TCHAR Line[PATH_MAX] = {};
+        
+        // write BOM
+        static uint16 ByteOrderMarker = 0xFEFF;
+        write(ReportFile, &ByteOrderMarker, sizeof(ByteOrderMarker));
+        
+        WriteLine(ReportFile, TEXT("<?xml version=\"1.0\" encoding=\"UTF-16\"?>"));
+        WriteLine(ReportFile, TEXT("<WERReportMetadata>"));
+        
+        WriteLine(ReportFile, TEXT("\t<OSVersionInformation>"));
+        WriteUTF16String(ReportFile, TEXT("\t\t<WindowsNTVersion>"));
+        WriteUTF16String(ReportFile, *GIOSAppInfo.OSVersion);
+        WriteLine(ReportFile, TEXT("</WindowsNTVersion>"));
+        
+        WriteUTF16String(ReportFile, TEXT("\t\t<Build>"));
+        WriteUTF16String(ReportFile, *GIOSAppInfo.OSVersion);
+        WriteUTF16String(ReportFile, TEXT(" ("));
+        WriteUTF16String(ReportFile, *GIOSAppInfo.OSBuild);
+        WriteLine(ReportFile, TEXT(")</Build>"));
+        
+        WriteUTF16String(ReportFile, TEXT("\t\t<Product>(0x30): IOS "));
+        WriteUTF16String(ReportFile, *GIOSAppInfo.OSVersion);
+        WriteLine(ReportFile, TEXT("</Product>"));
+        
+        WriteLine(ReportFile, TEXT("\t\t<Edition>IOS</Edition>"));
+        
+        WriteUTF16String(ReportFile, TEXT("\t\t<BuildString>IOS "));
+        WriteUTF16String(ReportFile, *GIOSAppInfo.OSVersion);
+        WriteUTF16String(ReportFile, TEXT(" ("));
+        WriteUTF16String(ReportFile, *GIOSAppInfo.OSBuild);
+        WriteLine(ReportFile, TEXT(")</BuildString>"));
+        
+        WriteUTF16String(ReportFile, TEXT("\t\t<Revision>"));
+        WriteUTF16String(ReportFile, *GIOSAppInfo.OSBuild);
+        WriteLine(ReportFile, TEXT("</Revision>"));
+        
+        WriteLine(ReportFile, TEXT("\t\t<Flavor>Multiprocessor Free</Flavor>"));
+        WriteLine(ReportFile, TEXT("\t\t<Architecture>X64</Architecture>"));
+        WriteUTF16String(ReportFile, TEXT("\t\t<LCID>"));
+        WriteUTF16String(ReportFile, *GIOSAppInfo.LCID);
+        WriteLine(ReportFile, TEXT("</LCID>"));
+        WriteLine(ReportFile, TEXT("\t</OSVersionInformation>"));
+        
+        WriteLine(ReportFile, TEXT("\t<ParentProcessInformation>"));
+        
+        WriteUTF16String(ReportFile, TEXT("\t\t<ParentProcessId>"));
+        WriteUTF16String(ReportFile, ItoTCHAR(getppid(), 10));
+        WriteLine(ReportFile, TEXT("</ParentProcessId>"));
+        
+        WriteUTF16String(ReportFile, TEXT("\t\t<ParentProcessPath>"));
+        WriteUTF16String(ReportFile, *GIOSAppInfo.ParentProcess);
+        WriteLine(ReportFile, TEXT("</ParentProcessPath>"));
+        
+        WriteLine(ReportFile, TEXT("\t\t<ParentProcessCmdLine></ParentProcessCmdLine>"));	// FIXME: supply valid?
+        WriteLine(ReportFile, TEXT("\t</ParentProcessInformation>"));
+        
+        WriteLine(ReportFile, TEXT("\t<ProblemSignatures>"));
+        WriteLine(ReportFile, TEXT("\t\t<EventType>APPCRASH</EventType>"));
+        
+        WriteUTF16String(ReportFile, TEXT("\t\t<Parameter0>UE4-"));
+        WriteUTF16String(ReportFile, *GIOSAppInfo.AppName);
+        WriteLine(ReportFile, TEXT("</Parameter0>"));
+        
+        WriteUTF16String(ReportFile, TEXT("\t\t<Parameter1>"));
+        WriteUTF16String(ReportFile, ItoTCHAR(FEngineVersion::Current().GetMajor(), 10));
+        WriteUTF16String(ReportFile, TEXT("."));
+        WriteUTF16String(ReportFile, ItoTCHAR(FEngineVersion::Current().GetMinor(), 10));
+        WriteUTF16String(ReportFile, TEXT("."));
+        WriteUTF16String(ReportFile, ItoTCHAR(FEngineVersion::Current().GetPatch(), 10));
+        WriteLine(ReportFile, TEXT("</Parameter1>"));
+        
+        // App time stamp
+        WriteLine(ReportFile, TEXT("\t\t<Parameter2>528f2d37</Parameter2>"));													// FIXME: supply valid?
+        
+        Dl_info DLInfo;
+        if(Info && Info->si_addr != 0 && dladdr(Info->si_addr, &DLInfo) != 0)
+        {
+            // Crash Module name
+            WriteUTF16String(ReportFile, TEXT("\t\t<Parameter3>"));
+            if (DLInfo.dli_fname && FCStringAnsi::Strlen(DLInfo.dli_fname))
+            {
+                FMemory::Memzero(Line, PATH_MAX * sizeof(TCHAR));
+                FUTF8ToTCHAR_Convert::Convert(Line, PATH_MAX, DLInfo.dli_fname, FCStringAnsi::Strlen(DLInfo.dli_fname));
+                WriteUTF16String(ReportFile, Line);
+            }
+            else
+            {
+                WriteUTF16String(ReportFile, TEXT("Unknown"));
+            }
+            WriteLine(ReportFile, TEXT("</Parameter3>"));
+            
+            // Check header
+            uint32 Version = 0;
+            uint32 TimeStamp = 0;
+            struct mach_header_64* Header = (struct mach_header_64*)DLInfo.dli_fbase;
+            struct load_command *CurrentCommand = (struct load_command *)( (char *)Header + sizeof(struct mach_header_64) );
+            if( Header->magic == MH_MAGIC_64 )
+            {
+                for( int32 i = 0; i < Header->ncmds; i++ )
+                {
+                    if( CurrentCommand->cmd == LC_LOAD_DYLIB )
+                    {
+                        struct dylib_command *DylibCommand = (struct dylib_command *) CurrentCommand;
+                        Version = DylibCommand->dylib.current_version;
+                        TimeStamp = DylibCommand->dylib.timestamp;
+                        Version = ((Version & 0xff) + ((Version >> 8) & 0xff) * 100 + ((Version >> 16) & 0xffff) * 10000);
+                        break;
+                    }
+                    
+                    CurrentCommand = (struct load_command *)( (char *)CurrentCommand + CurrentCommand->cmdsize );
+                }
+            }
+            
+            // Module version
+            WriteUTF16String(ReportFile, TEXT("\t\t<Parameter4>"));
+            WriteUTF16String(ReportFile, ItoTCHAR(Version, 10));
+            WriteLine(ReportFile, TEXT("</Parameter4>"));
+            
+            // Module time stamp
+            WriteUTF16String(ReportFile, TEXT("\t\t<Parameter5>"));
+            WriteUTF16String(ReportFile, ItoTCHAR(TimeStamp, 16));
+            WriteLine(ReportFile, TEXT("</Parameter5>"));
+            
+            // MethodDef token -> no equivalent
+            WriteLine(ReportFile, TEXT("\t\t<Parameter6>00000001</Parameter6>"));
+            
+            // IL Offset -> Function pointer
+            WriteUTF16String(ReportFile, TEXT("\t\t<Parameter7>"));
+            WriteUTF16String(ReportFile, ItoTCHAR((uint64)Info->si_addr, 16));
+            WriteLine(ReportFile, TEXT("</Parameter7>"));
+        }
+        
+        // Command line, must match the Windows version.
+        WriteUTF16String(ReportFile, TEXT("\t\t<Parameter8>!"));
+        WriteUTF16String(ReportFile, FCommandLine::GetOriginal());
+        WriteLine(ReportFile, TEXT("!</Parameter8>"));
+        
+        WriteUTF16String(ReportFile, TEXT("\t\t<Parameter9>"));
+        WriteUTF16String(ReportFile, *GIOSAppInfo.BranchBaseDir);
+        WriteLine(ReportFile, TEXT("</Parameter9>"));
+        
+        WriteLine(ReportFile, TEXT("\t</ProblemSignatures>"));
+        
+        WriteLine(ReportFile, TEXT("\t<DynamicSignatures>"));
+        
+        WriteUTF16String(ReportFile, TEXT("\t\t<Parameter1>"));
+        WriteUTF16String(ReportFile, *GIOSAppInfo.BiosUUID);
+        WriteLine(ReportFile, TEXT("</Parameter1>"));
+        
+        WriteUTF16String(ReportFile, TEXT("\t\t<Parameter2>"));
+        WriteUTF16String(ReportFile, *GIOSAppInfo.LCID);
+        WriteLine(ReportFile, TEXT("</Parameter2>"));
+        WriteLine(ReportFile, *FString::Printf(TEXT("\t\t<DeploymentName>%s</DeploymentName>"), FApp::GetDeploymentName()));
+        WriteLine(ReportFile, *FString::Printf(TEXT("\t\t<IsEnsure>%s</IsEnsure>"), bIsEnsure ? TEXT("1") : TEXT("0")));
+        WriteLine(ReportFile, *FString::Printf(TEXT("\t\t<IsAssert>%s</IsAssert>"), FDebug::bHasAsserted ? TEXT("1") : TEXT("0")));
+        WriteLine(ReportFile, *FString::Printf(TEXT("\t\t<CrashType>%s</CrashType>"), FGenericCrashContext::GetCrashTypeString(bIsEnsure, FDebug::bHasAsserted, GIsGPUCrashed)));
+        WriteLine(ReportFile, *FString::Printf(TEXT("\t\t<BuildVersion>%s</BuildVersion>"), FApp::GetBuildVersion()));
+        WriteLine(ReportFile, *FString::Printf(TEXT("\t\t<EngineModeEx>%s</EngineModeEx>"), FGenericCrashContext::EngineModeExString()));
+        
+        WriteLine(ReportFile, TEXT("\t</DynamicSignatures>"));
+        
+        WriteLine(ReportFile, TEXT("\t<SystemInformation>"));
+        
+        WriteUTF16String(ReportFile, TEXT("\t\t<MID>"));
+        WriteUTF16String(ReportFile, *GIOSAppInfo.MachineUUID);
+        WriteLine(ReportFile, TEXT("</MID>"));
+        
+        WriteLine(ReportFile, TEXT("\t\t<SystemManufacturer>Apple Inc.</SystemManufacturer>"));
+        
+        WriteUTF16String(ReportFile, TEXT("\t\t<SystemProductName>"));
+        WriteUTF16String(ReportFile, *GIOSAppInfo.MachineModel);
+        WriteLine(ReportFile, TEXT("</SystemProductName>"));
+        
+        WriteUTF16String(ReportFile, TEXT("\t\t<BIOSVersion>"));
+        WriteUTF16String(ReportFile, *GIOSAppInfo.BiosRelease);
+        WriteUTF16String(ReportFile, TEXT("-"));
+        WriteUTF16String(ReportFile, *GIOSAppInfo.BiosRevision);
+        WriteLine(ReportFile, TEXT("</BIOSVersion>"));
+        
+        WriteUTF16String(ReportFile, TEXT("\t\t<GraphicsCard>"));
+        WriteUTF16String(ReportFile, *GIOSAppInfo.PrimaryGPU);
+        WriteLine(ReportFile, TEXT("</GraphicsCard>"));
+        
+        WriteLine(ReportFile, TEXT("\t</SystemInformation>"));
+        
+        WriteLine(ReportFile, TEXT("</WERReportMetadata>"));
+        
+        close(ReportFile);
+    }
+}
+
+void FIOSCrashContext::CopyMinidump(char const* OutputPath, char const* InputPath) const
+{
+#if !PLATFORM_TVOS
+    NSError* Error = nil;
+    NSString* Path = FString(ANSI_TO_TCHAR(InputPath)).GetNSString();
+    NSData* CrashData = [NSData dataWithContentsOfFile: Path options: NSMappedRead error: &Error];
+    if (CrashData && !Error)
+    {
+        PLCrashReport* CrashLog = [[PLCrashReport alloc] initWithData: CrashData error: &Error];
+        if (CrashLog && !Error)
+        {
+            NSString* Report = [PLCrashReportTextFormatter stringValueForCrashReport: CrashLog withTextFormat: PLCrashReportTextFormatiOS];
+            FString CrashDump = FString(Report);
+            [Report writeToFile: Path atomically: YES encoding: NSUTF8StringEncoding error: &Error];
+        }
+        else
+        {
+            NSLog(@"****UE4 %@", [Error localizedDescription]);
+        }
+    }
+    else
+    {
+        NSLog(@"****UE4 %@", [Error localizedDescription]);
+    }
+    int ReportFile = open(OutputPath, O_CREAT|O_WRONLY, 0766);
+    int DumpFile = open(InputPath, O_RDONLY, 0766);
+    if (ReportFile != -1 && DumpFile != -1)
+    {
+        char Data[PATH_MAX];
+        
+        int Bytes = 0;
+        while((Bytes = read(DumpFile, Data, PATH_MAX)) > 0)
+        {
+            write(ReportFile, Data, Bytes);
+        }
+        
+        close(DumpFile);
+        close(ReportFile);
+        
+        unlink(InputPath);
+    }
+#endif
+}
+
+void FIOSCrashContext::GenerateInfoInFolder(char const* const InfoFolder, bool bIsEnsure) const
+{
+    // create a crash-specific directory
+    char CrashInfoFolder[PATH_MAX] = {};
+    FCStringAnsi::Strncpy(CrashInfoFolder, InfoFolder, PATH_MAX);
+    
+    if(!mkdir(CrashInfoFolder, 0766))
+    {
+        char FilePath[PATH_MAX] = {};
+        FCStringAnsi::Strncpy(FilePath, CrashInfoFolder, PATH_MAX);
+        FCStringAnsi::Strcat(FilePath, PATH_MAX, "/report.wer");
+        int ReportFile = open(FilePath, O_CREAT|O_WRONLY, 0766);
+        if (ReportFile != -1)
+        {
+            // write BOM
+            static uint16 ByteOrderMarker = 0xFEFF;
+            write(ReportFile, &ByteOrderMarker, sizeof(ByteOrderMarker));
+            
+            WriteUTF16String(ReportFile, TEXT("\r\nAppPath="));
+            WriteUTF16String(ReportFile, *GIOSAppInfo.AppPath);
+            WriteLine(ReportFile, TEXT("\r\n"));
+            
+            close(ReportFile);
+        }
+        
+        // generate "WER"
+        FCStringAnsi::Strncpy(FilePath, CrashInfoFolder, PATH_MAX);
+        FCStringAnsi::Strcat(FilePath, PATH_MAX, "/wermeta.xml");
+        GenerateWindowsErrorReport(FilePath, bIsEnsure);
+        
+        // generate "minidump" (Apple crash log format)
+        FCStringAnsi::Strncpy(FilePath, CrashInfoFolder, PATH_MAX);
+        FCStringAnsi::Strcat(FilePath, PATH_MAX, "/minidump.dmp");
+        CopyMinidump(FilePath, GIOSAppInfo.PLCrashReportPath);
+        
+        // generate "info.txt" custom data for our server
+        FCStringAnsi::Strncpy(FilePath, CrashInfoFolder, PATH_MAX);
+        FCStringAnsi::Strcat(FilePath, PATH_MAX, "/info.txt");
+        ReportFile = open(FilePath, O_CREAT|O_WRONLY, 0766);
+        if (ReportFile != -1)
+        {
+            WriteUTF16String(ReportFile, TEXT("GameName UE4-"));
+            WriteLine(ReportFile, *GIOSAppInfo.AppName);
+            
+            WriteUTF16String(ReportFile, TEXT("BuildVersion 1.0."));
+            WriteUTF16String(ReportFile, ItoTCHAR(FEngineVersion::Current().GetChangelist() >> 16, 10));
+            WriteUTF16String(ReportFile, TEXT("."));
+            WriteLine(ReportFile, ItoTCHAR(FEngineVersion::Current().GetChangelist() & 0xffff, 10));
+            
+            WriteUTF16String(ReportFile, TEXT("CommandLine "));
+            WriteLine(ReportFile, *GIOSAppInfo.CommandLine);
+            
+            WriteUTF16String(ReportFile, TEXT("BaseDir "));
+            WriteLine(ReportFile, *GIOSAppInfo.BranchBaseDir);
+            
+            WriteUTF16String(ReportFile, TEXT("MachineGuid "));
+            WriteLine(ReportFile, *GIOSAppInfo.MachineUUID);
+            
+            close(ReportFile);
+        }
+        
+        // Introduces a new runtime crash context. Will replace all Windows related crash reporting.
+        FCStringAnsi::Strncpy(FilePath, CrashInfoFolder, PATH_MAX);
+        FCStringAnsi::Strcat(FilePath, PATH_MAX, "/" );
+        FCStringAnsi::Strcat(FilePath, PATH_MAX, FGenericCrashContext::CrashContextRuntimeXMLNameA );
+        //SerializeAsXML( FilePath ); @todo uncomment after verification - need to do a bit more work on this for macOS
+        
+        // copy log
+        FCStringAnsi::Strncpy(FilePath, CrashInfoFolder, PATH_MAX);
+        FCStringAnsi::Strcat(FilePath, PATH_MAX, "/");
+        FCStringAnsi::Strcat(FilePath, PATH_MAX, (!GIOSAppInfo.AppName.IsEmpty() ? GIOSAppInfo.AppNameUTF8 : "UE4"));
+        FCStringAnsi::Strcat(FilePath, PATH_MAX, ".log");
+        
+        int LogSrc = open(GIOSAppInfo.AppLogPath, O_RDONLY);
+        int LogDst = open(FilePath, O_CREAT|O_WRONLY, 0766);
+        
+        char Data[PATH_MAX] = {};
+        int Bytes = 0;
+        while((Bytes = read(LogSrc, Data, PATH_MAX)) > 0)
+        {
+            write(LogDst, Data, Bytes);
+        }
+        
+        // If present, include the crash report config file to pass config values to the CRC
+        FCStringAnsi::Strncpy(FilePath, CrashInfoFolder, PATH_MAX);
+        FCStringAnsi::Strcat(FilePath, PATH_MAX, "/");
+        FCStringAnsi::Strcat(FilePath, PATH_MAX, FGenericCrashContext::CrashConfigFileNameA);
+        int ConfigSrc = open(TCHAR_TO_ANSI(GetCrashConfigFilePath()), O_RDONLY);
+        int ConfigDst = open(FilePath, O_CREAT | O_WRONLY, 0766);
+        
+        while ((Bytes = read(ConfigSrc, Data, PATH_MAX)) > 0)
+        {
+            write(ConfigDst, Data, Bytes);
+        }
+        
+        close(ConfigDst);
+        close(ConfigSrc);
+        
+        close(LogDst);
+        close(LogSrc);
+        // best effort, so don't care about result: couldn't copy -> tough, no log
+    }
+    else
+    {
+        NSLog(@"******* UE4 - Failed to make folder: %s", CrashInfoFolder);
+    }
+}
+
+void FIOSCrashContext::GenerateCrashInfo() const
+{
+    // create a crash-specific directory
+    char CrashInfoFolder[PATH_MAX] = {};
+    FCStringAnsi::Strncpy(CrashInfoFolder, GIOSAppInfo.CrashReportPath, PATH_MAX);
+    FCStringAnsi::Strcat(CrashInfoFolder, PATH_MAX, "/CrashReport-UE4-");
+    FCStringAnsi::Strcat(CrashInfoFolder, PATH_MAX, GIOSAppInfo.AppNameUTF8);
+    FCStringAnsi::Strcat(CrashInfoFolder, PATH_MAX, "-pid-");
+    FCStringAnsi::Strcat(CrashInfoFolder, PATH_MAX, ItoANSI(getpid(), 10));
+    FCStringAnsi::Strcat(CrashInfoFolder, PATH_MAX, "-");
+    FCStringAnsi::Strcat(CrashInfoFolder, PATH_MAX, ItoANSI(GIOSAppInfo.RunUUID.A, 16));
+    FCStringAnsi::Strcat(CrashInfoFolder, PATH_MAX, ItoANSI(GIOSAppInfo.RunUUID.B, 16));
+    FCStringAnsi::Strcat(CrashInfoFolder, PATH_MAX, ItoANSI(GIOSAppInfo.RunUUID.C, 16));
+    FCStringAnsi::Strcat(CrashInfoFolder, PATH_MAX, ItoANSI(GIOSAppInfo.RunUUID.D, 16));
+        
+    const bool bIsEnsure = false;
+    GenerateInfoInFolder(CrashInfoFolder, bIsEnsure);
+        
+    // for IOS we will need to send the report on the next run
+    
+    // Sandboxed applications re-raise the signal to trampoline into the system crash reporter as suppressing it may fall foul of Apple's Mac App Store rules.
+    // @todo Submit an application to the MAS & see whether Apple's reviewers highlight our crash reporting or trampolining to the system reporter.
+    if(GIOSAppInfo.bIsSandboxed)
+    {
+        struct sigaction Action;
+        FMemory::Memzero(&Action, sizeof(struct sigaction));
+        Action.sa_handler = SIG_DFL;
+        sigemptyset(&Action.sa_mask);
+        sigaction(SIGQUIT, &Action, NULL);
+        sigaction(SIGILL, &Action, NULL);
+        sigaction(SIGEMT, &Action, NULL);
+        sigaction(SIGFPE, &Action, NULL);
+        sigaction(SIGBUS, &Action, NULL);
+        sigaction(SIGSEGV, &Action, NULL);
+        sigaction(SIGSYS, &Action, NULL);
+        sigaction(SIGABRT, &Action, NULL);
+        sigaction(SIGTRAP, &Action, NULL);
+        
+        raise(Signal);
+    }
+    
+    _Exit(0);
+}
+
+void FIOSCrashContext::GenerateEnsureInfo() const
+{
+    // Prevent CrashReportClient from spawning another CrashReportClient.
+    const bool bCanRunCrashReportClient = FCString::Stristr( *(GIOSAppInfo.ExecutableName), TEXT( "CrashReportClient" ) ) == nullptr;
+    
+#if !PLATFORM_TVOS
+    if(bCanRunCrashReportClient)
+    {
+        SCOPED_AUTORELEASE_POOL;
+        
+        // Write the PLCrashReporter report to the expected location
+        NSData* CrashReport = [FIOSApplicationInfo::CrashReporter generateLiveReport];
+        [CrashReport writeToFile:[NSString stringWithUTF8String:GIOSAppInfo.PLCrashReportPath] atomically:YES];
+        
+        // Use a slightly different output folder name to not conflict with a subequent crash
+        const FGuid Guid = FGuid::NewGuid();
+        FString GameName = FApp::GetGameName();
+        FString EnsureLogFolder = FString(GIOSAppInfo.CrashReportPath) / FString::Printf(TEXT("EnsureReport-%s-%s"), *GameName, *Guid.ToString(EGuidFormats::Digits));
+        
+        const bool bIsEnsure = true;
+        GenerateInfoInFolder(TCHAR_TO_UTF8(*EnsureLogFolder), bIsEnsure);
+        
+        FString Arguments;
+        if (IsInteractiveEnsureMode())
+        {
+            Arguments = FString::Printf(TEXT("\"%s/\""), *EnsureLogFolder);
+        }
+        else
+        {
+            Arguments = FString::Printf(TEXT("\"%s/\" -Unattended"), *EnsureLogFolder);
+        }
+        
+        FString ReportClient = FPaths::ConvertRelativePathToFull(FPlatformProcess::GenerateApplicationPath(TEXT("CrashReportClient"), EBuildConfigurations::Development));
+        FPlatformProcess::ExecProcess(*ReportClient, *Arguments, nullptr, nullptr, nullptr);
+    }
+#endif
+}
+
+static FCriticalSection EnsureLock;
+static bool bReentranceGuard = false;
+
+void NewReportEnsure( const TCHAR* ErrorMessage )
+{
+    // Simple re-entrance guard.
+    EnsureLock.Lock();
+    
+    if( bReentranceGuard )
+    {
+        EnsureLock.Unlock();
+        return;
+    }
+    
+    bReentranceGuard = true;
+    
+#if !PLATFORM_TVOS
+    if(FIOSApplicationInfo::CrashReporter != nil)
+    {
+        siginfo_t Signal;
+        Signal.si_signo = SIGTRAP;
+        Signal.si_code = TRAP_TRACE;
+        Signal.si_addr = __builtin_return_address(0);
+        
+        FIOSCrashContext EnsureContext;
+        EnsureContext.InitFromSignal(SIGTRAP, &Signal, nullptr);
+        EnsureContext.GenerateEnsureInfo();
+    }
+#endif
+    
+    bReentranceGuard = false;
+    EnsureLock.Unlock();
+}
diff --git a/Engine/Source/Runtime/Core/Private/IOS/IOSPlatformOutputDevices.cpp b/Engine/Source/Runtime/Core/Private/IOS/IOSPlatformOutputDevices.cpp
new file mode 100644
--- /dev/null
+++ b/Engine/Source/Runtime/Core/Private/IOS/IOSPlatformOutputDevices.cpp
@@ -0,0 +1,87 @@
+// Copyright 1998-2017 Epic Games, Inc. All Rights Reserved.
+
+/*=============================================================================
+	IOSPlatformOutputDevices.mm: iOS implementations of OutputDevices functions
+=============================================================================*/
+
+#include "IOSPlatformOutputDevices.h"
+#include "Misc/OutputDeviceFile.h"
+
+#include "FeedbackContextAnsi.h"
+#include "../Private/iOS/IOSPlatformOutputDevicesPrivate.h"
+
+//////////////////////////////////
+// FIOSPlatformOutputDevices
+//////////////////////////////////
+
+class FOutputDeviceError* FIOSPlatformOutputDevices::GetError()
+{
+	static FOutputDeviceIOSError Singleton;
+	return &Singleton;
+}
+class FOutputDevice*	FIOSPlatformOutputDevices::GetLog()
+{
+    static FOutputDeviceFile Singleton(nullptr, true);
+    return &Singleton;
+}
+
+//class FFeedbackContext*				FIOSPlatformOutputDevices::GetWarn()
+//{
+//	static FOutputDeviceIOSDebug Singleton;
+//	return &Singleton;
+//}
+
+//////////////////////////////////
+// FOutputDeviceIOSDebug
+//////////////////////////////////
+
+FOutputDeviceIOSDebug::FOutputDeviceIOSDebug() 
+{}
+
+void FOutputDeviceIOSDebug::Serialize( const TCHAR* Msg, ELogVerbosity::Type Verbosity, const class FName& Category )
+{
+	FPlatformMisc::LowLevelOutputDebugStringf(TEXT("%s%s"),*FOutputDeviceHelper::FormatLogLine(Verbosity, Category, Msg, GPrintLogTimes),LINE_TERMINATOR); 
+}
+
+//////////////////////////////////
+// FOutputDeviceIOSError
+//////////////////////////////////
+
+FOutputDeviceIOSError::FOutputDeviceIOSError()
+:	ErrorPos(0)
+{}
+
+void FOutputDeviceIOSError::Serialize( const TCHAR* Msg, ELogVerbosity::Type Verbosity, const class FName& Category )
+{
+	// @todo iosmerge: this was a big mess copied from Mac that we don't want at all. can we use default like other consoles?
+	if( GIsGuarded )
+	{
+//		FOutputDevice::Serialize(Msg, Verbosity, Category);
+		FPlatformMisc::DebugBreak();
+	}
+	else
+	{
+		// We crashed outside the guarded code (e.g. appExit).
+		HandleError();
+		FPlatformMisc::RequestExit( true );
+	}
+}
+
+void FOutputDeviceIOSError::HandleError()
+{
+	// make sure we don't report errors twice
+	static int32 CallCount = 0;
+	int32 NewCallCount = FPlatformAtomics::InterlockedIncrement(&CallCount);
+	if (NewCallCount != 1)
+	{
+		UE_LOG(LogIOS, Error, TEXT("HandleError re-entered.") );
+		return;
+	}
+
+	GIsGuarded = 0;
+	GIsRunning = 0;
+	GIsCriticalError = 1;
+	GLogConsole = NULL;
+
+	GLog->Flush();
+}
diff --git a/Engine/Source/Runtime/Core/Private/IOS/IOSPlatformOutputDevicesPrivate.h b/Engine/Source/Runtime/Core/Private/IOS/IOSPlatformOutputDevicesPrivate.h
new file mode 100644
--- /dev/null
+++ b/Engine/Source/Runtime/Core/Private/IOS/IOSPlatformOutputDevicesPrivate.h
@@ -0,0 +1,54 @@
+// Copyright 1998-2017 Epic Games, Inc. All Rights Reserved.
+
+/*=============================================================================
+	IOSPlatformOutputDevicesPrivate.h: iOS implementations of output devices
+=============================================================================*/
+
+#pragma once
+
+#include "Misc/OutputDevice.h"
+#include "Misc/OutputDeviceError.h"
+
+/**
+ * Generic Output device that writes to Windows Event Log
+ */
+class FOutputDeviceIOSDebug :
+	public FOutputDevice
+{
+public:
+	/** Constructor, initializing member variables */
+	CORE_API FOutputDeviceIOSDebug();
+
+	/**
+	 * Serializes the passed in data unless the current event is suppressed.
+	 *
+	 * @param	Data	Text to log
+	 * @param	Event	Event name used for suppression purposes
+	 */
+	virtual void Serialize( const TCHAR* Msg, ELogVerbosity::Type Verbosity, const class FName& Category ) override;
+};
+
+
+class FOutputDeviceIOSError : public FOutputDeviceError
+{
+public:
+	/** Constructor, initializing member variables */
+	CORE_API FOutputDeviceIOSError();
+
+	/**
+	 * Serializes the passed in data unless the current event is suppressed.
+	 *
+	 * @param	Data	Text to log
+	 * @param	Event	Event name used for suppression purposes
+	 */
+	virtual void Serialize( const TCHAR* Msg, ELogVerbosity::Type Verbosity, const class FName& Category ) override;
+
+	/**
+	 * Error handling function that is being called from within the system wide global
+	 * error handler, e.g. using structured exception handling on the PC.
+	 */
+	void HandleError();
+
+private:
+	int32		ErrorPos;
+};
diff --git a/Engine/Source/Runtime/Core/Private/IOS/IOSPlatformProcess.cpp b/Engine/Source/Runtime/Core/Private/IOS/IOSPlatformProcess.cpp
new file mode 100644
--- /dev/null
+++ b/Engine/Source/Runtime/Core/Private/IOS/IOSPlatformProcess.cpp
@@ -0,0 +1,294 @@
+// Copyright 1998-2017 Epic Games, Inc. All Rights Reserved.
+
+/*=============================================================================
+IOSPlatformProcess.cpp: iOS implementations of Process functions
+=============================================================================*/
+
+#include "IOSPlatformProcess.h"
+#include "ApplePlatformRunnableThread.h"
+#include "IOSAppDelegate.h"
+#include <mach-o/dyld.h>
+
+// numbers recommended by Apple
+#define GAME_THREAD_PRIORITY 47
+#define RENDER_THREAD_PRIORITY 45
+
+
+const TCHAR* FIOSPlatformProcess::ComputerName()
+{
+	static TCHAR Result[256]=TEXT("");
+
+	if( !Result[0] )
+	{
+		ANSICHAR AnsiResult[ARRAY_COUNT(Result)];
+		gethostname(AnsiResult, ARRAY_COUNT(Result));
+		FCString::Strcpy(Result, ANSI_TO_TCHAR(AnsiResult));
+	}
+	return Result;
+}
+
+const TCHAR* FIOSPlatformProcess::BaseDir()
+{
+	return TEXT("");
+}
+
+FRunnableThread* FIOSPlatformProcess::CreateRunnableThread()
+{
+	return new FRunnableThreadApple();
+}
+
+void FIOSPlatformProcess::LaunchURL( const TCHAR* URL, const TCHAR* Parms, FString* Error )
+{
+	UE_LOG(LogIOS, Log,  TEXT("LaunchURL %s %s"), URL, Parms?Parms:TEXT("") );
+	NSString* CFUrl = (NSString*)FPlatformString::TCHARToCFString( URL );
+    CFUrl = [CFUrl stringByAddingPercentEncodingWithAllowedCharacters: [NSCharacterSet URLQueryAllowedCharacterSet]];
+	[[UIApplication sharedApplication] openURL: [NSURL URLWithString:( NSString *)CFUrl]];
+
+	if (Error != nullptr)
+	{
+		*Error = TEXT("");
+	}
+}
+
+bool FIOSPlatformProcess::CanLaunchURL(const TCHAR* URL)
+{
+	NSString* CFUrl = (NSString*)FPlatformString::TCHARToCFString(URL);
+    CFUrl = [CFUrl stringByAddingPercentEncodingWithAllowedCharacters: [NSCharacterSet URLQueryAllowedCharacterSet]];
+	bool Result = [[UIApplication sharedApplication] canOpenURL:[NSURL URLWithString : (NSString *)CFUrl]];
+
+	return Result;
+}
+
+FString FIOSPlatformProcess::GetGameBundleId()
+{
+	return FString([[NSBundle mainBundle] bundleIdentifier]);
+}
+
+void FIOSPlatformProcess::SetRealTimeMode()
+{
+	if ([IOSAppDelegate GetDelegate].OSVersion < 7 && FPlatformMisc::NumberOfCores() > 1)
+	{
+		mach_timebase_info_data_t TimeBaseInfo;
+		mach_timebase_info( &TimeBaseInfo );
+		double MsToAbs = ((double)TimeBaseInfo.denom / (double)TimeBaseInfo.numer) * 1000000.0;
+		uint32 NormalProcessingTimeMs = 20;
+		uint32 ConstraintProcessingTimeMs = 60;
+
+		thread_time_constraint_policy_data_t Policy;
+		Policy.period      = 0;
+		Policy.computation = (uint32_t)(NormalProcessingTimeMs * MsToAbs);
+		Policy.constraint  = (uint32_t)(ConstraintProcessingTimeMs * MsToAbs);
+		Policy.preemptible = true;
+
+		thread_policy_set(pthread_mach_thread_np(pthread_self()), THREAD_TIME_CONSTRAINT_POLICY, (thread_policy_t)&Policy, THREAD_TIME_CONSTRAINT_POLICY_COUNT);
+	}
+}
+
+// Set the game thread priority to very high, slightly above the render thread
+void FIOSPlatformProcess::SetupThread(const int Priority)
+{
+	struct sched_param Sched;
+	FMemory::Memzero(&Sched, sizeof(struct sched_param));
+
+	// Read the current priority and policy
+	int32 CurrentPolicy = SCHED_RR;
+	pthread_getschedparam(pthread_self(), &CurrentPolicy, &Sched);
+
+	// Set the new priority and policy (apple recommended FIFO for the two main non-working threads)
+	int32 Policy = SCHED_FIFO;
+	Sched.sched_priority = Priority;
+	pthread_setschedparam(pthread_self(), Policy, &Sched);
+}
+
+void FIOSPlatformProcess::SetupGameThread()
+{
+	SetupThread(GAME_THREAD_PRIORITY);
+}
+
+void FIOSPlatformProcess::SetupRenderThread()
+{
+	SetupThread(RENDER_THREAD_PRIORITY);
+}
+
+void FIOSPlatformProcess::SetThreadAffinityMask(uint64 AffinityMask)
+{
+	if ([IOSAppDelegate GetDelegate].OSVersion >= 8 && FPlatformMisc::NumberOfCores() > 1)
+	{
+		thread_affinity_policy AP;
+		AP.affinity_tag = AffinityMask;
+		thread_policy_set(pthread_mach_thread_np(pthread_self()), THREAD_AFFINITY_POLICY, (integer_t*)&AP, THREAD_AFFINITY_POLICY_COUNT);
+	}
+	else
+	{
+		FGenericPlatformProcess::SetThreadAffinityMask(AffinityMask);
+	}
+}
+
+const TCHAR* FIOSPlatformProcess::ExecutableName(bool bRemoveExtension)
+{
+	static TCHAR Result[512] = TEXT("");
+	if (!Result[0])
+	{
+		NSString *NSExeName = [[[NSBundle mainBundle] executablePath] lastPathComponent];
+		FPlatformString::CFStringToTCHAR((CFStringRef)NSExeName, Result);
+	}
+	return Result;
+}
+
+FString FIOSPlatformProcess::GenerateApplicationPath( const FString& AppName, EBuildConfigurations::Type BuildConfiguration)
+{
+    SCOPED_AUTORELEASE_POOL;
+    
+    FString PlatformName = TEXT("IOS");
+    FString ExecutableName = AppName;
+    if (BuildConfiguration != EBuildConfigurations::Development && BuildConfiguration != EBuildConfigurations::DebugGame)
+    {
+        ExecutableName += FString::Printf(TEXT("-%s-%s"), *PlatformName, EBuildConfigurations::ToString(BuildConfiguration));
+    }
+    
+    NSURL* CurrentBundleURL = [[NSBundle mainBundle] bundleURL];
+    NSString* CurrentBundleName = [[CurrentBundleURL lastPathComponent] stringByDeletingPathExtension];
+    if(FString(CurrentBundleName) == ExecutableName)
+    {
+        CFStringRef FilePath = CFURLCopyFileSystemPath((CFURLRef)CurrentBundleURL, kCFURLPOSIXPathStyle);
+        FString ExecutablePath = FString::Printf(TEXT("%s/%s"), *FString((NSString*)FilePath), *ExecutableName);
+        CFRelease(FilePath);
+        return ExecutablePath;
+    }
+    else
+    {
+        return FString();
+    }
+}
+
+const uint64 FIOSPlatformAffinity::GetMainGameMask()
+{
+	return MAKEAFFINITYMASK1(0);
+}
+
+const uint64 FIOSPlatformAffinity::GetRenderingThreadMask()
+{
+	static int Mask = 0;
+	if (Mask == 0)
+	{
+		switch (FPlatformMisc::NumberOfCores())
+		{
+		case 2:
+		case 3:
+			Mask = MAKEAFFINITYMASK1(1);
+			break;
+		default:
+			Mask = FGenericPlatformAffinity::GetRenderingThreadMask();
+			break;
+		}
+	}
+
+	return Mask;
+}
+
+const uint64 FIOSPlatformAffinity::GetRTHeartBeatMask()
+{
+	static int Mask = 0;
+	if (Mask == 0)
+	{
+		switch (FPlatformMisc::NumberOfCores())
+		{
+		case 2:
+		case 3:
+			Mask = MAKEAFFINITYMASK1(0);
+			break;
+		default:
+			Mask = FGenericPlatformAffinity::GetRTHeartBeatMask();
+			break;
+		}
+	}
+
+	return Mask;
+}
+
+const uint64 FIOSPlatformAffinity::GetPoolThreadMask()
+{
+	static int Mask = 0;
+	if (Mask == 0)
+	{
+		switch (FPlatformMisc::NumberOfCores())
+		{
+		case 2:
+			Mask = MAKEAFFINITYMASK1(1);
+			break;
+		case 3:
+			Mask = MAKEAFFINITYMASK1(2);
+			break;
+		default:
+			Mask = FGenericPlatformAffinity::GetPoolThreadMask();
+			break;
+		}
+	}
+
+	return Mask;
+}
+
+const uint64 FIOSPlatformAffinity::GetTaskGraphThreadMask()
+{
+	static int Mask = 0;
+	if (Mask == 0)
+	{
+		switch (FPlatformMisc::NumberOfCores())
+		{
+		case 2:
+			Mask = MAKEAFFINITYMASK1(1);
+			break;
+		case 3:
+			Mask = MAKEAFFINITYMASK1(2);
+			break;
+		default:
+			Mask = FGenericPlatformAffinity::GetTaskGraphThreadMask();
+			break;
+		}
+	}
+
+	return Mask;
+}
+
+const uint64 FIOSPlatformAffinity::GetStatsThreadMask()
+{
+	static int Mask = 0;
+	if (Mask == 0)
+	{
+		switch (FPlatformMisc::NumberOfCores())
+		{
+		case 2:
+			Mask = MAKEAFFINITYMASK1(0);
+			break;
+		case 3:
+			Mask = MAKEAFFINITYMASK1(2);
+			break;
+		default:
+			Mask = FGenericPlatformAffinity::GetStatsThreadMask();
+			break;
+		}
+	}
+
+	return Mask;
+}
+
+const uint64 FIOSPlatformAffinity::GetNoAffinityMask()
+{
+	static int Mask = 0;
+	if (Mask == 0)
+	{
+		switch (FPlatformMisc::NumberOfCores())
+		{
+		case 2:
+		case 3:
+			Mask = (1 << FPlatformMisc::NumberOfCores()) - 1;
+			break;
+		default:
+			Mask = FGenericPlatformAffinity::GetNoAffinityMask();
+			break;
+		}
+	}
+
+	return Mask;
+}
+
diff --git a/Engine/Source/Runtime/Core/Private/IOS/IOSPlatformString.cpp b/Engine/Source/Runtime/Core/Private/IOS/IOSPlatformString.cpp
new file mode 100644
--- /dev/null
+++ b/Engine/Source/Runtime/Core/Private/IOS/IOSPlatformString.cpp
@@ -0,0 +1,25 @@
+// Copyright 1998-2017 Epic Games, Inc. All Rights Reserved.
+
+/*=============================================================================
+	IOSPlatformString.mm: Additional NSString methods
+=============================================================================*/
+
+#include "IOSPlatformString.h"
+#include "IOSSystemIncludes.h"
+#include "UnrealString.h"
+#include "Containers/StringConv.h"
+
+@implementation NSString (FString_Extensions)
+
++ (NSString*) stringWithTCHARString:(const TCHAR*)MyTCHARString
+{
+	return [NSString stringWithCString:TCHAR_TO_UTF8(MyTCHARString) encoding:NSUTF8StringEncoding];
+}
+
++ (NSString*) stringWithFString:(const FString&)InFString
+{
+	return [NSString stringWithTCHARString:*InFString];
+}
+
+
+@end
diff --git a/Engine/Source/Runtime/Core/Private/IOS/IOSPlatformSurvey.cpp b/Engine/Source/Runtime/Core/Private/IOS/IOSPlatformSurvey.cpp
new file mode 100644
--- /dev/null
+++ b/Engine/Source/Runtime/Core/Private/IOS/IOSPlatformSurvey.cpp
@@ -0,0 +1,39 @@
+// Copyright 1998-2017 Epic Games, Inc. All Rights Reserved.
+
+/*=============================================================================
+	IOSPlatformSurvey.cpp: HardwareSurvey implementation
+=============================================================================*/
+
+#include "IOSPlatformSurvey.h"
+#include "IOSAppDelegate.h"
+#include "IOSWindow.h"
+
+bool FIOSPlatformSurvey::GetSurveyResults(FHardwareSurveyResults& OutResults, bool bWait)
+{
+	FCString::Strcpy(OutResults.Platform, FPlatformMisc::GetDefaultDeviceProfileName());
+	FCString::Strcpy(OutResults.OSVersion, *FString([[UIDevice currentDevice] systemVersion]));
+#if PLATFORM_64BITS
+	OutResults.OSBits = 64;
+#else
+	OutResults.OSBits = 32;
+#endif
+	FCString::Strcpy(OutResults.OSLanguage, *FString([[NSLocale preferredLanguages] objectAtIndex:0]));
+	FCString::Strcpy(OutResults.MultimediaAPI, FPlatformMisc::HasPlatformFeature(TEXT("Metal")) ? TEXT("Metal") : TEXT("ES2"));
+	OutResults.CPUCount = FPlatformMisc::NumberOfCores();
+
+	// display 0 is max size
+	CGRect MainFrame = [[UIScreen mainScreen] bounds];
+	if ([IOSAppDelegate GetDelegate].OSVersion < 8.0f && ![IOSAppDelegate GetDelegate].bDeviceInPortraitMode)
+	{
+		Swap(MainFrame.size.width, MainFrame.size.height);
+	}
+	float Scale = [[UIScreen mainScreen] scale];
+	OutResults.Displays[0].CurrentModeWidth = MainFrame.size.width * Scale;
+	OutResults.Displays[0].CurrentModeHeight = MainFrame.size.height * Scale;
+
+	// display 1 is current size
+	FPlatformRect ScreenSize = FIOSWindow::GetScreenRect();
+	OutResults.Displays[1].CurrentModeWidth = ScreenSize.Right - ScreenSize.Left;
+	OutResults.Displays[1].CurrentModeHeight = ScreenSize.Bottom - ScreenSize.Top;
+	return true;
+}
diff --git a/Engine/Source/Runtime/Core/Private/IOS/IOSView.cpp b/Engine/Source/Runtime/Core/Private/IOS/IOSView.cpp
new file mode 100644
--- /dev/null
+++ b/Engine/Source/Runtime/Core/Private/IOS/IOSView.cpp
@@ -0,0 +1,1055 @@
+// Copyright 1998-2017 Epic Games, Inc. All Rights Reserved.
+
+#include "IOSView.h"
+#include "IOSAppDelegate.h"
+#include "IOSApplication.h"
+#include "IOS/IOSInputInterface.h"
+#include "Delegates/Delegate.h"
+#include "Misc/ConfigCacheIni.h"
+#include "HAL/IConsoleManager.h"
+#include "Misc/CommandLine.h"
+#include "IOSPlatformProcess.h"
+
+#include <OpenGLES/ES2/gl.h>
+#import "IOSAsyncTask.h"
+#import <QuartzCore/QuartzCore.h>
+#import <OpenGLES/EAGLDrawable.h>
+#import <UIKit/UIGeometry.h>
+
+#if HAS_METAL
+id<MTLDevice> GMetalDevice = nil;
+#endif
+
+
+@interface IndexedPosition : UITextPosition {
+	NSUInteger _index;
+	id <UITextInputDelegate> _inputDelegate;
+}
+@property (nonatomic) NSUInteger index;
++ (IndexedPosition *)positionWithIndex:(NSUInteger)index;
+@end
+
+@interface IndexedRange : UITextRange {
+	NSRange _range;
+}
+@property (nonatomic) NSRange range;
++ (IndexedRange *)rangeWithNSRange:(NSRange)range;
+
+@end
+
+
+@implementation IndexedPosition
+@synthesize index = _index;
+
++ (IndexedPosition *)positionWithIndex:(NSUInteger)index {
+	IndexedPosition *pos = [[IndexedPosition alloc] init];
+	pos.index = index;
+	return pos;
+}
+
+@end
+
+@implementation IndexedRange
+@synthesize range = _range;
+
++ (IndexedRange *)rangeWithNSRange:(NSRange)nsrange {
+	if (nsrange.location == NSNotFound)
+		return nil;
+	IndexedRange *range = [[IndexedRange alloc] init];
+	range.range = nsrange;
+	return range;
+}
+
+- (UITextPosition *)start {
+	return [IndexedPosition positionWithIndex:self.range.location];
+}
+
+- (UITextPosition *)end {
+	return [IndexedPosition positionWithIndex:(self.range.location + self.range.length)];
+}
+
+-(BOOL)isEmpty {
+	return (self.range.length == 0);
+}
+@end
+
+
+
+@implementation FIOSView
+
+@synthesize keyboardType = KeyboardType;
+@synthesize autocorrectionType = AutocorrectionType;
+@synthesize autocapitalizationType = AutocapitalizationType;
+@synthesize secureTextEntry = bSecureTextEntry;
+@synthesize SwapCount, OnScreenColorRenderBuffer, OnScreenColorRenderBufferMSAA, markedTextStyle;
+
+/**
+ * @return The Layer Class for the window
+ */
++ (Class)layerClass
+{
+#if HAS_METAL
+	// make sure the project setting has enabled Metal support (per-project user settings in the editor)
+	bool bSupportsMetal = false;
+	bool bSupportsMetalMRT = false;
+	GConfig->GetBool(TEXT("/Script/IOSRuntimeSettings.IOSRuntimeSettings"), TEXT("bSupportsMetal"), bSupportsMetal, GEngineIni);
+	GConfig->GetBool(TEXT("/Script/IOSRuntimeSettings.IOSRuntimeSettings"), TEXT("bSupportsMetalMRT"), bSupportsMetalMRT, GEngineIni);
+
+	// does commandline override?
+	bool bForceES2 = FParse::Param(FCommandLine::Get(), TEXT("ES2"));
+
+	bool bTriedToInit = false;
+
+	// the check for the function pointer itself is to determine if the Metal framework exists, before calling it
+	if ((bSupportsMetal || bSupportsMetalMRT) && !bForceES2 && MTLCreateSystemDefaultDevice != NULL)
+	{
+		// if the device is unable to run with Metal (pre-A7), this will return nil
+		GMetalDevice = MTLCreateSystemDefaultDevice();
+
+		// just tracking for printout below
+		bTriedToInit = true;
+	}
+
+#if !UE_BUILD_SHIPPING
+	if (GMetalDevice == nil)
+	{
+		FPlatformMisc::LowLevelOutputDebugStringf(TEXT("Not using Metal because: [Project Settings Disabled Metal? %s :: Commandline Forced ES2? %s :: Older OS? %s :: Pre-A7 Device? %s]"),
+			bSupportsMetal ? TEXT("No") : TEXT("Yes"),
+			bForceES2? TEXT("Yes") : TEXT("No"),
+			(MTLCreateSystemDefaultDevice == NULL) ? TEXT("Yes") : TEXT("No"),
+			bTriedToInit ? TEXT("Yes") : TEXT("Unknown (didn't test)"));
+	}
+#endif
+
+	if (GMetalDevice != nil)
+	{
+		return [CAMetalLayer class];
+	}
+	else
+#endif
+	{
+		return [CAEAGLLayer class];
+	}
+}
+
+
+- (id)initWithFrame:(CGRect)Frame
+{
+	if ((self = [super initWithFrame:Frame]))
+	{
+		// figure out if we should start up GL or Metal
+#if HAS_METAL
+		// if the device is valid, we know Metal is usable (see +layerClass)
+		MetalDevice = GMetalDevice;
+		if (MetalDevice != nil)
+		{
+			bIsUsingMetal = true;
+
+			// grab the MetalLayer and typecast it to match what's in layerClass
+			CAMetalLayer* MetalLayer = (CAMetalLayer*)self.layer;
+			MetalLayer.presentsWithTransaction = NO;
+			MetalLayer.drawsAsynchronously = YES;
+
+			// set a background color to make sure the layer appears
+			CGFloat components[] = { 0.0, 0.0, 0.0, 1 };
+			MetalLayer.backgroundColor = CGColorCreate(CGColorSpaceCreateDeviceRGB(), components);
+
+			// set the device on the rendering layer and provide a pixel format
+			MetalLayer.device = MetalDevice;
+			MetalLayer.pixelFormat = MTLPixelFormatBGRA8Unorm;
+			MetalLayer.framebufferOnly = NO;
+		
+		}
+		else
+#endif
+		{
+			bIsUsingMetal = false;
+
+			// Get the layer
+			CAEAGLLayer *EaglLayer = (CAEAGLLayer *)self.layer;
+			EaglLayer.opaque = YES;
+			NSMutableDictionary* Dict = [NSMutableDictionary dictionary];
+			[Dict setValue : [NSNumber numberWithBool : NO] forKey : kEAGLDrawablePropertyRetainedBacking];
+			[Dict setValue : kEAGLColorFormatRGBA8 forKey : kEAGLDrawablePropertyColorFormat];
+			EaglLayer.drawableProperties = Dict;
+
+			// Initialize a single, static OpenGL ES 2.0 context, shared by all EAGLView objects
+			Context = [[EAGLContext alloc] initWithAPI:kEAGLRenderingAPIOpenGLES2];
+
+			// delete this on failure
+			if (!Context || ![EAGLContext setCurrentContext : Context])
+			{
+				[self release];
+				return nil;
+			}
+		}
+
+		NSLog(@"::: Created a UIView that will support %@ :::", bIsUsingMetal ? @"Metal" : @"@GLES");
+
+		// Initialize some variables
+		SwapCount = 0;
+
+		FMemory::Memzero(AllTouches, sizeof(AllTouches));
+		[self setAutoresizingMask: UIViewAutoresizingFlexibleWidth | UIViewAutoresizingFlexibleHeight];
+		bIsInitialized = false;
+		
+
+		[self InitKeyboard];
+	}
+	return self;
+}
+
+
+- (bool)CreateFramebuffer:(bool)bIsForOnDevice
+{
+	if (!bIsInitialized)
+	{
+		// look up what the device can support
+		const float NativeScale = [[UIScreen mainScreen] scale];
+
+		// look up the CVar for the scale factor
+		static IConsoleVariable* CVar = IConsoleManager::Get().FindConsoleVariable(TEXT("r.MobileContentScaleFactor"));
+		float RequestedContentScaleFactor = CVar->GetFloat();
+
+		// 0 means to leave the scale alone, use native
+		if (RequestedContentScaleFactor == 0.0f)
+		{
+#ifdef __IPHONE_8_0
+            if ([self.window.screen respondsToSelector:@selector(nativeScale)])
+            {
+                self.contentScaleFactor = self.window.screen.nativeScale;
+                UE_LOG(LogIOS, Log, TEXT("Setting contentScaleFactor to nativeScale which is = %f"), self.contentScaleFactor);
+            }
+            else
+#endif
+            {
+                UE_LOG(LogIOS, Log, TEXT("Leaving contentScaleFactor alone, with scale = %f"), NativeScale);
+            }
+		}
+		else
+		{
+			// for TV screens, always use scale factor of 1
+			self.contentScaleFactor = bIsForOnDevice ? RequestedContentScaleFactor : 1.0f;
+			UE_LOG(LogIOS, Log, TEXT("Setting contentScaleFactor to %0.4f (optimal = %0.4f)"), self.contentScaleFactor, NativeScale);
+		}
+
+
+		// handle Metal or GL sizing
+#if HAS_METAL
+		if (bIsUsingMetal)
+		{
+			CAMetalLayer* MetalLayer = (CAMetalLayer*)self.layer;
+			CGSize DrawableSize = self.bounds.size;
+			DrawableSize.width *= self.contentScaleFactor;
+			DrawableSize.height *= self.contentScaleFactor;
+			MetalLayer.drawableSize = DrawableSize;
+		}
+		else
+#endif
+		{
+			// make sure this is current
+			[self MakeCurrent];
+
+			// This is causing a pretty large slow down on iPad 3 and 4 using iOS 6, for now going to comment it out
+	//		if (self.contentScaleFactor == 1.0f || self.contentScaleFactor == 2.0f)
+	//		{
+	//			UE_LOG(LogIOS,Log,TEXT("Setting layer filter to NEAREST"));
+	//			CAEAGLLayer *EaglLayer = (CAEAGLLayer *)self.layer;
+	//			EaglLayer.magnificationFilter = kCAFilterNearest;
+	//		}
+
+			// Create our standard displayable surface
+			glGenRenderbuffers(1, &OnScreenColorRenderBuffer);
+			check(glGetError() == 0);
+			glBindRenderbuffer(GL_RENDERBUFFER, OnScreenColorRenderBuffer);
+			check(glGetError() == 0);
+			[Context renderbufferStorage:GL_RENDERBUFFER fromDrawable:(CAEAGLLayer*)self.layer];
+
+			// Get the size of the surface
+			GLint OnScreenWidth, OnScreenHeight;
+			glGetRenderbufferParameteriv(GL_RENDERBUFFER, GL_RENDERBUFFER_WIDTH, &OnScreenWidth);
+			glGetRenderbufferParameteriv(GL_RENDERBUFFER, GL_RENDERBUFFER_HEIGHT, &OnScreenHeight);
+
+			// NOTE: This resolve FBO is necessary even if we don't plan on using MSAA because otherwise
+			// the shaders will not warm properly. Future investigation as to why; it seems unnecessary.
+
+			// Create an FBO used to target the resolve surface
+			glGenFramebuffers(1, &ResolveFrameBuffer);
+			check(glGetError() == 0);
+			glBindFramebuffer(GL_FRAMEBUFFER, ResolveFrameBuffer);
+			check(glGetError() == 0);
+			glFramebufferRenderbuffer(GL_FRAMEBUFFER, GL_COLOR_ATTACHMENT0, GL_RENDERBUFFER, OnScreenColorRenderBuffer);
+			check(glGetError() == 0);
+			check(glCheckFramebufferStatus(GL_FRAMEBUFFER) == GL_FRAMEBUFFER_COMPLETE);
+
+#if defined(USE_DETAILED_IPHONE_MEM_TRACKING) && USE_DETAILED_IPHONE_MEM_TRACKING
+			//This value is used to allow the engine to track gl allocated memory see GetIPhoneOpenGLBackBufferSize
+			uint32 SingleBufferSize = OnScreenWidth * OnScreenHeight * 4/*rgba8*/;
+			IPhoneBackBufferMemSize = singleBufferSize * 3/*iphone back buffer system is tripple buffered*/;
+
+			UE_LOG(LogIOS, Log, TEXT("IPhone Back Buffer Size: %i MB"), (IPhoneBackBufferMemSize / 1024) / 1024.f);
+#endif
+		}
+
+		bIsInitialized = true;
+	}    
+	return true;
+}
+
+/**
+ * If view is resized, update the frame buffer so it is the same size as the display area.
+ */
+- (void)layoutSubviews
+{
+#if !PLATFORM_TVOS
+	UIDeviceOrientation orientation = [[UIDevice currentDevice] orientation];
+	FIOSApplication::OrientationChanged(orientation);
+#endif
+}
+
+-(void)UpdateRenderWidth:(uint32)Width andHeight:(uint32)Height
+{
+#if HAS_METAL
+	if (bIsUsingMetal)
+	{
+		if (MetalDevice != nil)
+		{
+			// grab the MetalLayer and typecast it to match what's in layerClass
+			CAMetalLayer* MetalLayer = (CAMetalLayer*)self.layer;
+			CGSize drawableSize = CGSizeMake(Width, Height);
+			checkf( FMath::TruncToInt(drawableSize.width) == FMath::TruncToInt(self.bounds.size.width * self.contentScaleFactor) && 
+				   FMath::TruncToInt(drawableSize.height) == FMath::TruncToInt(self.bounds.size.height * self.contentScaleFactor),
+				TEXT("[IOSView UpdateRenderWidth:andHeight:] passed in size doesn't match what we expected. Width: %d, Expected Width = %d (%.2f * %.2f). Height = %d, Expected Height = %d (%.2f * %.2f)"),
+				FMath::TruncToInt(drawableSize.width), FMath::TruncToInt(self.bounds.size.width * self.contentScaleFactor), (float)self.bounds.size.width, self.contentScaleFactor,
+				FMath::TruncToInt(drawableSize.height), FMath::TruncToInt(self.bounds.size.height * self.contentScaleFactor), (float)self.bounds.size.height, self.contentScaleFactor);
+			MetalLayer.drawableSize = drawableSize;
+		}
+		return;
+	}
+#endif
+
+	// Allocate color buffer based on the current layer size
+	glBindRenderbuffer(GL_RENDERBUFFER, OnScreenColorRenderBuffer);
+	check(glGetError() == 0);
+	[Context renderbufferStorage:GL_RENDERBUFFER fromDrawable:(CAEAGLLayer*)self.layer];
+	glBindFramebuffer(GL_FRAMEBUFFER, ResolveFrameBuffer);
+	check(glGetError() == 0);
+	glFramebufferRenderbuffer(GL_FRAMEBUFFER, GL_COLOR_ATTACHMENT0, GL_RENDERBUFFER, OnScreenColorRenderBuffer);
+	check(glGetError() == 0);
+	check(glCheckFramebufferStatus(GL_FRAMEBUFFER) == GL_FRAMEBUFFER_COMPLETE);
+}
+
+#if HAS_METAL
+- (id<CAMetalDrawable>)MakeDrawable
+{
+    return ![IOSAppDelegate GetDelegate].bIsSuspended ? [(CAMetalLayer*)self.layer nextDrawable] : nil;
+}
+#endif
+
+- (void)DestroyFramebuffer
+{
+	if (bIsInitialized)
+	{
+		// nothing to do here for Metal
+		if (!bIsUsingMetal)
+		{
+			// toss framebuffers
+			if(ResolveFrameBuffer)
+			{
+				glDeleteFramebuffers(1, &ResolveFrameBuffer);
+				ResolveFrameBuffer = 0;
+			}
+			if(OnScreenColorRenderBuffer)
+			{
+				glDeleteRenderbuffers(1, &OnScreenColorRenderBuffer);
+				OnScreenColorRenderBuffer = 0;
+			}
+// 			if( GMSAAAllowed )
+// 			{
+// 				if(OnScreenColorRenderBufferMSAA)
+// 				{
+// 					glDeleteRenderbuffers(1, &OnScreenColorRenderBufferMSAA);
+// 					OnScreenColorRenderBufferMSAA = 0;
+// 				}
+//			}
+ 		}
+
+		// we are ready to be re-initialized
+		bIsInitialized = false;
+	}
+}
+
+- (void)MakeCurrent
+{
+	if (!bIsUsingMetal)
+	{
+		[EAGLContext setCurrentContext:Context];
+	}
+}
+
+- (void)UnmakeCurrent
+{
+	if (!bIsUsingMetal)
+	{
+		[EAGLContext setCurrentContext:nil];
+	}
+}
+
+- (void)SwapBuffers
+{
+	if (!bIsUsingMetal)
+	{
+// 		// We may need this in the MSAA case
+// 		GLint CurrentFramebuffer = 0;
+// 		// @todo-mobile: Fix this when we have MSAA support
+// 		if( GMSAAAllowed && GMSAAEnabled )
+// 		{
+// 			// Get the currently bound FBO
+// 			glGetIntegerv(GL_FRAMEBUFFER_BINDING, &CurrentFramebuffer);
+// 
+// 			// Set up and perform the resolve (the READ is already set)
+// 			glBindFramebuffer(GL_DRAW_FRAMEBUFFER_APPLE, ResolveFrameBuffer);
+// 			glResolveMultisampleFramebufferAPPLE();
+// 
+// 			// After the resolve, we can discard the old attachments
+// 			GLenum attachments[] = { GL_COLOR_ATTACHMENT0, GL_DEPTH_ATTACHMENT, GL_STENCIL_ATTACHMENT };
+// 			glDiscardFramebufferEXT( GL_READ_FRAMEBUFFER_APPLE, 3, attachments );
+// 		}
+// 		else
+// 		{
+// 			// Discard the now-unncessary depth buffer
+// 			GLenum attachments[] = { GL_DEPTH_ATTACHMENT, GL_STENCIL_ATTACHMENT };
+// 			glDiscardFramebufferEXT( GL_READ_FRAMEBUFFER_APPLE, 2, attachments );
+// 		}
+// 
+		// Perform the actual present with the on-screen renderbuffer
+		//glBindRenderbuffer(GL_RENDERBUFFER, OnScreenColorRenderBuffer);
+		[Context presentRenderbuffer:GL_RENDERBUFFER];
+
+// 		if( GMSAAAllowed && GMSAAEnabled )
+// 		{
+// 			// Restore the DRAW framebuffer object
+// 			glBindFramebuffer(GL_DRAW_FRAMEBUFFER_APPLE, CurrentFramebuffer);
+// 		}
+	}
+
+	// increment our swap counter
+	SwapCount++;
+}
+
+/**
+ * Returns the unique ID for the given touch
+ */
+-(int32) GetTouchIndex:(UITouch*)Touch
+{
+	// look for existing touch
+	for (int Index = 0; Index < ARRAY_COUNT(AllTouches); Index++)
+	{
+		if (AllTouches[Index] == Touch)
+		{
+			return Index;
+		}
+	}
+	
+	// if we get here, it's a new touch, find a slot
+	for (int Index = 0; Index < ARRAY_COUNT(AllTouches); Index++)
+	{
+		if (AllTouches[Index] == nil)
+		{
+			AllTouches[Index] = Touch;
+			return Index;
+		}
+	}
+	
+	// if we get here, that means we are trying to use more than 5 touches at once, which is an error
+	return -1;
+}
+
+/**
+ * Pass touch events to the input queue for slate to pull off of, and trigger the debug console.
+ *
+ * @param View The view the event happened in
+ * @param Touches Set of touch events from the OS
+ */
+-(void) HandleTouches:(NSSet*)Touches ofType:(TouchType)Type
+{
+	CGFloat Scale = self.contentScaleFactor;
+
+	TArray<TouchInput> TouchesArray;
+	for (UITouch* Touch in Touches)
+	{
+		// get info from the touch
+		CGPoint Loc = [Touch locationInView:self];
+		CGPoint PrevLoc = [Touch previousLocationInView:self];
+
+		// convert TOuch pointer to a unique 0 based index
+		int32 TouchIndex = [self GetTouchIndex:Touch];
+		if (TouchIndex < 0)
+		{
+			continue;
+		}
+
+		// make a new touch event struct
+		TouchInput TouchMessage;
+		TouchMessage.Handle = TouchIndex;
+		TouchMessage.Type = Type;
+		TouchMessage.Position = FVector2D(FMath::Min<float>(self.frame.size.width - 1, Loc.x), FMath::Min<float>(self.frame.size.height - 1, Loc.y)) * Scale;
+		TouchMessage.LastPosition = FVector2D(FMath::Min<float>(self.frame.size.width - 1, PrevLoc.x), FMath::Min<float>(self.frame.size.height - 1, PrevLoc.y)) * Scale;
+		TouchesArray.Add(TouchMessage);
+
+		if (Type == TouchBegan)
+		{
+			TouchInput EmulatedMessage;
+			EmulatedMessage.Handle = TouchMessage.Handle;
+			EmulatedMessage.Type = TouchMoved;
+			EmulatedMessage.Position = TouchMessage.Position;
+			EmulatedMessage.LastPosition = TouchMessage.Position;
+			TouchesArray.Add(EmulatedMessage);
+		}
+		
+		// clear out the touch when it ends
+		if (Type == TouchEnded)
+		{
+			AllTouches[TouchIndex] = nil;
+		}
+	}
+
+	FIOSInputInterface::QueueTouchInput(TouchesArray);
+}
+
+/**
+ * Handle the various touch types from the OS
+ *
+ * @param touches Array of touch events
+ * @param event Event information
+ */
+- (void) touchesBegan:(NSSet*)touches withEvent:(UIEvent*)event 
+{
+	NumActiveTouches += touches.count;
+	[self HandleTouches:touches ofType:TouchBegan];
+
+#if !UE_BUILD_SHIPPING
+#if WITH_SIMULATOR
+	// use 2 on the simulator so that Option-Click will bring up console (option-click is for doing pinch gestures, which we don't care about, atm)
+	if( NumActiveTouches >= 2 )
+#else
+	// If there are 3 active touches, bring up the console
+	if( NumActiveTouches >= 4 )
+#endif
+	{
+		bool bShowConsole = true;
+		GConfig->GetBool(TEXT("/Script/Engine.InputSettings"), TEXT("bShowConsoleOnFourFingerTap"), bShowConsole, GInputIni);
+
+		if (bShowConsole)
+		{
+			// disable the integrated keyboard when launching the console
+/*			if (bIsUsingIntegratedKeyboard)
+			{
+				// press the console key twice to get the big one up
+				// @todo keyboard: Find a direct way to bering this up (it can get into a bad state where two presses won't do it correctly)
+				// and also the ` key could be changed via .ini
+				FIOSInputInterface::QueueKeyInput('`', '`');
+				FIOSInputInterface::QueueKeyInput('`', '`');
+				
+				[self ActivateKeyboard:true];
+			}
+			else*/
+			{
+				// Route the command to the main iOS thread (all UI must go to the main thread)
+				[[IOSAppDelegate GetDelegate] performSelectorOnMainThread:@selector(ShowConsole) withObject:nil waitUntilDone:NO];
+			}
+		}
+	}
+#endif
+}
+
+- (void) touchesMoved:(NSSet*)touches withEvent:(UIEvent*)event
+{
+	[self HandleTouches:touches ofType:TouchMoved];
+}
+
+- (void) touchesEnded:(NSSet*)touches withEvent:(UIEvent*)event
+{
+	NumActiveTouches -= touches.count;
+	[self HandleTouches:touches ofType:TouchEnded];
+}
+
+- (void) touchesCancelled:(NSSet*)touches withEvent:(UIEvent*)event
+{
+	NumActiveTouches -= touches.count;
+	[self HandleTouches:touches ofType:TouchEnded];
+}
+
+
+
+
+
+
+
+
+
+
+
+
+#pragma mark Keyboard
+
+
+-(BOOL)canBecomeFirstResponder
+{
+	return YES;
+}
+
+- (BOOL)hasText
+{
+	return YES;
+}
+
+- (void)insertText:(NSString *)theText
+{
+	// insert text one key at a time, as chars, not keydowns
+	for (int32 CharIndex = 0; CharIndex < [theText length]; CharIndex++)
+	{
+		int32 Char = [theText characterAtIndex:CharIndex];
+
+		// FPlatformMisc::LowLevelOutputDebugStringf(TEXT("sending key '%c' to game\n"), Char);
+
+		if (Char == '\n')
+		{
+			// send the enter keypress
+			FIOSInputInterface::QueueKeyInput(KEYCODE_ENTER, Char);
+			
+			// hide the keyboard
+			[self resignFirstResponder];
+		}
+		else
+		{
+			FIOSInputInterface::QueueKeyInput(Char, Char);
+		}
+	}
+}
+
+- (void)deleteBackward
+{
+	FIOSInputInterface::QueueKeyInput(KEYCODE_BACKSPACE, '\b');
+}
+
+-(void)ActivateKeyboard:(bool)bInSendEscapeOnClose
+{
+	FKeyboardConfig DefaultConfig;
+	[self ActivateKeyboard:bInSendEscapeOnClose keyboardConfig:DefaultConfig];
+}
+
+-(void)ActivateKeyboard:(bool)bInSendEscapeOnClose keyboardConfig:(FKeyboardConfig)KeyboardConfig
+{
+	FPlatformAtomics::InterlockedIncrement(&KeyboardShowCount);
+	
+	dispatch_async(dispatch_get_main_queue(),^ {
+		volatile int32 ShowCount = KeyboardShowCount;
+		if (ShowCount == 1)
+		{
+			self.keyboardType = KeyboardConfig.KeyboardType;
+			self.autocorrectionType = KeyboardConfig.AutocorrectionType;
+			self.autocapitalizationType = KeyboardConfig.AutocapitalizationType;
+			self.secureTextEntry = KeyboardConfig.bSecureTextEntry;
+		
+			// Remember the setting
+			bSendEscapeOnClose = bInSendEscapeOnClose;
+		
+			// Dismiss the existing keyboard, if one exists, so the style can be overridden.
+			[self endEditing:YES];
+			[self becomeFirstResponder];
+		}
+		
+		FPlatformAtomics::InterlockedDecrement(&KeyboardShowCount);
+	});
+}
+
+-(void)DeactivateKeyboard
+{
+	dispatch_async(dispatch_get_main_queue(),^ {
+		volatile int32 ShowCount = KeyboardShowCount;
+		if (ShowCount == 0)
+		{
+			// Wait briefly, in case a keyboard activation is triggered.
+			FPlatformProcess::Sleep(0.1F);
+			
+			ShowCount = KeyboardShowCount;
+			if (ShowCount == 0)
+			{
+				// Dismiss the existing keyboard, if one exists.
+				[self endEditing:YES];
+			}
+		}
+	});
+}
+
+-(BOOL)becomeFirstResponder
+{
+	volatile int32 ShowCount = KeyboardShowCount;
+	if (ShowCount >= 1)
+	{
+		return [super becomeFirstResponder];
+	}
+	else
+	{
+		return NO;
+	}
+}
+
+-(BOOL)resignFirstResponder
+{
+	if (bSendEscapeOnClose)
+	{
+		// tell the console to close itself
+		FIOSInputInterface::QueueKeyInput(KEYCODE_ESCAPE, 0);
+	}
+	
+	return [super resignFirstResponder];
+}
+
+
+// @todo keyboard: This is a helper define to show functions that _may_ need to be implemented as we go forward with keyboard support
+// for now, the very basics work, but most likely at some point for optimal functionality, we'll want to know the actual string
+// in the box, but that needs to come from Slate, and we currently don't have a way to get it
+#define REPORT_EVENT UE_LOG(LogIOS, Display, TEXT("Got a keyboard call, line %d"), __LINE__);
+
+
+- (NSString *)textInRange:(UITextRange *)range
+{
+	// @todo keyboard: This is called
+	return @"";
+	//IndexedRange *r = (IndexedRange *)range;
+	//return ([textStore substringWithRange:r.range]);
+}
+
+- (void)replaceRange:(UITextRange *)range withText:(NSString *)text
+{
+	REPORT_EVENT;
+	return;
+//	IndexedRange *r = (IndexedRange *)range;
+//	[textStore replaceCharactersInRange:r.range withString:text];
+}
+
+- (UITextRange *)selectedTextRange
+{
+	// @todo keyboard: This is called
+	return [IndexedRange rangeWithNSRange:NSMakeRange(0,0)];//self.textView.selectedTextRange];
+}
+
+
+- (void)setSelectedTextRange:(UITextRange *)range
+{
+	REPORT_EVENT;
+	//IndexedRange *indexedRange = (IndexedRange *)range;
+	//self.textView.selectedTextRange = indexedRange.range;
+}
+
+
+- (UITextRange *)markedTextRange
+{
+	// @todo keyboard: This is called
+	return nil;
+}
+
+
+- (void)setMarkedText:(NSString *)markedText selectedRange:(NSRange)selectedRange
+{
+	CachedMarkedText = markedText;
+	//NSLog(@"setting marked text to %@", markedText);
+}
+
+
+- (void)unmarkText
+{
+	if (CachedMarkedText != nil)
+	{
+		[self insertText:CachedMarkedText];
+		CachedMarkedText = nil;
+	}
+}
+
+
+- (UITextPosition *)beginningOfDocument
+{
+	// @todo keyboard: This is called
+	return [IndexedPosition positionWithIndex:0];
+}
+
+
+- (UITextPosition *)endOfDocument
+{
+	REPORT_EVENT;
+	return [IndexedPosition positionWithIndex:0];
+}
+
+
+- (UITextRange *)textRangeFromPosition:(UITextPosition *)fromPosition toPosition:(UITextPosition *)toPosition
+{
+	// @todo keyboard: This is called
+	// Generate IndexedPosition instances that wrap the to and from ranges.
+	IndexedPosition *fromIndexedPosition = (IndexedPosition *)fromPosition;
+	IndexedPosition *toIndexedPosition = (IndexedPosition *)toPosition;
+	NSRange range = NSMakeRange(MIN(fromIndexedPosition.index, toIndexedPosition.index), ABS(toIndexedPosition.index - fromIndexedPosition.index));
+ 
+	return [IndexedRange rangeWithNSRange:range];
+}
+
+
+- (UITextPosition *)positionFromPosition:(UITextPosition *)position offset:(NSInteger)offset
+{
+	// @todo keyboard: This is called
+	return nil;
+}
+
+
+- (UITextPosition *)positionFromPosition:(UITextPosition *)position inDirection:(UITextLayoutDirection)direction offset:(NSInteger)offset
+{
+	REPORT_EVENT;
+	return nil;
+}
+
+
+- (NSComparisonResult)comparePosition:(UITextPosition *)position toPosition:(UITextPosition *)other
+{
+	// @todo keyboard: This is called
+	return NSOrderedSame;
+}
+
+
+- (NSInteger)offsetFromPosition:(UITextPosition *)from toPosition:(UITextPosition *)toPosition
+{
+	REPORT_EVENT;
+	IndexedPosition *fromIndexedPosition = (IndexedPosition *)from;
+	IndexedPosition *toIndexedPosition = (IndexedPosition *)toPosition;
+	return (toIndexedPosition.index - fromIndexedPosition.index);
+}
+
+
+
+- (UITextPosition *)positionWithinRange:(UITextRange *)range farthestInDirection:(UITextLayoutDirection)direction
+{
+	REPORT_EVENT;
+	return nil;
+}
+
+
+- (UITextRange *)characterRangeByExtendingPosition:(UITextPosition *)position inDirection:(UITextLayoutDirection)direction
+{
+	REPORT_EVENT;
+	return nil;
+}
+
+
+- (UITextWritingDirection)baseWritingDirectionForPosition:(UITextPosition *)position inDirection:(UITextStorageDirection)direction
+{
+	REPORT_EVENT;
+	// assume left to right for now
+	return UITextWritingDirectionLeftToRight;
+}
+
+
+- (void)setBaseWritingDirection:(UITextWritingDirection)writingDirection forRange:(UITextRange *)range
+{
+	// @todo keyboard: This is called
+}
+
+
+
+- (CGRect)firstRectForRange:(UITextRange *)range
+{
+	REPORT_EVENT;
+	return CGRectMake(0,0,0,0);
+}
+
+
+- (CGRect)caretRectForPosition:(UITextPosition *)position
+{
+	// @todo keyboard: This is called
+	return CGRectMake(0,0,0,0);
+}
+
+
+- (UITextPosition *)closestPositionToPoint:(CGPoint)point
+{
+	REPORT_EVENT;
+	return nil;
+}
+
+- (UITextPosition *)closestPositionToPoint:(CGPoint)point withinRange:(UITextRange *)range
+{
+	REPORT_EVENT;
+	return nil;
+}
+
+- (UITextRange *)characterRangeAtPoint:(CGPoint)point
+{
+	REPORT_EVENT;
+	return nil;
+}
+
+- (NSArray *)selectionRectsForRange:(UITextRange *)range
+{
+	REPORT_EVENT;
+	return nil;
+}
+
+
+- (NSDictionary *)textStylingAtPosition:(UITextPosition *)position inDirection:(UITextStorageDirection)direction
+{
+	// @todo keyboard: This is called
+	return @{ };
+}
+
+
+- (void) setInputDelegate: (id <UITextInputDelegate>) delegate
+{
+	// @todo keyboard: This is called
+}
+
+- (id <UITextInputTokenizer>) tokenizer
+{
+	// @todo keyboard: This is called
+	return nil;
+}
+
+- (id <UITextInputDelegate>) inputDelegate
+{
+	// @todo keyboard: This is called
+	return nil;
+}
+
+
+
+
+
+#if !PLATFORM_TVOS
+- (void)keyboardWasShown:(NSNotification*)aNotification
+{
+	// send a callback to let the game know where to sldie the textbox up above
+	NSDictionary* info = [aNotification userInfo];
+	CGRect Frame = [[info objectForKey:UIKeyboardFrameEndUserInfoKey] CGRectValue];
+	
+	FPlatformRect ScreenRect;
+	ScreenRect.Top = Frame.origin.y;
+	ScreenRect.Bottom = (Frame.origin.y + Frame.size.height);
+	ScreenRect.Left = Frame.origin.x;
+	ScreenRect.Right = (Frame.origin.x + Frame.size.width);
+	
+	[FIOSAsyncTask CreateTaskWithBlock:^bool(void){
+		[IOSAppDelegate GetDelegate].IOSApplication->OnVirtualKeyboardShown().Broadcast(ScreenRect);
+		return true;
+	 }];
+}
+
+// Called when the UIKeyboardWillHideNotification is sent
+- (void)keyboardWillBeHidden:(NSNotification*)aNotification
+{
+	[FIOSAsyncTask CreateTaskWithBlock:^bool(void){
+		[IOSAppDelegate GetDelegate].IOSApplication->OnVirtualKeyboardHidden().Broadcast();
+		return true;
+	 }];
+}
+#endif
+
+- (void)InitKeyboard
+{
+#if !PLATFORM_TVOS
+	KeyboardShowCount = 0;
+	
+	bool bUseIntegratedKeyboard = false;
+	GConfig->GetBool(TEXT("/Script/IOSRuntimeSettings.IOSRuntimeSettings"), TEXT("bUseIntegratedKeyboard"), bUseIntegratedKeyboard, GEngineIni);
+	
+	// get notifications when the keyboard is in view
+	bIsUsingIntegratedKeyboard = FParse::Param(FCommandLine::Get(), TEXT("NewKeyboard")) || bUseIntegratedKeyboard;
+	if (bIsUsingIntegratedKeyboard)
+	{
+		[[NSNotificationCenter defaultCenter] addObserver:self
+												 selector:@selector(keyboardWasShown:)
+													 name:UIKeyboardDidShowNotification object:nil];
+		
+		[[NSNotificationCenter defaultCenter] addObserver:self
+												 selector:@selector(keyboardWillBeHidden:)
+													 name:UIKeyboardWillHideNotification object:nil];
+		
+	}
+#endif
+}
+
+
+
+
+
+@end
+
+
+#pragma mark IOSViewController
+
+@implementation IOSViewController
+
+/**
+ * The ViewController was created, so now we need to create our view to be controlled (an EAGLView)
+ */
+- (void) loadView
+{
+	// get the landcape size of the screen
+	CGRect Frame = [[UIScreen mainScreen] bounds];
+	if (![IOSAppDelegate GetDelegate].bDeviceInPortraitMode)
+	{
+		Swap(Frame.size.width, Frame.size.height);
+	}
+
+	self.view = [[UIView alloc] initWithFrame:Frame];
+
+	// settings copied from InterfaceBuilder
+#if defined(__IPHONE_7_0)
+	if ([IOSAppDelegate GetDelegate].OSVersion >= 7.0)
+	{
+		self.edgesForExtendedLayout = UIRectEdgeNone;
+	}
+#endif
+
+	self.view.clearsContextBeforeDrawing = NO;
+#if !PLATFORM_TVOS
+	self.view.multipleTouchEnabled = NO;
+#endif
+}
+
+/**
+ * View was unloaded from us
+ */ 
+- (void) viewDidUnload
+{
+	UE_LOG(LogIOS, Log, TEXT("IOSViewController unloaded the view. This is unexpected, tell Josh Adams"));
+}
+
+/**
+ * Tell the OS what the default supported orientations are
+ */
+- (NSUInteger)supportedInterfaceOrientations
+{
+	return UIInterfaceOrientationMaskAll;
+}
+
+/**
+ * Tell the OS that our view controller can auto-rotate between supported orientations
+ */
+- (BOOL)shouldAutorotate
+{
+	return YES;
+}
+
+/**
+ * Tell the OS to hide the status bar (iOS 7 method for hiding)
+ */
+- (BOOL)prefersStatusBarHidden
+{
+	return YES;
+}
+
+
+
+
+
+
+
+@end
diff --git a/Engine/Source/Runtime/Core/Private/IOS/IOSWindow.cpp b/Engine/Source/Runtime/Core/Private/IOS/IOSWindow.cpp
new file mode 100644
--- /dev/null
+++ b/Engine/Source/Runtime/Core/Private/IOS/IOSWindow.cpp
@@ -0,0 +1,98 @@
+// Copyright 1998-2017 Epic Games, Inc. All Rights Reserved.
+
+#include "IOSWindow.h"
+#include "IOSAppDelegate.h"
+#include "IOSView.h"
+
+FIOSWindow::~FIOSWindow()
+{
+	// NOTE: The Window is invalid here!
+	//       Use NativeWindow_Destroy() instead.
+}
+
+TSharedRef<FIOSWindow> FIOSWindow::Make()
+{
+	return MakeShareable( new FIOSWindow() );
+}
+
+FIOSWindow::FIOSWindow()
+{
+}
+
+void FIOSWindow::Initialize( class FIOSApplication* const Application, const TSharedRef< FGenericWindowDefinition >& InDefinition, const TSharedPtr< FIOSWindow >& InParent, const bool bShowImmediately )
+{
+	OwningApplication = Application;
+	Definition = InDefinition;
+
+	Window = [[UIApplication sharedApplication] keyWindow];
+
+#if !PLATFORM_TVOS
+	if(InParent.Get() != NULL)
+	{
+		dispatch_async(dispatch_get_main_queue(),^ {
+#ifdef __IPHONE_8_0
+			if ([UIAlertController class])
+			{
+				UIAlertController* AlertController = [UIAlertController alertControllerWithTitle:@""
+														message:@"Error: Only one UIWindow may be created on iOS."
+														preferredStyle:UIAlertControllerStyleAlert];
+				UIAlertAction* okAction = [UIAlertAction
+											actionWithTitle:NSLocalizedString(@"OK", nil)
+											style:UIAlertActionStyleDefault
+											handler:^(UIAlertAction* action)
+											{
+												[AlertController dismissViewControllerAnimated : YES completion : nil];
+											}
+				];
+
+				[AlertController addAction : okAction];
+				[[IOSAppDelegate GetDelegate].IOSController presentViewController : AlertController animated : YES completion : nil];
+			}
+			else
+#endif
+			{
+#if __IPHONE_OS_VERSION_MIN_REQUIRED < __IPHONE_9_0
+				UIAlertView* AlertView = [[UIAlertView alloc] initWithTitle:@""
+											message:@"Error: Only one UIWindow may be created on iOS."
+											delegate:nil
+											cancelButtonTitle:NSLocalizedString(@"Ok", nil)
+											otherButtonTitles:nil];
+
+				[AlertView show];
+				[AlertView release];
+#endif
+			}
+		} );
+	}
+#endif
+}
+
+FPlatformRect FIOSWindow::GetScreenRect()
+{
+	// get the main view's frame
+	IOSAppDelegate* AppDelegate = (IOSAppDelegate*)[[UIApplication sharedApplication] delegate];
+	UIView* View = AppDelegate.IOSView;
+	CGRect Frame = [View frame];
+	CGFloat Scale = View.contentScaleFactor;
+
+	FPlatformRect ScreenRect;
+	ScreenRect.Top = Frame.origin.y * Scale;
+	ScreenRect.Bottom = (Frame.origin.y + Frame.size.height) * Scale;
+	ScreenRect.Left = Frame.origin.x * Scale;
+	ScreenRect.Right = (Frame.origin.x + Frame.size.width) * Scale;
+
+	return ScreenRect;
+}
+
+
+bool FIOSWindow::GetFullScreenInfo( int32& X, int32& Y, int32& Width, int32& Height ) const
+{
+	FPlatformRect ScreenRect = GetScreenRect();
+
+	X = ScreenRect.Left;
+	Y = ScreenRect.Top;
+	Width = ScreenRect.Right - ScreenRect.Left;
+	Height = ScreenRect.Bottom - ScreenRect.Top;
+
+	return true;
+}
diff --git a/Engine/Source/Runtime/Core/Private/IOS/IOSWindow.h b/Engine/Source/Runtime/Core/Private/IOS/IOSWindow.h
new file mode 100644
--- /dev/null
+++ b/Engine/Source/Runtime/Core/Private/IOS/IOSWindow.h
@@ -0,0 +1,54 @@
+// Copyright 1998-2017 Epic Games, Inc. All Rights Reserved.
+
+#pragma once
+
+#include "GenericWindow.h"
+#include "GenericApplication.h"
+#import <UIKit/UIKit.h>
+
+/**
+ * A platform specific implementation of FNativeWindow.
+ * Native windows provide platform-specific backing for and are always owned by an SWindow.
+ */
+class CORE_API FIOSWindow : public FGenericWindow
+{
+public:
+	~FIOSWindow();
+
+	/** Create a new FIOSWindow.
+	 *
+	 * @param OwnerWindow		The SlateWindow for which we are crating a backing IOSWindow
+	 * @param InParent			Parent iOS window; usually NULL.
+	 */
+	static TSharedRef<FIOSWindow> Make();
+
+	/** Returns a rect describing the main screen */
+	static FPlatformRect GetScreenRect();
+
+	/** Returns a void pointer to the hWnd (for other apis)*/
+	virtual void* GetOSWindowHandle() const override { return Window; }
+
+	void Initialize( class FIOSApplication* const Application, const TSharedRef< FGenericWindowDefinition >& InDefinition, const TSharedPtr< FIOSWindow >& InParent, const bool bShowImmediately );
+
+	/** Returns the rectangle of the screen the window is associated with */
+	virtual bool GetFullScreenInfo( int32& X, int32& Y, int32& Width, int32& Height ) const override;
+
+protected:
+	/** @return true if the native window is currently in fullscreen mode, false otherwise */
+	virtual EWindowMode::Type GetWindowMode() const override { return EWindowMode::Fullscreen; }
+
+private:
+	/**
+	 * Protect the constructor; only TSharedRefs of this class can be made.
+	 */
+	FIOSWindow();
+
+	FIOSApplication* OwningApplication;
+
+	/** iOS window handle, typically, only one should ever exist */
+	UIWindow *Window;
+
+	/** Store the window region size for querying whether a point lies within the window */
+	int32 RegionX;
+	int32 RegionY;
+};
