Fix includes to files in an "iOS" directory.

diff --git a/Engine/Source/Runtime/Core/Private/IOS/IOSPlatformOutputDevices.cpp b/Engine/Source/Runtime/Core/Private/IOS/IOSPlatformOutputDevices.cpp
--- a/Engine/Source/Runtime/Core/Private/IOS/IOSPlatformOutputDevices.cpp
+++ b/Engine/Source/Runtime/Core/Private/IOS/IOSPlatformOutputDevices.cpp
@@ -8,7 +8,7 @@
 #include "Misc/OutputDeviceFile.h"
 
 #include "FeedbackContextAnsi.h"
-#include "../Private/iOS/IOSPlatformOutputDevicesPrivate.h"
+#include "../Private/IOS/IOSPlatformOutputDevicesPrivate.h"
 
 //////////////////////////////////
 // FIOSPlatformOutputDevices
diff --git a/Engine/Source/Runtime/Core/Public/HAL/Platform.h b/Engine/Source/Runtime/Core/Public/HAL/Platform.h
--- a/Engine/Source/Runtime/Core/Public/HAL/Platform.h
+++ b/Engine/Source/Runtime/Core/Public/HAL/Platform.h
@@ -75,7 +75,7 @@
 #elif PLATFORM_MAC
 	#include "Mac/MacPlatformCompilerPreSetup.h"
 #elif PLATFORM_IOS
-	#include "iOS/IOSPlatformCompilerPreSetup.h"
+	#include "IOS/IOSPlatformCompilerPreSetup.h"
 #elif PLATFORM_ANDROID
 	#include "Android/AndroidPlatformCompilerPreSetup.h"
 #elif PLATFORM_HTML5
