Fixing up case of "iOS" folders (pt 2)

diff --git a/Engine/Source/Runtime/Core/Public/IOS/IOSAsyncTask.h b/Engine/Source/Runtime/Core/Public/IOS/IOSAsyncTask.h
new file mode 100644
--- /dev/null
+++ b/Engine/Source/Runtime/Core/Public/IOS/IOSAsyncTask.h
@@ -0,0 +1,48 @@
+// Copyright 1998-2017 Epic Games, Inc. All Rights Reserved.
+
+/*=============================================================================
+	IPhoneAppDelegate.h: IPhone application class / main loop
+=============================================================================*/
+
+#pragma once
+
+#import <Foundation/Foundation.h>
+
+@interface FIOSAsyncTask : NSObject
+{
+@private
+	/** Whether or not the task is ready to have GameThread callback called (set on iOS thread) */
+	int32 bIsReadyForGameThread;
+}
+
+/** Extra data for this async task */
+@property (retain) id UserData;
+
+/**
+ * Code to run on the game thread when the async task completes 
+ * @return true when the task is complete, and it will be pulled from the list of tasks and destroyed
+ */
+@property (copy) bool (^GameThreadCallback)(void);
+
+
+/**
+ * Create an async task object, set the block, and mark it as ready to be processed on main thread
+ * For advanced uses (ie, setting UserData or delaying when it's ready for game thread, use the usual [[alloc] init] pattern
+ *
+ * Note, this doesn't return an AsyncObject so no one is tempted to use it - once it's ready for game thread, it's dangerous to touch it
+ * This is callable on any thread, not just main thread
+ */
++ (void)CreateTaskWithBlock:(bool (^)(void))Block;
+
+/**
+ * Mark that the task is complete on the iOS thread, and now the
+ * GameThread can be fired (the Task is unsafe to use after this call)
+ */
+- (void)FinishedTask;
+
+/**
+ * Tick all currently running tasks
+ */
++ (void)ProcessAsyncTasks;
+
+@end
diff --git a/Engine/Source/Runtime/Core/Public/IOS/IOSPlatform.h b/Engine/Source/Runtime/Core/Public/IOS/IOSPlatform.h
new file mode 100644
--- /dev/null
+++ b/Engine/Source/Runtime/Core/Public/IOS/IOSPlatform.h
@@ -0,0 +1,100 @@
+// Copyright 1998-2017 Epic Games, Inc. All Rights Reserved.
+
+/*================================================================================
+	IOSPlatform.h: Setup for the iOS platform
+==================================================================================*/
+
+#pragma once
+
+/**
+* iOS specific types
+**/
+struct FIOSPlatformTypes : public FGenericPlatformTypes
+{
+	typedef size_t				SIZE_T;
+	typedef decltype(NULL)		TYPE_OF_NULL;
+	typedef char16_t			CHAR16;
+};
+
+typedef FIOSPlatformTypes FPlatformTypes;
+
+// Base defines, must define these for the platform, there are no defaults
+#define PLATFORM_DESKTOP				0
+
+#ifdef __LP64__
+#define PLATFORM_64BITS					1
+#else
+#define PLATFORM_64BITS					0
+#endif
+
+// Base defines, defaults are commented out
+#define PLATFORM_LITTLE_ENDIAN							1
+#define PLATFORM_SUPPORTS_PRAGMA_PACK					1
+#define PLATFORM_COMPILER_DISTINGUISHES_INT_AND_LONG	1
+#define PLATFORM_TCHAR_IS_4_BYTES						1
+#define PLATFORM_USE_SYSTEM_VSWPRINTF					0
+#define PLATFORM_HAS_BSD_TIME							1
+#define PLATFORM_HAS_BSD_IPV6_SOCKETS					1
+#define PLATFORM_HAS_BSD_SOCKET_FEATURE_MSG_DONTWAIT	1
+#define PLATFORM_MAX_FILEPATH_LENGTH					MAX_PATH
+#define PLATFORM_SUPPORTS_TEXTURE_STREAMING				1
+#define PLATFORM_BUILTIN_VERTEX_HALF_FLOAT				0
+#define PLATFORM_SUPPORTS_MULTIPLE_NATIVE_WINDOWS		0
+#define PLATFORM_ALLOW_NULL_RHI							1
+#define PLATFORM_ENABLE_VECTORINTRINSICS_NEON			PLATFORM_64BITS // disable vector intrinsics to make it compatible with 32-bit in Xcode 8.3
+#define PLATFORM_SUPPORTS_STACK_SYMBOLS					1
+#define PLATFORM_NUM_AUDIODECOMPRESSION_PRECACHE_BUFFERS		0
+#if PLATFORM_TVOS
+#define PLATFORM_USES_ES2								0
+#define PLATFORM_HAS_TOUCH_MAIN_SCREEN					0
+#else
+#define PLATFORM_USES_ES2								1
+#define PLATFORM_HAS_TOUCH_MAIN_SCREEN					1
+#endif
+#define PLATFORM_UI_HAS_MOBILE_SCROLLBARS				1
+#define PLATFORM_UI_NEEDS_TOOLTIPS						0
+#define PLATFORM_UI_NEEDS_FOCUS_OUTLINES				0
+
+// @todo iOS: temporarily use Ansi allocator as wxWidgets cause problems with MallocTBB
+#define FORCE_ANSI_ALLOCATOR 1
+
+// Function type macros.
+#define VARARGS															/* Functions with variable arguments */
+#define CDECL															/* Standard C function */
+#define STDCALL															/* Standard calling convention */
+#if UE_BUILD_DEBUG
+#define FORCEINLINE inline 												/* Don't force code to be inline */
+#else
+#define FORCEINLINE inline __attribute__ ((always_inline))				/* Force code to be inline */
+#endif
+#define FORCENOINLINE __attribute__((noinline))							/* Force code to NOT be inline */
+#define FUNCTION_CHECK_RETURN_END __attribute__ ((warn_unused_result))	/* Warn that callers should not ignore the return value. */
+#define FUNCTION_NO_RETURN_END __attribute__ ((noreturn))				/* Indicate that the function never returns. */
+
+#define ABSTRACT abstract
+
+// Strings.
+#define LINE_TERMINATOR TEXT("\n")
+#define LINE_TERMINATOR_ANSI "\n"
+
+// Alignment.
+#define GCC_PACK(n) __attribute__((packed,aligned(n)))
+#define GCC_ALIGN(n) __attribute__((aligned(n)))
+#define REQUIRES_ALIGNED_ACCESS 1
+
+// operator new/delete operators
+// As of 10.9 we need to use _NOEXCEPT & cxx_noexcept compatible definitions
+#if __has_feature(cxx_noexcept)
+#define OPERATOR_NEW_THROW_SPEC
+#else
+#define OPERATOR_NEW_THROW_SPEC throw (std::bad_alloc)
+#endif
+#define OPERATOR_DELETE_THROW_SPEC noexcept
+#define OPERATOR_NEW_NOTHROW_SPEC  noexcept
+#define OPERATOR_DELETE_NOTHROW_SPEC  noexcept
+
+// DLL export and import definitions
+#define DLLEXPORT
+#define DLLIMPORT
+
+#define MAX_PATH 1024
diff --git a/Engine/Source/Runtime/Core/Public/IOS/IOSPlatformCompilerPreSetup.h b/Engine/Source/Runtime/Core/Public/IOS/IOSPlatformCompilerPreSetup.h
new file mode 100644
--- /dev/null
+++ b/Engine/Source/Runtime/Core/Public/IOS/IOSPlatformCompilerPreSetup.h
@@ -0,0 +1,162 @@
+// Copyright 1998-2017 Epic Games, Inc. All Rights Reserved.
+
+#pragma once
+
+
+#ifndef DISABLE_DEPRECATION
+	#pragma clang diagnostic warning "-Wdeprecated-declarations"
+
+	/**
+	 * Macro for marking up deprecated code, functions and types.
+	 *
+	 * Features that are marked as deprecated are scheduled to be removed from the code base
+	 * in a future release. If you are using a deprecated feature in your code, you should
+	 * replace it before upgrading to the next release. See the Upgrade Notes in the release
+	 * notes for the release in which the feature was marked deprecated.
+	 *
+	 * Sample usage (note the slightly different syntax for classes and structures):
+	 *
+	 *		DEPRECATED(4.6, "Message")
+	 *		void Function();
+	 *
+	 *		struct DEPRECATED(4.6, "Message") MODULE_API MyStruct
+	 *		{
+	 *			// StructImplementation
+	 *		};
+	 *		class DEPRECATED(4.6, "Message") MODULE_API MyClass
+	 *		{
+	 *			// ClassImplementation
+	 *		};
+	 *
+	 *		Unfortunately, VC++ will complain about using member functions and fields from deprecated
+	 *		class/structs even for class/struct implementation e.g.:
+	 *		class DEPRECATED(4.8, "") DeprecatedClass
+	 *		{
+	 *		public:
+	 *			DeprecatedClass() {}
+	 *
+	 *			float GetMyFloat()
+	 *			{
+	 *				return MyFloat; // This line will cause warning that deprecated field is used.
+	 *			}
+	 *		private:
+	 *			float MyFloat;
+	 *		};
+	 *
+	 *		To get rid of this warning, place all code not called in class implementation in non-deprecated
+	 *		base class and deprecate only derived one. This may force you to change some access specifiers
+	 *		from private to protected, e.g.:
+	 *
+	 *		class DeprecatedClass_Base_DEPRECATED
+	 *		{
+	 *		protected: // MyFloat is protected now, so DeprecatedClass has access to it.
+	 *			float MyFloat;
+	 *		};
+	 *
+	 *		class DEPRECATED(4.8, "") DeprecatedClass : DeprecatedClass_Base_DEPRECATED
+	 *		{
+	 *		public:
+	 *			DeprecatedClass() {}
+	 *
+	 *			float GetMyFloat()
+	 *			{
+	 *				return MyFloat;
+	 *			}
+	 *		};
+	 * @param VERSION The release number in which the feature was marked deprecated.
+	 * @param MESSAGE A message containing upgrade notes.
+	 */
+	#define DEPRECATED(VERSION, MESSAGE) __attribute__((deprecated(MESSAGE " Please update your code to the new API before upgrading to the next release, otherwise your project will no longer compile.")))
+
+	#define PRAGMA_DISABLE_DEPRECATION_WARNINGS \
+		_Pragma ("clang diagnostic push") \
+		_Pragma ("clang diagnostic ignored \"-Wdeprecated-declarations\"")
+
+	#define PRAGMA_ENABLE_DEPRECATION_WARNINGS \
+		_Pragma ("clang diagnostic pop")		
+
+#endif // DISABLE_DEPRECATION
+
+#ifndef PRAGMA_DISABLE_SHADOW_VARIABLE_WARNINGS
+	#define PRAGMA_DISABLE_SHADOW_VARIABLE_WARNINGS \
+		_Pragma ("clang diagnostic push") \
+		_Pragma ("clang diagnostic ignored \"-Wshadow\"")
+#endif // PRAGMA_DISABLE_SHADOW_VARIABLE_WARNINGS
+
+#ifndef PRAGMA_ENABLE_SHADOW_VARIABLE_WARNINGS
+	#define PRAGMA_ENABLE_SHADOW_VARIABLE_WARNINGS \
+		_Pragma("clang diagnostic pop")
+#endif // PRAGMA_ENABLE_SHADOW_VARIABLE_WARNINGS
+
+#ifndef PRAGMA_DISABLE_UNDEFINED_IDENTIFIER_WARNINGS
+	#define PRAGMA_DISABLE_UNDEFINED_IDENTIFIER_WARNINGS \
+		_Pragma("clang diagnostic push") \
+		_Pragma("clang diagnostic ignored \"-Wundef\"")
+#endif // PRAGMA_DISABLE_UNDEFINED_IDENTIFIER_WARNINGS
+
+#ifndef PRAGMA_ENABLE_UNDEFINED_IDENTIFIER_WARNINGS
+	#define PRAGMA_ENABLE_UNDEFINED_IDENTIFIER_WARNINGS \
+		_Pragma("clang diagnostic pop")
+#endif // PRAGMA_ENABLE_UNDEFINED_IDENTIFIER_WARNINGS
+
+#ifndef PRAGMA_DISABLE_MISSING_VIRTUAL_DESTRUCTOR_WARNINGS
+	#define PRAGMA_DISABLE_MISSING_VIRTUAL_DESTRUCTOR_WARNINGS \
+		_Pragma("clang diagnostic push") \
+		_Pragma("clang diagnostic ignored \"-Wdelete-non-virtual-dtor\"")
+#endif // PRAGMA_DISABLE_MISSING_VIRTUAL_DESTRUCTOR_WARNINGS
+
+#ifndef PRAGMA_ENABLE_MISSING_VIRTUAL_DESTRUCTOR_WARNINGS
+	#define PRAGMA_ENABLE_MISSING_VIRTUAL_DESTRUCTOR_WARNINGS \
+		_Pragma("clang diagnostic pop")
+#endif // PRAGMA_ENABLE_MISSING_VIRTUAL_DESTRUCTOR_WARNINGS
+
+#ifndef PRAGMA_POP
+	#define PRAGMA_POP \
+		_Pragma("clang diagnostic pop")
+#endif // PRAGMA_POP
+
+// Disable common CA warnings around SDK includes
+#ifndef THIRD_PARTY_INCLUDES_START
+	#define THIRD_PARTY_INCLUDES_START \
+		PRAGMA_DISABLE_SHADOW_VARIABLE_WARNINGS \
+		PRAGMA_DISABLE_UNDEFINED_IDENTIFIER_WARNINGS
+#endif
+
+#ifndef THIRD_PARTY_INCLUDES_END
+	#define THIRD_PARTY_INCLUDES_END \
+		PRAGMA_ENABLE_UNDEFINED_IDENTIFIER_WARNINGS \
+		PRAGMA_ENABLE_SHADOW_VARIABLE_WARNINGS
+#endif
+
+#ifndef EMIT_CUSTOM_WARNING_AT_LINE
+#define EMIT_CUSTOM_WARNING_AT_LINE(Line, Warning) \
+	_Pragma(PREPROCESSOR_TO_STRING(message(WARNING_LOCATION(Line) Warning)))
+#endif // EMIT_CUSTOM_WARNING_AT_LINE
+
+// Make certain warnings always be warnings, even despite -Werror.
+// Rationale: we don't want to suppress those as there are plans to address them (e.g. UE-12341), but breaking builds due to these warnings is very expensive
+// since they cannot be caught by all compilers that we support. They are deemed to be relatively safe to be ignored, at least until all SDKs/toolchains start supporting them.
+#pragma clang diagnostic warning "-Wreorder"
+#pragma clang diagnostic warning "-Wparentheses-equality"
+#pragma clang diagnostic ignored "-Winconsistent-missing-override"
+#pragma clang diagnostic ignored "-Wundefined-bool-conversion"
+#pragma clang diagnostic ignored "-Wunused-local-typedef"
+#pragma clang diagnostic ignored "-Wdelete-non-virtual-dtor"
+#pragma clang diagnostic ignored "-Wlogical-op-parentheses"
+#pragma clang diagnostic ignored "-Wconstant-logical-operand"
+
+// Apple LLVM 8.1.0 (Xcode 8.3) introduced -Wundefined-var-template
+#if (__clang_major__ > 8) || (__clang_major__ == 8 && __clang_minor__ >= 1)
+#pragma clang diagnostic ignored "-Wundefined-var-template"
+#pragma clang diagnostic ignored "-Wnullability-inferred-on-nested-type"
+#pragma clang diagnostic ignored "-Wobjc-protocol-property-synthesis"
+#pragma clang diagnostic ignored "-Wnullability-completeness-on-arrays"
+#pragma clang diagnostic ignored "-Wnull-dereference"
+#pragma clang diagnostic ignored "-Wnonportable-include-path"
+#endif
+
+// We can pragma optimisation's on and off as of Apple LLVM 7.3.0 but not before.
+#if (__clang_major__ > 7) || (__clang_major__ == 7 && __clang_minor__ >= 3)
+#define PRAGMA_DISABLE_OPTIMIZATION_ACTUAL _Pragma("clang optimize off")
+#define PRAGMA_ENABLE_OPTIMIZATION_ACTUAL  _Pragma("clang optimize on")
+#endif
diff --git a/Engine/Source/Runtime/Core/Public/IOS/IOSPlatformMath.h b/Engine/Source/Runtime/Core/Public/IOS/IOSPlatformMath.h
new file mode 100644
--- /dev/null
+++ b/Engine/Source/Runtime/Core/Public/IOS/IOSPlatformMath.h
@@ -0,0 +1,50 @@
+// Copyright 1998-2017 Epic Games, Inc. All Rights Reserved.
+
+/*=============================================================================================
+	IOSPlatformMath.h: IOS platform Math functions
+==============================================================================================*/
+
+#pragma once
+#include "GenericPlatform/GenericPlatformMath.h"
+
+/**
+* IOS implementation of the Math OS functions
+**/
+struct FIOSPlatformMath : public FGenericPlatformMath
+{
+	/**
+	 * Counts the number of leading zeros in the bit representation of the value
+	 *
+	 * @param Value the value to determine the number of leading zeros for
+	 *
+	 * @return the number of zeros before the first "on" bit
+	 */
+	static FORCEINLINE uint32 CountLeadingZeros(uint32 Value)
+	{
+		if (Value == 0)
+		{
+			return 32;
+		}
+	
+		return __builtin_clz(Value);
+	}
+	
+	/**
+	 * Counts the number of trailing zeros in the bit representation of the value
+	 *
+	 * @param Value the value to determine the number of trailing zeros for
+	 *
+	 * @return the number of zeros after the last "on" bit
+	 */
+	static FORCEINLINE uint32 CountTrailingZeros(uint32 Value)
+	{
+		if (Value == 0)
+		{
+			return 32;
+		}
+	
+		return __builtin_ctz(Value);
+	}
+};
+
+typedef FIOSPlatformMath FPlatformMath;
diff --git a/Engine/Source/Runtime/Core/Public/IOS/IOSPlatformMemory.h b/Engine/Source/Runtime/Core/Public/IOS/IOSPlatformMemory.h
new file mode 100644
--- /dev/null
+++ b/Engine/Source/Runtime/Core/Public/IOS/IOSPlatformMemory.h
@@ -0,0 +1,37 @@
+// Copyright 1998-2017 Epic Games, Inc. All Rights Reserved.
+
+
+/*=============================================================================================
+	IOSPlatformMemory.h: IOS platform memory functions
+==============================================================================================*/
+
+#pragma once
+#include "GenericPlatform/GenericPlatformMemory.h"
+#include "Apple/ApplePlatformMemory.h"
+
+/**
+ *	IOS implementation of the FGenericPlatformMemoryStats.
+ */
+struct FPlatformMemoryStats : public FGenericPlatformMemoryStats
+{};
+
+/**
+* IOS implementation of the memory OS functions
+**/
+struct CORE_API FIOSPlatformMemory : public FGenericPlatformMemory
+{
+	//~ Begin FGenericPlatformMemory Interface
+	static void Init();
+	static FPlatformMemoryStats GetStats();
+	static const FPlatformMemoryConstants& GetConstants();
+	static FMalloc* BaseAllocator();
+	static bool PageProtect(void* const Ptr, const SIZE_T Size, const bool bCanRead, const bool bCanWrite);
+	static void* BinnedAllocFromOS( SIZE_T Size );
+	static void BinnedFreeToOS( void* Ptr, SIZE_T Size );
+	//~ End FGenericPlatformMemory Interface
+};
+
+typedef FIOSPlatformMemory FPlatformMemory;
+
+
+
diff --git a/Engine/Source/Runtime/Core/Public/IOS/IOSPlatformMisc.h b/Engine/Source/Runtime/Core/Public/IOS/IOSPlatformMisc.h
new file mode 100644
--- /dev/null
+++ b/Engine/Source/Runtime/Core/Public/IOS/IOSPlatformMisc.h
@@ -0,0 +1,260 @@
+// Copyright 1998-2017 Epic Games, Inc. All Rights Reserved.
+
+/*=============================================================================================
+	IOSPlatformMisc.h: iOS platform misc functions
+==============================================================================================*/
+
+#pragma once
+#include "GenericPlatform/GenericPlatformMisc.h"
+#include "IOS/IOSSystemIncludes.h"
+
+#ifndef IOS_PROFILING_ENABLED
+#define IOS_PROFILING_ENABLED (UE_BUILD_DEBUG | UE_BUILD_DEVELOPMENT)
+#endif
+
+#ifdef __OBJC__
+
+class FScopeAutoreleasePool
+{
+public:
+
+	FScopeAutoreleasePool()
+	{
+		Pool = [[NSAutoreleasePool alloc] init];
+	}
+
+	~FScopeAutoreleasePool()
+	{
+		[Pool release];
+	}
+
+private:
+
+	NSAutoreleasePool*	Pool;
+};
+
+#define SCOPED_AUTORELEASE_POOL const FScopeAutoreleasePool PREPROCESSOR_JOIN(Pool,__LINE__);
+
+#endif // __OBJC__
+
+/**
+* iOS implementation of the misc OS functions
+**/
+struct CORE_API FIOSPlatformMisc : public FGenericPlatformMisc
+{
+    static void PlatformPreInit();
+	static void PlatformInit();
+    static void PlatformHandleSplashScreen(bool ShowSplashScreen = false);
+	static class GenericApplication* CreateApplication();
+	static void GetEnvironmentVariable(const TCHAR* VariableName, TCHAR* Result, int32 ResultLength);
+	static void* GetHardwareWindow();
+
+	static bool AllowThreadHeartBeat()
+	{
+		return false;
+	}
+
+#if !UE_BUILD_SHIPPING
+	static bool IsDebuggerPresent()
+	{
+		// Based on http://developer.apple.com/library/mac/#qa/qa1361/_index.html
+
+		struct kinfo_proc Info;
+		int32 Mib[] = { CTL_KERN, KERN_PROC, KERN_PROC_PID, getpid() };
+		SIZE_T Size = sizeof(Info);
+
+		sysctl( Mib, sizeof( Mib ) / sizeof( *Mib ), &Info, &Size, NULL, 0 );
+
+		return ( Info.kp_proc.p_flag & P_TRACED ) != 0;
+	}
+	FORCEINLINE static void DebugBreak()
+	{
+		if( IsDebuggerPresent() )
+		{
+			//Signal interupt to our process, this is caught by the main thread, this is not immediate but you can continue
+			//when triggered by check macros you will need to inspect other threads for the appFailAssert call.
+			//kill( getpid(), SIGINT );
+
+			//If you want an immediate break use the trap instruction, continued execuction is halted
+#if WITH_SIMULATOR
+            __asm__ ( "int $3" );
+#elif PLATFORM_64BITS
+			__asm__ ( "svc 0" );
+#else
+            __asm__ ( "trap" );
+#endif
+		}
+	}
+#endif
+
+	/** Break into debugger. Returning false allows this function to be used in conditionals. */
+	FORCEINLINE static bool DebugBreakReturningFalse()
+	{
+#if !UE_BUILD_SHIPPING
+		DebugBreak();
+#endif
+		return false;
+	}
+
+	/** Prompts for remote debugging if debugger is not attached. Regardless of result, breaks into debugger afterwards. Returns false for use in conditionals. */
+	FORCEINLINE static bool DebugBreakAndPromptForRemoteReturningFalse(bool bIsEnsure = false)
+	{
+#if !UE_BUILD_SHIPPING
+		if (!IsDebuggerPresent())
+		{
+			PromptForRemoteDebugging(bIsEnsure);
+		}
+
+		DebugBreak();
+#endif
+
+		return false;
+	}
+
+	FORCEINLINE static void MemoryBarrier()
+	{
+		__sync_synchronize();
+	}
+
+	static void LowLevelOutputDebugString(const TCHAR *Message);
+	static const TCHAR* GetSystemErrorMessage(TCHAR* OutBuffer, int32 BufferCount, int32 Error);
+	static void ClipboardCopy(const TCHAR* Str);
+	static void ClipboardPaste(class FString& Dest);
+	static EAppReturnType::Type MessageBoxExt( EAppMsgType::Type MsgType, const TCHAR* Text, const TCHAR* Caption );
+	static bool ControlScreensaver(EScreenSaverAction Action);
+	static int32 NumberOfCores();
+	static void LoadPreInitModules();
+	static void SetMemoryWarningHandler(void (* Handler)(const FGenericMemoryWarningContext& Context));
+	static bool HasPlatformFeature(const TCHAR* FeatureName);
+	static FString GetDefaultLanguage();
+	static FString GetDefaultLocale();
+	static bool SetStoredValue(const FString& InStoreId, const FString& InSectionName, const FString& InKeyName, const FString& InValue);
+	static bool GetStoredValue(const FString& InStoreId, const FString& InSectionName, const FString& InKeyName, FString& OutValue);
+	static bool DeleteStoredValue(const FString& InStoreId, const FString& InSectionName, const FString& InKeyName);
+	static TArray<uint8> GetSystemFontBytes();
+	static TArray<FString> GetPreferredLanguages();
+	static uint32 GetCharKeyMap(uint32* KeyCodes, FString* KeyNames, uint32 MaxMappings);
+	static uint32 GetKeyMap( uint32* KeyCodes, FString* KeyNames, uint32 MaxMappings );
+	static FString GetLocalCurrencyCode();
+	static FString GetLocalCurrencySymbol();
+	static void GetValidTargetPlatforms(class TArray<class FString>& TargetPlatformNames);
+	static bool HasActiveWiFiConnection();
+	static EScreenPhysicalAccuracy ComputePhysicalScreenDensity(int32& ScreenDensity);
+
+	static void ResetGamepadAssignments();
+	static void ResetGamepadAssignmentToController(int32 ControllerId);
+	static bool IsControllerAssignedToGamepad(int32 ControllerId);
+
+	static int GetAudioVolume();
+	static bool AreHeadphonesPluggedIn();
+	static int GetBatteryLevel();
+	static bool IsRunningOnBattery();
+
+	static void RegisterForRemoteNotifications();
+
+	static class IPlatformChunkInstall* GetPlatformChunkInstall();
+	
+#if IOS_PROFILING_ENABLED
+	static void BeginNamedEvent(const struct FColor& Color,const TCHAR* Text);
+	static void BeginNamedEvent(const struct FColor& Color,const ANSICHAR* Text);
+	static void EndNamedEvent();
+#endif
+	
+	//////// Platform specific
+	static void* CreateAutoreleasePool();
+	static void ReleaseAutoreleasePool(void *Pool);
+	static void HandleLowMemoryWarning();
+	static bool IsPackagedForDistribution();
+	DEPRECATED(4.14, "GetUniqueDeviceId is deprecated. Use GetDeviceId instead.")
+	static FString GetUniqueDeviceId();
+	/**
+	 * Implemented using UIDevice::identifierForVendor,
+	 * so all the caveats that apply to that API call apply here.
+	 */
+	static FString GetDeviceId();
+	static FString GetOSVersion();
+	static FString GetUniqueAdvertisingId();
+
+	// Possible iOS devices
+	enum EIOSDevice
+	{
+		// add new devices to the top, and add to IOSDeviceNames below!
+		IOS_IPhone4,
+		IOS_IPhone4S,
+		IOS_IPhone5, // also the IPhone5C
+		IOS_IPhone5S,
+		IOS_IPodTouch5,
+		IOS_IPodTouch6,
+		IOS_IPad2,
+		IOS_IPad3,
+		IOS_IPad4,
+		IOS_IPadMini,
+		IOS_IPadMini2, // also the iPadMini3
+		IOS_IPadMini4,
+		IOS_IPadAir,
+		IOS_IPadAir2,
+		IOS_IPhone6,
+		IOS_IPhone6Plus,
+		IOS_IPhone6S,
+		IOS_IPhone6SPlus,
+        IOS_IPhone7,
+        IOS_IPhone7Plus,
+		IOS_IPadPro,
+		IOS_AppleTV,
+		IOS_IPhoneSE,
+		IOS_IPadPro_129,
+		IOS_IPadPro_97,
+		IOS_Unknown,
+	};
+
+	static EIOSDevice GetIOSDeviceType();
+
+	static FORCENOINLINE const TCHAR* GetDefaultDeviceProfileName()
+	{
+		static const TCHAR* IOSDeviceNames[] = 
+		{
+			L"IPhone4",
+			L"IPhone4S",
+			L"IPhone5",
+			L"IPhone5S",
+			L"IPodTouch5",
+			L"IPodTouch6",
+			L"IPad2",
+			L"IPad3",
+			L"IPad4",
+			L"IPadMini",
+			L"IPadMini2",
+			L"IPadMini4",
+			L"IPadAir",
+			L"IPadAir2",
+			L"IPhone6",
+			L"IPhone6Plus",
+			L"IPhone6S",
+			L"IPhone6SPlus",
+            L"IPhone7",
+            L"IPhone7Plus",
+			L"IPadPro",
+			L"AppleTV",
+			L"IPhoneSE",
+			L"IPadPro129",
+			L"IPadPro97",
+			L"Unknown",
+		};
+		static_assert((sizeof(IOSDeviceNames) / sizeof(IOSDeviceNames[0])) == ((int32)IOS_Unknown + 1), "Mismatched IOSDeviceNames and EIOSDevice.");
+		
+		// look up into the string array by the enum
+		return IOSDeviceNames[(int32)GetIOSDeviceType()];
+	}
+
+	static FString GetCPUVendor();
+	static FString GetCPUBrand();
+	static void GetOSVersions(FString& out_OSVersionLabel, FString& out_OSSubVersionLabel);
+	static int32 IOSVersionCompare(uint8 Major, uint8 Minor, uint8 Revision);
+	
+    static void SetGracefulTerminationHandler();
+    static void SetCrashHandler(void(*CrashHandler)(const FGenericCrashContext& Context));
+private:
+	static class FIOSApplication* CachedApplication;
+};
+
+typedef FIOSPlatformMisc FPlatformMisc;
diff --git a/Engine/Source/Runtime/Core/Public/IOS/IOSPlatformProperties.h b/Engine/Source/Runtime/Core/Public/IOS/IOSPlatformProperties.h
new file mode 100644
--- /dev/null
+++ b/Engine/Source/Runtime/Core/Public/IOS/IOSPlatformProperties.h
@@ -0,0 +1,80 @@
+// Copyright 1998-2017 Epic Games, Inc. All Rights Reserved.
+
+/*================================================================================
+	IOSPlatformProperties.h - Basic static properties of a platform 
+	These are shared between:
+		the runtime platform - via FPlatformProperties
+		the target platforms - via ITargetPlatform
+==================================================================================*/
+
+#pragma once
+
+#include "GenericPlatformProperties.h"
+
+
+/**
+ * Implements iOS platform properties.
+ */
+struct FIOSPlatformProperties
+	: public FGenericPlatformProperties
+{
+	static FORCEINLINE bool HasEditorOnlyData( )
+	{
+		return false;
+	}
+
+	static FORCEINLINE const char* PlatformName( )
+	{
+		return "IOS";
+	}
+
+	static FORCEINLINE const char* IniPlatformName( )
+	{
+		return "IOS";
+	}
+
+	static FORCEINLINE bool IsGameOnly()
+	{
+		return true;
+	}
+	
+	static FORCEINLINE bool RequiresCookedData( )
+	{
+		return true;
+	}
+
+	static FORCEINLINE bool SupportsBuildTarget( EBuildTargets::Type BuildTarget )
+	{
+		return (BuildTarget == EBuildTargets::Game);
+	}
+
+	static FORCEINLINE bool SupportsLowQualityLightmaps()
+	{
+		return true;
+	}
+
+	static FORCEINLINE bool SupportsHighQualityLightmaps()
+	{
+		return true;
+	}
+
+	static FORCEINLINE bool SupportsTextureStreaming()
+	{
+		return true;
+	}
+
+	static FORCEINLINE bool HasFixedResolution()
+	{
+		return true;
+	}
+
+	static FORCEINLINE bool AllowsFramerateSmoothing()
+	{
+		return true;
+	}
+
+	static FORCEINLINE bool SupportsAudioStreaming()
+	{
+		return true;
+	}
+};
diff --git a/Engine/Source/Runtime/Core/Public/IOS/IOSView.h b/Engine/Source/Runtime/Core/Public/IOS/IOSView.h
new file mode 100644
--- /dev/null
+++ b/Engine/Source/Runtime/Core/Public/IOS/IOSView.h
@@ -0,0 +1,121 @@
+// Copyright 1998-2017 Epic Games, Inc. All Rights Reserved.
+
+#pragma once
+
+#include "CoreTypes.h"
+#import <UIKit/UIKit.h>
+#import <OpenGLES/EAGL.h>
+
+#if HAS_METAL
+#import <Metal/Metal.h>
+#import <QuartzCore/CAMetalLayer.h>
+#endif
+
+struct FKeyboardConfig
+{
+	UIKeyboardType KeyboardType;
+	UITextAutocorrectionType AutocorrectionType;
+	UITextAutocapitalizationType AutocapitalizationType;
+	BOOL bSecureTextEntry;
+	
+	FORCEINLINE FKeyboardConfig() :
+		KeyboardType(UIKeyboardTypeDefault),
+		AutocorrectionType(UITextAutocorrectionTypeNo),
+		AutocapitalizationType(UITextAutocapitalizationTypeNone),
+		bSecureTextEntry(NO) {}
+};
+
+@interface FIOSView : UIView  <UIKeyInput, UITextInput>
+{
+@public
+	// are we initialized?
+	bool bIsInitialized;
+
+//@private
+	// keeps track of the number of active touches
+	// used to bring up the three finger touch debug console after 3 active touches are registered
+	int NumActiveTouches;
+
+	// track the touches by pointer (which will stay the same for a given finger down) - note we don't deref the pointers in this array
+	UITouch* AllTouches[10];
+
+
+	//// GL MEMBERS
+	// the GL context
+	EAGLContext* Context;
+
+	// the internal MSAA FBO used to resolve the color buffer at present-time
+	GLuint ResolveFrameBuffer;
+
+
+	//// METAL MEMBERS
+#if HAS_METAL
+	// global metal device
+	id<MTLDevice> MetalDevice;
+#endif
+
+	// are we using the Metal API?
+	bool bIsUsingMetal;
+	
+	//// KEYBOARD MEMBERS
+	
+	// whether or not to use the new style virtual keyboard that sends events to the engine instead of using an alert
+	bool bIsUsingIntegratedKeyboard;
+	bool bSendEscapeOnClose;
+
+	// caches for the TextInput
+	NSString* CachedMarkedText;
+	
+	UIKeyboardType KeyboardType;
+	UITextAutocorrectionType AutocorrectionType;
+	UITextAutocapitalizationType AutocapitalizationType;
+	BOOL bSecureTextEntry;
+	
+	volatile int32 KeyboardShowCount;
+}
+
+
+//// SHARED FUNCTIONALITY
+@property (nonatomic) GLuint SwapCount;
+
+-(bool)CreateFramebuffer:(bool)bIsForOnDevice;
+-(void)DestroyFramebuffer;
+-(void)UpdateRenderWidth:(uint32)Width andHeight:(uint32)Height;
+
+//// GL FUNCTIONALITY
+@property (nonatomic) GLuint OnScreenColorRenderBuffer;
+@property (nonatomic) GLuint OnScreenColorRenderBufferMSAA;
+
+- (void)MakeCurrent;
+- (void)UnmakeCurrent;
+- (void)SwapBuffers;
+
+//// METAL FUNCTIONALITY
+#if HAS_METAL
+// Return a drawable object (ie a back buffer texture) for the RHI to render to
+- (id<CAMetalDrawable>)MakeDrawable;
+#endif
+
+
+//// KEYBOARD FUNCTIONALITY
+-(void)InitKeyboard;
+-(void)ActivateKeyboard:(bool)bInSendEscapeOnClose;
+-(void)ActivateKeyboard:(bool)bInSendEscapeOnClose keyboardConfig:(FKeyboardConfig)KeyboardConfig;
+-(void)DeactivateKeyboard;
+@end
+
+
+/**
+ * A view controller subclass that handles loading our IOS view as well as autorotation
+ */
+#if PLATFORM_TVOS
+#import <GameController/GameController.h>
+// if TVOS doesn't use the GCEventViewController, it will background the app when the user presses Menu/Pause
+@interface IOSViewController : GCEventViewController
+#else
+@interface IOSViewController : UIViewController
+#endif
+{
+
+}
+@end
