CIS fix for non-unity editor on Mac (missing include)

diff --git a/Engine/Source/Runtime/Core/Private/Mac/MacWindow.cpp b/Engine/Source/Runtime/Core/Private/Mac/MacWindow.cpp
--- a/Engine/Source/Runtime/Core/Private/Mac/MacWindow.cpp
+++ b/Engine/Source/Runtime/Core/Private/Mac/MacWindow.cpp
@@ -5,6 +5,7 @@
 #include "MacCursor.h"
 #include "CocoaTextView.h"
 #include "CocoaThread.h"
+#include "MacPlatformApplicationMisc.h"
 
 FMacWindow::FMacWindow()
 :	WindowHandle(NULL)
