CIS fixes.

diff --git a/Engine/Source/Runtime/Core/Private/Mac/MacPlatformApplicationMisc.cpp b/Engine/Source/Runtime/Core/Private/Mac/MacPlatformApplicationMisc.cpp
--- a/Engine/Source/Runtime/Core/Private/Mac/MacPlatformApplicationMisc.cpp
+++ b/Engine/Source/Runtime/Core/Private/Mac/MacPlatformApplicationMisc.cpp
@@ -5,6 +5,7 @@
 #include "MacApplication.h"
 #include "Math/Color.h"
 #include "Mac/MacMallocZone.h"
+#include "Misc/App.h"
 
 #include <dlfcn.h>
 #include <IOKit/IOKitLib.h>
