Added a cvar to prevent touch events from sending multiple mouse events and screwing everything up in the editor.  Disabled by default.   Enabling this will prevent pen tablets from working since until we switch to the windows 8 sdk (and can use WM_POINTER* events) we cannot detect the difference

diff --git a/Engine/Source/Runtime/Core/Private/Windows/WindowsApplication.cpp b/Engine/Source/Runtime/Core/Private/Windows/WindowsApplication.cpp
--- a/Engine/Source/Runtime/Core/Private/Windows/WindowsApplication.cpp
+++ b/Engine/Source/Runtime/Core/Private/Windows/WindowsApplication.cpp
@@ -50,6 +50,18 @@ THIRD_PARTY_INCLUDES_END
 
 DEFINE_LOG_CATEGORY(LogWindowsDesktop);
 
+/**
+ * Hack to get around multiple mouse events being triggered for touch events.
+ * Enabling this will prevent pen tablets from working since until we switch to the windows 8 sdk (and can use WM_POINTER*) events we cannot detect the difference
+ */
+static int32 bPreventDuplicateMouseEventsForTouch = false;
+
+FAutoConsoleVariableRef	CVarPreventDuplicateMouseEventsForTouch(
+	TEXT("Slate.PreventDuplicateMouseEventsForTouch"),
+	bPreventDuplicateMouseEventsForTouch,
+	TEXT("Hack to get around multiple mouse events being triggered for touch events.  Enabling this will prevent pen tablets from working since until we switch to the windows 8 sdk (and can use WM_POINTER* events) we cannot detect the difference")
+);
+
 const FIntPoint FWindowsApplication::MinimizedWindowPosition(-32000,-32000);
 
 FWindowsApplication* WindowsApplication = nullptr;
@@ -2120,7 +2132,7 @@ bool FWindowsApplication::IsInputMessage( uint32 msg )
 
 bool FWindowsApplication::IsFakeMouseInputMessage(uint32 msg)
 {
-	if (IsMouseInputMessage(msg))
+	if (bPreventDuplicateMouseEventsForTouch && IsMouseInputMessage(msg))
 	{
 		// This is only legal to call when handling messages in the pump, and is not valid
 		// to call in a deferred fashion.
