Mac: Restore code to replace hyphens in environment variables with underscores. Allows builders to set the shared DDC path through UE_SharedDataCachePath.

diff --git a/Engine/Source/Runtime/Core/Private/Apple/ApplePlatformMisc.cpp b/Engine/Source/Runtime/Core/Private/Apple/ApplePlatformMisc.cpp
--- a/Engine/Source/Runtime/Core/Private/Apple/ApplePlatformMisc.cpp
+++ b/Engine/Source/Runtime/Core/Private/Apple/ApplePlatformMisc.cpp
@@ -15,7 +15,11 @@
 
 void FApplePlatformMisc::GetEnvironmentVariable(const TCHAR* VariableName, TCHAR* Result, int32 ResultLength)
 {
-	ANSICHAR *AnsiResult = getenv(TCHAR_TO_ANSI(VariableName));
+	// Replace hyphens with underscores. Some legacy UE environment variables (eg. UE-SharedDataCachePath) are in widespread
+	// usage in their hyphenated form, but are not normally valid shell variables.
+	FString FixedVariableName = InVariableName;
+	FixedVariableName.ReplaceInline(TEXT("-"), TEXT("_"));
+	ANSICHAR *AnsiResult = getenv(TCHAR_TO_ANSI(*FixedVariableName));
 	if (AnsiResult)
 	{
 		wcsncpy(Result, ANSI_TO_TCHAR(AnsiResult), ResultLength);
