Fix typo (renamed variable).

diff --git a/Engine/Source/Runtime/Core/Private/Apple/ApplePlatformMisc.cpp b/Engine/Source/Runtime/Core/Private/Apple/ApplePlatformMisc.cpp
--- a/Engine/Source/Runtime/Core/Private/Apple/ApplePlatformMisc.cpp
+++ b/Engine/Source/Runtime/Core/Private/Apple/ApplePlatformMisc.cpp
@@ -17,7 +17,7 @@ void FApplePlatformMisc::GetEnvironmentVariable(const TCHAR* VariableName, TCHAR
 {
 	// Replace hyphens with underscores. Some legacy UE environment variables (eg. UE-SharedDataCachePath) are in widespread
 	// usage in their hyphenated form, but are not normally valid shell variables.
-	FString FixedVariableName = InVariableName;
+	FString FixedVariableName = VariableName;
 	FixedVariableName.ReplaceInline(TEXT("-"), TEXT("_"));
 	ANSICHAR *AnsiResult = getenv(TCHAR_TO_ANSI(*FixedVariableName));
 	if (AnsiResult)
