Copying //UE4/Dev-Sequencer to //UE4/Main (Source: //UE4/Dev-Sequencer @ 3617642)

diff --git a/Engine/Source/Runtime/AppFramework/Public/Widgets/Colors/SColorPicker.h b/Engine/Source/Runtime/AppFramework/Public/Widgets/Colors/SColorPicker.h
--- a/Engine/Source/Runtime/AppFramework/Public/Widgets/Colors/SColorPicker.h
+++ b/Engine/Source/Runtime/AppFramework/Public/Widgets/Colors/SColorPicker.h
@@ -18,9 +18,11 @@ class SColorThemesViewer;
 class SComboButton;
 class SThemeColorBlocksBar;
 
+
 /** Called when the color picker cancel button is pressed */
 DECLARE_DELEGATE_OneParam(FOnColorPickerCancelled, FLinearColor);
 
+
 /**
  * Enumerates color channels (do not reorder).
  */
@@ -170,11 +172,12 @@ public:
 	 *
 	 * @param InArgs Declaration from which to construct the widget.
 	 */
-	void Construct(const FArguments& InArgs );
+	void Construct(const FArguments& InArgs);
 
 	/** Delegate to override color picker creation behavior */
 	DECLARE_DELEGATE_OneParam(FOnColorPickerCreationOverride, const TSharedRef<SColorPicker>&);
 	static FOnColorPickerCreationOverride OnColorPickerNonModalCreateOverride;
+
 	/** Delegate to override color picker destruction behavior */
 	DECLARE_DELEGATE(FOnColorPickerDestructionOverride);
 	static FOnColorPickerDestructionOverride OnColorPickerDestroyOverride;
@@ -188,11 +191,11 @@ protected:
 	void RestoreColors();
 
 	/** Set all the colors to this new color */
-	void SetColors( const FLinearColor& InColor );
+	void SetColors(const FLinearColor& InColor);
 
-	bool ApplyNewTargetColor( bool bForceUpdate = false );
+	bool ApplyNewTargetColor(bool bForceUpdate = false);
 
-	void GenerateDefaultColorPickerContent( bool bAdvancedSectionExpanded );
+	void GenerateDefaultColorPickerContent(bool bAdvancedSectionExpanded);
 	void GenerateInlineColorPickerContent();
 
 	FLinearColor GetCurrentColor() const
@@ -204,8 +207,8 @@ protected:
 	void DiscardColor();
 
 	/** Sets new color in ether RGB or HSV */
-	bool SetNewTargetColorRGB( const FLinearColor& NewValue, bool bForceUpdate = false );
-	bool SetNewTargetColorHSV( const FLinearColor& NewValue, bool bForceUpdate = false );
+	bool SetNewTargetColorRGB(const FLinearColor& NewValue, bool bForceUpdate = false);
+	bool SetNewTargetColorHSV(const FLinearColor& NewValue, bool bForceUpdate = false);
 
 	void UpdateColorPick();
 	void UpdateColorPickMouseUp();
@@ -224,7 +227,7 @@ protected:
 	 * @param Channel The color channel to create the widget for.
 	 * @return The new slider.
 	 */
-	TSharedRef<SWidget> MakeColorSlider( EColorPickerChannels Channel ) const;
+	TSharedRef<SWidget> MakeColorSlider(EColorPickerChannels Channel) const;
 
 	/**
 	 * Creates a color spin box widget for the specified channel.
@@ -232,7 +235,7 @@ protected:
 	 * @param Channel The color channel to create the widget for.
 	 * @return The new spin box.
 	 */
-	TSharedRef<SWidget> MakeColorSpinBox( EColorPickerChannels Channel ) const;
+	TSharedRef<SWidget> MakeColorSpinBox(EColorPickerChannels Channel) const;
 
 	/**
 	 * Creates the color preview box widget.
@@ -244,16 +247,16 @@ protected:
 private:
 
 	// Callback for the active timer to animate the color post-construct
-	EActiveTimerReturnType AnimatePostConstruct( double InCurrentTime, float InDeltaTime );
+	EActiveTimerReturnType AnimatePostConstruct(double InCurrentTime, float InDeltaTime);
 
 	// Callback for getting the end color of a color spin box gradient.
-	FLinearColor GetGradientEndColor( EColorPickerChannels Channel ) const;
+	FLinearColor GetGradientEndColor(EColorPickerChannels Channel) const;
 
 	// Callback for getting the start color of a color spin box gradient.
-	FLinearColor GetGradientStartColor( EColorPickerChannels Channel ) const;
+	FLinearColor GetGradientStartColor(EColorPickerChannels Channel) const;
 
 	// Callback for handling expansion of the 'Advanced' area.
-	void HandleAdvancedAreaExpansionChanged( bool Expanded );
+	void HandleAdvancedAreaExpansionChanged(bool Expanded);
 
 	// Callback for getting the visibility of the alpha channel portion in color blocks.
 	EVisibility HandleAlphaColorBlockVisibility() const;
@@ -262,28 +265,28 @@ private:
 	FReply HandleCancelButtonClicked();
 
 	// Callback for pressing a mouse button in the color area.
-	FReply HandleColorAreaMouseDown( const FGeometry& MyGeometry, const FPointerEvent& MouseEvent );
+	FReply HandleColorAreaMouseDown(const FGeometry& MyGeometry, const FPointerEvent& MouseEvent);
 
 	// Callback for clicking the color picker mode button.
 	FReply HandleColorPickerModeButtonClicked();
 
 	// Callback for getting the visibility of the given color picker mode.
-	EVisibility HandleColorPickerModeVisibility( EColorPickerModes Mode ) const;
+	EVisibility HandleColorPickerModeVisibility(EColorPickerModes Mode) const;
 
 	// Callback for getting the end color of a color slider.
-	FLinearColor HandleColorSliderEndColor( EColorPickerChannels Channel ) const;
+	FLinearColor HandleColorSliderEndColor(EColorPickerChannels Channel) const;
 
 	// Callback for getting the start color of a color slider.
-	FLinearColor HandleColorSliderStartColor( EColorPickerChannels Channel ) const;
+	FLinearColor HandleColorSliderStartColor(EColorPickerChannels Channel) const;
 
 	// Callback for value changes in the color spectrum picker.
-	void HandleColorSpectrumValueChanged( FLinearColor NewValue );
+	void HandleColorSpectrumValueChanged(FLinearColor NewValue);
 
 	// Callback for getting the value of a color spin box.
-	float HandleColorSpinBoxValue( EColorPickerChannels Channel ) const;
+	float HandleColorSpinBoxValue(EColorPickerChannels Channel) const;
 
 	// Callback for value changes in a color spin box.
-	void HandleColorSpinBoxValueChanged( float NewValue, EColorPickerChannels Channel );
+	void HandleColorSpinBoxValueChanged(float NewValue, EColorPickerChannels Channel);
 
 	// Callback for completed eye dropper interactions.
 	void HandleEyeDropperButtonComplete(bool bCancelled);
@@ -301,7 +304,7 @@ private:
 	void HandleHexLinearInputTextCommitted(const FText& Text, ETextCommit::Type CommitType);
 
 	// Callback for changing the HSV value of the current color.
-	void HandleHSVColorChanged( FLinearColor NewValue );
+	void HandleHSVColorChanged(FLinearColor NewValue);
 
 	// Callback for when interactive user input begins.
 	void HandleInteractiveChangeBegin();
@@ -310,34 +313,34 @@ private:
 	void HandleInteractiveChangeEnd();
 
 	// Callback for when interactive user input ends.
-	void HandleInteractiveChangeEnd( float NewValue );
+	void HandleInteractiveChangeEnd(float NewValue);
 
 	// Callback for clicking the new color preview block.
-	FReply HandleNewColorBlockMouseButtonDown( const FGeometry& MyGeometry, const FPointerEvent& MouseEvent, bool bCheckAlpha );
+	FReply HandleNewColorBlockMouseButtonDown(const FGeometry& MyGeometry, const FPointerEvent& MouseEvent, bool bCheckAlpha);
 
 	// Callback for clicking the OK button.
 	FReply HandleOkButtonClicked();
 
 	// Callback for clicking the old color preview block.
-	FReply HandleOldColorBlockMouseButtonDown( const FGeometry& MyGeometry, const FPointerEvent& MouseEvent, bool bCheckAlpha );
+	FReply HandleOldColorBlockMouseButtonDown(const FGeometry& MyGeometry, const FPointerEvent& MouseEvent, bool bCheckAlpha);
 
 	// Callback for checking whether sRGB colors should be rendered.
 	bool HandleColorPickerUseSRGB() const;
 
 	// Callback for when the parent window has been closed.
-	void HandleParentWindowClosed( const TSharedRef<SWindow>& Window );
+	void HandleParentWindowClosed(const TSharedRef<SWindow>& Window);
 
 	// Callback for changing the RGB value of the current color.
-	void HandleRGBColorChanged( FLinearColor NewValue );
+	void HandleRGBColorChanged(FLinearColor NewValue);
 
 	// Callback for changing the checked state of the sRGB check box.
-	void HandleSRGBCheckBoxCheckStateChanged( ECheckBoxState InIsChecked );
+	void HandleSRGBCheckBoxCheckStateChanged(ECheckBoxState InIsChecked);
 
 	// Callback for determining whether the sRGB check box should be checked.
 	ECheckBoxState HandleSRGBCheckBoxIsChecked() const;
 
 	// Callback for selecting a color in the color theme bar.
-	void HandleThemeBarColorSelected( FLinearColor NewValue );
+	void HandleThemeBarColorSelected(FLinearColor NewValue);
 
 	// Callback for getting the theme bar's color theme.
 	TSharedPtr<class FColorTheme> HandleThemeBarColorTheme() const;
@@ -371,7 +374,7 @@ private:
 	/** Color start point to animate from */
 	FLinearColor ColorBegin;
 
-	// Holds the color picker's mode.
+	/** Holds the color picker's mode. */
 	EColorPickerModes CurrentMode;
 
 	/** Time, used for color animation */
@@ -449,7 +452,6 @@ private:
 	/** Is true if the color picker creation behavior can be overridden */
 	bool bValidCreationOverrideExists;
 
-
 private:
 
 	/** Invoked when a new value is selected on the color wheel */
@@ -461,8 +463,6 @@ private:
 	/** Invoked when the color picker cancel button is pressed */
 	FOnColorPickerCancelled OnColorPickerCancelled;
 
-
-
 	/** Invoked when a slider drag, color wheel drag or dropper grab starts */
 	FSimpleDelegate OnInteractivePickBegin;
 
@@ -538,6 +538,7 @@ struct FColorPickerArgs
 	/** Overrides the initial color set on the color picker. */
 	FLinearColor InitialColorOverride;
 
+	/** Default constructor. */
 	FColorPickerArgs()
 		: bIsModal(false)
 		, ParentWidget(nullptr)
@@ -562,13 +563,11 @@ struct FColorPickerArgs
 };
 
 
-/**
- * Opens up the static color picker, destroying any previously existing one.
- */
-APPFRAMEWORK_API bool OpenColorPicker( const FColorPickerArgs& Args );
+/** Open up the static color picker, destroying any previously existing one. */
+APPFRAMEWORK_API bool OpenColorPicker(const FColorPickerArgs& Args);
 
 /**
- * Destroys the current color picker. Necessary if the values the color picker
+ * Destroy the current color picker. Necessary if the values the color picker
  * currently targets become invalid.
  */
 APPFRAMEWORK_API void DestroyColorPicker();
diff --git a/Engine/Source/Runtime/AudioMixer/Private/Components/SynthComponent.cpp b/Engine/Source/Runtime/AudioMixer/Private/Components/SynthComponent.cpp
--- a/Engine/Source/Runtime/AudioMixer/Private/Components/SynthComponent.cpp
+++ b/Engine/Source/Runtime/AudioMixer/Private/Components/SynthComponent.cpp
@@ -16,7 +16,19 @@ void USynthSound::Init(USynthComponent* InSynthComponent, int32 InNumChannels)
 	OwningSynthComponent = InSynthComponent;
 	bVirtualizeWhenSilent = true;
 	NumChannels = InNumChannels;
-	bCanProcessAsync = true;
+
+	// Turn off async generation in old audio engine on mac.
+#if PLATFORM_MAC
+	if (!InSynthComponent->GetAudioDevice()->IsAudioMixerEnabled())
+	{
+		bCanProcessAsync = false;
+	}
+	else
+#endif // #if PLATFORM_MAC
+	{
+		bCanProcessAsync = true;
+	}
+
 	Duration = INDEFINITELY_LOOPING_DURATION;
 	bLooping = true;
 	SampleRate = InSynthComponent->GetAudioDevice()->SampleRate;
@@ -53,6 +65,7 @@ USynthComponent::USynthComponent(const FObjectInitializer& ObjectInitializer)
 
 	bIsSynthPlaying = false;
 	bIsInitialized = false;
+	bIsUISound = false;
 
 	// Set the default sound class
 	SoundClass = USoundBase::DefaultSoundClassObject;
@@ -261,6 +274,7 @@ void USynthComponent::Start()
 		// Copy the attenuation and concurrency data from the synth component to the audio component
 		AudioComponent->AttenuationSettings = AttenuationSettings;
 		AudioComponent->bOverrideAttenuation = bOverrideAttenuation;
+		AudioComponent->bIsUISound = bIsUISound;
 		AudioComponent->ConcurrencySettings = ConcurrencySettings;
 		AudioComponent->AttenuationOverrides = AttenuationOverrides;
 		AudioComponent->SoundClassOverride = SoundClass;
diff --git a/Engine/Source/Runtime/AudioMixer/Public/Components/SynthComponent.h b/Engine/Source/Runtime/AudioMixer/Public/Components/SynthComponent.h
--- a/Engine/Source/Runtime/AudioMixer/Public/Components/SynthComponent.h
+++ b/Engine/Source/Runtime/AudioMixer/Public/Components/SynthComponent.h
@@ -116,6 +116,10 @@ public:
 	UPROPERTY(EditAnywhere, BlueprintReadWrite, Category = Effects)
 	TArray<FSoundSubmixSendInfo> SoundSubmixSends;
 
+	/** Whether or not this sound plays when the game is paused in the UI */
+	UPROPERTY(EditAnywhere, BlueprintReadWrite, Category = Sound)
+	uint8 bIsUISound : 1;
+
 	/** Call if creating this synth component not via an actor component in BP, but in code or some other location. */
 	void Initialize();
 
diff --git a/Engine/Source/Runtime/AutomationMessages/Private/AutomationMessagesModule.cpp b/Engine/Source/Runtime/AutomationMessages/Private/AutomationMessagesModule.cpp
--- a/Engine/Source/Runtime/AutomationMessages/Private/AutomationMessagesModule.cpp
+++ b/Engine/Source/Runtime/AutomationMessages/Private/AutomationMessagesModule.cpp
@@ -12,11 +12,12 @@ class FAutomationMessagesModule
 {
 public:
 
-	virtual void StartupModule( ) override { }
+	//~ IModuleInterface interface
 
-	virtual void ShutdownModule( ) override { }
+	virtual void StartupModule() override { }
+	virtual void ShutdownModule() override { }
 
-	virtual bool SupportsDynamicReloading( ) override
+	virtual bool SupportsDynamicReloading() override
 	{
 		return true;
 	}
diff --git a/Engine/Source/Runtime/AutomationWorker/Private/AutomationAnalyticParams.h b/Engine/Source/Runtime/AutomationWorker/Private/AutomationAnalyticParams.h
--- a/Engine/Source/Runtime/AutomationWorker/Private/AutomationAnalyticParams.h
+++ b/Engine/Source/Runtime/AutomationWorker/Private/AutomationAnalyticParams.h
@@ -2,7 +2,6 @@
 
 #pragma once
 
-#include "CoreMinimal.h"
 
 namespace EAutomationEventName
 {
@@ -15,6 +14,7 @@ namespace EAutomationEventName
 	};
 }
 
+
 namespace EAutomationAnalyticParam
 {
 	enum Type
@@ -45,4 +45,3 @@ namespace EAutomationAnalyticParam
 		NUM_ENGINE_PARAMS
 	};
 }
-
diff --git a/Engine/Source/Runtime/AutomationWorker/Private/AutomationAnalytics.cpp b/Engine/Source/Runtime/AutomationWorker/Private/AutomationAnalytics.cpp
--- a/Engine/Source/Runtime/AutomationWorker/Private/AutomationAnalytics.cpp
+++ b/Engine/Source/Runtime/AutomationWorker/Private/AutomationAnalytics.cpp
@@ -1,27 +1,36 @@
 // Copyright 1998-2017 Epic Games, Inc. All Rights Reserved.
 
 #include "AutomationAnalytics.h"
-#include "Misc/CommandLine.h"
-#include "Misc/App.h"
-#include "Misc/EngineVersion.h"
-#include "Interfaces/IAutomationWorkerModule.h"
-#include "AutomationWorkerMessages.h"
 #include "AutomationWorkerPrivate.h"
+
+#include "AnalyticsET.h"
 #include "AnalyticsEventAttribute.h"
+#include "AutomationWorkerMessages.h"
+#include "IAnalyticsProvider.h"
 #include "IAnalyticsProviderET.h"
-#include "AnalyticsET.h"
+#include "Misc/App.h"
+#include "Misc/CommandLine.h"
+#include "Misc/EngineVersion.h"
+
+#include "IAutomationWorkerModule.h"
+
 
 DEFINE_LOG_CATEGORY(LogAutomationAnalytics);
 
+
+/* Static initialization
+ *****************************************************************************/
+
 bool FAutomationAnalytics::bIsInitialized;
 TSharedPtr<IAnalyticsProviderET> FAutomationAnalytics::Analytics;
 TArray<FString> FAutomationAnalytics::AutomationEventNames;
 TArray<FString> FAutomationAnalytics::AutomationParamNames;
 FString FAutomationAnalytics::MachineSpec;
 
-/**
-* Get analytics pointer
-*/
+
+/* FAutomationAnalytics static functions
+ *****************************************************************************/
+
 IAnalyticsProvider& FAutomationAnalytics::GetProvider()
 {
 	checkf(bIsInitialized && IsAvailable(), TEXT("FAutomationAnalytics::GetProvider called outside of Initialize/Shutdown."));
@@ -29,6 +38,7 @@ IAnalyticsProvider& FAutomationAnalytics::GetProvider()
 	return *Analytics.Get();
 }
 
+
 void FAutomationAnalytics::Initialize()
 {
 	checkf(!bIsInitialized, TEXT("FAutomationAnalytics::Initialize called more than once."));
@@ -51,24 +61,34 @@ void FAutomationAnalytics::Initialize()
 	}
 }
 
+
+bool FAutomationAnalytics::IsAvailable()
+{
+	return Analytics.IsValid();
+}
+
+
 void FAutomationAnalytics::Shutdown()
 {
 	Analytics.Reset();
 	bIsInitialized = false;
 }
 
+
 FString FAutomationAnalytics::GetAutomationEventName(EAutomationEventName::Type InEventName)
 {
 	check(InEventName < AutomationEventNames.Num());
 	return AutomationEventNames[InEventName];
 }
 
+
 FString FAutomationAnalytics::GetAutomationParamName(EAutomationAnalyticParam::Type InEngineParam)
 {
 	check(InEngineParam < AutomationParamNames.Num());
 	return AutomationParamNames[InEngineParam];
 }
 
+
 void FAutomationAnalytics::InititalizeAnalyticParameterNames()
 {
 	// Engine Event Names
@@ -122,7 +142,7 @@ void FAutomationAnalytics::InititalizeAnalyticParameterNames()
 	}
 }
 
-//Timestamp,Platform
+
 void FAutomationAnalytics::SetInitialParameters(TArray<FAnalyticsEventAttribute>& ParamArray)
 {
 	int32 TimeStamp;
@@ -136,33 +156,31 @@ void FAutomationAnalytics::SetInitialParameters(TArray<FAnalyticsEventAttribute>
 }
 
 
-/*
-* @EventName FPSCapture
-*
-* @Trigger 
-*
-* @Type Static
-*
-* @EventParam TimeStamp int32 The time in seconds when the test when the capture ended
-* @EventParam Platform string The platform in which the test is run on
-* @EventParam CL string The Changelist for the build
-* @EventParam Spec string The spec of the machine running the test
-* @EventParam MapName string The map that the test was run on
-* @EventParam MatineeName string The name of the matinee that fired the event
-* @EventParam FPS string 
-* @EventParam BuildConfiguration string Debug/Development/Test/Shipping
-* @EventParam AverageFrameTime string Time for a frame in ms
-* @EventParam AverageGameThreadTime string Time for the game thread in ms
-* @EventParam AverageRenderThreadTime string Time for the rendering thread in ms
-* @EventParam AverageGPUTime string Time for the GPU to flush in ms
-* @EventParam PercentOfFramesAtLeast30FPS string 
-* @EventParam PercentOfFramesAtLeast60FPS string 
-*
-* @Comments 
-*
-*/
 void FAutomationAnalytics::FireEvent_FPSCapture(const FAutomationPerformanceSnapshot& PerfSnapshot)
 {
+	// @EventName FPSCapture
+	//
+	// @Trigger 
+	//
+	// @Type Static
+	//
+	// @EventParam TimeStamp int32 The time in seconds when the test when the capture ended
+	// @EventParam Platform string The platform in which the test is run on
+	// @EventParam CL string The Changelist for the build
+	// @EventParam Spec string The spec of the machine running the test
+	// @EventParam MapName string The map that the test was run on
+	// @EventParam MatineeName string The name of the matinee that fired the event
+	// @EventParam FPS string 
+	// @EventParam BuildConfiguration string Debug/Development/Test/Shipping
+	// @EventParam AverageFrameTime string Time for a frame in ms
+	// @EventParam AverageGameThreadTime string Time for the game thread in ms
+	// @EventParam AverageRenderThreadTime string Time for the rendering thread in ms
+	// @EventParam AverageGPUTime string Time for the GPU to flush in ms
+	// @EventParam PercentOfFramesAtLeast30FPS string 
+	// @EventParam PercentOfFramesAtLeast60FPS string 
+	//
+	// @Comments
+
 	if (Analytics.IsValid())
 	{
 		TArray<FAnalyticsEventAttribute> ParamArray;
@@ -185,30 +203,29 @@ void FAutomationAnalytics::FireEvent_FPSCapture(const FAutomationPerformanceSnap
 	}
 }
 
-/*
-* @EventName AutomationTestResults
-*
-* @Trigger Fires when an automation test has completed
-*
-* @Type Static
-*
-* @EventParam TimeStamp int32 The time in seconds when the test ended
-* @EventParam Platform string The platform in which the test is run on
-* @EventParam CL string The Changelist for the build
-* @EventParam TestName string The execution string of the test being run
-* @EventParam BeautifiedTestName string The name displayed in the UI
-* @EventParam ExecutionCount string
-* @EventParam IsSuccess bool Where the test was a success or not
-* @EventParam Duration float The duration the test took to complete
-* @EventParam ErrorCount int32 The amount of errors during the test
-* @EventParam WarningCount int32 The amount of warnings during the test
-* @EventParam CL string The Changelist for the build
-*
-* @Comments 
-*
-*/
+
 void FAutomationAnalytics::FireEvent_AutomationTestResults(const FAutomationWorkerRunTestsReply* TestResults, const FString& BeautifiedTestName)
 {
+	// @EventName AutomationTestResults
+	//
+	// @Trigger Fires when an automation test has completed
+	//
+	// @Type Static
+	//
+	// @EventParam TimeStamp int32 The time in seconds when the test ended
+	// @EventParam Platform string The platform in which the test is run on
+	// @EventParam CL string The Changelist for the build
+	// @EventParam TestName string The execution string of the test being run
+	// @EventParam BeautifiedTestName string The name displayed in the UI
+	// @EventParam ExecutionCount string
+	// @EventParam IsSuccess bool Where the test was a success or not
+	// @EventParam Duration float The duration the test took to complete
+	// @EventParam ErrorCount int32 The amount of errors during the test
+	// @EventParam WarningCount int32 The amount of warnings during the test
+	// @EventParam CL string The Changelist for the build
+	//
+	// @Comments
+
 	if (Analytics.IsValid())
 	{
 		TArray<FAnalyticsEventAttribute> ParamArray;
diff --git a/Engine/Source/Runtime/AutomationWorker/Private/AutomationAnalytics.h b/Engine/Source/Runtime/AutomationWorker/Private/AutomationAnalytics.h
--- a/Engine/Source/Runtime/AutomationWorker/Private/AutomationAnalytics.h
+++ b/Engine/Source/Runtime/AutomationWorker/Private/AutomationAnalytics.h
@@ -2,33 +2,46 @@
 
 #pragma once
 
-#include "CoreMinimal.h"
+#include "CoreTypes.h"
+#include "Containers/Array.h"
+#include "Containers/UnrealString.h"
+#include "Templates/SharedPointer.h"
+#include "Templates/UnrealTemplate.h"
+
 #include "AutomationAnalyticParams.h"
 
 class IAnalyticsProvider;
 class IAnalyticsProviderET;
+
 struct FAnalyticsEventAttribute;
 struct FAutomationPerformanceSnapshot;
 struct FAutomationWorkerRunTestsReply;
 
-class FAutomationAnalytics : FNoncopyable
+
+class FAutomationAnalytics
+	: FNoncopyable
 {
 public:
+
 	/**
-	* Return the provider instance. Not valid outside of Initialize/Shutdown calls.
-	* Note: must check IsAvailable() first else this code will assert if the provider is not valid.
-	*/
+	 * Return the provider instance. Not valid outside of Initialize/Shutdown calls.
+	 *
+	 * Note: must check IsAvailable() first else this code will assert if the provider is not valid.
+	 */
 	static IAnalyticsProvider& GetProvider();
-	/** Helper function to determine if the provider is valid. */
-	static bool IsAvailable() { return Analytics.IsValid(); }
+
 	/** Called to initialize the singleton. */
 	static void Initialize();
+	
+	/** Helper function to determine if the provider is valid. */
+	static bool IsAvailable();
+
 	/** Called to shut down the singleton */
 	static void Shutdown();
 
 	/*
-	* Helper function to retrieve the parameter name from a given enum
-	*/
+	 * Helper function to retrieve the parameter name from a given enum
+	 */
 	static FString GetAutomationEventName(EAutomationEventName::Type InEventName);
 	static FString GetAutomationParamName(EAutomationAnalyticParam::Type InEngineParam);
 
@@ -37,6 +50,8 @@ public:
 	/** Helper to check if the Analytics provider is Initialized */
 	static bool IsInitialized();
 
+public:
+
 	//////////////////////////////////////////////////////////////////////////
 	//	Events
 	//////////////////////////////////////////////////////////////////////////
@@ -57,6 +72,7 @@ private:
 	static FString MachineSpec;
 
 	static bool bIsInitialized;
+
 	/** This allows us to "opt-out" of engine automation analytics. We still create a session, but we don't send any automation events. */
 	static TSharedPtr<IAnalyticsProviderET> Analytics;
 };
diff --git a/Engine/Source/Runtime/AutomationWorker/Private/AutomationWorkerModule.cpp b/Engine/Source/Runtime/AutomationWorker/Private/AutomationWorkerModule.cpp
--- a/Engine/Source/Runtime/AutomationWorker/Private/AutomationWorkerModule.cpp
+++ b/Engine/Source/Runtime/AutomationWorker/Private/AutomationWorkerModule.cpp
@@ -1,28 +1,31 @@
 // Copyright 1998-2017 Epic Games, Inc. All Rights Reserved.
 
 #include "AutomationWorkerModule.h"
+
+#include "AutomationAnalytics.h"
+#include "AutomationWorkerMessages.h"
 #include "HAL/FileManager.h"
+#include "MessageEndpoint.h"
+#include "MessageEndpointBuilder.h"
+#include "JsonObjectConverter.h"
 #include "Misc/FileHelper.h"
 #include "Misc/Paths.h"
 #include "Misc/App.h"
 #include "Modules/ModuleManager.h"
-#include "Helpers/MessageEndpointBuilder.h"
-#include "AutomationWorkerMessages.h"
-#include "AutomationAnalytics.h"
-#include "JsonObjectConverter.h"
 
 #if WITH_ENGINE
-#include "ImageUtils.h"
-#include "Engine/Engine.h"
-#include "Engine/GameViewportClient.h"
-#include "UnrealClient.h"
-#include "Tests/AutomationCommon.h"
+	#include "Engine/Engine.h"
+	#include "Engine/GameViewportClient.h"
+	#include "ImageUtils.h"
+	#include "Tests/AutomationCommon.h"
+	#include "UnrealClient.h"
 #endif
 
 #if WITH_EDITOR
-#include "AssetRegistryModule.h"
+	#include "AssetRegistryModule.h"
 #endif
 
+
 #define LOCTEXT_NAMESPACE "AutomationTest"
 
 IMPLEMENT_MODULE(FAutomationWorkerModule, AutomationWorker);
@@ -272,7 +275,7 @@ void FAutomationWorkerModule::SendTests( const FMessageAddress& ControllerAddres
 /* FAutomationWorkerModule callbacks
  *****************************************************************************/
 
-void FAutomationWorkerModule::HandleFindWorkersMessage(const FAutomationWorkerFindWorkers& Message, const IMessageContextRef& Context)
+void FAutomationWorkerModule::HandleFindWorkersMessage(const FAutomationWorkerFindWorkers& Message, const TSharedRef<IMessageContext, ESPMode::ThreadSafe>& Context)
 {
 	// Set the Instance name to be the same as the session browser. This information should be shared at some point
 	if ((Message.SessionId == FApp::GetSessionId()) && (Message.Changelist == 10000))
@@ -330,7 +333,7 @@ void FAutomationWorkerModule::SendWorkerFound()
 }
 
 
-void FAutomationWorkerModule::HandleNextNetworkCommandReplyMessage( const FAutomationWorkerNextNetworkCommandReply& Message, const IMessageContextRef& Context )
+void FAutomationWorkerModule::HandleNextNetworkCommandReplyMessage( const FAutomationWorkerNextNetworkCommandReply& Message, const TSharedRef<IMessageContext, ESPMode::ThreadSafe>& Context )
 {
 	// Allow the next command to execute
 	bExecuteNextNetworkCommand = true;
@@ -340,19 +343,19 @@ void FAutomationWorkerModule::HandleNextNetworkCommandReplyMessage( const FAutom
 }
 
 
-void FAutomationWorkerModule::HandlePingMessage( const FAutomationWorkerPing& Message, const IMessageContextRef& Context )
+void FAutomationWorkerModule::HandlePingMessage( const FAutomationWorkerPing& Message, const TSharedRef<IMessageContext, ESPMode::ThreadSafe>& Context )
 {
 	MessageEndpoint->Send(new FAutomationWorkerPong(), Context->GetSender());
 }
 
 
-void FAutomationWorkerModule::HandleResetTests( const FAutomationWorkerResetTests& Message, const IMessageContextRef& Context )
+void FAutomationWorkerModule::HandleResetTests( const FAutomationWorkerResetTests& Message, const TSharedRef<IMessageContext, ESPMode::ThreadSafe>& Context )
 {
 	FAutomationTestFramework::Get().ResetTests();
 }
 
 
-void FAutomationWorkerModule::HandleRequestTestsMessage( const FAutomationWorkerRequestTests& Message, const IMessageContextRef& Context )
+void FAutomationWorkerModule::HandleRequestTestsMessage( const FAutomationWorkerRequestTests& Message, const TSharedRef<IMessageContext, ESPMode::ThreadSafe>& Context )
 {
 	FAutomationTestFramework::Get().LoadTestModules();
 	FAutomationTestFramework::Get().SetDeveloperDirectoryIncluded(Message.DeveloperDirectoryIncluded);
@@ -379,18 +382,18 @@ void FAutomationWorkerModule::HandlePostTestingEvent()
 }
 
 
-void FAutomationWorkerModule::HandleScreenShotCompared(const FAutomationWorkerImageComparisonResults& Message, const IMessageContextRef& Context)
+void FAutomationWorkerModule::HandleScreenShotCompared(const FAutomationWorkerImageComparisonResults& Message, const TSharedRef<IMessageContext, ESPMode::ThreadSafe>& Context)
 {
 	// Image comparison finished.
 	FAutomationTestFramework::Get().NotifyScreenshotComparisonComplete(Message.bNew, Message.bSimilar, Message.MaxLocalDifference, Message.GlobalDifference, Message.ErrorMessage);
 }
 
-void FAutomationWorkerModule::HandleTestDataRetrieved(const FAutomationWorkerTestDataResponse& Message, const IMessageContextRef& Context)
+void FAutomationWorkerModule::HandleTestDataRetrieved(const FAutomationWorkerTestDataResponse& Message, const TSharedRef<IMessageContext, ESPMode::ThreadSafe>& Context)
 {
 	FAutomationTestFramework::Get().NotifyTestDataRetrieved(Message.bIsNew, Message.JsonData);
 }
 
-void FAutomationWorkerModule::HandlePerformanceDataRetrieved(const FAutomationWorkerPerformanceDataResponse& Message, const IMessageContextRef& Context)
+void FAutomationWorkerModule::HandlePerformanceDataRetrieved(const FAutomationWorkerPerformanceDataResponse& Message, const TSharedRef<IMessageContext, ESPMode::ThreadSafe>& Context)
 {
 	FAutomationTestFramework::Get().NotifyPerformanceDataRetrieved(Message.bSuccess, Message.ErrorMessage);
 }
@@ -435,7 +438,7 @@ void FAutomationWorkerModule::HandleScreenShotCapturedWithName(const TArray<FCol
 #endif
 
 
-void FAutomationWorkerModule::HandleRunTestsMessage( const FAutomationWorkerRunTests& Message, const IMessageContextRef& Context )
+void FAutomationWorkerModule::HandleRunTestsMessage( const FAutomationWorkerRunTests& Message, const TSharedRef<IMessageContext, ESPMode::ThreadSafe>& Context )
 {
 	ExecutionCount = Message.ExecutionCount;
 	TestName = Message.TestName;
diff --git a/Engine/Source/Runtime/AutomationWorker/Private/AutomationWorkerModule.h b/Engine/Source/Runtime/AutomationWorker/Private/AutomationWorkerModule.h
--- a/Engine/Source/Runtime/AutomationWorker/Private/AutomationWorkerModule.h
+++ b/Engine/Source/Runtime/AutomationWorker/Private/AutomationWorkerModule.h
@@ -2,11 +2,15 @@
 
 #pragma once
 
-#include "CoreMinimal.h"
-#include "Misc/AutomationTest.h"
+#include "CoreTypes.h"
+#include "Containers/Array.h"
+#include "Containers/UnrealString.h"
 #include "IMessageContext.h"
-#include "Helpers/MessageEndpoint.h"
-#include "Interfaces/IAutomationWorkerModule.h"
+#include "Misc/AutomationTest.h"
+
+#include "IAutomationWorkerModule.h"
+
+class FMessageEndpoint;
 
 struct FAutomationWorkerFindWorkers;
 struct FAutomationWorkerImageComparisonResults;
@@ -17,6 +21,8 @@ struct FAutomationWorkerPing;
 struct FAutomationWorkerRequestTests;
 struct FAutomationWorkerResetTests;
 struct FAutomationWorkerRunTests;
+struct FMessageAddress;
+
 
 /**
  * Implements the Automation Worker module.
@@ -26,48 +32,33 @@ class FAutomationWorkerModule
 {
 public:
 
-	//~ Begin IModuleInterface Interface
+	//~ IModuleInterface interface
 
 	virtual void StartupModule() override;
 	virtual void ShutdownModule() override;
-
 	virtual bool SupportsDynamicReloading() override;
 
-	//~ End IModuleInterface Interface
-
 public:
 
-	//~ Begin IAutomationWorkerModule Interface
+	//~ IAutomationWorkerModule interface
 
 	virtual void Tick() override;
 
-	//~ End IAutomationWorkerModule Interface
-
 protected:
 
-	/**
-	 * Execute all latent commands and when complete send results message back to the automation controller
-	 */
+	/** Execute all latent commands and when complete send results message back to the automation controller. */
 	bool ExecuteLatentCommands();
 
-	/**
-	 * Execute all network commands and when complete send results message back to the automation controller
-	 */
+	/** Execute all network commands and when complete send results message back to the automation controller. */
 	bool ExecuteNetworkCommands();
 
-	/**
-	 * Initializes the automation worker.
-	 */
+	/** Initializes the automation worker. */
 	void Initialize();
 
-	/**
-	 * Network phase is complete (if there were any network commands).  Send ping back to the controller
-	 */
+	/** Network phase is complete (if there were any network commands).  Send ping back to the controller. */
 	void ReportNetworkCommandComplete();
 
-	/**
-	 * Test is complete. Send results back to controller
-	 */
+	/** Test is complete. Send results back to controller. */
 	void ReportTestComplete();
 
 	/** 
@@ -75,47 +66,47 @@ protected:
 	 *
 	 * @param ControllerAddress The message address of the controller that requested the tests.
 	 */
-	void SendTests( const FMessageAddress& ControllerAddress );
+	void SendTests(const FMessageAddress& ControllerAddress);
 
 private:
 
-	// Handles FAutomationWorkerFindWorkers messages.
-	void HandleFindWorkersMessage( const FAutomationWorkerFindWorkers& Message, const IMessageContextRef& Context );
+	/** Handles FAutomationWorkerFindWorkers messages. */
+	void HandleFindWorkersMessage(const FAutomationWorkerFindWorkers& Message, const TSharedRef<IMessageContext, ESPMode::ThreadSafe>& Context);
 
-	//deferred handler for sending "find worker" response in case the asset registry isn't loaded yet
+	/** Deferred handler for sending "find worker" response in case the asset registry isn't loaded yet. */
 	void SendWorkerFound();
 
-	// Handles message endpoint shutdowns.
+	/** Handles message endpoint shutdowns. */
 	void HandleMessageEndpointShutdown();
 
-	// Handles FAutomationWorkerNextNetworkCommandReply messages.
-	void HandleNextNetworkCommandReplyMessage( const FAutomationWorkerNextNetworkCommandReply& Message, const IMessageContextRef& Context );
+	/** Handles FAutomationWorkerNextNetworkCommandReply messages. */
+	void HandleNextNetworkCommandReplyMessage(const FAutomationWorkerNextNetworkCommandReply& Message, const TSharedRef<IMessageContext, ESPMode::ThreadSafe>& Context);
 
-	// Handles FAutomationWorkerPing messages.
-	void HandlePingMessage( const FAutomationWorkerPing& Message, const IMessageContextRef& Context );
+	/** Handles FAutomationWorkerPing messages. */
+	void HandlePingMessage(const FAutomationWorkerPing& Message, const TSharedRef<IMessageContext, ESPMode::ThreadSafe>& Context);
 
-	// Handles FAutomationWorkerResetTests messages.
-	void HandleResetTests( const FAutomationWorkerResetTests& Message, const IMessageContextRef& Context );
+	/** Handles FAutomationWorkerResetTests messages. */
+	void HandleResetTests(const FAutomationWorkerResetTests& Message, const TSharedRef<IMessageContext, ESPMode::ThreadSafe>& Context);
 
-	// Handles FAutomationWorkerRequestTests messages.
-	void HandleRequestTestsMessage( const FAutomationWorkerRequestTests& Message, const IMessageContextRef& Context );
+	/** Handles FAutomationWorkerRequestTests messages. */
+	void HandleRequestTestsMessage(const FAutomationWorkerRequestTests& Message, const TSharedRef<IMessageContext, ESPMode::ThreadSafe>& Context);
 
-	// Handles FAutomationWorkerRunTests messages.
-	void HandleRunTestsMessage( const FAutomationWorkerRunTests& Message, const IMessageContextRef& Context );
+	/** Handles FAutomationWorkerRunTests messages. */
+	void HandleRunTestsMessage(const FAutomationWorkerRunTests& Message, const TSharedRef<IMessageContext, ESPMode::ThreadSafe>& Context);
 
-	// Handles FAutomationWorkerImageComparisonResults messages.
-	void HandleScreenShotCompared(const FAutomationWorkerImageComparisonResults& Message, const IMessageContextRef& Context);
+	/** Handles FAutomationWorkerImageComparisonResults messages. */
+	void HandleScreenShotCompared(const FAutomationWorkerImageComparisonResults& Message, const TSharedRef<IMessageContext, ESPMode::ThreadSafe>& Context);
 
-	// Handles FAutomationWorkerTestDataResponse messages.
-	void HandleTestDataRetrieved(const FAutomationWorkerTestDataResponse& Message, const IMessageContextRef& Context);
+	/** Handles FAutomationWorkerTestDataResponse messages. */
+	void HandleTestDataRetrieved(const FAutomationWorkerTestDataResponse& Message, const TSharedRef<IMessageContext, ESPMode::ThreadSafe>& Context);
 
-	// Handles FAutomationWorkerPerformanceDataResponse messages.
-	void HandlePerformanceDataRetrieved(const FAutomationWorkerPerformanceDataResponse& Message, const IMessageContextRef& Context);
+	/** Handles FAutomationWorkerPerformanceDataResponse messages. */
+	void HandlePerformanceDataRetrieved(const FAutomationWorkerPerformanceDataResponse& Message, const TSharedRef<IMessageContext, ESPMode::ThreadSafe>& Context);
 
-	// Handles FAutomationTestFramework PreTestingEvents.
+	/** Handles FAutomationTestFramework PreTestingEvents. */
 	void HandlePreTestingEvent();
 
-	// Handles FAutomationTestFramework PostTestingEvents.
+	/** Handles FAutomationTestFramework PostTestingEvents. */
 	void HandlePostTestingEvent();
 
 #if WITH_ENGINE
@@ -123,21 +114,21 @@ private:
 	void HandleScreenShotCapturedWithName(const TArray<FColor>& RawImageData, const FAutomationScreenshotData& Data);
 #endif
 
-	//dispatches analytics events to the data collector
+	/** Dispatches analytics events to the data collector. */
 	void SendAnalyticsEvents(TArray<FString>& InAnalyticsItems);
 
-	// Helper for Performance Capture Analytics
-	void RecordPerformanceAnalytics( const FAutomationPerformanceSnapshot& PerfSnapshot );
+	/** Helper for Performance Capture Analytics. */
+	void RecordPerformanceAnalytics(const FAutomationPerformanceSnapshot& PerfSnapshot);
 
 private:
 
-	// The collection of test data we are to send to a controller
+	/** The collection of test data we are to send to a controller. */
 	TArray<FAutomationTestInfo> TestInfo;
 
 private:
 
 	/** Holds the messaging endpoint. */
-	FMessageEndpointPtr MessageEndpoint;
+	TSharedPtr<FMessageEndpoint, ESPMode::ThreadSafe> MessageEndpoint;
 
 	/** Message address of the controller sending the test request. */
 	FMessageAddress TestRequesterAddress;
diff --git a/Engine/Source/Runtime/AutomationWorker/Private/AutomationWorkerPrivate.h b/Engine/Source/Runtime/AutomationWorker/Private/AutomationWorkerPrivate.h
--- a/Engine/Source/Runtime/AutomationWorker/Private/AutomationWorkerPrivate.h
+++ b/Engine/Source/Runtime/AutomationWorker/Private/AutomationWorkerPrivate.h
@@ -2,7 +2,7 @@
 
 #pragma once
 
-#include "CoreMinimal.h"
+#include "Logging/LogMacros.h"
 
-DECLARE_LOG_CATEGORY_EXTERN(LogAutomationAnalytics, Log, All);
 
+DECLARE_LOG_CATEGORY_EXTERN(LogAutomationAnalytics, Log, All);
diff --git a/Engine/Source/Runtime/Core/Private/Android/AndroidApplicationMisc.cpp b/Engine/Source/Runtime/Core/Private/Android/AndroidApplicationMisc.cpp
--- a/Engine/Source/Runtime/Core/Private/Android/AndroidApplicationMisc.cpp
+++ b/Engine/Source/Runtime/Core/Private/Android/AndroidApplicationMisc.cpp
@@ -1,7 +1,10 @@
 // Copyright 1998-2017 Epic Games, Inc. All Rights Reserved.
 
 #include "AndroidApplicationMisc.h"
+
 #include "AndroidApplication.h"
+#include "AndroidInputInterface.h"
+
 
 GenericApplication* FAndroidApplicationMisc::CreateApplication()
 {
diff --git a/Engine/Source/Runtime/Core/Private/Android/AndroidFile.cpp b/Engine/Source/Runtime/Core/Private/Android/AndroidFile.cpp
--- a/Engine/Source/Runtime/Core/Private/Android/AndroidFile.cpp
+++ b/Engine/Source/Runtime/Core/Private/Android/AndroidFile.cpp
@@ -44,13 +44,13 @@ namespace
 		}
 
 		return FFileStatData(
-			AndroidEpoch + FTimespan(0, 0, FileInfo.st_ctime), 
-			AndroidEpoch + FTimespan(0, 0, FileInfo.st_atime), 
-			AndroidEpoch + FTimespan(0, 0, FileInfo.st_mtime), 
+			AndroidEpoch + FTimespan::FromSeconds(FileInfo.st_ctime), 
+			AndroidEpoch + FTimespan::FromSeconds(FileInfo.st_atime), 
+			AndroidEpoch + FTimespan::FromSeconds(FileInfo.st_mtime), 
 			FileSize,
 			bIsDirectory,
 			!!(FileInfo.st_mode & S_IWUSR)
-			);
+		);
 	}
 }
 
diff --git a/Engine/Source/Runtime/Core/Private/Android/AndroidJavaMediaPlayer.cpp b/Engine/Source/Runtime/Core/Private/Android/AndroidJavaMediaPlayer.cpp
--- a/Engine/Source/Runtime/Core/Private/Android/AndroidJavaMediaPlayer.cpp
+++ b/Engine/Source/Runtime/Core/Private/Android/AndroidJavaMediaPlayer.cpp
@@ -22,17 +22,20 @@ static jfieldID FindField(JNIEnv* JEnv, jclass Class, const ANSICHAR* FieldName,
 	return Field;
 }
 
-FJavaAndroidMediaPlayer::FJavaAndroidMediaPlayer(bool vulkanRenderer)
-	: FJavaClassObject(GetClassName(), "(Z)V", vulkanRenderer)
+FJavaAndroidMediaPlayer::FJavaAndroidMediaPlayer(bool swizzlePixels, bool vulkanRenderer)
+	: FJavaClassObject(GetClassName(), "(ZZ)V", swizzlePixels, vulkanRenderer)
 	, GetDurationMethod(GetClassMethod("getDuration", "()I"))
 	, ResetMethod(GetClassMethod("reset", "()V"))
 	, GetCurrentPositionMethod(GetClassMethod("getCurrentPosition", "()I"))
+	, DidCompleteMethod(GetClassMethod("didComplete", "()Z"))
 	, IsLoopingMethod(GetClassMethod("isLooping", "()Z"))
 	, IsPlayingMethod(GetClassMethod("isPlaying", "()Z"))
+	, IsPreparedMethod(GetClassMethod("isPrepared", "()Z"))
 	, SetDataSourceURLMethod(GetClassMethod("setDataSourceURL", "(Ljava/lang/String;)Z"))
 	, SetDataSourceFileMethod(GetClassMethod("setDataSource", "(Ljava/lang/String;JJ)Z"))
 	, SetDataSourceAssetMethod(GetClassMethod("setDataSource", "(Landroid/content/res/AssetManager;Ljava/lang/String;JJ)Z"))
 	, PrepareMethod(GetClassMethod("prepare", "()V"))
+	, PrepareAsyncMethod(GetClassMethod("prepareAsync", "()V"))
 	, SeekToMethod(GetClassMethod("seekTo", "(I)V"))
 	, SetLoopingMethod(GetClassMethod("setLooping", "(Z)V"))
 	, ReleaseMethod(GetClassMethod("release", "()V"))
@@ -49,7 +52,15 @@ FJavaAndroidMediaPlayer::FJavaAndroidMediaPlayer(bool vulkanRenderer)
 	, GetCaptionTracksMethod(GetClassMethod("GetCaptionTracks", "()[Lcom/epicgames/ue4/MediaPlayer14$CaptionTrackInfo;"))
 	, GetVideoTracksMethod(GetClassMethod("GetVideoTracks", "()[Lcom/epicgames/ue4/MediaPlayer14$VideoTrackInfo;"))
 	, DidResolutionChangeMethod(GetClassMethod("didResolutionChange", "()Z"))
+	, GetExternalTextureIdMethod(GetClassMethod("getExternalTextureId", "()I"))
+	, UpdateVideoFrameMethod(GetClassMethod("updateVideoFrame", "(I)Lcom/epicgames/ue4/MediaPlayer14$FrameUpdateInfo;"))
 {
+	VideoTexture = nullptr;
+	bVideoTextureValid = false;
+
+	UScale = VScale = 1.0f;
+	UOffset = VOffset = 0.0f;
+
 	bTrackInfoSupported = FAndroidMisc::GetAndroidBuildVersion() >= 16;
 
 	if (bTrackInfoSupported)
@@ -59,6 +70,18 @@ FJavaAndroidMediaPlayer::FJavaAndroidMediaPlayer(bool vulkanRenderer)
 
 	JNIEnv* JEnv = FAndroidApplication::GetJavaEnv();
 
+	// get field IDs for FrameUpdateInfo class members
+	jclass localFrameUpdateInfoClass = FAndroidApplication::FindJavaClass("com/epicgames/ue4/MediaPlayer14$FrameUpdateInfo");
+	FrameUpdateInfoClass = (jclass)JEnv->NewGlobalRef(localFrameUpdateInfoClass);
+	JEnv->DeleteLocalRef(localFrameUpdateInfoClass);
+	FrameUpdateInfo_CurrentPosition = FindField(JEnv, FrameUpdateInfoClass, "CurrentPosition", "I", false);
+	FrameUpdateInfo_FrameReady = FindField(JEnv, FrameUpdateInfoClass, "FrameReady", "Z", false);
+	FrameUpdateInfo_RegionChanged = FindField(JEnv, FrameUpdateInfoClass, "RegionChanged", "Z", false);
+	FrameUpdateInfo_UScale = FindField(JEnv, FrameUpdateInfoClass, "UScale", "F", false);
+	FrameUpdateInfo_UOffset = FindField(JEnv, FrameUpdateInfoClass, "UOffset", "F", false);
+	FrameUpdateInfo_VScale = FindField(JEnv, FrameUpdateInfoClass, "VScale", "F", false);
+	FrameUpdateInfo_VOffset = FindField(JEnv, FrameUpdateInfoClass, "VOffset", "F", false);
+
 	// get field IDs for AudioTrackInfo class members
 	jclass localAudioTrackInfoClass = FAndroidApplication::FindJavaClass("com/epicgames/ue4/MediaPlayer14$AudioTrackInfo");
 	AudioTrackInfoClass = (jclass)JEnv->NewGlobalRef(localAudioTrackInfoClass);
@@ -100,6 +123,8 @@ int32 FJavaAndroidMediaPlayer::GetDuration()
 
 void FJavaAndroidMediaPlayer::Reset()
 {
+	UScale = VScale = 1.0f;
+	UOffset = VOffset = 0.0f;
 	CallMethod<void>(ResetMethod);
 }
 
@@ -124,18 +149,34 @@ bool FJavaAndroidMediaPlayer::IsPlaying()
 	return CallMethod<bool>(IsPlayingMethod);
 }
 
+bool FJavaAndroidMediaPlayer::IsPrepared()
+{
+	return CallMethod<bool>(IsPreparedMethod);
+}
+
+bool FJavaAndroidMediaPlayer::DidComplete()
+{
+	return CallMethod<bool>(DidCompleteMethod);
+}
+
 bool FJavaAndroidMediaPlayer::SetDataSource(const FString & Url)
 {
+	UScale = VScale = 1.0f;
+	UOffset = VOffset = 0.0f;
 	return CallMethod<bool>(SetDataSourceURLMethod, GetJString(Url));
 }
 
 bool FJavaAndroidMediaPlayer::SetDataSource(const FString& MoviePathOnDevice, int64 offset, int64 size)
 {
+	UScale = VScale = 1.0f;
+	UOffset = VOffset = 0.0f;
 	return CallMethod<bool>(SetDataSourceFileMethod, GetJString(MoviePathOnDevice), offset, size);
 }
 
 bool FJavaAndroidMediaPlayer::SetDataSource(jobject AssetMgr, const FString& AssetPath, int64 offset, int64 size)
 {
+	UScale = VScale = 1.0f;
+	UOffset = VOffset = 0.0f;
 	return CallMethod<bool>(SetDataSourceAssetMethod, AssetMgr, GetJString(AssetPath), offset, size);
 }
 
@@ -153,6 +194,20 @@ bool FJavaAndroidMediaPlayer::Prepare()
 	return true;
 }
 
+bool FJavaAndroidMediaPlayer::PrepareAsync()
+{
+	// This can return an exception in some cases (URL without internet, for example)
+	JNIEnv*	JEnv = FAndroidApplication::GetJavaEnv();
+	JEnv->CallVoidMethod(Object, PrepareAsyncMethod.Method);
+	if (JEnv->ExceptionCheck())
+	{
+		JEnv->ExceptionDescribe();
+		JEnv->ExceptionClear();
+		return false;
+	}
+	return true;
+}
+
 void FJavaAndroidMediaPlayer::SeekTo(int32 Milliseconds)
 {
 	CallMethod<void>(SeekToMethod, Milliseconds);
@@ -190,13 +245,27 @@ void FJavaAndroidMediaPlayer::SetAudioEnabled(bool enabled /*= true*/)
 
 bool FJavaAndroidMediaPlayer::GetVideoLastFrameData(void* & outPixels, int64 & outCount)
 {
-	jobject buffer = CallMethod<jobject>(GetVideoLastFrameDataMethod);
+	// This can return an exception in some cases
+	JNIEnv*	JEnv = FAndroidApplication::GetJavaEnv();
+	jobject buffer = JEnv->CallObjectMethod(Object, GetVideoLastFrameDataMethod.Method);
+	if (JEnv->ExceptionCheck())
+	{
+		JEnv->ExceptionDescribe();
+		JEnv->ExceptionClear();
+		if (nullptr != buffer)
+		{
+			// the CallObjectMethod returns a local ref, but Java will still own the real buffer
+			JEnv->DeleteLocalRef(buffer);
+		}
+		return false;
+	}
 	if (nullptr != buffer)
 	{
-		JNIEnv*	JEnv = FAndroidApplication::GetJavaEnv();
 		outPixels = JEnv->GetDirectBufferAddress(buffer);
 		outCount = JEnv->GetDirectBufferCapacity(buffer);
-		// we don't reduce the buffer refcount because it is still owned by the Java class, and will be reused.
+
+		// the CallObjectMethod returns a local ref, but Java will still own the real buffer
+		JEnv->DeleteLocalRef(buffer);
 	}
 	if (nullptr == buffer || nullptr == outPixels || 0 == outCount)
 	{
@@ -220,6 +289,49 @@ bool FJavaAndroidMediaPlayer::DidResolutionChange()
 	return CallMethod<bool>(DidResolutionChangeMethod);
 }
 
+int32 FJavaAndroidMediaPlayer::GetExternalTextureId()
+{
+	return CallMethod<int32>(GetExternalTextureIdMethod);
+}
+
+bool FJavaAndroidMediaPlayer::UpdateVideoFrame(int32 ExternalTextureId, int32 *CurrentPosition, bool *bRegionChanged)
+{
+	// This can return an exception in some cases
+	JNIEnv*	JEnv = FAndroidApplication::GetJavaEnv();
+	jobject Result = JEnv->CallObjectMethod(Object, UpdateVideoFrameMethod.Method, ExternalTextureId);
+	if (JEnv->ExceptionCheck())
+	{
+		JEnv->ExceptionDescribe();
+		JEnv->ExceptionClear();
+		if (nullptr != Result)
+		{
+			JEnv->DeleteLocalRef(Result);
+		}
+		*CurrentPosition = -1;
+		*bRegionChanged = false;
+		return false;
+	}
+
+	if (nullptr == Result)
+	{
+		*CurrentPosition = -1;
+		*bRegionChanged = false;
+		return false;
+	}
+
+	*CurrentPosition = (int32)JEnv->GetIntField(Result, FrameUpdateInfo_CurrentPosition);
+	bool bFrameReady = (bool)JEnv->GetBooleanField(Result, FrameUpdateInfo_FrameReady);
+	*bRegionChanged = (bool)JEnv->GetBooleanField(Result, FrameUpdateInfo_RegionChanged);
+	UScale = (float)JEnv->GetFloatField(Result, FrameUpdateInfo_UScale);
+	UOffset = (float)JEnv->GetFloatField(Result, FrameUpdateInfo_UOffset);
+	VScale = (float)JEnv->GetFloatField(Result, FrameUpdateInfo_VScale);
+	VOffset = (float)JEnv->GetFloatField(Result, FrameUpdateInfo_VOffset);
+
+	JEnv->DeleteLocalRef(Result);
+
+	return bFrameReady;
+}
+
 bool FJavaAndroidMediaPlayer::GetVideoLastFrame(int32 destTexture)
 {
 	// This can return an exception in some cases
diff --git a/Engine/Source/Runtime/Core/Private/Apple/ApplePlatformFile.cpp b/Engine/Source/Runtime/Core/Private/Apple/ApplePlatformFile.cpp
--- a/Engine/Source/Runtime/Core/Private/Apple/ApplePlatformFile.cpp
+++ b/Engine/Source/Runtime/Core/Private/Apple/ApplePlatformFile.cpp
@@ -29,13 +29,13 @@ namespace
 		}
 
 		return FFileStatData(
-			MacEpoch + FTimespan(0, 0, FileInfo.st_ctime), 
-			MacEpoch + FTimespan(0, 0, FileInfo.st_atime), 
-			MacEpoch + FTimespan(0, 0, FileInfo.st_mtime), 
+			MacEpoch + FTimespan::FromSeconds(FileInfo.st_ctime), 
+			MacEpoch + FTimespan::FromSeconds(FileInfo.st_atime), 
+			MacEpoch + FTimespan::FromSeconds(FileInfo.st_mtime), 
 			FileSize,
 			bIsDirectory,
 			!!(FileInfo.st_mode & S_IWUSR)
-			);
+		);
 	}
 }
 
diff --git a/Engine/Source/Runtime/Core/Private/HTML5/HTML5PlatformFile.cpp b/Engine/Source/Runtime/Core/Private/HTML5/HTML5PlatformFile.cpp
--- a/Engine/Source/Runtime/Core/Private/HTML5/HTML5PlatformFile.cpp
+++ b/Engine/Source/Runtime/Core/Private/HTML5/HTML5PlatformFile.cpp
@@ -52,9 +52,9 @@ namespace
 		}
 
 		return FFileStatData(
-			HTML5Epoch + FTimespan(0, 0, FileInfo.st_ctime),
-			HTML5Epoch + FTimespan(0, 0, FileInfo.st_atime),
-			HTML5Epoch + FTimespan(0, 0, FileInfo.st_mtime),
+			HTML5Epoch + FTimespan::FromSeconds(FileInfo.st_ctime),
+			HTML5Epoch + FTimespan::FromSeconds(FileInfo.st_atime),
+			HTML5Epoch + FTimespan::FromSeconds(FileInfo.st_mtime),
 			FileSize,
 			bIsDirectory,
 			!!(FileInfo.st_mode & S_IWUSR)
diff --git a/Engine/Source/Runtime/Core/Private/IOS/IOSPlatformFile.cpp b/Engine/Source/Runtime/Core/Private/IOS/IOSPlatformFile.cpp
--- a/Engine/Source/Runtime/Core/Private/IOS/IOSPlatformFile.cpp
+++ b/Engine/Source/Runtime/Core/Private/IOS/IOSPlatformFile.cpp
@@ -31,9 +31,9 @@ namespace
 		}
 
 		return FFileStatData(
-			IOSEpoch + FTimespan(0, 0, FileInfo.st_ctime), 
-			IOSEpoch + FTimespan(0, 0, FileInfo.st_atime), 
-			IOSEpoch + FTimespan(0, 0, FileInfo.st_mtime), 
+			IOSEpoch + FTimespan::FromSeconds(FileInfo.st_ctime), 
+			IOSEpoch + FTimespan::FromSeconds(FileInfo.st_atime), 
+			IOSEpoch + FTimespan::FromSeconds(FileInfo.st_mtime), 
 			FileSize,
 			bIsDirectory,
 			!!(FileInfo.st_mode & S_IWUSR)
diff --git a/Engine/Source/Runtime/Core/Private/Linux/LinuxPlatformFile.cpp b/Engine/Source/Runtime/Core/Private/Linux/LinuxPlatformFile.cpp
--- a/Engine/Source/Runtime/Core/Private/Linux/LinuxPlatformFile.cpp
+++ b/Engine/Source/Runtime/Core/Private/Linux/LinuxPlatformFile.cpp
@@ -28,13 +28,13 @@ namespace
 		}
 
 		return FFileStatData(
-			UnixEpoch + FTimespan(0, 0, FileInfo.st_ctime), 
-			UnixEpoch + FTimespan(0, 0, FileInfo.st_atime), 
-			UnixEpoch + FTimespan(0, 0, FileInfo.st_mtime), 
+			UnixEpoch + FTimespan::FromSeconds(FileInfo.st_ctime), 
+			UnixEpoch + FTimespan::FromSeconds(FileInfo.st_atime), 
+			UnixEpoch + FTimespan::FromSeconds(FileInfo.st_mtime), 
 			FileSize,
 			bIsDirectory,
 			!!(FileInfo.st_mode & S_IWUSR)
-			);
+		);
 	}
 }
 
diff --git a/Engine/Source/Runtime/Core/Private/Linux/LinuxPlatformSplash.cpp b/Engine/Source/Runtime/Core/Private/Linux/LinuxPlatformSplash.cpp
--- a/Engine/Source/Runtime/Core/Private/Linux/LinuxPlatformSplash.cpp
+++ b/Engine/Source/Runtime/Core/Private/Linux/LinuxPlatformSplash.cpp
@@ -1,11 +1,7 @@
 // Copyright 1998-2017 Epic Games, Inc. All Rights Reserved.
 
-/*================================================================================
-	SplashScreen.cpp: Splash screen for game/editor startup
-================================================================================*/
-
-
 #include "Linux/LinuxPlatformSplash.h"
+
 #include "Misc/AssertionMacros.h"
 #include "HAL/UnrealMemory.h"
 #include "Misc/EngineVersionBase.h"
@@ -35,11 +31,13 @@
 #include "SDL.h"
 #include "SDL_thread.h"
 #include "SDL_timer.h"
-#include "Interfaces/IImageWrapperModule.h"
+#include "IImageWrapper.h"
+#include "IImageWrapperModule.h"
 
 #include "ft2build.h"
 #include FT_FREETYPE_H
 
+
 /**
  * Splash screen functions and static globals
  */
@@ -83,8 +81,8 @@ static FText GSplashScreenText[ SplashTextType::NumTextTypes ];
 static Rect GSplashScreenTextRects[ SplashTextType::NumTextTypes ];
 static FString GSplashPath;
 static FString GIconPath;
-static IImageWrapperPtr GSplashImageWrapper;
-static IImageWrapperPtr GIconImageWrapper;
+static TSharedPtr<IImageWrapper> GSplashImageWrapper;
+static TSharedPtr<IImageWrapper> GIconImageWrapper;
 static FCriticalSection GSplashScreenTextCriticalSection;
 
 static int32 SplashWidth = 0, SplashHeight = 0;
@@ -533,7 +531,7 @@ static int RenderString (GLuint tex_idx)
  *
  * @return Splash surface
  */
-static SDL_Surface* LinuxSplash_LoadImage(const TCHAR* InImagePath, IImageWrapperPtr & OutImageWrapper)
+static SDL_Surface* LinuxSplash_LoadImage(const TCHAR* InImagePath, TSharedPtr<IImageWrapper>& OutImageWrapper)
 {
 	TArray<uint8> RawFileData;
 	FString ImagePath(InImagePath);
diff --git a/Engine/Source/Runtime/Core/Private/Mac/MacPlatformSplash.cpp b/Engine/Source/Runtime/Core/Private/Mac/MacPlatformSplash.cpp
--- a/Engine/Source/Runtime/Core/Private/Mac/MacPlatformSplash.cpp
+++ b/Engine/Source/Runtime/Core/Private/Mac/MacPlatformSplash.cpp
@@ -6,6 +6,7 @@
 #include "EngineVersion.h"
 #include "EngineBuildSettings.h"
 #include "CocoaThread.h"
+#include "ScopeLock.h"
 
 /**
  * Simple window class that overrides a couple of functions, so the splash window can be moved to front even if it's borderless
diff --git a/Engine/Source/Runtime/Core/Private/Math/UnitConversion.cpp b/Engine/Source/Runtime/Core/Private/Math/UnitConversion.cpp
--- a/Engine/Source/Runtime/Core/Private/Math/UnitConversion.cpp
+++ b/Engine/Source/Runtime/Core/Private/Math/UnitConversion.cpp
@@ -76,6 +76,8 @@ FParseCandidate ParseCandidates[] = {
 
 	{ TEXT("ppi"),					EUnit::PixelsPerInch },			{ TEXT("dpi"),		EUnit::PixelsPerInch },
 
+	{ TEXT("Percent"),				EUnit::Percentage },			{ TEXT("%"),	EUnit::Percentage },
+
 	{ TEXT("times"),				EUnit::Multiplier },			{ TEXT("x"),	EUnit::Multiplier },			{ TEXT("multiplier"),		EUnit::Multiplier },
 };
 
@@ -106,6 +108,8 @@ const TCHAR* const DisplayStrings[] = {
 
 	TEXT("ppi"),
 
+	TEXT("%"),
+
 	TEXT("x"),
 };
 
@@ -135,6 +139,8 @@ const EUnitType UnitTypes[] = {
 
 	EUnitType::PixelDensity,
 
+	EUnitType::Multipliers,
+
 	EUnitType::Arbitrary,
 };
 
@@ -639,6 +645,16 @@ namespace UnitConversion
 		}
 	}
 
+	double MultiplierUnificationFactor(EUnit From)
+	{
+		switch (From)
+		{
+			case EUnit::Percentage:			return 0.01;
+			case EUnit::Multiplier:						// fallthrough
+			default: 						return 1.0;
+		}
+	}
+
 	TValueOrError<FNumericUnit<double>, FText> TryParseExpression(const TCHAR* InExpression, EUnit From, const FNumericUnit<double>& InExistingValue)
 	{
 		const FUnitExpressionParser Parser(From);
diff --git a/Engine/Source/Runtime/Core/Private/Misc/DateTime.cpp b/Engine/Source/Runtime/Core/Private/Misc/DateTime.cpp
--- a/Engine/Source/Runtime/Core/Private/Misc/DateTime.cpp
+++ b/Engine/Source/Runtime/Core/Private/Misc/DateTime.cpp
@@ -5,6 +5,7 @@
 #include "Templates/TypeHash.h"
 #include "UObject/PropertyPortFlags.h"
 
+
 /* FDateTime constants
  *****************************************************************************/
 
@@ -15,7 +16,7 @@
 /* FDateTime structors
  *****************************************************************************/
 
-FDateTime::FDateTime( int32 Year, int32 Month, int32 Day, int32 Hour, int32 Minute, int32 Second, int32 Millisecond )
+FDateTime::FDateTime(int32 Year, int32 Month, int32 Day, int32 Hour, int32 Minute, int32 Second, int32 Millisecond)
 {
 	check(Validate(Year, Month, Day, Hour, Minute, Second, Millisecond));
 
@@ -26,8 +27,8 @@
 		++TotalDays;
 	}
 
-	Year--;											// the current year is not a full year yet
-	Month--;										// the current month is not a full month yet
+	--Year;											// the current year is not a full year yet
+	--Month;										// the current month is not a full month yet
 
 	TotalDays += Year * 365;
 	TotalDays += Year / 4;							// leap year day every four years...
@@ -47,7 +48,7 @@
 /* FDateTime interface
  *****************************************************************************/
 
-bool FDateTime::ExportTextItem( FString& ValueStr, FDateTime const& DefaultValue, UObject* Parent, int32 PortFlags, UObject* ExportRootScope ) const
+bool FDateTime::ExportTextItem(FString& ValueStr, FDateTime const& DefaultValue, UObject* Parent, int32 PortFlags, UObject* ExportRootScope) const
 {
 	if (0 != (PortFlags & EPropertyPortFlags::PPF_ExportCpp))
 	{
@@ -61,7 +62,7 @@
 }
 
 
-void FDateTime::GetDate( int32& OutYear, int32& OutMonth, int32& OutDay ) const
+void FDateTime::GetDate(int32& OutYear, int32& OutMonth, int32& OutDay) const
 {
 	// Based on FORTRAN code in:
 	// Fliegel, H. F. and van Flandern, T. C.,
@@ -152,25 +153,27 @@
 }
 
 
-bool FDateTime::ImportTextItem( const TCHAR*& Buffer, int32 PortFlags, UObject* Parent, FOutputDevice* ErrorText )
+bool FDateTime::ImportTextItem(const TCHAR*& Buffer, int32 PortFlags, UObject* Parent, FOutputDevice* ErrorText)
 {
-	if (FPlatformString::Strlen(Buffer) < 19)
+	const int32 ExportDateTimeLen = 19;
+
+	if (FPlatformString::Strlen(Buffer) < ExportDateTimeLen)
 	{
 		return false;
 	}
 
-	if (!Parse(FString(Buffer).Left(19), *this))
+	if (!Parse(FString(Buffer).Left(ExportDateTimeLen), *this))
 	{
 		return false;
 	}
 
-	Buffer += 19;
+	Buffer += ExportDateTimeLen;
 
 	return true;
 }
 
 
-bool FDateTime::Serialize( FArchive& Ar )
+bool FDateTime::Serialize(FArchive& Ar)
 {
 	Ar << *this;
 
@@ -178,14 +181,11 @@
 }
 
 
-FString FDateTime::ToIso8601() const
-{
-	return ToString(TEXT("%Y-%m-%dT%H:%M:%S.%sZ"));
-}
-
 FString FDateTime::ToHttpDate() const
 {
 	FString DayStr;
+	FString MonthStr;
+
 	switch (GetDayOfWeek())
 	{
 		case EDayOfWeek::Monday:	DayStr = TEXT("Mon");	break;
@@ -197,7 +197,6 @@
 		case EDayOfWeek::Sunday:	DayStr = TEXT("Sun");	break;
 	}
 
-	FString MonthStr;
 	switch (GetMonthOfYear())
 	{
 		case EMonthOfYear::January:		MonthStr = TEXT("Jan");	break;
@@ -215,16 +214,24 @@
 	}
 
 	FString Time = FString::Printf(TEXT("%02i:%02i:%02i"), GetHour(), GetMinute(), GetSecond());
+
 	return FString::Printf(TEXT("%s, %02d %s %d %s GMT"), *DayStr, GetDay(), *MonthStr, GetYear(), *Time);
 }
 
+
+FString FDateTime::ToIso8601() const
+{
+	return ToString(TEXT("%Y-%m-%dT%H:%M:%S.%sZ"));
+}
+
+
 FString FDateTime::ToString() const
 {
 	return ToString(TEXT("%Y.%m.%d-%H.%M.%S"));
 }
 
 
-FString FDateTime::ToString( const TCHAR* Format ) const
+FString FDateTime::ToString(const TCHAR* Format) const
 {
 	FString Result;
 
@@ -268,7 +275,7 @@
 /* FDateTime static interface
  *****************************************************************************/
 
-int32 FDateTime::DaysInMonth( int32 Year, int32 Month )
+int32 FDateTime::DaysInMonth(int32 Year, int32 Month)
 {
 	check((Month >= 1) && (Month <= 12));
 
@@ -281,7 +288,7 @@
 }
 
 
-int32 FDateTime::DaysInYear( int32 Year )
+int32 FDateTime::DaysInYear(int32 Year)
 {
 	if (IsLeapYear(Year))
 	{
@@ -292,7 +299,7 @@
 }
 
 
-bool FDateTime::IsLeapYear( int32 Year )
+bool FDateTime::IsLeapYear(int32 Year)
 {
 	if ((Year % 4) == 0)
 	{
@@ -314,17 +321,15 @@
 }
 
 
-bool FDateTime::Parse( const FString& DateTimeString, FDateTime& OutDateTime )
+bool FDateTime::Parse(const FString& DateTimeString, FDateTime& OutDateTime)
 {
 	// first replace -, : and . with space
 	FString FixedString = DateTimeString.Replace(TEXT("-"), TEXT(" "));
-
 	FixedString.ReplaceInline(TEXT(":"), TEXT(" "), ESearchCase::CaseSensitive);
 	FixedString.ReplaceInline(TEXT("."), TEXT(" "), ESearchCase::CaseSensitive);
 
-	TArray<FString> Tokens;
-
 	// split up on a single delimiter
+	TArray<FString> Tokens;
 	FixedString.ParseIntoArray(Tokens, TEXT(" "), true);
 
 	// make sure it parsed it properly (within reason)
@@ -353,160 +358,6 @@
 }
 
 
-bool FDateTime::ParseIso8601( const TCHAR* DateTimeString, FDateTime& OutDateTime )
-{
-	// DateOnly: YYYY-MM-DD
-	// DateTime: YYYY-mm-ddTHH:MM:SS(.ssss)(Z|+th:tm|-th:tm)
-
-	const TCHAR* Ptr = DateTimeString;
-	TCHAR* Next = nullptr;
-
-	int32 Year = 0, Month = 0, Day = 0;
-	int32 Hour = 0, Minute = 0, Second = 0, Millisecond = 0;
-	int32 TzHour = 0, TzMinute = 0;
-
-	// get date
-	Year = FCString::Strtoi(Ptr, &Next, 10);
-
-	if ((Next <= Ptr) || (*Next == TCHAR('\0')))
-	{
-		return false;
-	}
-
-	Ptr = Next + 1; // skip separator
-	Month = FCString::Strtoi(Ptr, &Next, 10);
-
-	if ((Next <= Ptr) || (*Next == TCHAR('\0')))
-	{
-		return false;
-	}
-
-	Ptr = Next + 1; // skip separator
-	Day = FCString::Strtoi(Ptr, &Next, 10);
-
-	if (Next <= Ptr)
-	{
-		return false;
-	}
-
-	// check whether this is date and time
-	if (*Next == TCHAR('T'))
-	{
-		Ptr = Next + 1;
-
-		// parse time
-		Hour = FCString::Strtoi(Ptr, &Next, 10);
-
-		if ((Next <= Ptr) || (*Next == TCHAR('\0')))
-		{
-			return false;
-		}
-
-		Ptr = Next + 1; // skip separator
-		Minute = FCString::Strtoi(Ptr, &Next, 10);
-
-		if ((Next <= Ptr) || (*Next == TCHAR('\0')))
-		{
-			return false;
-		}
-
-		Ptr = Next + 1; // skip separator
-		Second = FCString::Strtoi(Ptr, &Next, 10);
-
-		if (Next <= Ptr)
-		{
-			return false;
-		}
-
-		// check for milliseconds
-		if (*Next == TCHAR('.'))
-		{
-			Ptr = Next + 1;
-			Millisecond = FCString::Strtoi(Ptr, &Next, 10);
-
-			// should be no more than 3 digits
-			if ((Next <= Ptr) || (Next > Ptr + 3))
-			{
-				return false;
-			}
-
-			for (int32 Digits = Next - Ptr; Digits < 3; ++Digits)
-			{
-				Millisecond *= 10;
-			}
-		}
-
-		// see if the timezone offset is included
-		if (*Next == TCHAR('+') || *Next == TCHAR('-'))
-		{
-			// include the separator since it's + or -
-			Ptr = Next;
-
-			// parse the timezone offset
-			TzHour = FCString::Strtoi(Ptr, &Next, 10);
-
-			if ((Next <= Ptr) || (*Next == TCHAR('\0')))
-			{
-				return false;
-			}
-
-			Ptr = Next + 1; // skip separator
-			TzMinute = FCString::Strtoi(Ptr, &Next, 10);
-
-			if (Next <= Ptr)
-			{
-				return false;
-			}
-		}
-		else if ((*Next != TCHAR('\0')) && (*Next != TCHAR('Z')))
-		{
-			return false;
-		}
-	}
-	else if (*Next != TCHAR('\0'))
-	{
-		return false;
-	}
-
-	if (!Validate(Year, Month, Day, Hour, Minute, Second, Millisecond))
-	{
-		return false;
-	}
-
-	FDateTime Final(Year, Month, Day, Hour, Minute, Second, Millisecond);
-
-	// adjust for the timezone (bringing the DateTime into UTC)
-	int32 TzOffsetMinutes = (TzHour < 0) ? TzHour * 60 - TzMinute : TzHour * 60 + TzMinute;
-	Final -= FTimespan(0, TzOffsetMinutes, 0);
-	OutDateTime = Final;
-
-	return true;
-}
-
-/**
- * Parse a string of the HTTP-date format from a web server
- * https://www.w3.org/Protocols/rfc2616/rfc2616-sec3.html#sec3.3.1
- *
- *   HTTP-date    = rfc1123-date | rfc850-date | asctime-date
- *	 rfc1123-date = wkday "," SP date1 SP time SP "GMT"
- *	 rfc850-date  = weekday "," SP date2 SP time SP "GMT"
- *   asctime-date = wkday SP date3 SP time SP 4DIGIT
- *   date1        = 2DIGIT SP month SP 4DIGIT
- *				  ; day month year (e.g., 02 Jun 1982)
- *   date2        = 2DIGIT "-" month "-" 2DIGIT
- *				  ; day-month-year (e.g., 02-Jun-82)
- *   date3        = month SP ( 2DIGIT | ( SP 1DIGIT ))
- *				  ; month day (e.g., Jun  2)
- *   time         = 2DIGIT ":" 2DIGIT ":" 2DIGIT
- *				  ; 00:00:00 - 23:59:59
- *   wkday        = "Mon" | "Tue" | "Wed"
- *				  | "Thu" | "Fri" | "Sat" | "Sun"
- *   weekday      = "Monday" | "Tuesday" | "Wednesday"
- *				  | "Thursday" | "Friday" | "Saturday" | "Sunday"
- *   month        = "Jan" | "Feb" | "Mar" | "Apr"
- *				  | "May" | "Jun" | "Jul" | "Aug"
- *				  | "Sep" | "Oct" | "Nov" | "Dec"
- */
 bool FDateTime::ParseHttpDate(const FString& HttpDate, FDateTime& OutDateTime)
 {
 	auto ParseTime = [](const FString& Time, int32& Hour, int32& Minute, int32& Second) -> bool
@@ -517,11 +368,13 @@
 
 		// split up on a single delimiter
 		int32 NumTokens = Time.ParseIntoArray(Tokens, TEXT(":"), true);
+
 		if (NumTokens == 3)
 		{
 			Hour = FCString::Atoi(*Tokens[0]);
 			Minute = FCString::Atoi(*Tokens[1]);
 			Second = FCString::Atoi(*Tokens[2]);
+
 			return (Hour >= 0 && Hour < 24) && (Minute >= 0 && Minute <= 59) && (Second >= 0 && Second <= 59);
 		}
 
@@ -531,6 +384,7 @@
 	auto ParseWkday = [](const FString& WkDay) -> int32
 	{
 		const int32 NumChars = WkDay.Len();
+
 		if (NumChars == 3)
 		{
 			if (WkDay.Equals(TEXT("Mon")))
@@ -569,6 +423,7 @@
 	auto ParseWeekday = [](const FString& WeekDay) -> int32
 	{
 		const int32 NumChars = WeekDay.Len();
+
 		if (NumChars >= 6 && NumChars <= 9)
 		{
 			if (WeekDay.Equals(TEXT("Monday")))
@@ -607,6 +462,7 @@
 	auto ParseMonth = [](const FString& Month) -> int32
 	{
 		const int32 NumChars = Month.Len();
+
 		if (NumChars == 3)
 		{
 			if (Month.Equals(TEXT("Jan")))
@@ -734,6 +590,7 @@
 				if (Tokens.Num() == 6)
 				{
 					int32 WkDay = ParseWkday(Tokens[0]);
+
 					if (WkDay > 0)
 					{
 						// rfc1123 - date = wkday "," SP date1 SP time SP "GMT"
@@ -747,6 +604,7 @@
 				{
 					// rfc850 - date = weekday "," SP date2 SP time SP "GMT"
 					int32 WeekDay = ParseWeekday(Tokens[0]);
+
 					if (WeekDay > 0)
 					{
 						if (ParseDate2(Tokens[1], Month, Day, Year))
@@ -761,6 +619,7 @@
 			{
 				// asctime - date = wkday SP date3 SP time SP 4DIGIT
 				int32 WkDay = ParseWkday(Tokens[0]);
+
 				if (WkDay > 0)
 				{
 					if (ParseDate3(Tokens[1], Tokens[2], Month, Day))
@@ -780,6 +639,7 @@
 			{
 				// convert the tokens to numbers
 				OutDateTime = FDateTime(Year, Month, Day, Hour, Minute, Second, Millisecond);
+
 				return true;
 			}
 		}
@@ -788,6 +648,138 @@
 	return false;
 }
 
+
+bool FDateTime::ParseIso8601(const TCHAR* DateTimeString, FDateTime& OutDateTime)
+{
+	// DateOnly: YYYY-MM-DD
+	// DateTime: YYYY-mm-ddTHH:MM:SS(.ssss)(Z|+th:tm|-th:tm)
+
+	const TCHAR* Ptr = DateTimeString;
+	TCHAR* Next = nullptr;
+
+	int32 Year = 0, Month = 0, Day = 0;
+	int32 Hour = 0, Minute = 0, Second = 0, Millisecond = 0;
+	int32 TzHour = 0, TzMinute = 0;
+
+	// get date
+	Year = FCString::Strtoi(Ptr, &Next, 10);
+
+	if ((Next <= Ptr) || (*Next == TCHAR('\0')))
+	{
+		return false;
+	}
+
+	Ptr = Next + 1; // skip separator
+	Month = FCString::Strtoi(Ptr, &Next, 10);
+
+	if ((Next <= Ptr) || (*Next == TCHAR('\0')))
+	{
+		return false;
+	}
+
+	Ptr = Next + 1; // skip separator
+	Day = FCString::Strtoi(Ptr, &Next, 10);
+
+	if (Next <= Ptr)
+	{
+		return false;
+	}
+
+	// check whether this is date and time
+	if (*Next == TCHAR('T'))
+	{
+		Ptr = Next + 1;
+
+		// parse time
+		Hour = FCString::Strtoi(Ptr, &Next, 10);
+
+		if ((Next <= Ptr) || (*Next == TCHAR('\0')))
+		{
+			return false;
+		}
+
+		Ptr = Next + 1; // skip separator
+		Minute = FCString::Strtoi(Ptr, &Next, 10);
+
+		if ((Next <= Ptr) || (*Next == TCHAR('\0')))
+		{
+			return false;
+		}
+
+		Ptr = Next + 1; // skip separator
+		Second = FCString::Strtoi(Ptr, &Next, 10);
+
+		if (Next <= Ptr)
+		{
+			return false;
+		}
+
+		// check for milliseconds
+		if (*Next == TCHAR('.'))
+		{
+			Ptr = Next + 1;
+			Millisecond = FCString::Strtoi(Ptr, &Next, 10);
+
+			// should be no more than 3 digits
+			if ((Next <= Ptr) || (Next > Ptr + 3))
+			{
+				return false;
+			}
+
+			for (int32 Digits = Next - Ptr; Digits < 3; ++Digits)
+			{
+				Millisecond *= 10;
+			}
+		}
+
+		// see if the timezone offset is included
+		if (*Next == TCHAR('+') || *Next == TCHAR('-'))
+		{
+			// include the separator since it's + or -
+			Ptr = Next;
+
+			// parse the timezone offset
+			TzHour = FCString::Strtoi(Ptr, &Next, 10);
+
+			if ((Next <= Ptr) || (*Next == TCHAR('\0')))
+			{
+				return false;
+			}
+
+			Ptr = Next + 1; // skip separator
+			TzMinute = FCString::Strtoi(Ptr, &Next, 10);
+
+			if (Next <= Ptr)
+			{
+				return false;
+			}
+		}
+		else if ((*Next != TCHAR('\0')) && (*Next != TCHAR('Z')))
+		{
+			return false;
+		}
+	}
+	else if (*Next != TCHAR('\0'))
+	{
+		return false;
+	}
+
+	if (!Validate(Year, Month, Day, Hour, Minute, Second, Millisecond))
+	{
+		return false;
+	}
+
+	FDateTime Final(Year, Month, Day, Hour, Minute, Second, Millisecond);
+
+	// adjust for the timezone (bringing the DateTime into UTC)
+	int32 TzOffsetMinutes = (TzHour < 0) ? TzHour * 60 - TzMinute : TzHour * 60 + TzMinute;
+	Final -= FTimespan::FromMinutes(TzOffsetMinutes);
+	OutDateTime = Final;
+
+	return true;
+}
+
+
 FDateTime FDateTime::UtcNow()
 {
 	int32 Year, Month, Day, DayOfWeek;
@@ -799,7 +791,7 @@
 }
 
 
-bool FDateTime::Validate( int32 Year, int32 Month, int32 Day, int32 Hour, int32 Minute, int32 Second, int32 Millisecond )
+bool FDateTime::Validate(int32 Year, int32 Month, int32 Day, int32 Hour, int32 Minute, int32 Second, int32 Millisecond)
 {
 	return (Year >= 1) && (Year <= 9999) &&
 		(Month >= 1) && (Month <= 12) &&
@@ -810,16 +802,3 @@
 		(Millisecond >= 0) && (Millisecond <= 999);
 }
 
-
-/* FDateTime friend functions
- *****************************************************************************/
-
-FArchive& operator<<( FArchive& Ar, FDateTime& DateTime )
-{
-	return Ar << DateTime.Ticks;
-}
-
-uint32 GetTypeHash(const FDateTime& DateTime)
-{
-	return GetTypeHash(DateTime.Ticks);
-}
diff --git a/Engine/Source/Runtime/Core/Private/Misc/FileHelper.cpp b/Engine/Source/Runtime/Core/Private/Misc/FileHelper.cpp
--- a/Engine/Source/Runtime/Core/Private/Misc/FileHelper.cpp
+++ b/Engine/Source/Runtime/Core/Private/Misc/FileHelper.cpp
@@ -517,18 +517,34 @@ void FMaintenance::DeleteOldLogs()
 {
 	int32 PurgeLogsDays = -1; // -1 means don't delete old files
 	int32 MaxLogFilesOnDisk = -1; // -1 means keep all files
+
 	GConfig->GetInt(TEXT("LogFiles"), TEXT("PurgeLogsDays"), PurgeLogsDays, GEngineIni);
 	GConfig->GetInt(TEXT("LogFiles"), TEXT("MaxLogFilesOnDisk"), MaxLogFilesOnDisk, GEngineIni);
+
 	if (PurgeLogsDays >= 0 || MaxLogFilesOnDisk >= 0)
 	{
-		// get a list of files in the log dir
-		TArray<FString> Files;
-		IFileManager::Get().FindFiles(Files, *FString::Printf(TEXT("%s*.*"), *FPaths::ProjectLogDir()), true, false);
-		for (FString& Filename : Files)
+		// get list of files in the log directory (grouped by log name)
+		TMap<FString, TArray<FString>> LogToPaths;
 		{
-			Filename = FPaths::ProjectLogDir() / Filename;
+			TArray<FString> Files;
+			IFileManager::Get().FindFiles(Files, *FString::Printf(TEXT("%s*.*"), *FPaths::ProjectLogDir()), true, false);
+
+			for (FString& Filename : Files)
+			{
+				const int32 BackupPostfixIndex = Filename.Find(BACKUP_LOG_FILENAME_POSTFIX);
+
+				if (BackupPostfixIndex >= 0)
+				{
+					const FString LogName = Filename.Left(BackupPostfixIndex);
+					TArray<FString>& FilePaths = LogToPaths.FindOrAdd(LogName);
+					FilePaths.Add(FPaths::ProjectLogDir() / Filename);
+				}
+			}
 		}
 
+		// delete old log files in each group
+		double MaxFileAgeSeconds = 60.0 * 60.0 * 24.0 * double(PurgeLogsDays);
+
 		struct FSortByDateNewestFirst
 		{
 			bool operator()(const FString& A, const FString& B) const
@@ -538,30 +554,37 @@ void FMaintenance::DeleteOldLogs()
 				return TimestampB < TimestampA;
 			}
 		};
-		Files.Sort(FSortByDateNewestFirst());
 
-		// delete all those with the backup text in their name and that are older than the specified number of days
-		double MaxFileAgeSeconds = 60.0 * 60.0 * 24.0 * double(PurgeLogsDays);
-		for (int32 FileIndex = Files.Num() - 1; FileIndex >= 0; --FileIndex)
+		for (TPair<FString, TArray<FString>>& Pair : LogToPaths)
 		{
-			const FString& Filename = Files[FileIndex];
-			if (FOutputDeviceFile::IsBackupCopy(*Filename) && IFileManager::Get().GetFileAgeSeconds(*Filename) > MaxFileAgeSeconds)
+			TArray<FString>& FilePaths = Pair.Value;
+
+			// sort the file paths by date
+			FilePaths.Sort(FSortByDateNewestFirst());
+
+			// delete files that are older than the desired number of days
+			for (int32 PathIndex = FilePaths.Num() - 1; PathIndex >= 0; --PathIndex)
 			{
-				UE_LOG(LogStreaming, Log, TEXT("Deleting old log file %s"), *Filename);
-				IFileManager::Get().Delete(*Filename);
-				Files.RemoveAt(FileIndex);
+				const FString& FilePath = FilePaths[PathIndex];
+
+				if (IFileManager::Get().GetFileAgeSeconds(*FilePath) > MaxFileAgeSeconds)
+				{
+					UE_LOG(LogStreaming, Log, TEXT("Deleting old log file %s"), *FilePath);
+					IFileManager::Get().Delete(*FilePath);
+					FilePaths.RemoveAt(PathIndex);
+				}
 			}
-		}
 
-		// If required, trim the number of files on disk
-		if (MaxLogFilesOnDisk >= 0 && Files.Num() > MaxLogFilesOnDisk)
-		{
-			for (int32 FileIndex = Files.Num() - 1; FileIndex >= 0 && Files.Num() > MaxLogFilesOnDisk; --FileIndex)
+			// trim the number of files on disk if desired
+			if (MaxLogFilesOnDisk >= 0 && FilePaths.Num() > MaxLogFilesOnDisk)
 			{
-				if (FOutputDeviceFile::IsBackupCopy(*Files[FileIndex]))
+				for (int32 PathIndex = FilePaths.Num() - 1; PathIndex >= 0 && FilePaths.Num() > MaxLogFilesOnDisk; --PathIndex)
 				{
-					IFileManager::Get().Delete(*Files[FileIndex]);
-					Files.RemoveAt(FileIndex);
+					if (FOutputDeviceFile::IsBackupCopy(*FilePaths[PathIndex]))
+					{
+						IFileManager::Get().Delete(*FilePaths[PathIndex]);
+						FilePaths.RemoveAt(PathIndex);
+					}
 				}
 			}
 		}
diff --git a/Engine/Source/Runtime/Core/Private/Misc/OutputDeviceFile.cpp b/Engine/Source/Runtime/Core/Private/Misc/OutputDeviceFile.cpp
--- a/Engine/Source/Runtime/Core/Private/Misc/OutputDeviceFile.cpp
+++ b/Engine/Source/Runtime/Core/Private/Misc/OutputDeviceFile.cpp
@@ -354,10 +354,9 @@ void FOutputDeviceFile::CreateBackupCopy(const TCHAR* Filename)
 {
 	if (IFileManager::Get().FileSize(Filename) > 0)
 	{
-		FString SystemTime = FDateTime::Now().ToString();
 		FString Name, Extension;
 		FString(Filename).Split(TEXT("."), &Name, &Extension, ESearchCase::CaseSensitive, ESearchDir::FromEnd);
-		FString BackupFilename = FString::Printf(TEXT("%s%s%s.%s"), *Name, BACKUP_LOG_FILENAME_POSTFIX, *SystemTime, *Extension);
+		FString BackupFilename = FString::Printf(TEXT("%s%s%s.%s"), *Name, BACKUP_LOG_FILENAME_POSTFIX, *GSystemStartTime, *Extension);
 		IFileManager::Get().Copy(*BackupFilename, Filename, false);
 	}
 }
diff --git a/Engine/Source/Runtime/Core/Private/Misc/Timespan.cpp b/Engine/Source/Runtime/Core/Private/Misc/Timespan.cpp
--- a/Engine/Source/Runtime/Core/Private/Misc/Timespan.cpp
+++ b/Engine/Source/Runtime/Core/Private/Misc/Timespan.cpp
@@ -5,18 +5,19 @@
 #include "Containers/UnrealString.h"
 #include "UObject/PropertyPortFlags.h"
 
+
 /* FTimespan interface
  *****************************************************************************/
 
 bool FTimespan::ExportTextItem(FString& ValueStr, FTimespan const& DefaultValue, UObject* Parent, int32 PortFlags, UObject* ExportRootScope) const
 {
-	if (0 != (PortFlags & EPropertyPortFlags::PPF_ExportCpp))
+	if ((PortFlags & EPropertyPortFlags::PPF_ExportCpp) != 0)
 	{
 		ValueStr += FString::Printf(TEXT("FTimespan(0x%016X)"), Ticks);
 		return true;
 	}
 
-	ValueStr += ToString(TEXT("%N%d.%h:%m:%s.%f"));
+	ValueStr += ToString(TEXT("%D.%h:%m:%s.%n"));
 
 	return true;
 }
@@ -24,8 +25,21 @@ bool FTimespan::ExportTextItem(FString& ValueStr, FTimespan const& DefaultValue,
 
 bool FTimespan::ImportTextItem(const TCHAR*& Buffer, int32 PortFlags, UObject* Parent, FOutputDevice* ErrorText)
 {
-	// @todo gmp: implement FTimespan::ImportTextItem
-	return false;
+	const int32 ExportTimespanLen = 27;
+
+	if (FPlatformString::Strlen(Buffer) < ExportTimespanLen)
+	{
+		return false;
+	}
+
+	if (!Parse(FString(Buffer).Left(ExportTimespanLen), *this))
+	{
+		return false;
+	}
+
+	Buffer += ExportTimespanLen;
+
+	return true;
 }
 
 
@@ -41,10 +55,10 @@ FString FTimespan::ToString() const
 {
 	if (GetDays() == 0)
 	{
-		return ToString(TEXT("%n%h:%m:%s.%f"));
+		return ToString(TEXT("%h:%m:%s.%f"));
 	}
 
-	return ToString(TEXT("%n%d.%h:%m:%s.%f"));
+	return ToString(TEXT("%d.%h:%m:%s.%f"));
 }
 
 
@@ -52,27 +66,23 @@ FString FTimespan::ToString(const TCHAR* Format) const
 {
 	FString Result;
 
+	Result += (Ticks < 0) ? TCHAR('-') : TCHAR('+');
+
 	while (*Format != TCHAR('\0'))
 	{
 		if ((*Format == TCHAR('%')) && (*++Format != TCHAR('\0')))
 		{
 			switch (*Format)
 			{
-			case TCHAR('n'): if (Ticks < 0) Result += TCHAR('-'); break;
-			case TCHAR('N'): Result += (Ticks < 0) ? TCHAR('-') : TCHAR('+'); break;
 			case TCHAR('d'): Result += FString::Printf(TEXT("%i"), FMath::Abs(GetDays())); break;
+			case TCHAR('D'): Result += FString::Printf(TEXT("%08i"), FMath::Abs(GetDays())); break;
 			case TCHAR('h'): Result += FString::Printf(TEXT("%02i"), FMath::Abs(GetHours())); break;
 			case TCHAR('m'): Result += FString::Printf(TEXT("%02i"), FMath::Abs(GetMinutes())); break;
 			case TCHAR('s'): Result += FString::Printf(TEXT("%02i"), FMath::Abs(GetSeconds())); break;
-			case TCHAR('f'): Result += FString::Printf(TEXT("%03i"), FMath::Abs(GetMilliseconds())); break;
-			case TCHAR('D'): Result += FString::Printf(TEXT("%f"), FMath::Abs(GetTotalDays())); break;
-			case TCHAR('H'): Result += FString::Printf(TEXT("%f"), FMath::Abs(GetTotalHours())); break;
-			case TCHAR('M'): Result += FString::Printf(TEXT("%f"), FMath::Abs(GetTotalMinutes())); break;
-			case TCHAR('S'): Result += FString::Printf(TEXT("%f"), FMath::Abs(GetTotalSeconds())); break;
-			case TCHAR('F'): Result += FString::Printf(TEXT("%f"), FMath::Abs(GetTotalMilliseconds())); break;
-
+			case TCHAR('f'): Result += FString::Printf(TEXT("%03i"), FMath::Abs(GetFractionMilli())); break;
+			case TCHAR('u'): Result += FString::Printf(TEXT("%06i"), FMath::Abs(GetFractionMicro())); break;
+			case TCHAR('n'): Result += FString::Printf(TEXT("%09i"), FMath::Abs(GetFractionNano())); break;
 			default:
-
 				Result += *Format;
 			}
 		}
@@ -91,77 +101,78 @@ FString FTimespan::ToString(const TCHAR* Format) const
 /* FTimespan static interface
  *****************************************************************************/
 
-FTimespan FTimespan::FromDays(double Days)
-{
-	check((Days >= MinValue().GetTotalDays()) && (Days <= MaxValue().GetTotalDays()));
-
-	return FTimespan(Days * ETimespan::TicksPerDay);
-}
-
-
-FTimespan FTimespan::FromHours(double Hours)
-{
-	check((Hours >= MinValue().GetTotalHours()) && (Hours <= MaxValue().GetTotalHours()));
-
-	return FTimespan(Hours * ETimespan::TicksPerHour);
-}
-
-
-FTimespan FTimespan::FromMicroseconds(double Microseconds)
-{
-	check((Microseconds >= MinValue().GetTotalMicroseconds()) && (Microseconds <= MaxValue().GetTotalMicroseconds()));
-
-	return FTimespan(Microseconds * ETimespan::TicksPerMicrosecond);
-}
-
-
-FTimespan FTimespan::FromMilliseconds(double Milliseconds)
+bool FTimespan::Parse(const FString& TimespanString, FTimespan& OutTimespan)
 {
-	check((Milliseconds >= MinValue().GetTotalMilliseconds()) && (Milliseconds <= MaxValue().GetTotalMilliseconds()));
-
-	return FTimespan(Milliseconds * ETimespan::TicksPerMillisecond);
-}
+	// @todo gmp: implement stricter FTimespan parsing; this implementation is too forgiving
 
+	// get string tokens
+	const bool HasFractional = TimespanString.Contains(TEXT("."));
+	FString TokenString = (HasFractional) ? TimespanString.Replace(TEXT("."), TEXT(":")) : TimespanString;
+	TokenString.ReplaceInline(TEXT(","), TEXT(":"));
 
-FTimespan FTimespan::FromMinutes(double Minutes)
-{
-	check((Minutes >= MinValue().GetTotalMinutes()) && (Minutes <= MaxValue().GetTotalMinutes()));
-
-	return FTimespan(Minutes * ETimespan::TicksPerMinute);
-}
+	const bool Negative = TokenString.StartsWith(TEXT("-"));
+	TokenString.ReplaceInline(TEXT("-"), TEXT(":"));
+	TokenString.ReplaceInline(TEXT("+"), TEXT(":"));
 
+	TArray<FString> Tokens;
+	TokenString.ParseIntoArray(Tokens, TEXT(":"), true);
 
-FTimespan FTimespan::FromSeconds(double Seconds)
-{
-	check((Seconds >= MinValue().GetTotalSeconds()) && (Seconds <= MaxValue().GetTotalSeconds()));
+	if (!HasFractional)
+	{
+		Tokens.AddDefaulted();
+	}
 
-	return FTimespan(Seconds * ETimespan::TicksPerSecond);
-}
+	// poor man's token verification
+	for (const FString& Token : Tokens)
+	{
+		if (!Token.IsEmpty() && !Token.IsNumeric())
+		{
+			return false;
+		}
+	}
 
+	// add missing tokens
+	if (Tokens.Num() < 5)
+	{
+		Tokens.InsertDefaulted(0, 5 - Tokens.Num());
+	}
+	else if (Tokens.Num() > 5)
+	{
+		return false;
+	}
 
-bool FTimespan::Parse(const FString& TimespanString, FTimespan& OutTimespan)
-{
-	// @todo gmp: implement stricter FTimespan parsing; this implementation is too forgiving.
-	FString TokenString = TimespanString.Replace(TEXT("."), TEXT(":"));
-	TokenString.ReplaceInline(TEXT(","), TEXT(":"));
+	// pad fractional token with zeros
+	if (HasFractional)
+	{
+		const int32 FractionalLen = Tokens[4].Len();
 
-	bool Negative = TokenString.StartsWith(TEXT("-"));
-	TokenString.ReplaceInline(TEXT("-"), TEXT(":"), ESearchCase::CaseSensitive);
+		if (FractionalLen > 9)
+		{
+			Tokens[4] = Tokens[4].Left(9);
+		}
+		else if (FractionalLen < 9)
+		{
+			Tokens[4] += FString(TEXT("000000000")).Left(9 - FractionalLen);
+		}
+	}
 
-	TArray<FString> Tokens;
-	TokenString.ParseIntoArray(Tokens, TEXT(":"), true);
+	const int32 Days = FCString::Atoi(*Tokens[0]);
+	const int32 Hours = FCString::Atoi(*Tokens[1]);
+	const int32 Minutes = FCString::Atoi(*Tokens[2]);
+	const int32 Seconds = FCString::Atoi(*Tokens[3]);
+	const int32 FractionNano = FCString::Atoi(*Tokens[4]);
 
-	if (Tokens.Num() == 4)
+	if ((Days > (ETimespan::MaxTicks / ETimespan::TicksPerDay) - 1))
 	{
-		Tokens.Insert(TEXT("0"), 0);
+		return false;
 	}
 
-	if (Tokens.Num() != 5)
+	if ((Hours > 23) || (Minutes > 59) || (Seconds > 59) || (FractionNano > 999999999))
 	{
 		return false;
 	}
 
-	OutTimespan.Assign(FCString::Atoi(*Tokens[0]), FCString::Atoi(*Tokens[1]), FCString::Atoi(*Tokens[2]), FCString::Atoi(*Tokens[3]), FCString::Atoi(*Tokens[4]), 0);
+	OutTimespan.Assign(Days, Hours, Minutes, Seconds, FractionNano);
 
 	if (Negative)
 	{
@@ -172,28 +183,20 @@ bool FTimespan::Parse(const FString& TimespanString, FTimespan& OutTimespan)
 }
 
 
-/* FTimespan friend functions
+/* FTimespan implementation
  *****************************************************************************/
 
-FArchive& operator<<(FArchive& Ar, FTimespan& Timespan)
+void FTimespan::Assign(int32 Days, int32 Hours, int32 Minutes, int32 Seconds, int32 FractionNano)
 {
-	return Ar << Timespan.Ticks;
-}
-
-
-uint32 GetTypeHash(const FTimespan& Timespan)
-{
-	return GetTypeHash(Timespan.Ticks);
-}
+	int64 TotalTicks = 0;
 
+	TotalTicks += Days * ETimespan::TicksPerDay;
+	TotalTicks += Hours * ETimespan::TicksPerHour;
+	TotalTicks += Minutes * ETimespan::TicksPerMinute;
+	TotalTicks += Seconds * ETimespan::TicksPerSecond;
+	TotalTicks += FractionNano / ETimespan::NanosecondsPerTick;
 
-/* FTimespan implementation
- *****************************************************************************/
-
-void FTimespan::Assign(int32 Days, int32 Hours, int32 Minutes, int32 Seconds, int32 Milliseconds, int32 Microseconds)
-{
-	int64 TotalTicks = ETimespan::TicksPerMicrosecond * (1000 * (1000 * (60 * 60 * 24 * (int64)Days + 60 * 60 * (int64)Hours + 60 * (int64)Minutes + (int64)Seconds) + (int64)Milliseconds) + (int64)Microseconds);
-	check((TotalTicks >= MinValue().GetTicks()) && (TotalTicks <= MaxValue().GetTicks()));
+	check((TotalTicks >= ETimespan::MinTicks) && (TotalTicks <= ETimespan::MaxTicks));
 
 	Ticks = TotalTicks;
 }
diff --git a/Engine/Source/Runtime/Core/Private/Tests/Misc/TimespanTest.cpp b/Engine/Source/Runtime/Core/Private/Tests/Misc/TimespanTest.cpp
--- a/Engine/Source/Runtime/Core/Private/Tests/Misc/TimespanTest.cpp
+++ b/Engine/Source/Runtime/Core/Private/Tests/Misc/TimespanTest.cpp
@@ -5,11 +5,13 @@
 #include "Misc/Timespan.h"
 #include "Misc/AutomationTest.h"
 
+
 #if WITH_DEV_AUTOMATION_TESTS
 
 IMPLEMENT_SIMPLE_AUTOMATION_TEST(FTimespanTest, "System.Core.Misc.Timespan", EAutomationTestFlags::ApplicationContextMask | EAutomationTestFlags::SmokeFilter)
 
-bool FTimespanTest::RunTest( const FString& Parameters )
+
+bool FTimespanTest::RunTest(const FString& Parameters)
 {
 	// constructors must create equal objects
 	{
@@ -18,25 +20,26 @@ bool FTimespanTest::RunTest( const FString& Parameters )
 		FTimespan ts1_3 = FTimespan(0, 3, 2, 1, 0);
 
 		TestEqual(TEXT("Constructors must create equal objects (Hours/Minutes/Seconds vs. Days/Hours/Minutes/Seconds)"), ts1_1, ts1_2);
-		TestEqual(TEXT("Constructors must create equal objects (Hours/Minutes/Seconds vs. Days/Hours/Minutes/Seconds/Milliseconds)"), ts1_1, ts1_3);
+		TestEqual(TEXT("Constructors must create equal objects (Hours/Minutes/Seconds vs. Days/Hours/Minutes/Seconds/FractionNano)"), ts1_1, ts1_3);
 	}
 
 	// component getters must return correct values
 	{
-		FTimespan ts2_1 = FTimespan(1, 2, 3, 4, 5, 6);
+		FTimespan ts2_1 = FTimespan(1, 2, 3, 4, 123456789);
 
 		TestEqual(TEXT("Component getters must return correct values (Days)"), ts2_1.GetDays(), 1);
 		TestEqual(TEXT("Component getters must return correct values (Hours)"), ts2_1.GetHours(), 2);
 		TestEqual(TEXT("Component getters must return correct values (Minutes)"), ts2_1.GetMinutes(), 3);
 		TestEqual(TEXT("Component getters must return correct values (Seconds)"), ts2_1.GetSeconds(), 4);
-		TestEqual(TEXT("Component getters must return correct values (Milliseconds)"), ts2_1.GetMilliseconds(), 5);
-		TestEqual(TEXT("Component getters must return correct values (Microseconds)"), ts2_1.GetMicroseconds(), 6);
+		TestEqual(TEXT("Component getters must return correct values (FractionMilli)"), ts2_1.GetFractionMilli(), 123);
+		TestEqual(TEXT("Component getters must return correct values (FractionMicro)"), ts2_1.GetFractionMicro(), 123456);
+		TestEqual(TEXT("Component getters must return correct values (FractionNano)"), ts2_1.GetFractionNano(), 123456700);
 	}
 
 	// durations of positive and negative time spans must match
 	{
-		FTimespan ts3_1 = FTimespan(1, 2, 3, 4, 5, 6);
-		FTimespan ts3_2 = FTimespan(-1, -2, -3, -4, -5, -6);
+		FTimespan ts3_1 = FTimespan(1, 2, 3, 4, 123456789);
+		FTimespan ts3_2 = FTimespan(-1, -2, -3, -4, -123456789);
 
 		TestEqual(TEXT("Durations of positive and negative time spans must match"), ts3_1.GetDuration(), ts3_2.GetDuration());
 	}
@@ -53,19 +56,43 @@ bool FTimespanTest::RunTest( const FString& Parameters )
 
 	// string conversions must return correct strings
 	{
-		FTimespan ts5_1 = FTimespan(1, 2, 3, 4, 5);
+		FTimespan ts5_1 = FTimespan(1, 2, 3, 4, 123456789);
 
-		TestEqual<FString>(TEXT("String conversion (Default)"), ts5_1.ToString(), TEXT("1.02:03:04.005"));
-		TestEqual<FString>(TEXT("String conversion (%n%d.%h:%m:%s.%f)"), ts5_1.ToString(TEXT("%n%d.%h:%m:%s.%f")), TEXT("1.02:03:04.005"));
+		TestEqual<FString>(TEXT("String conversion (Default)"), ts5_1.ToString(), TEXT("+1.02:03:04.123"));
+		TestEqual<FString>(TEXT("String conversion (%d.%h:%m:%s.%f)"), ts5_1.ToString(TEXT("%d.%h:%m:%s.%f")), TEXT("+1.02:03:04.123"));
+		TestEqual<FString>(TEXT("String conversion (%d.%h:%m:%s.%u)"), ts5_1.ToString(TEXT("%d.%h:%m:%s.%u")), TEXT("+1.02:03:04.123456"));
+		TestEqual<FString>(TEXT("String conversion (%D.%h:%m:%s.%n)"), ts5_1.ToString(TEXT("%D.%h:%m:%s.%n")), TEXT("+00000001.02:03:04.123456700"));
 	}
 
 	// parsing valid strings must succeed
 	{
-		FTimespan ts6_1 = FTimespan(1, 2, 3, 4, 5);
-		FTimespan ts6_2;
+		FTimespan ts6_t;
+
+		FTimespan ts6_1 = FTimespan(1, 2, 3, 4, 123000000);
+		FTimespan ts6_2 = FTimespan(1, 2, 3, 4, 123456000);
+		FTimespan ts6_3 = FTimespan(1, 2, 3, 4, 123456700);
+
+		TestTrue(TEXT("Parsing valid strings must succeed (+1.02:03:04.123)"), FTimespan::Parse(TEXT("+1.02:03:04.123"), ts6_t));
+		TestEqual(TEXT("Parsing valid strings must result in correct values (+1.02:03:04.123)"), ts6_t, ts6_1);
+
+		TestTrue(TEXT("Parsing valid strings must succeed (+1.02:03:04.123456)"), FTimespan::Parse(TEXT("+1.02:03:04.123456"), ts6_t));
+		TestEqual(TEXT("Parsing valid strings must result in correct values (+1.02:03:04.123456)"), ts6_t, ts6_2);
 
-		TestTrue(TEXT("Parsing valid strings must succeed (1.02:03:04.005)"), FTimespan::Parse(TEXT("1.02:03:04.005"), ts6_2));
-		TestEqual(TEXT("Parsing valid strings must result in correct values (1.02:03:04.005)"), ts6_2, ts6_1);
+		TestTrue(TEXT("Parsing valid strings must succeed (+1.02:03:04.123456789)"), FTimespan::Parse(TEXT("+1.02:03:04.123456789"), ts6_t));
+		TestEqual(TEXT("Parsing valid strings must result in correct values (+1.02:03:04.123456789)"), ts6_t, ts6_3);
+
+		FTimespan ts6_4 = FTimespan(-1, -2, -3, -4, -123000000);
+		FTimespan ts6_5 = FTimespan(-1, -2, -3, -4, -123456000);
+		FTimespan ts6_6 = FTimespan(-1, -2, -3, -4, -123456700);
+
+		TestTrue(TEXT("Parsing valid strings must succeed (-1.02:03:04.123)"), FTimespan::Parse(TEXT("-1.02:03:04.123"), ts6_t));
+		TestEqual(TEXT("Parsing valid strings must result in correct values (-1.02:03:04.123)"), ts6_t, ts6_4);
+
+		TestTrue(TEXT("Parsing valid strings must succeed (-1.02:03:04.123456)"), FTimespan::Parse(TEXT("-1.02:03:04.123456"), ts6_t));
+		TestEqual(TEXT("Parsing valid strings must result in correct values (-1.02:03:04.123456)"), ts6_t, ts6_5);
+
+		TestTrue(TEXT("Parsing valid strings must succeed (-1.02:03:04.123456789)"), FTimespan::Parse(TEXT("-1.02:03:04.123456789"), ts6_t));
+		TestEqual(TEXT("Parsing valid strings must result in correct values (-1.02:03:04.123456789)"), ts6_t, ts6_6);
 	}
 
 	// parsing invalid strings must fail
@@ -80,4 +107,5 @@ bool FTimespanTest::RunTest( const FString& Parameters )
 	return true;
 }
 
+
 #endif //WITH_DEV_AUTOMATION_TESTS
diff --git a/Engine/Source/Runtime/Core/Public/Android/AndroidJavaMediaPlayer.h b/Engine/Source/Runtime/Core/Public/Android/AndroidJavaMediaPlayer.h
--- a/Engine/Source/Runtime/Core/Public/Android/AndroidJavaMediaPlayer.h
+++ b/Engine/Source/Runtime/Core/Public/Android/AndroidJavaMediaPlayer.h
@@ -4,6 +4,8 @@
 
 #include "CoreMinimal.h"
 #include "AndroidJava.h"
+#include "RHI.h"
+#include "RHIResources.h"
 
 // Wrapper for com/epicgames/ue4/MediaPlayer*.java.
 class FJavaAndroidMediaPlayer : public FJavaClassObject
@@ -42,17 +44,20 @@ public:
 	};
 
 public:
-	FJavaAndroidMediaPlayer(bool vulkanRenderer);
+	FJavaAndroidMediaPlayer(bool swizzlePixels, bool vulkanRenderer);
 	int32 GetDuration();
 	void Reset();
 	void Stop();
 	int32 GetCurrentPosition();
 	bool IsLooping();
 	bool IsPlaying();
+	bool IsPrepared();
+	bool DidComplete();
 	bool SetDataSource(const FString & Url);
 	bool SetDataSource(const FString& MoviePathOnDevice, int64 offset, int64 size);
 	bool SetDataSource(jobject AssetMgr, const FString& AssetPath, int64 offset, int64 size);
 	bool Prepare();
+	bool PrepareAsync();
 	void SeekTo(int32 Milliseconds);
 	void SetLooping(bool Looping);
 	void Release();
@@ -69,6 +74,8 @@ public:
 	bool GetCaptionTracks(TArray<FCaptionTrack>& CaptionTracks);
 	bool GetVideoTracks(TArray<FVideoTrack>& VideoTracks);
 	bool DidResolutionChange();
+	int32 GetExternalTextureId();
+	bool UpdateVideoFrame(int32 ExternalTextureId, int32 *CurrentPosition, bool *bRegionChanged);
 
 private:
 	static FName GetClassName();
@@ -81,10 +88,13 @@ private:
 	FJavaClassMethod GetCurrentPositionMethod;
 	FJavaClassMethod IsLoopingMethod;
 	FJavaClassMethod IsPlayingMethod;
+	FJavaClassMethod IsPreparedMethod;
+	FJavaClassMethod DidCompleteMethod;
 	FJavaClassMethod SetDataSourceURLMethod;
 	FJavaClassMethod SetDataSourceFileMethod;
 	FJavaClassMethod SetDataSourceAssetMethod;
 	FJavaClassMethod PrepareMethod;
+	FJavaClassMethod PrepareAsyncMethod;
 	FJavaClassMethod SeekToMethod;
 	FJavaClassMethod SetLoopingMethod;
 	FJavaClassMethod ReleaseMethod;
@@ -101,6 +111,18 @@ private:
 	FJavaClassMethod GetCaptionTracksMethod;
 	FJavaClassMethod GetVideoTracksMethod;
 	FJavaClassMethod DidResolutionChangeMethod;
+	FJavaClassMethod GetExternalTextureIdMethod;
+	FJavaClassMethod UpdateVideoFrameMethod;
+
+	// FrameUpdateInfo member field ids
+	jclass FrameUpdateInfoClass;
+	jfieldID FrameUpdateInfo_CurrentPosition;
+	jfieldID FrameUpdateInfo_FrameReady;
+	jfieldID FrameUpdateInfo_RegionChanged;
+	jfieldID FrameUpdateInfo_UScale;
+	jfieldID FrameUpdateInfo_UOffset;
+	jfieldID FrameUpdateInfo_VScale;
+	jfieldID FrameUpdateInfo_VOffset;
 
 	// AudioDeviceInfo member field ids
 	jclass AudioTrackInfoClass;
@@ -128,4 +150,37 @@ private:
 	jfieldID VideoTrackInfo_Width;
 	jfieldID VideoTrackInfo_Height;
 	jfieldID VideoTrackInfo_FrameRate;
+
+	FTextureRHIRef VideoTexture;
+	bool bVideoTextureValid;
+
+	float UScale, UOffset;
+	float VScale, VOffset;
+
+public:
+	FTextureRHIRef GetVideoTexture()
+	{
+		return VideoTexture;
+	}
+
+	void SetVideoTexture(FTextureRHIRef Texture)
+	{
+		VideoTexture = Texture;
+	}
+
+	void SetVideoTextureValid(bool Condition)
+	{
+		bVideoTextureValid = Condition;
+	}
+
+	bool IsVideoTextureValid()
+	{
+		return bVideoTextureValid;
+	}
+
+	float GetUScale() { return UScale; }
+	float GetUOffset() { return UOffset; }
+
+	float GetVScale() { return VScale; }
+	float GetVOffset() { return VOffset; }
 };
diff --git a/Engine/Source/Runtime/Core/Public/Containers/CircularQueue.h b/Engine/Source/Runtime/Core/Public/Containers/CircularQueue.h
--- a/Engine/Source/Runtime/Core/Public/Containers/CircularQueue.h
+++ b/Engine/Source/Runtime/Core/Public/Containers/CircularQueue.h
@@ -25,10 +25,10 @@ public:
 	/**
 	 * Default constructor.
 	 *
-	 * @param Size The number of elements that the queue can hold (will be rounded up to the next power of 2).
+	 * @param CapacityPlusOne The number of elements that the queue can hold (will be rounded up to the next power of 2).
 	 */
-	TCircularQueue(uint32 Size)
-		: Buffer(Size)
+	TCircularQueue(uint32 CapacityPlusOne)
+		: Buffer(CapacityPlusOne)
 		, Head(0)
 		, Tail(0)
 	{ }
@@ -41,6 +41,9 @@ public:
 	/**
 	 * Gets the number of elements in the queue.
 	 *
+	 * The result reflects the calling thread's current view. Since no
+	 * locking is used, different threads may return different results.
+	 *
 	 * @return Number of queued elements.
 	 */
 	uint32 Count() const
@@ -60,6 +63,7 @@ public:
 	 *
 	 * @param OutElement Will contain the element if the queue is not empty.
 	 * @return true if an element has been returned, false if the queue was empty.
+	 * @note To be called only from consumer thread.
 	 */
 	bool Dequeue(ElementType& OutElement)
 	{
@@ -77,6 +81,7 @@ public:
 	/**
 	 * Empties the queue.
 	 *
+	 * @note To be called only from consumer thread.
 	 * @see IsEmpty
 	 */
 	void Empty()
@@ -89,6 +94,7 @@ public:
 	 *
 	 * @param Element The element to add.
 	 * @return true if the item was added, false if the queue was full.
+	 * @note To be called only from producer thread.
 	 */
 	bool Enqueue(const ElementType& Element)
 	{
@@ -108,6 +114,9 @@ public:
 	/**
 	 * Checks whether the queue is empty.
 	 *
+	 * The result reflects the calling thread's current view. Since no
+	 * locking is used, different threads may return different results.
+	 *
 	 * @return true if the queue is empty, false otherwise.
 	 * @see Empty, IsFull
 	 */
@@ -119,6 +128,9 @@ public:
 	/**
 	 * Checks whether the queue is full.
 	 *
+	 * The result reflects the calling thread's current view. Since no
+	 * locking is used, different threads may return different results.
+	 *
 	 * @return true if the queue is full, false otherwise.
 	 * @see IsEmpty
 	 */
@@ -132,6 +144,7 @@ public:
 	 *
 	 * @param OutItem Will contain the item if the queue is not empty.
 	 * @return true if an item has been returned, false if the queue was empty.
+	 * @note To be called only from consumer thread.
 	 */
 	bool Peek(ElementType& OutItem)
 	{
