Copying //UE4/Dev-Sequencer to //UE4/Main (Source: //UE4/Dev-Sequencer @ 3617642)

diff --git a/Engine/Source/Runtime/Core/Public/Containers/Map.h b/Engine/Source/Runtime/Core/Public/Containers/Map.h
--- a/Engine/Source/Runtime/Core/Public/Containers/Map.h
+++ b/Engine/Source/Runtime/Core/Public/Containers/Map.h
@@ -49,6 +49,7 @@ public:
 	}
 };
 
+
 /** An initializer type for keys that's passed to the pair set when adding a new key. */
 template <typename KeyInitType>
 class TKeyInitializer
@@ -59,8 +60,7 @@ public:
 	/** Initialization constructor. */
 	FORCEINLINE explicit TKeyInitializer(KeyInitType InKey)
 		: Key(InKey)
-	{
-	}
+	{ }
 
 	template <typename KeyType, typename ValueType>
 	operator TPair<KeyType, ValueType>() const
@@ -138,15 +138,13 @@ protected:
 	template<typename OtherSetAllocator>
 	TMapBase(TMapBase<KeyType, ValueType, OtherSetAllocator, KeyFuncs>&& Other)
 		: Pairs(MoveTemp(Other.Pairs))
-	{
-	}
+	{ }
 
 	/** Constructor for copying elements from a TMap with a different SetAllocator */
 	template<typename OtherSetAllocator>
 	TMapBase(const TMapBase<KeyType, ValueType, OtherSetAllocator, KeyFuncs>& Other)
 		: Pairs(Other.Pairs)
-	{
-	}
+	{ }
 
 	/** Assignment operator for moving elements from a TMap with a different SetAllocator */
 	template<typename OtherSetAllocator>
@@ -210,8 +208,12 @@ public:
 	}
 
 	/**
-	 * Removes all elements from the map, potentially leaving space allocated for an expected number of elements about to be added.
-	 * @param ExpectedNumElements - The number of elements about to be added to the set.
+	 * Removes all elements from the map.
+	 *
+	 * This method potentially leaves space allocated for an expected
+	 * number of elements about to be added.
+	 *
+	 * @param ExpectedNumElements The number of elements about to be added to the set.
 	 */
 	FORCEINLINE void Empty(int32 ExpectedNumElements = 0)
 	{
@@ -255,8 +257,9 @@ public:
 	}
 
 	/**
-	 * Returns the unique keys contained within this map
-	 * @param	OutKeys	- Upon return, contains the set of unique keys in this map.
+	 * Get the unique keys contained within this map.
+	 *
+	 * @param OutKeys Upon return, contains the set of unique keys in this map.
 	 * @return The number of unique keys in the map.
 	 */
 	template<typename Allocator> int32 GetKeys(TArray<KeyType, Allocator>& OutKeys) const
@@ -274,26 +277,33 @@ public:
 	}
 
 	/** 
-	 * Helper function to return the amount of memory allocated by this container 
-	 * @return number of bytes allocated by this container
+	 * Helper function to return the amount of memory allocated by this container .
+	 *
+	 * @return Number of bytes allocated by this container.
+	 * @see CountBytes
 	 */
 	FORCEINLINE uint32 GetAllocatedSize() const
 	{
 		return Pairs.GetAllocatedSize();
 	}
 
-	/** Tracks the container's memory use through an archive. */
+	/**
+	 * Track the container's memory use through an archive.
+	 *
+	 * @param Ar The archive to use.
+	 * @see GetAllocatedSize
+	 */
 	FORCEINLINE void CountBytes(FArchive& Ar)
 	{
 		Pairs.CountBytes(Ar);
 	}
 
 	/**
-	 * Sets the value associated with a key.
+	 * Set the value associated with a key.
 	 *
-	 * @param InKey - The key to associate the value with.
-	 * @param InValue - The value to associate with the key.
-	 * @return A reference to the value as stored in the map.  The reference is only valid until the next change to any key in the map.
+	 * @param InKey The key to associate the value with.
+	 * @param InValue The value to associate with the key.
+	 * @return A reference to the value as stored in the map. The reference is only valid until the next change to any key in the map.
 	 */
 	FORCEINLINE ValueType& Add(const KeyType&  InKey, const ValueType&  InValue) { return Emplace(                   InKey ,                    InValue ); }
 	FORCEINLINE ValueType& Add(const KeyType&  InKey,       ValueType&& InValue) { return Emplace(                   InKey , MoveTempIfPossible(InValue)); }
@@ -301,10 +311,10 @@ public:
 	FORCEINLINE ValueType& Add(      KeyType&& InKey,       ValueType&& InValue) { return Emplace(MoveTempIfPossible(InKey), MoveTempIfPossible(InValue)); }
 
 	/**
-	 * Sets a default value associated with a key.
+	 * Set a default value associated with a key.
 	 *
-	 * @param InKey - The key to associate the value with.
-	 * @return A reference to the value as stored in the map.  The reference is only valid until the next change to any key in the map.
+	 * @param InKey The key to associate the value with.
+	 * @return A reference to the value as stored in the map. The reference is only valid until the next change to any key in the map.
 	 */
 	FORCEINLINE ValueType& Add(const KeyType&  InKey) { return Emplace(                   InKey ); }
 	FORCEINLINE ValueType& Add(      KeyType&& InKey) { return Emplace(MoveTempIfPossible(InKey)); }
@@ -312,10 +322,9 @@ public:
 	/**
 	 * Sets the value associated with a key.
 	 *
-	 * @param InKey - The key to associate the value with.
-	 * @param InValue - The value to associate with the key.
-	 * @return A reference to the value as stored in the map.  The reference is only valid until the next change to any key in the map.
-	 */
+	 * @param InKey The key to associate the value with.
+	 * @param InValue The value to associate with the key.
+	 * @return A reference to the value as stored in the map. The reference is only valid until the next change to any key in the map.	 */
 	template <typename InitKeyType, typename InitValueType>
 	ValueType& Emplace(InitKeyType&& InKey, InitValueType&& InValue)
 	{
@@ -325,10 +334,10 @@ public:
 	}
 
 	/**
-	 * Sets a default value associated with a key.
+	 * Set a default value associated with a key.
 	 *
-	 * @param InKey - The key to associate the value with.
-	 * @return A reference to the value as stored in the map.  The reference is only valid until the next change to any key in the map.
+	 * @param InKey The key to associate the value with.
+	 * @return A reference to the value as stored in the map. The reference is only valid until the next change to any key in the map.
 	 */
 	template <typename InitKeyType>
 	ValueType& Emplace(InitKeyType&& InKey)
@@ -339,8 +348,9 @@ public:
 	}
 
 	/**
-	 * Removes all value associations for a key.
-	 * @param InKey - The key to remove associated values for.
+	 * Remove all value associations for a key.
+	 *
+	 * @param InKey The key to remove associated values for.
 	 * @return The number of values that were associated with the key.
 	 */
 	FORCEINLINE int32 Remove(KeyConstPointerType InKey)
@@ -350,10 +360,14 @@ public:
 	}
 
 	/**
-	 * Returns the key associated with the specified value.  The time taken is O(N) in the number of pairs.
-	 * @param	Value - The value to search for
-	 * @return	A pointer to the key associated with the specified value, or nullptr if the value isn't contained in this map.  The pointer
-	 *			is only valid until the next change to any key in the map.
+	 * Find the key associated with the specified value.
+	 *
+	 * The time taken is O(N) in the number of pairs.
+	 *
+	 * @param Value The value to search for
+	 * @return A pointer to the key associated with the specified value,
+	 *     or nullptr if the value isn't contained in this map. The pointer
+	 *     is only valid until the next change to any key in the map.
 	 */
 	const KeyType* FindKey(ValueInitType Value) const
 	{
@@ -368,9 +382,10 @@ public:
 	}
 
 	/**
-	 * Returns the value associated with a specified key.
-	 * @param	Key - The key to search for.
-	 * @return	A pointer to the value associated with the specified key, or nullptr if the key isn't contained in this map.  The pointer
+	 * Find the value associated with a specified key.
+	 *
+	 * @param Key The key to search for.
+	 * @return A pointer to the value associated with the specified key, or nullptr if the key isn't contained in this map.  The pointer
 	 *			is only valid until the next change to any key in the map.
 	 */
 	FORCEINLINE ValueType* Find(KeyConstPointerType Key)
@@ -388,11 +403,13 @@ public:
 	}
 
 private:
+
 	/**
-	 * Returns the value associated with a specified key, or if none exists, 
+	 * Find the value associated with a specified key, or if none exists, 
 	 * adds a value using the default constructor.
-	 * @param	Key - The key to search for.
-	 * @return	A reference to the value associated with the specified key.
+	 *
+	 * @param Key The key to search for.
+	 * @return A reference to the value associated with the specified key.
 	 */
 	template <typename ArgType>
 	FORCEINLINE ValueType& FindOrAddImpl(ArgType&& Arg)
@@ -404,20 +421,23 @@ private:
 	}
 
 public:
+
 	/**
-	 * Returns the value associated with a specified key, or if none exists, 
+	 * Find the value associated with a specified key, or if none exists, 
 	 * adds a value using the default constructor.
-	 * @param	Key - The key to search for.
-	 * @return	A reference to the value associated with the specified key.
+	 *
+	 * @param Key The key to search for.
+	 * @return A reference to the value associated with the specified key.
 	 */
 	FORCEINLINE ValueType& FindOrAdd(const KeyType&  Key) { return FindOrAddImpl(                   Key ); }
 	FORCEINLINE ValueType& FindOrAdd(      KeyType&& Key) { return FindOrAddImpl(MoveTempIfPossible(Key)); }
 
 	/**
-	 * Returns the value associated with a specified key, or if none exists, 
+	 * Find the value associated with a specified key, or if none exists, 
 	 * adds a value using the key as the constructor parameter.
-	 * @param	Key - The key to search for.
-	 * @return	A reference to the value associated with the specified key.
+	 *
+	 * @param Key The key to search for.
+	 * @return A reference to the value associated with the specified key.
 	 */
 	//@todo UE4 merge - this prevents FConfigCacheIni from compiling
 	/*ValueType& FindOrAddKey(KeyInitType Key)
@@ -434,9 +454,10 @@ public:
 	}*/
 
 	/**
-	 * Returns a reference to the value associated with a specified key.
-	 * @param	Key - The key to search for.
-	 * @return	The value associated with the specified key, or triggers an assertion if the key does not exist.
+	 * Find a reference to the value associated with a specified key.
+	 *
+	 * @param Key The key to search for.
+	 * @return The value associated with the specified key, or triggers an assertion if the key does not exist.
 	 */
 	FORCEINLINE const ValueType& FindChecked(KeyConstPointerType Key) const
 	{
@@ -446,9 +467,10 @@ public:
 	}
 
 	/**
-	 * Returns a reference to the value associated with a specified key.
-	 * @param	Key - The key to search for.
-	 * @return	The value associated with the specified key, or triggers an assertion if the key does not exist.
+	 * Find a reference to the value associated with a specified key.
+	 *
+	 * @param Key The key to search for.
+	 * @return The value associated with the specified key, or triggers an assertion if the key does not exist.
 	 */
 	FORCEINLINE ValueType& FindChecked(KeyConstPointerType Key)
 	{
@@ -458,9 +480,10 @@ public:
 	}
 
 	/**
-	 * Returns the value associated with a specified key.
-	 * @param	Key - The key to search for.
-	 * @return	The value associated with the specified key, or the default value for the ValueType if the key isn't contained in this map.
+	 * Find the value associated with a specified key.
+	 *
+	 * @param Key The key to search for.
+	 * @return The value associated with the specified key, or the default value for the ValueType if the key isn't contained in this map.
 	 */
 	FORCEINLINE ValueType FindRef(KeyConstPointerType Key) const
 	{
@@ -473,8 +496,9 @@ public:
 	}
 
 	/**
-	 * Checks if map contains the specified key.
-	 * @param Key - The key to check for.
+	 * Check if map contains the specified key.
+	 *
+	 * @param Key The key to check for.
 	 * @return true if the map contains the key.
 	 */
 	FORCEINLINE bool Contains(KeyConstPointerType Key) const
@@ -483,7 +507,9 @@ public:
 	}
 
 	/**
-	 * Generates an array from the keys in this map.
+	 * Generate an array from the keys in this map.
+	 *
+	 * @param OutArray Will contain the collection of keys.
 	 */
 	template<typename Allocator> void GenerateKeyArray(TArray<KeyType, Allocator>& OutArray) const
 	{
@@ -495,7 +521,9 @@ public:
 	}
 
 	/**
-	 * Generates an array from the values in this map.
+	 * Generate an array from the values in this map.
+	 *
+	 * @param OutArray Will contain the collection of values.
 	 */
 	template<typename Allocator> void GenerateValueArray(TArray<ValueType, Allocator>& OutArray) const
 	{
@@ -514,7 +542,8 @@ public:
 
 	/**
 	 * Describes the map's contents through an output device.
-	 * @param Ar - The output device to describe the map's contents through.
+	 *
+	 * @param Ar The output device to describe the map's contents through.
 	 */
 	void Dump(FOutputDevice& Ar)
 	{
@@ -737,6 +766,7 @@ private:
 	FORCEINLINE friend TConstIterator end  (const TMapBase& MapBase) { return TConstIterator(         end  (MapBase.Pairs)); }
 };
 
+
 /** The base type of sortable maps. */
 template <typename KeyType, typename ValueType, typename SetAllocator, typename KeyFuncs>
 class TSortableMapBase : public TMapBase<KeyType, ValueType, SetAllocator, KeyFuncs>
@@ -919,9 +949,11 @@ public:
 	}
 
 	/**
-	 * Removes the pair with the specified key and copies the value that was removed to the ref parameter
-	 * @param Key - the key to search for
-	 * @param OutRemovedValue - if found, the value that was removed (not modified if the key was not found)
+	 * Remove the pair with the specified key and copies the value
+	 * that was removed to the ref parameter
+	 *
+	 * @param Key The key to search for
+	 * @param OutRemovedValue If found, the value that was removed (not modified if the key was not found)
 	 * @return whether or not the key was found
 	 */
 	FORCEINLINE bool RemoveAndCopyValue(KeyInitType Key,ValueType& OutRemovedValue)
@@ -936,9 +968,11 @@ public:
 	}
 	
 	/**
-	 * Finds a pair with the specified key, removes it from the map, and returns the value part of the pair.
+	 * Find a pair with the specified key, removes it from the map, and returns the value part of the pair.
+	 *
 	 * If no pair was found, an exception is thrown.
-	 * @param Key - the key to search for
+	 *
+	 * @param Key the key to search for
 	 * @return whether or not the key was found
 	 */
 	FORCEINLINE ValueType FindAndRemoveChecked(KeyConstPointerType Key)
@@ -951,8 +985,10 @@ public:
 	}
 
 	/**
-	 * Move all items from another map into our map (if any keys are in both, the value from the other map wins) and empty the other map.
-	 * @param OtherMap - The other map of items to move the elements from.
+	 * Move all items from another map into our map (if any keys are in both,
+	 * the value from the other map wins) and empty the other map.
+	 *
+	 * @param OtherMap The other map of items to move the elements from.
 	 */
 	template<typename OtherSetAllocator>
 	void Append(TMap<KeyType, ValueType, OtherSetAllocator, KeyFuncs>&& OtherMap)
@@ -967,8 +1003,10 @@ public:
 	}
 
 	/**
-	 * Add all items from another map to our map (if any keys are in both, the value from the other map wins)
-	 * @param OtherMap - The other map of items to add.
+	 * Add all items from another map to our map (if any keys are in both,
+	 * the value from the other map wins).
+	 *
+	 * @param OtherMap The other map of items to add.
 	 */
 	template<typename OtherSetAllocator>
 	void Append(const TMap<KeyType, ValueType, OtherSetAllocator, KeyFuncs>& OtherMap)
@@ -984,6 +1022,7 @@ public:
 	FORCEINLINE const ValueType& operator[](KeyConstPointerType Key) const { return this->FindChecked(Key); }
 };
 
+
 /** A TMapBase specialization that allows multiple values to be associated with each key. */
 template<typename KeyType,typename ValueType,typename SetAllocator /* = FDefaultSetAllocator */,typename KeyFuncs /*= TDefaultMapHashableKeyFuncs<KeyType,ValueType,true>*/>
 class TMultiMap : public TSortableMapBase<KeyType, ValueType, SetAllocator, KeyFuncs>
@@ -1048,9 +1087,10 @@ public:
 
 	/**
 	 * Finds all values associated with the specified key.
-	 * @param Key - The key to find associated values for.
-	 * @param OutValues - Upon return, contains the values associated with the key.
-	 * @param bMaintainOrder - true if the Values array should be in the same order as the map's pairs.
+	 *
+	 * @param Key The key to find associated values for.
+	 * @param OutValues Upon return, contains the values associated with the key.
+	 * @param bMaintainOrder true if the Values array should be in the same order as the map's pairs.
 	 */
 	template<typename Allocator> void MultiFind(KeyInitType Key,TArray<ValueType, Allocator>& OutValues,bool bMaintainOrder = false) const
 	{
@@ -1067,10 +1107,11 @@ public:
 
 	/**
 	 * Finds all values associated with the specified key.
-	 * @param Key - The key to find associated values for.
-	 * @param OutValues - Upon return, contains pointers to the values associated with the key.
+	 *
+	 * @param Key The key to find associated values for.
+	 * @param OutValues Upon return, contains pointers to the values associated with the key.
 	 *					Pointers are only valid until the next change to any key in the map.
-	 * @param bMaintainOrder - true if the Values array should be in the same order as the map's pairs.
+	 * @param bMaintainOrder true if the Values array should be in the same order as the map's pairs.
 	 */
 	template<typename Allocator> void MultiFindPointer(KeyInitType Key,TArray<const ValueType*, Allocator>& OutValues,bool bMaintainOrder = false) const
 	{
@@ -1098,11 +1139,12 @@ public:
 	}
 
 	/**
-	 * Adds a key-value association to the map.  The association doesn't replace any of the key's existing associations.
+	 * Add a key-value association to the map.  The association doesn't replace any of the key's existing associations.
 	 * However, if both the key and value match an existing association in the map, no new association is made and the existing association's
 	 * value is returned.
-	 * @param InKey - The key to associate.
-	 * @param InValue - The value to associate.
+	 *
+	 * @param InKey The key to associate.
+	 * @param InValue The value to associate.
 	 * @return A reference to the value as stored in the map; the reference is only valid until the next change to any key in the map.
 	 */
 	FORCEINLINE ValueType& AddUnique(const KeyType&  InKey, const ValueType&  InValue) { return EmplaceUnique(                   InKey ,                    InValue ); }
@@ -1111,11 +1153,14 @@ public:
 	FORCEINLINE ValueType& AddUnique(      KeyType&& InKey,       ValueType&& InValue) { return EmplaceUnique(MoveTempIfPossible(InKey), MoveTempIfPossible(InValue)); }
 
 	/**
-	 * Adds a key-value association to the map.  The association doesn't replace any of the key's existing associations.
-	 * However, if both the key and value match an existing association in the map, no new association is made and the existing association's
-	 * value is returned.
-	 * @param InKey - The key to associate.
-	 * @param InValue - The value to associate.
+	 * Add a key-value association to the map.
+	 *
+	 * The association doesn't replace any of the key's existing associations.
+	 * However, if both key and value match an existing association in the map,
+	 * no new association is made and the existing association's value is returned.
+	 *
+	 * @param InKey The key to associate.
+	 * @param InValue The value to associate.
 	 * @return A reference to the value as stored in the map; the reference is only valid until the next change to any key in the map.
 	 */
 	template <typename InitKeyType, typename InitValueType>
@@ -1131,8 +1176,9 @@ public:
 	}
 
 	/**
-	 * Removes all value associations for a key.
-	 * @param InKey - The key to remove associated values for.
+	 * Remove all value associations for a key.
+	 *
+	 * @param InKey The key to remove associated values for.
 	 * @return The number of values that were associated with the key.
 	 */
 	FORCEINLINE int32 Remove(KeyConstPointerType InKey)
@@ -1141,9 +1187,10 @@ public:
 	}
 
 	/**
-	 * Removes associations between the specified key and value from the map.
-	 * @param InKey - The key part of the pair to remove.
-	 * @param InValue - The value part of the pair to remove.
+	 * Remove associations between the specified key and value from the map.
+	 *
+	 * @param InKey The key part of the pair to remove.
+	 * @param InValue The value part of the pair to remove.
 	 * @return The number of associations removed.
 	 */
 	int32 Remove(KeyInitType InKey,ValueInitType InValue)
@@ -1163,9 +1210,10 @@ public:
 	}
 
 	/**
-	 * Removes the first association between the specified key and value from the map.
-	 * @param InKey - The key part of the pair to remove.
-	 * @param InValue - The value part of the pair to remove.
+	 * Remove the first association between the specified key and value from the map.
+	 *
+	 * @param InKey The key part of the pair to remove.
+	 * @param InValue The value part of the pair to remove.
 	 * @return The number of associations removed.
 	 */
 	int32 RemoveSingle(KeyInitType InKey,ValueInitType InValue)
@@ -1188,9 +1236,10 @@ public:
 	}
 
 	/**
-	 * Finds an association between a specified key and value. (const)
-	 * @param Key - The key to find.
-	 * @param Value - The value to find.
+	 * Find an association between a specified key and value. (const)
+	 *
+	 * @param Key The key to find.
+	 * @param Value The value to find.
 	 * @return If the map contains a matching association, a pointer to the value in the map is returned.  Otherwise nullptr is returned.
 	 *			The pointer is only valid as long as the map isn't changed.
 	 */
@@ -1200,9 +1249,10 @@ public:
 	}
 
 	/**
-	 * Finds an association between a specified key and value.
-	 * @param Key - The key to find.
-	 * @param Value - The value to find.
+	 * Find an association between a specified key and value.
+	 *
+	 * @param Key The key to find.
+	 * @param Value The value to find.
 	 * @return If the map contains a matching association, a pointer to the value in the map is returned.  Otherwise nullptr is returned.
 	 *			The pointer is only valid as long as the map isn't changed.
 	 */
@@ -1240,6 +1290,7 @@ public:
 	}
 };
 
+
 struct FScriptMapLayout
 {
 	int32 KeyOffset;
@@ -1248,6 +1299,7 @@ struct FScriptMapLayout
 	FScriptSetLayout SetLayout;
 };
 
+
 // Untyped map type for accessing TMap data, like FScriptArray for TArray.
 // Must have the same memory representation as a TMap.
 class FScriptMap
@@ -1309,7 +1361,7 @@ public:
 	 * Adds an uninitialized object to the map.
 	 * The map will need rehashing at some point after this call to make it valid.
 	 *
-	 * @return  The index of the added element.
+	 * @return The index of the added element.
 	 */
 	int32 AddUninitialized(const FScriptMapLayout& Layout)
 	{
@@ -1414,6 +1466,7 @@ public:
 	void operator=(const FScriptMap&) { check(false); }
 };
 
+
 template <>
 struct TIsZeroConstructType<FScriptMap>
 {
@@ -1427,6 +1480,7 @@ struct TContainerTraits<TMap<KeyType, ValueType, SetAllocator, KeyFuncs>> : publ
 	enum { MoveWillEmptyContainer = TContainerTraits<typename TMap<KeyType, ValueType, SetAllocator, KeyFuncs>::ElementSetType>::MoveWillEmptyContainer };
 };
 
+
 template <typename KeyType, typename ValueType, typename SetAllocator,typename KeyFuncs>
 struct TContainerTraits<TMultiMap<KeyType, ValueType, SetAllocator, KeyFuncs>> : public TContainerTraitsBase<TMultiMap<KeyType, ValueType, SetAllocator, KeyFuncs>>
 {
diff --git a/Engine/Source/Runtime/Core/Public/Containers/Queue.h b/Engine/Source/Runtime/Core/Public/Containers/Queue.h
--- a/Engine/Source/Runtime/Core/Public/Containers/Queue.h
+++ b/Engine/Source/Runtime/Core/Public/Containers/Queue.h
@@ -65,7 +65,8 @@ public:
 	 *
 	 * @param OutValue Will hold the returned value.
 	 * @return true if a value was returned, false if the queue was empty.
-	 * @see Enqueue, IsEmpty, Peek
+	 * @note To be called only from consumer thread.
+	 * @see Empty, Enqueue, IsEmpty, Peek
 	 */
 	bool Dequeue(ItemType& OutItem)
 	{
@@ -87,7 +88,12 @@ public:
 		return true;
 	}
 
-	/** Empty the queue, discarding all items. */
+	/**
+	 * Empty the queue, discarding all items.
+	 *
+	 * @note To be called only from consumer thread.
+	 * @see Dequeue, IsEmpty, Peek
+	 */
 	void Empty()
 	{
 		ItemType DummyItem;
@@ -99,7 +105,8 @@ public:
 	 *
 	 * @param Item The item to add.
 	 * @return true if the item was added, false otherwise.
-	 * @see Dequeue, IsEmpty, Peek
+	 * @note To be called only from producer thread(s).
+	 * @see Dequeue
 	 */
 	bool Enqueue(const ItemType& Item)
 	{
@@ -135,7 +142,8 @@ public:
 	 *
 	 * @param Item The item to add.
 	 * @return true if the item was added, false otherwise.
-	 * @see Dequeue, IsEmpty, Peek
+	 * @note To be called only from producer thread(s).
+	 * @see Dequeue
 	 */
 	bool Enqueue(ItemType&& Item)
 	{
@@ -170,7 +178,8 @@ public:
 	 * Checks whether the queue is empty.
 	 *
 	 * @return true if the queue is empty, false otherwise.
-	 * @see Dequeue, Enqueue, Peek
+	 * @note To be called only from consumer thread.
+	 * @see Dequeue, Empty, Peek
 	 */
 	bool IsEmpty() const
 	{
@@ -182,7 +191,8 @@ public:
 	 *
 	 * @param OutItem Will hold the peeked at item.
 	 * @return true if an item was returned, false if the queue was empty.
-	 * @see Dequeue, Enqueue, IsEmpty
+	 * @note To be called only from consumer thread.
+	 * @see Dequeue, Empty, IsEmpty
 	 */
 	bool Peek(ItemType& OutItem) const
 	{
diff --git a/Engine/Source/Runtime/Core/Public/Containers/SortedMap.h b/Engine/Source/Runtime/Core/Public/Containers/SortedMap.h
--- a/Engine/Source/Runtime/Core/Public/Containers/SortedMap.h
+++ b/Engine/Source/Runtime/Core/Public/Containers/SortedMap.h
@@ -8,13 +8,14 @@
 #include "UObject/NameTypes.h"
 
 /** 
- * A Map of keys to value, implemented as a sorted TArray of TPairs. It has a mostly identical interface to TMap and is designed as a drop in replacement
- * Keys must be unique, there is no equivalent sorted version of TMultiMap
- * It uses half as much memory as TMap, but adding/removing elements is O(n) and finding is O(Log n)
- * In practice it is faster than TMap for low element counts, and slower as n increases
- * This map is always kept sorted by the key type so cannot be sorted manually
+ * A Map of keys to value, implemented as a sorted TArray of TPairs.
+ *
+ * It has a mostly identical interface to TMap and is designed as a drop in replacement. Keys must be unique,
+ * there is no equivalent sorted version of TMultiMap. It uses half as much memory as TMap, but adding and 
+ * removing elements is O(n), and finding is O(Log n). In practice it is faster than TMap for low element
+ * counts, and slower as n increases, This map is always kept sorted by the key type so cannot be sorted manually.
  */
-template <typename KeyType, typename ValueType, typename ArrayAllocator /*= FDefaultAllocator*/, typename SortPredicate /* = TLess<KeyType> */ >
+template <typename KeyType, typename ValueType, typename ArrayAllocator /*= FDefaultAllocator*/, typename SortPredicate /*= TLess<KeyType>*/ >
 class TSortedMap
 {
 	template <typename OtherKeyType, typename OtherValueType, typename OtherArrayAllocator, typename OtherSortPredicate>
@@ -33,21 +34,21 @@ public:
 	TSortedMap& operator=(TSortedMap&&) = default;
 	TSortedMap& operator=(const TSortedMap&) = default;
 
-	/** Constructor for moving elements from a TSortedMap with a different ArrayAllocator */
+	/** Constructor for moving elements from a TSortedMap with a different ArrayAllocator. */
 	template<typename OtherArrayAllocator>
 	TSortedMap(TSortedMap<KeyType, ValueType, OtherArrayAllocator, SortPredicate>&& Other)
 		: Pairs(MoveTemp(Other.Pairs))
 	{
 	}
 
-	/** Constructor for copying elements from a TSortedMap with a different ArrayAllocator */
+	/** Constructor for copying elements from a TSortedMap with a different ArrayAllocator. */
 	template<typename OtherArrayAllocator>
 	TSortedMap(const TSortedMap<KeyType, ValueType, OtherArrayAllocator, SortPredicate>& Other)
 		: Pairs(Other.Pairs)
 	{
 	}
 
-	/** Assignment operator for moving elements from a TSortedMap with a different ArrayAllocator */
+	/** Assignment operator for moving elements from a TSortedMap with a different ArrayAllocator. */
 	template<typename OtherArrayAllocator>
 	TSortedMap& operator=(TSortedMap<KeyType, ValueType, OtherArrayAllocator, SortPredicate>&& Other)
 	{
@@ -55,7 +56,7 @@ public:
 		return *this;
 	}
 
-	/** Assignment operator for copying elements from a TSortedMap with a different ArrayAllocator */
+	/** Assignment operator for copying elements from a TSortedMap with a different ArrayAllocator. */
 	template<typename OtherArrayAllocator>
 	TSortedMap& operator=(const TSortedMap<KeyType, ValueType, OtherArrayAllocator, SortPredicate>& Other)
 	{
@@ -63,13 +64,13 @@ public:
 		return *this;
 	}
 
-	/** Equality operator. This is efficient because pairs are always sorted */
+	/** Equality operator. This is efficient because pairs are always sorted. */
 	FORCEINLINE bool operator==(const TSortedMap& Other) const
 	{
 		return Pairs == Other.Pairs;
 	}
 
-	/** Inequality operator. This is efficient because pairs are always sorted */
+	/** Inequality operator. This is efficient because pairs are always sorted. */
 	FORCEINLINE bool operator!=(const TSortedMap& Other) const
 	{
 		return Pairs != Other.Pairs;
@@ -77,14 +78,15 @@ public:
 
 	/**
 	 * Removes all elements from the map, potentially leaving space allocated for an expected number of elements about to be added.
-	 * @param ExpectedNumElements - The number of elements about to be added to the set.
+	 *
+	 * @param ExpectedNumElements The number of elements about to be added to the set.
 	 */
 	FORCEINLINE void Empty(int32 ExpectedNumElements = 0)
 	{
 		Pairs.Empty(ExpectedNumElements);
 	}
 
-    /** Efficiently empties out the map but preserves all allocations and capacities */
+    /** Efficiently empties out the map but preserves all allocations and capacities. */
     FORCEINLINE void Reset()
     {
         Pairs.Reset();
@@ -96,7 +98,7 @@ public:
 		Pairs.Shrink();
 	}
 
-	/** Preallocates enough memory to contain Number elements */
+	/** Preallocates enough memory to contain Number elements. */
 	FORCEINLINE void Reserve(int32 Number)
 	{
 		Pairs.Reserve(Number);
@@ -109,8 +111,9 @@ public:
 	}
 
 	/** 
-	 * Helper function to return the amount of memory allocated by this container 
-	 * @return number of bytes allocated by this container
+	 * Helper function to return the amount of memory allocated by this container.
+	 *
+	 * @return number of bytes allocated by this container.
 	 */
 	FORCEINLINE uint32 GetAllocatedSize() const
 	{
@@ -126,9 +129,9 @@ public:
 	/**
 	 * Sets the value associated with a key.
 	 *
-	 * @param InKey - The key to associate the value with.
-	 * @param InValue - The value to associate with the key.
-	 * @return A reference to the value as stored in the map.  The reference is only valid until the next change to any key in the map.
+	 * @param InKey The key to associate the value with.
+	 * @param InValue The value to associate with the key.
+	 * @return A reference to the value as stored in the map (only valid until the next change to any key in the map).
 	 */
 	FORCEINLINE ValueType& Add(const KeyType&  InKey, const ValueType&  InValue) { return Emplace(         InKey ,          InValue ); }
 	FORCEINLINE ValueType& Add(const KeyType&  InKey,       ValueType&& InValue) { return Emplace(         InKey , MoveTemp(InValue)); }
@@ -138,7 +141,7 @@ public:
 	/**
 	 * Sets a default value associated with a key.
 	 *
-	 * @param InKey - The key to associate the value with.
+	 * @param InKey The key to associate the value with.
 	 * @return A reference to the value as stored in the map.  The reference is only valid until the next change to any key in the map.
 	 */
 	FORCEINLINE ValueType& Add(const KeyType&  InKey) { return Emplace(         InKey ); }
@@ -149,7 +152,7 @@ public:
 	 *
 	 * @param InKey - The key to associate the value with.
 	 * @param InValue - The value to associate with the key.
-	 * @return A reference to the value as stored in the map.  The reference is only valid until the next change to any key in the map.
+	 * @return A reference to the value as stored in the map (only valid until the next change to any key in the map).
 	 */
 	template <typename InitKeyType, typename InitValueType>
 	ValueType& Emplace(InitKeyType&& InKey, InitValueType&& InValue)
@@ -164,8 +167,8 @@ public:
 	/**
 	 * Sets a default value associated with a key.
 	 *
-	 * @param InKey - The key to associate the value with.
-	 * @return A reference to the value as stored in the map.  The reference is only valid until the next change to any key in the map.
+	 * @param InKey The key to associate the value with.
+	 * @return A reference to the value as stored in the map (only valid until the next change to any key in the map).
 	 */
 	template <typename InitKeyType>
 	ValueType& Emplace(InitKeyType&& InKey)
@@ -179,7 +182,8 @@ public:
 
 	/**
 	 * Removes all value associations for a key.
-	 * @param InKey - The key to remove associated values for.
+	 *
+	 * @param InKey The key to remove associated values for.
 	 * @return The number of values that were associated with the key.
 	 */
 	FORCEINLINE int32 Remove(KeyConstPointerType InKey)
@@ -198,9 +202,9 @@ public:
 
 	/**
 	 * Returns the key associated with the specified value.  The time taken is O(N) in the number of pairs.
-	 * @param	Value - The value to search for
-	 * @return	A pointer to the key associated with the specified value, or nullptr if the value isn't contained in this map.  The pointer
-	 *			is only valid until the next change to any key in the map.
+	 *
+	 * @param Value The value to search for.
+	 * @return A pointer to the key associated with the specified value, or nullptr if the value isn't contained in this map (only valid until the next change to any key in the map).
 	 */
 	const KeyType* FindKey(ValueInitType Value) const
 	{
@@ -216,9 +220,9 @@ public:
 
 	/**
 	 * Returns the value associated with a specified key.
-	 * @param	Key - The key to search for.
-	 * @return	A pointer to the value associated with the specified key, or nullptr if the key isn't contained in this map.  The pointer
-	 *			is only valid until the next change to any key in the map.
+	 *
+	 * @param Key The key to search for.
+	 * @return A pointer to the value associated with the specified key, or nullptr if the key isn't contained in this map.  The pointer (only valid until the next change to any key in the map).
 	 */
 	FORCEINLINE ValueType* Find(KeyConstPointerType Key)
 	{
@@ -240,16 +244,18 @@ public:
 	/**
 	 * Returns the value associated with a specified key, or if none exists, 
 	 * adds a value using the default constructor.
-	 * @param	Key - The key to search for.
-	 * @return	A reference to the value associated with the specified key.
+	 *
+	 * @param Key The key to search for.
+	 * @return A reference to the value associated with the specified key.
 	 */
 	FORCEINLINE ValueType& FindOrAdd(const KeyType&  Key) { return FindOrAddImpl(         Key ); }
 	FORCEINLINE ValueType& FindOrAdd(      KeyType&& Key) { return FindOrAddImpl(MoveTemp(Key)); }
 
 	/**
 	 * Returns a reference to the value associated with a specified key.
-	 * @param	Key - The key to search for.
-	 * @return	The value associated with the specified key, or triggers an assertion if the key does not exist.
+	 *
+	 * @param Key The key to search for.
+	 * @return The value associated with the specified key, or triggers an assertion if the key does not exist.
 	 */
 	FORCEINLINE ValueType& FindChecked(KeyConstPointerType Key)
 	{
@@ -267,8 +273,9 @@ public:
 
 	/**
 	 * Returns the value associated with a specified key.
-	 * @param	Key - The key to search for.
-	 * @return	The value associated with the specified key, or the default value for the ValueType if the key isn't contained in this map.
+	 *
+	 * @param Key The key to search for.
+	 * @return The value associated with the specified key, or the default value for the ValueType if the key isn't contained in this map.
 	 */
 	FORCEINLINE ValueType FindRef(KeyConstPointerType Key) const
 	{
@@ -282,7 +289,8 @@ public:
 
 	/**
 	 * Checks if map contains the specified key.
-	 * @param Key - The key to check for.
+	 *
+	 * @param Key The key to check for.
 	 * @return true if the map contains the key.
 	 */
 	FORCEINLINE bool Contains(KeyConstPointerType Key) const
@@ -295,8 +303,9 @@ public:
 	}
 
 	/**
-	 * Returns the unique keys contained within this map
-	 * @param OutKeys - Upon return, contains the set of unique keys in this map.
+	 * Returns the unique keys contained within this map.
+	 *
+	 * @param OutKeys Upon return, contains the set of unique keys in this map.
 	 * @return The number of unique keys in the map.
 	 */
 	template<typename Allocator> int32 GetKeys(TArray<KeyType, Allocator>& OutKeys) const
@@ -348,7 +357,8 @@ public:
 
 	/**
 	 * Describes the map's contents through an output device.
-	 * @param Ar - The output device to describe the map's contents through.
+	 *
+	 * @param Ar The output device to describe the map's contents through.
 	 */
 	void Dump(FOutputDevice& Ar)
 	{
@@ -356,10 +366,11 @@ public:
 	}
 
 	/**
-	 * Removes the pair with the specified key and copies the value that was removed to the ref parameter
-	 * @param Key - the key to search for
-	 * @param OutRemovedValue - if found, the value that was removed (not modified if the key was not found)
-	 * @return whether or not the key was found
+	 * Removes the pair with the specified key and copies the value that was removed to the ref parameter.
+	 *
+	 * @param Key The key to search for
+	 * @param OutRemovedValue If found, the value that was removed (not modified if the key was not found)
+	 * @return Whether or not the key was found
 	 */
 	FORCEINLINE bool RemoveAndCopyValue(KeyInitType Key, ValueType& OutRemovedValue)
 	{
@@ -377,8 +388,9 @@ public:
 	/**
 	 * Finds a pair with the specified key, removes it from the map, and returns the value part of the pair.
 	 * If no pair was found, an exception is thrown.
-	 * @param Key - the key to search for
-	 * @return whether or not the key was found
+	 *
+	 * @param Key The key to search for
+	 * @return Whether or not the key was found
 	 */
 	FORCEINLINE ValueType FindAndRemoveChecked(KeyConstPointerType Key)
 	{
@@ -392,7 +404,8 @@ public:
 
 	/**
 	 * Move all items from another map into our map (if any keys are in both, the value from the other map wins) and empty the other map.
-	 * @param OtherMap - The other map of items to move the elements from.
+	 *
+	 * @param OtherMap The other map of items to move the elements from.
 	 */
 	template<typename OtherArrayAllocator, typename OtherSortPredicate>
 	void Append(TSortedMap<KeyType, ValueType, OtherArrayAllocator, OtherSortPredicate>&& OtherMap)
@@ -407,8 +420,9 @@ public:
 	}
 
 	/**
-	 * Add all items from another map to our map (if any keys are in both, the value from the other map wins)
-	 * @param OtherMap - The other map of items to add.
+	 * Add all items from another map to our map (if any keys are in both, the value from the other map wins).
+	 *
+	 * @param OtherMap The other map of items to add.
 	 */
 	template<typename OtherArrayAllocator, typename OtherSortPredicate>
 	void Append(const TSortedMap<KeyType, ValueType, OtherArrayAllocator, OtherSortPredicate>& OtherMap)
@@ -662,6 +676,7 @@ private:
 	FORCEINLINE friend RangedForConstIteratorType	end(const TSortedMap& MapBase)	{ return end(MapBase.Pairs); }
 };
 
+
 template <typename KeyType, typename ValueType, typename ArrayAllocator, typename SortPredicate>
 struct TContainerTraits<TSortedMap<KeyType, ValueType, ArrayAllocator, SortPredicate>> : public TContainerTraitsBase<TSortedMap<KeyType, ValueType, ArrayAllocator, SortPredicate>>
 {
diff --git a/Engine/Source/Runtime/Core/Public/GenericPlatform/IInputInterface.h b/Engine/Source/Runtime/Core/Public/GenericPlatform/IInputInterface.h
--- a/Engine/Source/Runtime/Core/Public/GenericPlatform/IInputInterface.h
+++ b/Engine/Source/Runtime/Core/Public/GenericPlatform/IInputInterface.h
@@ -4,6 +4,8 @@
 
 #include "CoreTypes.h"
 #include "Containers/Array.h"
+#include "Math/Color.h"
+
 
 // General identifiers for potential force feedback channels. These will be mapped according to the
 // platform specific implementation.
@@ -125,5 +127,5 @@ public:
 	/*
 	 * Sets the controller for the given controller.  Ignored if controller does not support a color.
 	 */
-	virtual void SetLightColor(int32 ControllerId, struct FColor Color) = 0;
+	virtual void SetLightColor(int32 ControllerId, FColor Color) = 0;
 };
diff --git a/Engine/Source/Runtime/Core/Public/IOS/IOSAsyncTask.h b/Engine/Source/Runtime/Core/Public/IOS/IOSAsyncTask.h
--- a/Engine/Source/Runtime/Core/Public/IOS/IOSAsyncTask.h
+++ b/Engine/Source/Runtime/Core/Public/IOS/IOSAsyncTask.h
@@ -1,13 +1,12 @@
 // Copyright 1998-2017 Epic Games, Inc. All Rights Reserved.
 
-/*=============================================================================
-	IPhoneAppDelegate.h: IPhone application class / main loop
-=============================================================================*/
-
 #pragma once
 
+#include "CoreTypes.h"
+
 #import <Foundation/Foundation.h>
 
+
 @interface FIOSAsyncTask : NSObject
 {
 @private
diff --git a/Engine/Source/Runtime/Core/Public/Math/Range.h b/Engine/Source/Runtime/Core/Public/Math/Range.h
--- a/Engine/Source/Runtime/Core/Public/Math/Range.h
+++ b/Engine/Source/Runtime/Core/Public/Math/Range.h
@@ -3,10 +3,12 @@
 #pragma once
 
 #include "CoreTypes.h"
-#include "Misc/AssertionMacros.h"
 #include "Containers/Array.h"
+#include "Misc/AssertionMacros.h"
 #include "Misc/DateTime.h"
 #include "Math/RangeBound.h"
+#include "Serialization/Archive.h"
+
 
 /**
  * Template for ranges.
@@ -40,10 +42,8 @@
  * recommended to limit all operations to canonical ranges of the form [A, B) in
  * which the lower bound is included and the upper bound is excluded from the range.
  *
- * @param ElementType - The type of elements represented by the range
- *
- * @see RangeBound
- * @see TInterval
+ * @param ElementType The type of elements represented by the range
+ * @see RangeBound, TInterval
  */
 template<typename ElementType> class TRange
 {
@@ -55,7 +55,7 @@ public:
 	TRange() { }
 
 	/**
-	 * Creates a range with a single element.
+	 * Create a range with a single element.
 	 *
 	 * The created range is of the form [A, A].
 	 *
@@ -67,7 +67,7 @@ public:
 	{ }
 	
 	/**
-	 * Creates and initializes a new range with the given lower and upper bounds.
+	 * Create and initializes a new range with the given lower and upper bounds.
 	 *
 	 * The created range is of the form [A, B).
 	 *
@@ -80,7 +80,7 @@ public:
 	{ }
 
 	/**
-	 * Creates and initializes a new range with the given lower and upper bounds.
+	 * Create and initializes a new range with the given lower and upper bounds.
 	 *
 	 * @param InLowerBound The range's lower bound.
 	 * @param InUpperBound The range's upper bound.
@@ -93,7 +93,7 @@ public:
 public:
 
 	/**
-	 * Compares this range with the specified range for equality.
+	 * Compare this range with the specified range for equality.
 	 *
 	 * @param Other The range to compare with.
 	 * @return true if the ranges are equal, false otherwise.
@@ -109,7 +109,7 @@ public:
 	}
 
 	/**
-	 * Compares this range with the specified range for inequality.
+	 * Compare this range with the specified range for inequality.
 	 *
 	 * @param Other The range to compare with.
 	 * @return true if the ranges are not equal, false otherwise.
@@ -127,10 +127,10 @@ public:
 public:
 
 	/**
-	 * Checks whether this range adjoins to another.
+	 * Check whether this range adjoins to another.
 	 *
 	 * Two ranges are adjoint if they are next to each other without overlapping, i.e.
-	 *		[A, B) and [C, C) or
+	 *		[A, B) and [B, C) or
 	 *		[A, B] and (B, C)
 	 *
 	 * @param Other The other range.
@@ -159,7 +159,7 @@ public:
 	}
 
 	/**
-	 * Checks whether this range conjoins the two given ranges.
+	 * Check whether this range conjoins the two given ranges.
 	 *
 	 * A range conjoins two non-overlapping ranges if it adjoins both of them, i.e.
 	 *		[B, C) conjoins the two ranges [A, B) and [C, D).
@@ -179,7 +179,7 @@ public:
 	}
 	
 	/**
-	 * Checks whether this range contains the specified element.
+	 * Check whether this range contains the specified element.
 	 *
 	 * @param Element The element to check.
 	 * @return true if the range contains the element, false otherwise.
@@ -191,7 +191,7 @@ public:
 	}
 	
 	/**
-	 * Checks whether this range contains another range.
+	 * Check whether this range contains another range.
 	 *
 	 * @param Other The range to check.
 	 * @return true if the range contains the other range, false otherwise.
@@ -203,7 +203,7 @@ public:
 	}
 
 	/**
-	 * Checks if this range is contiguous with another range.
+	 * Check if this range is contiguous with another range.
 	 *
 	 * Two ranges are contiguous if they are adjoint or overlapping.
 	 *
@@ -216,9 +216,10 @@ public:
 	}
 	
 	/**
-	 * Gets the range's lower bound.
+	 * Get the range's lower bound.
 	 *
 	 * @return Lower bound.
+	 * @see GetLowerBoundValue, GetUpperBound, HasLowerBound
 	 */
 	BoundsType GetLowerBound() const
 	{
@@ -226,12 +227,12 @@ public:
 	}
 
 	/**
-	 * Gets the value of the lower bound.
+	 * Get the value of the lower bound.
 	 *
 	 * Use HasLowerBound() to ensure that this range actually has a lower bound.
 	 *
 	 * @return Bound value.
-	 * @see GetUpperBoundValue, HasLowerBound
+	 * @see GetLowerBound, GetUpperBoundValue, HasLowerBound
 	 */
 	const ElementType& GetLowerBoundValue() const
 	{
@@ -239,9 +240,10 @@ public:
 	}
 
 	/**
-	 * Gets the range's upper bound.
+	 * Get the range's upper bound.
 	 *
 	 * @return Upper bound.
+	 * @see GetLowerBound, GetUpperBoundValue, HasUpperBound
 	 */
 	BoundsType GetUpperBound() const
 	{
@@ -249,12 +251,12 @@ public:
 	}
 
 	/**
-	 * Gets the value of the upper bound.
+	 * Get the value of the upper bound.
 	 *
 	 * Use HasUpperBound() to ensure that this range actually has an upper bound.
 	 *
 	 * @return Bound value.
-	 * @see GetLowerBoundValue, HasUpperBound
+	 * @see GetLowerBoundValue, GetUpperBound, HasUpperBound
 	 */
 	const ElementType& GetUpperBoundValue() const
 	{
@@ -262,9 +264,10 @@ public:
 	}
 
 	/**
-	 * Checks whether the range has a lower bound.
+	 * Check whether the range has a lower bound.
 	 *
 	 * @return true if the range has a lower bound, false otherwise.
+	 * @see GetLowerBound, GetLowerBoundValue, HasUpperBound
 	 */
 	bool HasLowerBound() const
 	{
@@ -272,9 +275,10 @@ public:
 	}
 	
 	/**
-	 * Checks whether the range has an upper bound.
+	 * Check whether the range has an upper bound.
 	 *
 	 * @return true if the range has an upper bound, false otherwise.
+	 * @see GetUpperBound, GetUpperBoundValue, HasLowerBound
 	 */
 	bool HasUpperBound() const
 	{
@@ -282,12 +286,13 @@ public:
 	}
 
 	/**
-	 * Checks whether this range is degenerate.
+	 * Check whether this range is degenerate.
 	 *
 	 * A range is degenerate if it contains only a single element, i.e. has the following form:
 	 *		[A, A]
 	 *
 	 * @return true if the range is degenerate, false otherwise.
+	 * @see IsEmpty
 	 */
 	bool IsDegenerate() const
 	{
@@ -295,7 +300,7 @@ public:
 	}
 	
 	/**
-	 * Checks whether this range is empty.
+	 * Check whether this range is empty.
 	 *
 	 * A range is empty if it contains no elements, i.e.
 	 *		(A, A)
@@ -303,6 +308,7 @@ public:
 	 *		[A, A)
 	 *
 	 * @return true if the range is empty, false otherwise.
+	 * @see IsDegenerate
 	 */
 	bool IsEmpty() const
 	{
@@ -320,7 +326,7 @@ public:
 	}
 
 	/**
-	 * Checks whether this range overlaps with another.
+	 * Check whether this range overlaps with another.
 	 *
 	 * @param Other The other range.
 	 * @return true if the ranges overlap, false otherwise.
@@ -361,7 +367,7 @@ public:
 	}
 
 	/**
-	 * Computes the size (diameter, length, width) of this range.
+	 * Compute the size (diameter, length, width) of this range.
 	 *
 	 * The size of a closed range is the difference between its upper and lower bound values.
 	 * Use IsClosed() on the lower and upper bounds before calling this method in order to
@@ -377,10 +383,10 @@ public:
 	}
 	
 	/**
-	 * Splits the range into two ranges at the specified element.
+	 * Split the range into two ranges at the specified element.
 	 *
-	 * If a range [A, C) does not contain the element B, the original range is returned.
-	 * Otherwise the range is split into two ranges [A, B) and [B, C), each of which may be empty.
+	 * If a range [A, C) does not contain the element X, the original range is returned.
+	 * Otherwise the range is split into two ranges [A, X) and [X, C), each of which may be empty.
 	 *
 	 * @param Element The element at which to split the range.
 	 */
@@ -404,7 +410,7 @@ public:
 public:
 
 	/**
-	 * Calculates the difference between two ranges, i.e. X - Y.
+	 * Calculate the difference between two ranges, i.e. X - Y.
 	 *
 	 * @param X The first range to subtract from.
 	 * @param Y The second range to subtract with.
@@ -439,17 +445,7 @@ public:
 	}
 
 	/**
-	 * Returns an empty range.
-	 *
-	 * @return Empty range.
-	 */
-	static FORCEINLINE TRange Empty()
-	{
-		return TRange(BoundsType::Exclusive(ElementType()), BoundsType::Exclusive(ElementType()));
-	}
-
-	/**
-	 * Computes the hull of two ranges.
+	 * Compute the hull of two ranges.
 	 *
 	 * The hull is the smallest range that contains both ranges.
 	 *
@@ -474,7 +470,7 @@ public:
 	}
 	
 	/**
-	 * Computes the hull of many ranges.
+	 * Compute the hull of many ranges.
 	 *
 	 * @param Ranges The ranges to hull.
 	 * @return The hull.
@@ -498,7 +494,7 @@ public:
 	}
 
 	/**
-	 * Computes the intersection of two ranges.
+	 * Compute the intersection of two ranges.
 	 *
 	 * The intersection of two ranges is the largest range that is contained by both ranges.
 	 *
@@ -523,7 +519,7 @@ public:
 	}
 	
 	/**
-	 * Computes the intersection of many ranges.
+	 * Compute the intersection of many ranges.
 	 *
 	 * @param Ranges The ranges to intersect.
 	 * @return The intersection.
@@ -547,7 +543,7 @@ public:
 	}
 
 	/**
-	 * Returns the union of two contiguous ranges.
+	 * Return the union of two contiguous ranges.
 	 *
 	 * A union is a range or series of ranges that contains both ranges.
 	 *
@@ -583,9 +579,10 @@ public:
 public:
 
 	/**
-	 * Creates an unbounded (open) range that contains all elements of the domain.
+	 * Create an unbounded (open) range that contains all elements of the domain.
 	 *
 	 * @return A new range.
+	 * @see AtLeast, AtMost, Empty, Exclusive, GreaterThan, Inclusive, LessThan
 	 */
 	static FORCEINLINE TRange All()
 	{
@@ -593,10 +590,11 @@ public:
 	}
 
 	/**
-	 * Creates a left-bounded range that contains all elements greater than or equal to the specified value.
+	 * Create a left-bounded range that contains all elements greater than or equal to the specified value.
 	 *
 	 * @param Value The value.
 	 * @return A new range.
+	 * @see All, AtMost, Empty, Exclusive, GreaterThan, Inclusive, LessThan
 	 */
 	static FORCEINLINE TRange AtLeast(const ElementType& Value)
 	{
@@ -604,10 +602,11 @@ public:
 	}
 
 	/**
-	 * Creates a right-bounded range that contains all elements less than or equal to the specified value.
+	 * Create a right-bounded range that contains all elements less than or equal to the specified value.
 	 *
 	 * @param Value The value.
 	 * @return A new range.
+	 * @see All, AtLeast, Empty, Exclusive, GreaterThan, Inclusive, LessThan
 	 */
 	static FORCEINLINE TRange AtMost(const ElementType& Value)
 	{
@@ -615,10 +614,35 @@ public:
 	}
 
 	/**
-	 * Creates a left-bounded range that contains all elements greater than the specified value.
+	 * Return an empty range.
+	 *
+	 * @return Empty range.
+	 * @see All, AtLeast, AtMost, Exclusive, GreaterThan, Inclusive, LessThan
+	 */
+	static FORCEINLINE TRange Empty()
+	{
+		return TRange(BoundsType::Exclusive(ElementType()), BoundsType::Exclusive(ElementType()));
+	}
+
+	/**
+	 * Create a range that excludes the given minimum and maximum values.
+	 *
+	 * @param MinThe minimum value to be included.
+	 * @param Max The maximum value to be included.
+	 * @return A new range.
+	 * @see All, AtLeast, AtMost, Empty, Exclusive, GreaterThan, Inclusive, LessThan
+	 */
+	static FORCEINLINE TRange Exclusive(const ElementType& Min, const ElementType& Max)
+	{
+		return TRange(BoundsType::Exclusive(Min), BoundsType::Exclusive(Max));
+	}
+
+	/**
+	 * Create a left-bounded range that contains all elements greater than the specified value.
 	 *
 	 * @param Value The value.
 	 * @return A new range.
+	 * @see All, AtLeast, AtMost, Empty, Exclusive, Inclusive, LessThan
 	 */
 	static FORCEINLINE TRange GreaterThan(const ElementType& Value)
 	{
@@ -626,10 +650,24 @@ public:
 	}
 	
 	/**
-	 * Creates a right-bounded range that contains all elements less than the specified value.
+	 * Create a range that includes the given minimum and maximum values.
+	 *
+	 * @param Min The minimum value to be included.
+	 * @param Max The maximum value to be included.
+	 * @return A new range.
+	 * @see All, AtLeast, AtMost, Empty, Exclusive, GreaterThan, LessThan
+	 */
+	static FORCEINLINE TRange Inclusive(const ElementType& Min, const ElementType& Max)
+	{
+		return TRange(BoundsType::Inclusive(Min), BoundsType::Inclusive(Max));
+	}
+
+	/**
+	 * Create a right-bounded range that contains all elements less than the specified value.
 	 *
 	 * @param Value The value.
 	 * @return A new range.
+	 * @see All, AtLeast, AtMost, Empty, Exclusive, GreaterThan, Inclusive
 	 */
 	static FORCEINLINE TRange LessThan(const ElementType& Value)
 	{
diff --git a/Engine/Source/Runtime/Core/Public/Math/RangeSet.h b/Engine/Source/Runtime/Core/Public/Math/RangeSet.h
--- a/Engine/Source/Runtime/Core/Public/Math/RangeSet.h
+++ b/Engine/Source/Runtime/Core/Public/Math/RangeSet.h
@@ -5,14 +5,17 @@
 #include "CoreTypes.h"
 #include "Containers/Array.h"
 #include "Math/Range.h"
+#include "Serialization/Archive.h"
+
 
 /**
  * Template for range sets.
  *
- * @todo gmp: TRangeSet needs an overhaul
+ * @todo gmp: Implement more efficient storage of range sets
  */
 template<typename ElementType> class TRangeSet
 {
+	typedef TRangeBound<ElementType> BoundsType;
 	typedef TRange<ElementType> RangeType;
 
 public:
@@ -28,42 +31,43 @@ public:
 	/**
 	 * Adds a range to the set.
 	 *
-	 * This method merges overlapping ranges into a single range (i.e. {[1, 4], [4, 6]} becomes [1, 6]).
-	 * Adjacent ranges (i.e. {[1, 3], [4, 6]}) are not merged.
+	 * This method merges overlapping ranges into a single range (i.e. {[1, 5], [4, 6]} becomes [1, 6]).
+	 * Adjacent ranges (i.e. {[1, 4), [4, 6)} are also merged.
 	 *
 	 * @param Range The range to add.
-	 *//*
+	 */
 	void Add(RangeType Range)
 	{
 		for (int32 Index = 0; Index < Ranges.Num(); ++Index)
 		{
-			RangeType& Current = Ranges(Index);
+			const RangeType& Current = Ranges[Index];
 
-			if ((Current.GetUpperBound() < Range.GetLowerBound()) || Current.GetLowerBound() > Range.GetUpperBound())
+			if (Current.Adjoins(Range) || Current.Overlaps(Range))
 			{
-				continue;
-			}
-
-			Range = RangeType(FMath::Min(Current.GetLowerBound(), Range.GetLowerBound()), FMath::Max(Current.GetUpperBound(), Range.GetUpperBound()));
+				Range = RangeType(
+					BoundsType::MinLower(Current.GetLowerBound(), Range.GetLowerBound()),
+					BoundsType::MaxUpper(Current.GetUpperBound(), Range.GetUpperBound())
+				);
 
-			Ranges.RemoveAtSwap(Index--);
+				Ranges.RemoveAtSwap(Index--);
+			}
 		}
 
 		Ranges.Add(Range);
-	}*/
+	}
 
 	/**
 	 * Merges another range set into this set.
 	 *
 	 * @param Other The range set to merge.
 	 */
-/*	void Merge(const TRangeSet& Other)
+	void Merge(const TRangeSet& Other)
 	{
 		for (typename TArray<RangeType>::TConstIterator It(Other.Ranges); It; ++It)
 		{
 			Add(*It);
 		}
-	}*/
+	}
 
 	/**
 	 * Removes a range from the set.
@@ -71,29 +75,27 @@ public:
 	 * Ranges that overlap with the removed range will be split.
 	 *
 	 * @param Range The range to remove.
-	 */
-/*	void Remove(const RangeType& Range)
+	 *//*
+	void Remove(const RangeType& Range)
 	{
 		for (int32 Index = 0; Index < Ranges.Num(); ++Index)
 		{
-			RangeType& Current = Ranges(Index);
+			const RangeType& Current = Ranges(Index);
 
-			if ((Current.GetUpperBound() < Range.GetLowerBound()) || (Current.GetLowerBound() > Range.GetUpperBound()))
+			if (Current.Overlaps(Range))
 			{
-				continue;
-			}
+				if (Current.GetLowerBound() < Range.GetLowerBound())
+				{
+					Ranges.Add(RangeType(Current.GetLowerBound(), Range.GetLowerBound()));
+				}
 
-			if (Current.GetLowerBound() < Range.GetLowerBound())
-			{
-				Ranges.Add(RangeType(Current.GetLowerBound(), Range.GetLowerBound()));
-			}
+				if (Current.GetUpperBound() > Range.GetUpperBound())
+				{
+					Ranges.Add(RangeType(Range.GetUpperBound(), Current.GetUpperBound()));
+				}
 
-			if (Current.GetUpperBound() > Range.GetUpperBound())
-			{
-				Ranges.Add(RangeType(Range.GetUpperBound(), Current.GetUpperBound()));
+				Ranges.RemoveAtSwap(Index--);
 			}
-
-			Ranges.RemoveAtSwap(Index--);
 		}
 	}*/
 
@@ -142,17 +144,104 @@ public:
 
 		return false;
 	}
-	
+
+	/**
+	 * Gets the range set's lowest bound.
+	 *
+	 * @return Lowest bound.
+	 * @see GetMaxBound, GetMinBoundValue, HasMinBound
+	 */
+	BoundsType GetMinBound() const
+	{
+		BoundsType Result;
+
+		for (const auto& Range : Ranges)
+		{
+			Result = BoundsType::MinLower(Result, Range.GetLowerBound());
+		}
+
+		return Result;
+	}
+
+	/**
+	 * Gets the value of the lowest bound.
+	 *
+	 * Use HasMinBound() to ensure that this range set actually has a lowest bound.
+	 *
+	 * @return Bound value.
+	 * @see GetMaxBoundValue, GetMinBound, HasMinBound
+	 */
+	const ElementType& GetMinBoundValue() const
+	{
+		return GetMinBound().GetValue();
+	}
+
+	/**
+	 * Gets the range set's uppermost bound.
+	 *
+	 * @return Uppermost bound.
+	 * @see GetMaxBoundValue, GetMinBound, HasMaxBound
+	 */
+	BoundsType GetMaxBound() const
+	{
+		BoundsType Result;
+
+		for (const auto& Range : Ranges)
+		{
+			Result = BoundsType::MaxUpper(Result, Range.GetUpperBound());
+		}
+
+		return Result;
+	}
+
+	/**
+	 * Gets the value of the uppermost bound.
+	 *
+	 * Use HasMaxBound() to ensure that this range actually has an upper bound.
+	 *
+	 * @return Bound value.
+	 * @see GetMaxBound, GetMinBoundValue, HasMaxBound
+	 */
+	const ElementType& GetMaxBoundValue() const
+	{
+		return GetMaxBound().GetValue();
+	}
+
+
 	/**
 	 * Returns a read-only collection of the ranges contained in this set.
 	 *
-	 * @return Array of ranges.
+	 * @param Allocator The array allocator to use.
+	 * @param OutRanges Will contain the collection of ranges.
+	 */
+	template<typename Allocator>
+	const void GetRanges(TArray<RangeType, Allocator>& OutRanges) const
+	{
+		OutRanges = Ranges;
+	}
+	
+	/**
+	 * Checks whether the range has a lowest bound.
+	 *
+	 * @return true if the range has a lowest bound, false otherwise.
+	 * @see GetMinBound, GetMinBoundValue, HasMaxBound
 	 */
-	const TArray<RangeType>& GetRanges() const
+	bool HasMinBound() const
 	{
-		return Ranges;
+		return GetMinBound().IsClosed();
 	}
 	
+	/**
+	 * Checks whether the range has an uppermost bound.
+	 *
+	 * @return true if the range has an uppermost bound, false otherwise.
+	 * @see GetUpperBound, GetUpperBoundValue, HasMinBound
+	 */
+	bool HasMaxBound() const
+	{
+		return GetMaxBound().IsClosed();
+	}
+
 	/**
 	 * Checks whether this range set is empty.
 	 *
@@ -202,7 +291,7 @@ public:
 
 		return false;
 	}
-	
+
 public:
 
 	/**
diff --git a/Engine/Source/Runtime/Core/Public/Math/UnitConversion.h b/Engine/Source/Runtime/Core/Public/Math/UnitConversion.h
--- a/Engine/Source/Runtime/Core/Public/Math/UnitConversion.h
+++ b/Engine/Source/Runtime/Core/Public/Math/UnitConversion.h
@@ -52,8 +52,8 @@ enum class EUnit : uint8
 	/** Pixel density units */
 	PixelsPerInch,
 
-	/** Arbitrary multiplier */
-	Multiplier,
+	/** Arbitrary multipliers */
+	Percentage,	Multiplier,
 
 	/** Symbolic entry, not specifiable on meta data */
 	Unspecified
@@ -62,7 +62,7 @@ enum class EUnit : uint8
 /** Enumeration that specifies particular classes of unit */
 enum class EUnitType
 {
-	Distance, Angle, Speed, Temperature, Mass, Force, Frequency, DataSize, LuminousFlux, Time, PixelDensity, Arbitrary,
+	Distance, Angle, Speed, Temperature, Mass, Force, Frequency, DataSize, LuminousFlux, Time, PixelDensity, Multipliers, Arbitrary,
 
 	// Symbolic entry - do not use directly
 	NumberOf,
diff --git a/Engine/Source/Runtime/Core/Public/Math/UnitConversion.inl b/Engine/Source/Runtime/Core/Public/Math/UnitConversion.inl
--- a/Engine/Source/Runtime/Core/Public/Math/UnitConversion.inl
+++ b/Engine/Source/Runtime/Core/Public/Math/UnitConversion.inl
@@ -34,6 +34,8 @@ namespace UnitConversion
 	CORE_API double DataSizeUnificationFactor(EUnit From);
 	/** Find the common quantization factor for the specified time unit. Quantizes to hours. */
 	CORE_API double TimeUnificationFactor(EUnit From);
+	/** Find the common quantization factor for the specified multiplier unit. Quantizes to 1.0 scale (where 1.0 == 100%). */
+	CORE_API double MultiplierUnificationFactor(EUnit From);
 
 	/** Attempt to parse an expression into a numeric unit */
 	CORE_API TValueOrError<FNumericUnit<double>, FText> TryParseExpression(const TCHAR* InExpression, EUnit DefaultUnit, const FNumericUnit<double>& InExistingValue);
@@ -80,7 +82,8 @@ T FUnitConversion::Convert(T InValue, EUnit From, EUnit To)
 		case EUnitType::Frequency:			return InValue * FrequencyUnificationFactor(From) 		* (1.0 / FrequencyUnificationFactor(To));
 		case EUnitType::DataSize:			return InValue * DataSizeUnificationFactor(From) 		* (1.0 / DataSizeUnificationFactor(To));
 		case EUnitType::LuminousFlux:		return InValue;
-		case EUnitType::Time:				return InValue * TimeUnificationFactor(From) 			* (1.0 / TimeUnificationFactor(To));			
+		case EUnitType::Time:				return InValue * TimeUnificationFactor(From) 			* (1.0 / TimeUnificationFactor(To));
+		case EUnitType::Multipliers:		return InValue * MultiplierUnificationFactor(From) 		* (1.0 / MultiplierUnificationFactor(To));
 		// Temperature conversion is not just a simple multiplication, so needs special treatment
 		case EUnitType::Temperature:
 		{
diff --git a/Engine/Source/Runtime/Core/Public/Misc/Char.h b/Engine/Source/Runtime/Core/Public/Misc/Char.h
--- a/Engine/Source/Runtime/Core/Public/Misc/Char.h
+++ b/Engine/Source/Runtime/Core/Public/Misc/Char.h
@@ -101,6 +101,8 @@ public:
 	static inline bool IsUpper(CharType Char);
 	static inline bool IsLower(CharType Char);
 	static inline bool IsAlpha(CharType Char);
+	static inline bool IsGraph(CharType Char);
+	static inline bool IsPrint(CharType Char);
 	static inline bool IsPunct(CharType Char);
 
 	static inline bool IsAlnum(CharType Char);
@@ -138,6 +140,8 @@ template <> inline TChar<WIDECHAR>::CharType TChar<WIDECHAR>::ToLower(CharType C
 template <> inline bool TChar<WIDECHAR>::IsUpper(CharType Char)							{ return ::iswupper(Char) != 0; }
 template <> inline bool TChar<WIDECHAR>::IsLower(CharType Char)							{ return ::iswlower(Char) != 0; }
 template <> inline bool TChar<WIDECHAR>::IsAlpha(CharType Char)							{ return ::iswalpha(Char) != 0; }
+template <> inline bool TChar<WIDECHAR>::IsGraph(CharType Char)							{ return ::iswgraph(Char) != 0; }
+template <> inline bool TChar<WIDECHAR>::IsPrint(CharType Char)							{ return ::iswprint(Char) != 0; }
 template <> inline bool TChar<WIDECHAR>::IsPunct(CharType Char)							{ return ::iswpunct(Char) != 0; }
 template <> inline bool TChar<WIDECHAR>::IsAlnum(CharType Char)							{ return ::iswalnum(Char) != 0; }
 template <> inline bool TChar<WIDECHAR>::IsDigit(CharType Char)							{ return ::iswdigit(Char) != 0; }
@@ -152,6 +156,8 @@ template <> inline TChar<ANSICHAR>::CharType TChar<ANSICHAR>::ToLower(CharType C
 template <> inline bool TChar<ANSICHAR>::IsUpper(CharType Char)							{ return ::isupper((unsigned char)Char) != 0; }
 template <> inline bool TChar<ANSICHAR>::IsLower(CharType Char)							{ return ::islower((unsigned char)Char) != 0; }
 template <> inline bool TChar<ANSICHAR>::IsAlpha(CharType Char)							{ return ::isalpha((unsigned char)Char) != 0; }
+template <> inline bool TChar<ANSICHAR>::IsGraph(CharType Char)							{ return ::isgraph((unsigned char)Char) != 0; }
+template <> inline bool TChar<ANSICHAR>::IsPrint(CharType Char)							{ return ::isprint((unsigned char)Char) != 0; }
 template <> inline bool TChar<ANSICHAR>::IsPunct(CharType Char)							{ return ::ispunct((unsigned char)Char) != 0; }
 template <> inline bool TChar<ANSICHAR>::IsAlnum(CharType Char)							{ return ::isalnum((unsigned char)Char) != 0; }
 template <> inline bool TChar<ANSICHAR>::IsDigit(CharType Char)							{ return ::isdigit((unsigned char)Char) != 0; }
diff --git a/Engine/Source/Runtime/Core/Public/Misc/DateTime.h b/Engine/Source/Runtime/Core/Public/Misc/DateTime.h
--- a/Engine/Source/Runtime/Core/Public/Misc/DateTime.h
+++ b/Engine/Source/Runtime/Core/Public/Misc/DateTime.h
@@ -6,6 +6,11 @@
 #include "Containers/UnrealString.h"
 #include "Misc/Timespan.h"
 
+class FArchive;
+class FOutputDevice;
+class UObject;
+
+
 /**
  * Enumerates the days of the week in 7-day calendars.
  */
@@ -75,7 +80,7 @@ public:
 	 *
 	 * @param Ticks The ticks representing the date and time.
 	 */
-	FDateTime( int64 InTicks )
+	FDateTime(int64 InTicks)
 		: Ticks(InTicks)
 	{ }
 
@@ -90,7 +95,7 @@ public:
 	 * @param Second The second (optional).
 	 * @param Millisecond The millisecond (optional).
 	 */
-	CORE_API FDateTime( int32 Year, int32 Month, int32 Day, int32 Hour = 0, int32 Minute = 0, int32 Second = 0, int32 Millisecond = 0 );
+	CORE_API FDateTime(int32 Year, int32 Month, int32 Day, int32 Hour = 0, int32 Minute = 0, int32 Second = 0, int32 Millisecond = 0);
 
 public:
 
@@ -100,7 +105,7 @@ public:
 	 * @return A date whose value is the sum of this date and the given time span.
 	 * @see FTimespan
 	 */
-	FDateTime operator+( const FTimespan& Other ) const
+	FDateTime operator+(const FTimespan& Other) const
 	{
 		return FDateTime(Ticks + Other.GetTicks());
 	}
@@ -111,7 +116,7 @@ public:
 	 * @return This date.
 	 * @see FTimespan
 	 */
-	FDateTime& operator+=( const FTimespan& Other )
+	FDateTime& operator+=(const FTimespan& Other)
 	{
 		Ticks += Other.GetTicks();
 
@@ -124,7 +129,7 @@ public:
 	 * @return A time span whose value is the difference of this date and the given date.
 	 * @see FTimespan
 	 */
-	FTimespan operator-( const FDateTime& Other ) const
+	FTimespan operator-(const FDateTime& Other) const
 	{
 		return FTimespan(Ticks - Other.Ticks);
 	}
@@ -135,7 +140,7 @@ public:
 	 * @return A date whose value is the difference of this date and the given time span.
 	 * @see FTimespan
 	 */
-	FDateTime operator-( const FTimespan& Other ) const
+	FDateTime operator-(const FTimespan& Other) const
 	{
 		return FDateTime(Ticks - Other.GetTicks());
 	}
@@ -146,7 +151,7 @@ public:
 	 * @return This date.
 	 * @see FTimespan
 	 */
-	FDateTime& operator-=( const FTimespan& Other )
+	FDateTime& operator-=(const FTimespan& Other)
 	{
 		Ticks -= Other.GetTicks();
 
@@ -159,7 +164,7 @@ public:
 	 * @param other The date to compare with.
 	 * @return true if the dates are equal, false otherwise.
 	 */
-	bool operator==( const FDateTime& Other ) const
+	bool operator==(const FDateTime& Other) const
 	{
 		return (Ticks == Other.Ticks);
 	}
@@ -170,7 +175,7 @@ public:
 	 * @param other The date to compare with.
 	 * @return true if the dates are not equal, false otherwise.
 	 */
-	bool operator!=( const FDateTime& Other ) const
+	bool operator!=(const FDateTime& Other) const
 	{
 		return (Ticks != Other.Ticks);
 	}
@@ -181,7 +186,7 @@ public:
 	 * @param other The date to compare with.
 	 * @return true if this date is greater, false otherwise.
 	 */
-	bool operator>( const FDateTime& Other ) const
+	bool operator>(const FDateTime& Other) const
 	{
 		return (Ticks > Other.Ticks);
 	}
@@ -192,7 +197,7 @@ public:
 	 * @param other The date to compare with.
 	 * @return true if this date is greater or equal, false otherwise.
 	 */
-	bool operator>=( const FDateTime& Other ) const
+	bool operator>=(const FDateTime& Other) const
 	{
 		return (Ticks >= Other.Ticks);
 	}
@@ -203,7 +208,7 @@ public:
 	 * @param other The date to compare with.
 	 * @return true if this date is less, false otherwise.
 	 */
-	bool operator<( const FDateTime& Other ) const
+	bool operator<(const FDateTime& Other) const
 	{
 		return (Ticks < Other.Ticks);
 	}
@@ -214,7 +219,7 @@ public:
 	 * @param other The date to compare with.
 	 * @return true if this date is less or equal, false otherwise.
 	 */
-	bool operator<=( const FDateTime& Other ) const
+	bool operator<=(const FDateTime& Other) const
 	{
 		return (Ticks <= Other.Ticks);
 	}
@@ -232,7 +237,7 @@ public:
 	 * @return true on success, false otherwise.
 	 * @see ImportTextItem
 	 */
-	CORE_API bool ExportTextItem( FString& ValueStr, FDateTime const& DefaultValue, UObject* Parent, int32 PortFlags, UObject* ExportRootScope ) const;
+	CORE_API bool ExportTextItem(FString& ValueStr, FDateTime const& DefaultValue, UObject* Parent, int32 PortFlags, UObject* ExportRootScope) const;
 
 	/**
 	 * Gets the date part of this date.
@@ -253,7 +258,7 @@ public:
 	 * @param OutMonth Will contain the number of the month (1-12).
 	 * @param OutDay Will contain the number of the day (1-31).
 	 */
-	CORE_API void GetDate( int32& OutYear, int32& OutMonth, int32& OutDay ) const;
+	CORE_API void GetDate(int32& OutYear, int32& OutMonth, int32& OutDay) const;
 
 	/**
 	 * Gets this date's day part (1 to 31).
@@ -418,7 +423,7 @@ public:
 	 * @return true on success, false otherwise.
 	 * @see ExportTextItem
 	 */
-	CORE_API bool ImportTextItem( const TCHAR*& Buffer, int32 PortFlags, UObject* Parent, FOutputDevice* ErrorText );
+	CORE_API bool ImportTextItem(const TCHAR*& Buffer, int32 PortFlags, UObject* Parent, FOutputDevice* ErrorText);
 
 	/**
 	 * Gets whether this date's time is in the afternoon.
@@ -448,27 +453,27 @@ public:
 	 * @param Ar The archive to serialize from or into.
 	 * @return true on success, false otherwise.
 	 */
-	CORE_API bool Serialize( FArchive& Ar );
+	CORE_API bool Serialize(FArchive& Ar);
 
 	/**
-	 * Returns the ISO-8601 string representation of the FDateTime.
+	 * Returns the RFC 1123 string representation of the FDateTime.
 	 *
 	 * The resulting string assumes that the FDateTime is in UTC.
 	 * 
 	 * @return String representation.
-	 * @see ParseIso8601, ToString
+	 * @see ParseHttpDate, ToIso8601, ToString
 	 */
-	CORE_API FString ToIso8601() const;
+	CORE_API FString ToHttpDate() const;
 
 	/**
-	 * Returns the RFC 1123 string representation of the FDateTime.
+	 * Returns the ISO-8601 string representation of the FDateTime.
 	 *
 	 * The resulting string assumes that the FDateTime is in UTC.
 	 * 
 	 * @return String representation.
-	 * @see ParseHttpDate, ToString
+	 * @see ParseIso8601, ToHttpDate, ToString
 	 */
-	CORE_API FString ToHttpDate() const;
+	CORE_API FString ToIso8601() const;
 
 	/**
 	 * Returns the string representation of this date using a default format.
@@ -488,7 +493,7 @@ public:
 	 * @return String representation.
 	 * @see Parse, ToIso8601
 	 */
-	CORE_API FString ToString( const TCHAR* Format ) const;
+	CORE_API FString ToString(const TCHAR* Format) const;
 
 	/**
 	 * Returns this date as the number of seconds since the Unix Epoch (January 1st of 1970).
@@ -511,7 +516,7 @@ public:
 	 * @return The number of days
 	 * @see DaysInYear
 	 */
-	static CORE_API int32 DaysInMonth( int32 Year, int32 Month );
+	static CORE_API int32 DaysInMonth(int32 Year, int32 Month);
 
 	/**
 	 * Gets the number of days in the given year.
@@ -520,7 +525,7 @@ public:
 	 * @return The number of days.
 	 * @see DaysInMonth
 	 */
-	static CORE_API int32 DaysInYear( int32 Year );
+	static CORE_API int32 DaysInYear(int32 Year);
 
 	/**
 	 * Returns the proleptic Gregorian date for the given Julian Day.
@@ -529,7 +534,7 @@ public:
 	 * @return Gregorian date and time.
 	 * @see GetJulianDay
 	 */
-	static FDateTime FromJulianDay( double JulianDay )
+	static FDateTime FromJulianDay(double JulianDay)
 	{
 		return FDateTime((int64)((JulianDay - 1721425.5) * ETimespan::TicksPerDay));
 	}
@@ -541,7 +546,7 @@ public:
 	 * @return Gregorian date and time.
 	 * @see ToUnixTimestamp
 	 */
-	static FDateTime FromUnixTimestamp( int64 UnixTime )
+	static FDateTime FromUnixTimestamp(int64 UnixTime)
 	{
 		return FDateTime(1970, 1, 1) + FTimespan(UnixTime * ETimespan::TicksPerSecond);
 	}
@@ -556,7 +561,7 @@ public:
 	 * @param Year The year to check.
 	 * @return true if the year is a leap year, false otherwise.
 	 */
-	static CORE_API bool IsLeapYear( int32 Year );
+	static CORE_API bool IsLeapYear(int32 Year);
 
 	/**
 	 * Returns the maximum date value.
@@ -603,29 +608,42 @@ public:
 	 * @param DateTimeString The string to convert.
 	 * @param OutDateTime Will contain the parsed date and time.
 	 * @return true if the string was converted successfully, false otherwise.
+	 * @see ParseHttpDate, ParseIso8601, ToString
 	 */
-	static CORE_API bool Parse( const FString& DateTimeString, FDateTime& OutDateTime );
+	static CORE_API bool Parse(const FString& DateTimeString, FDateTime& OutDateTime);
 
 	/**
-	 * Parses a date string in ISO-8601 format.
-	 * 
-	 * @param DateTimeString The string to be parsed
-	 * @param OutDateTime FDateTime object (in UTC) corresponding to the input string (which may have been in any timezone).
+	 * Parses a date string in HTTP-date format (rfc1123-date | rfc850-date | asctime-date)
+	 * https://www.w3.org/Protocols/rfc2616/rfc2616-sec3.html#sec3.3.1
+	 *
+ 	 * HTTP-date    = rfc1123-date | rfc850-date | asctime-date
+	 * rfc1123-date = wkday "," SP date1 SP time SP "GMT"
+	 * rfc850-date  = weekday "," SP date2 SP time SP "GMT"
+	 * asctime-date = wkday SP date3 SP time SP 4DIGIT
+	 * date1        = 2DIGIT SP month SP 4DIGIT ; day month year (e.g., 02 Jun 1982)
+	 * date2        = 2DIGIT "-" month "-" 2DIGIT ; day-month-year (e.g., 02-Jun-82)
+	 * date3        = month SP (2DIGIT | (SP 1DIGIT)) ; month day (e.g., Jun  2)
+	 * time         = 2DIGIT ":" 2DIGIT ":" 2DIGIT ; 00:00:00 - 23:59:59
+	 * wkday        = "Mon" | "Tue" | "Wed" | "Thu" | "Fri" | "Sat" | "Sun"
+	 * weekday      = "Monday" | "Tuesday" | "Wednesday" | "Thursday" | "Friday" | "Saturday" | "Sunday"
+	 * month        = "Jan" | "Feb" | "Mar" | "Apr" | "May" | "Jun" | "Jul" | "Aug" | "Sep" | "Oct" | "Nov" | "Dec"
+	 *
+	 * @param HttpDate The string to be parsed
+	 * @param OutDateTime FDateTime object (assumes UTC) corresponding to the input string.
 	 * @return true if the string was converted successfully, false otherwise.
-	 * @see Parse, ToIso8601
+	 * @see Parse, ToHttpDate, ParseIso8601
 	 */
-	static CORE_API bool ParseIso8601( const TCHAR* DateTimeString, FDateTime& OutDateTime );
+	static CORE_API bool ParseHttpDate(const FString& HttpDate, FDateTime& OutDateTime);
 
 	/**
-	 * Parses a date string in HTTP-date format (rfc1123-date | rfc850-date | asctime-date)
-	 * https://www.w3.org/Protocols/rfc2616/rfc2616-sec3.html#sec3.3.1
+	 * Parses a date string in ISO-8601 format.
 	 * 
-	 * @param HttpDate The string to be parsed
-	 * @param OutDateTime FDateTime object (assumes UTC) corresponding to the input string.
+	 * @param DateTimeString The string to be parsed
+	 * @param OutDateTime FDateTime object (in UTC) corresponding to the input string (which may have been in any timezone).
 	 * @return true if the string was converted successfully, false otherwise.
-	 * @see Parse
+	 * @see Parse, ParseHttpDate, ToIso8601
 	 */
-	static CORE_API bool ParseHttpDate( const FString& HttpDate, FDateTime& OutDateTime );
+	static CORE_API bool ParseIso8601(const TCHAR* DateTimeString, FDateTime& OutDateTime);
 
 	/**
 	 * Gets the local date on this computer.
@@ -667,7 +685,7 @@ public:
 	 *
 	 * @return true if the components are valid, false otherwise.
 	 */
-	static CORE_API bool Validate( int32 Year, int32 Month, int32 Day, int32 Hour, int32 Minute, int32 Second, int32 Millisecond );
+	static CORE_API bool Validate(int32 Year, int32 Month, int32 Day, int32 Hour, int32 Minute, int32 Second, int32 Millisecond);
 
 public:
 
@@ -677,10 +695,11 @@ public:
 	 * @param Ar The archive to serialize from or into.
 	 * @param DateTime The date and time value to serialize.
 	 * @return The archive.
-	 *
-	 * @todo gmp: Figure out better include order in Core.h so this can be inlined.
 	 */
-	friend CORE_API FArchive& operator<<( FArchive& Ar, FDateTime& DateTime );
+	friend CORE_API FArchive& operator<<(FArchive& Ar, FDateTime& DateTime)
+	{
+		return Ar << DateTime.Ticks;
+	}
 
 	/**
 	 * Gets the hash for the specified date and time.
@@ -688,7 +707,10 @@ public:
 	 * @param DateTime The date and time to get the hash for.
 	 * @return Hash value.
 	 */
-	friend CORE_API uint32 GetTypeHash(const FDateTime& DateTime);
+	friend CORE_API uint32 GetTypeHash(const FDateTime& DateTime)
+	{
+		return GetTypeHash(DateTime.Ticks);
+	}
 
 protected:
 
diff --git a/Engine/Source/Runtime/Core/Public/Misc/Guid.h b/Engine/Source/Runtime/Core/Public/Misc/Guid.h
--- a/Engine/Source/Runtime/Core/Public/Misc/Guid.h
+++ b/Engine/Source/Runtime/Core/Public/Misc/Guid.h
@@ -7,6 +7,11 @@
 #include "Misc/Crc.h"
 #include "Containers/UnrealString.h"
 
+class FArchive;
+class FOutputDevice;
+class UObject;
+
+
 /**
  * Enumerates known GUID formats.
  */
@@ -79,7 +84,7 @@ public:
 	 * @param InC The third component.
 	 * @param InD The fourth component.
 	 */
-	FGuid( uint32 InA, uint32 InB, uint32 InC, uint32 InD )
+	FGuid(uint32 InA, uint32 InB, uint32 InC, uint32 InD)
 		: A(InA), B(InB), C(InC), D(InD)
 	{ }
 
@@ -92,7 +97,7 @@ public:
 	 * @param Y The second GUID to compare.
 	 * @return true if the GUIDs are equal, false otherwise.
 	 */
-	friend bool operator==( const FGuid& X, const FGuid& Y )
+	friend bool operator==(const FGuid& X, const FGuid& Y)
 	{
 		return ((X.A ^ Y.A) | (X.B ^ Y.B) | (X.C ^ Y.C) | (X.D ^ Y.D)) == 0;
 	}
@@ -104,7 +109,7 @@ public:
 	 * @param Y The second GUID to compare.
 	 * @return true if the GUIDs are not equal, false otherwise.
 	 */
-	friend bool operator!=( const FGuid& X, const FGuid& Y )
+	friend bool operator!=(const FGuid& X, const FGuid& Y)
 	{
 		return ((X.A ^ Y.A) | (X.B ^ Y.B) | (X.C ^ Y.C) | (X.D ^ Y.D)) != 0;
 	}
@@ -116,12 +121,12 @@ public:
 	 * @param Y The second GUID to compare.
 	 * @return true if the first GUID is less than the second one.
 	 */
-	friend bool operator<( const FGuid& X, const FGuid& Y )
+	friend bool operator<(const FGuid& X, const FGuid& Y)
 	{
 		return	((X.A < Y.A) ? true : ((X.A > Y.A) ? false :
 				((X.B < Y.B) ? true : ((X.B > Y.B) ? false :
 				((X.C < Y.C) ? true : ((X.C > Y.C) ? false :
-				((X.D < Y.D) ? true : ((X.D > Y.D) ? false : false )))))))); //-V583
+				((X.D < Y.D) ? true : ((X.D > Y.D) ? false : false)))))))); //-V583
 	}
 
 	/**
@@ -130,7 +135,7 @@ public:
 	 * @param Index The index of the component to return (0...3).
 	 * @return The component.
 	 */
-	uint32& operator[]( int32 Index )
+	uint32& operator[](int32 Index)
 	{
 		checkSlow(Index >= 0);
 		checkSlow(Index < 4);
@@ -152,7 +157,7 @@ public:
 	 * @param Index The index of the component to return (0...3).
 	 * @return The component.
 	 */
-	const uint32& operator[]( int32 Index ) const
+	const uint32& operator[](int32 Index) const
 	{
 		checkSlow(Index >= 0);
 		checkSlow(Index < 4);
@@ -174,12 +179,12 @@ public:
 	 * @param Ar The archive to serialize from or into.
 	 * @param G The GUID to serialize.
 	 */
-	friend FArchive& operator<<( FArchive& Ar, FGuid& G )
+	friend FArchive& operator<<(FArchive& Ar, FGuid& G)
 	{
 		return Ar << G.A << G.B << G.C << G.D;
 	}
 
-	bool Serialize( FArchive& Ar )
+	bool Serialize(FArchive& Ar)
 	{
 		Ar << *this;
 		return true;
@@ -198,7 +203,7 @@ public:
 	 * @return true on success, false otherwise.
 	 * @see ImportTextItem
 	 */
-	CORE_API bool ExportTextItem( FString& ValueStr, FGuid const& DefaultValue, UObject* Parent, int32 PortFlags, class UObject* ExportRootScope ) const;
+	CORE_API bool ExportTextItem(FString& ValueStr, FGuid const& DefaultValue, UObject* Parent, int32 PortFlags, class UObject* ExportRootScope) const;
 
 	/**
 	 * Imports the GUIDs value from a text buffer.
@@ -210,7 +215,7 @@ public:
 	 * @return true on success, false otherwise.
 	 * @see ExportTextItem
 	 */
-	CORE_API bool ImportTextItem( const TCHAR*& Buffer, int32 PortFlags, class UObject* Parent, FOutputDevice* ErrorText );
+	CORE_API bool ImportTextItem(const TCHAR*& Buffer, int32 PortFlags, UObject* Parent, FOutputDevice* ErrorText);
 
 	/**
 	 * Invalidates the GUID.
@@ -251,7 +256,7 @@ public:
 	 * @param Format The string format to use.
 	 * @return The string representation.
 	 */
-	CORE_API FString ToString( EGuidFormats Format ) const;
+	CORE_API FString ToString(EGuidFormats Format) const;
 
 public:
 
@@ -261,7 +266,7 @@ public:
 	 * @param Guid The GUID to calculate the hash for.
 	 * @return The hash.
 	 */
-	friend uint32 GetTypeHash( const FGuid& Guid )
+	friend uint32 GetTypeHash(const FGuid& Guid)
 	{
 		return FCrc::MemCrc_DEPRECATED(&Guid, sizeof(FGuid));
 	}
@@ -283,7 +288,7 @@ public:
 	 * @return true if the string was converted successfully, false otherwise.
 	 * @see ParseExact, ToString
 	 */
-	static CORE_API bool Parse( const FString& GuidString, FGuid& OutGuid );
+	static CORE_API bool Parse(const FString& GuidString, FGuid& OutGuid);
 
 	/**
 	 * Converts a string with the specified format to a GUID.
@@ -294,7 +299,7 @@ public:
 	 * @return true if the string was converted successfully, false otherwise.
 	 * @see Parse, ToString
 	 */
-	static CORE_API bool ParseExact( const FString& GuidString, EGuidFormats Format, FGuid& OutGuid );
+	static CORE_API bool ParseExact(const FString& GuidString, EGuidFormats Format, FGuid& OutGuid);
 
 //private:
 public:
@@ -312,6 +317,7 @@ public:
 	uint32 D;
 };
 
+
 template <> struct TIsPODType<FGuid> { enum { Value = true }; };
 
 namespace Lex
