Copying //UE4/Dev-Sequencer to //UE4/Main (Source: //UE4/Dev-Sequencer @ 3617642)

diff --git a/Engine/Source/Runtime/Core/Public/Misc/Timespan.h b/Engine/Source/Runtime/Core/Public/Misc/Timespan.h
--- a/Engine/Source/Runtime/Core/Public/Misc/Timespan.h
+++ b/Engine/Source/Runtime/Core/Public/Misc/Timespan.h
@@ -3,12 +3,27 @@
 #pragma once
 
 #include "CoreTypes.h"
-#include "CoreFwd.h"
+#include "Containers/UnrealString.h"
 
-class FString;
+class FArchive;
+class FOutputDevice;
+class UObject;
 
+
+/**
+ * Time span related constants.
+ */
 namespace ETimespan
 {
+	/** The maximum number of ticks that can be represented in FTimespan. */
+	const int64 MaxTicks = 9223372036854775807;
+
+	/** The minimum number of ticks that can be represented in FTimespan. */
+	const int64 MinTicks = -9223372036854775807 - 1;
+
+	/** The number of nanoseconds per tick. */
+	const int64 NanosecondsPerTick = 100;
+
 	/** The number of timespan ticks per day. */
 	const int64 TicksPerDay = 864000000000;
 
@@ -44,6 +59,10 @@ namespace ETimespan
  * based arithmetic, such as calculating the difference between two dates or adding a certain amount
  * of time to a given date.
  *
+ * When initializing time span values from single components, consider using the FromHours, FromMinutes,
+ * FromSeconds, Zero, MinValue and related methods instead of calling the overloaded constructors as
+ * they will make your code easier to read and understand.
+ *
  * @see FDateTime
  */
 struct FTimespan
@@ -54,58 +73,68 @@ public:
 	FTimespan() { }
 
 	/**
-	 * Creates and initializes a new time interval with the specified number of ticks.
+	 * Create and initialize a new time interval with the specified number of ticks.
+	 *
+	 * For better readability, consider using MinValue, MaxValue and Zero.
 	 *
 	 * @param Ticks The number of ticks.
+	 * @see MaxValue, MinValue, Zero
 	 */
 	FTimespan(int64 InTicks)
 		: Ticks(InTicks)
-	{ }
+	{
+		check((InTicks >= ETimespan::MinTicks) && (InTicks <= ETimespan::MaxTicks));
+	}
 
 	/**
-	 * Creates and initializes a new time interval with the specified number of hours, minutes and seconds.
+	 * Create and initialize a new time interval with the specified number of hours, minutes and seconds.
+	 *
+	 * For better readability, consider using FromHours, FromMinutes and FromSeconds.
 	 *
 	 * @param Hours The hours component.
 	 * @param Minutes The minutes component.
 	 * @param Seconds The seconds component.
+	 * @see FromHours, FromMinutes, FromSeconds
 	 */
 	FTimespan(int32 Hours, int32 Minutes, int32 Seconds)
 	{
-		Assign(0, Hours, Minutes, Seconds, 0, 0);
+		Assign(0, Hours, Minutes, Seconds, 0);
 	}
 
 	/**
-	 * Creates and initializes a new time interval with the specified number of days, hours, minutes and seconds.
+	 * Create and initialize a new time interval with the specified number of days, hours, minutes and seconds.
+	 *
+	 * For better readability, consider using FromDays, FromHours, FromMinutes and FromSeconds.
 	 *
 	 * @param Days The days component.
 	 * @param Hours The hours component.
 	 * @param Minutes The minutes component.
 	 * @param Seconds The seconds component.
+	 * @see FromDays, FromHours, FromMinutes, FromSeconds
 	 */
 	FTimespan(int32 Days, int32 Hours, int32 Minutes, int32 Seconds)
 	{
-		Assign(Days, Hours, Minutes, Seconds, 0, 0);
+		Assign(Days, Hours, Minutes, Seconds, 0);
 	}
 
 	/**
-	 * Creates and initializes a new time interval with the specified number of days, hours, minutes and seconds.
+	 * Create and initialize a new time interval with the specified number of days, hours, minutes and seconds.
 	 *
 	 * @param Days The days component.
 	 * @param Hours The hours component.
 	 * @param Minutes The minutes component.
 	 * @param Seconds The seconds component.
-	 * @param Milliseconds The milliseconds component.
-	 * @param Microseconds The microseconds component (default = 0).
+	 * @param FractionNano The fractional seconds (in nanosecond resolution).
 	 */
-	FTimespan(int32 Days, int32 Hours, int32 Minutes, int32 Seconds, int32 Milliseconds, int32 Microseconds = 0)
-	{
-		Assign(Days, Hours, Minutes, Seconds, Milliseconds, Microseconds);
-	}
+ 	FTimespan(int32 Days, int32 Hours, int32 Minutes, int32 Seconds, int32 FractionNano)
+ 	{
+ 		Assign(Days, Hours, Minutes, Seconds, FractionNano);
+ 	}
 
 public:
 
 	/**
-	 * Returns result of adding the given time span to this time span.
+	 * Return the result of adding the given time span to this time span.
 	 *
 	 * @return A time span whose value is the sum of this time span and the given time span.
 	 */
@@ -126,7 +155,7 @@ public:
 	}
 
 	/**
-	 * Returns the inverse of this time span.
+	 * Return the inverse of this time span.
 	 *
 	 * The value of this time span must be greater than FTimespan::MinValue(), or else an overflow will occur.
 	 *
@@ -138,7 +167,7 @@ public:
 	}
 
 	/**
-	 * Returns the result of subtracting the given time span from this time span.
+	 * Return the result of subtracting the given time span from this time span.
 	 *
 	 * @param Other The time span to compare with.
 	 * @return A time span whose value is the difference of this time span and the given time span.
@@ -149,7 +178,7 @@ public:
 	}
 
 	/**
-	 * Subtracts the given time span from this time span.
+	 * Subtract the given time span from this time span.
 	 *
 	 * @param Other The time span to subtract.
 	 * @return This time span.
@@ -161,30 +190,53 @@ public:
 	}
 
 	/**
-	 * Returns the result of multiplying the this time span with the given scalar.
+	 * Return the result of multiplying the this time span with the given scalar.
 	 *
 	 * @param Scalar The scalar to multiply with.
 	 * @return A time span whose value is the product of this time span and the given scalar.
 	 */
-	FTimespan operator*(float Scalar) const
+	FTimespan operator*(double Scalar) const
 	{
 		return FTimespan((int64)(Ticks * Scalar));
 	}
 
 	/**
-	 * Multiplies this time span with the given scalar.
+	 * Multiply this time span with the given scalar.
 	 *
 	 * @param Scalar The scalar to multiply with.
 	 * @return This time span.
 	 */
-	FTimespan& operator*=(float Scalar)
+	FTimespan& operator*=(double Scalar)
 	{
 		Ticks = (int64)(Ticks * Scalar);
 		return *this;
 	}
 
 	/**
-	 * Returns the result of calculating the modulus of this time span with another time span.
+	 * Return the result of dividing the this time span by the given scalar.
+	 *
+	 * @param Scalar The scalar to divide by.
+	 * @return A time span whose value is the quotient of this time span and the given scalar.
+	 */
+	FTimespan operator/(double Scalar) const
+	{
+		return FTimespan((int64)(Ticks / Scalar));
+	}
+
+	/**
+	 * Divide this time span by the given scalar.
+	 *
+	 * @param Scalar The scalar to divide by.
+	 * @return This time span.
+	 */
+	FTimespan& operator/=(double Scalar)
+	{
+		Ticks = (int64)(Ticks / Scalar);
+		return *this;
+	}
+
+	/**
+	 * Return the result of calculating the modulus of this time span with another time span.
 	 *
 	 * @param Other The time span to divide by.
 	 * @return A time span representing the remainder of the modulus operation.
@@ -195,7 +247,7 @@ public:
 	}
 
 	/**
-	 * Calculates this time span modulo another.
+	 * Calculate this time span modulo another.
 	 *
 	 * @param Other The time span to divide by.
 	 * @return This time span.
@@ -207,7 +259,7 @@ public:
 	}
 
 	/**
-	 * Compares this time span with the given time span for equality.
+	 * Compare this time span with the given time span for equality.
 	 *
 	 * @param Other The time span to compare with.
 	 * @return true if the time spans are equal, false otherwise.
@@ -218,7 +270,7 @@ public:
 	}
 
 	/**
-	 * Compares this time span with the given time span for inequality.
+	 * Compare this time span with the given time span for inequality.
 	 *
 	 * @param Other The time span to compare with.
 	 * @return true if the time spans are not equal, false otherwise.
@@ -229,7 +281,7 @@ public:
 	}
 
 	/**
-	 * Checks whether this time span is greater than the given time span.
+	 * Check whether this time span is greater than the given time span.
 	 *
 	 * @param Other The time span to compare with.
 	 * @return true if this time span is greater, false otherwise.
@@ -240,7 +292,7 @@ public:
 	}
 
 	/**
-	 * Checks whether this time span is greater than or equal to the given time span.
+	 * Check whether this time span is greater than or equal to the given time span.
 	 *
 	 * @param Other The time span to compare with.
 	 * @return true if this time span is greater or equal, false otherwise.
@@ -251,7 +303,7 @@ public:
 	}
 
 	/**
-	 * Checks whether this time span is less than the given time span.
+	 * Check whether this time span is less than the given time span.
 	 *
 	 * @param Other The time span to compare with.
 	 * @return true if this time span is less, false otherwise.
@@ -262,7 +314,7 @@ public:
 	}
 
 	/**
-	 * Checks whether this time span is less than or equal to the given time span.
+	 * Check whether this time span is less than or equal to the given time span.
 	 *
 	 * @param Other The time span to compare with.
 	 * @return true if this time span is less or equal, false otherwise.
@@ -275,7 +327,7 @@ public:
 public:
 
 	/**
-	 * Exports this time span value to a string.
+	 * Export this time span value to a string.
 	 *
 	 * @param ValueStr Will hold the string value.
 	 * @param DefaultValue The default value.
@@ -288,7 +340,7 @@ public:
 	CORE_API bool ExportTextItem(FString& ValueStr, FTimespan const& DefaultValue, UObject* Parent, int32 PortFlags, UObject* ExportRootScope) const;
 
 	/**
-	 * Gets the days component of this time span.
+	 * Get the days component of this time span.
 	 *
 	 * @return Days component.
 	 */
@@ -298,7 +350,7 @@ public:
 	}
 
 	/**
-	 * Returns a time span with the absolute value of this time span.
+	 * Get a time span with the absolute value of this time span.
 	 *
 	 * This method may overflow the timespan if its value is equal to MinValue.
 	 *
@@ -311,40 +363,50 @@ public:
 	}
 
 	/**
-	 * Gets the hours component of this time span.
+	 * Gets the fractional seconds (in microsecond resolution).
 	 *
-	 * @return Hours component.
-	 * @see GetTotalHours
+	 * @return Number of microseconds in fractional part.
+	 * @see GetTotalMicroseconds
 	 */
-	int32 GetHours() const
+	int32 GetFractionMicro() const
 	{
-		return (int32)((Ticks / ETimespan::TicksPerHour) % 24);
+		return (int32)((Ticks % ETimespan::TicksPerSecond) / ETimespan::TicksPerMicrosecond);
 	}
 
 	/**
-	 * Gets the microseconds component of this time span.
+	 * Gets the fractional seconds (in millisecond resolution).
 	 *
-	 * @return Microseconds component.
-	 * @see GetTotalMicroseconds
+	 * @return Number of milliseconds in fractional part.
+	 * @see GetTotalMilliseconds
 	 */
-	int32 GetMicroseconds() const
+	int32 GetFractionMilli() const
 	{
-		return (int32)((Ticks / ETimespan::TicksPerMicrosecond) % 1000);
+		return (int32)((Ticks % ETimespan::TicksPerSecond) / ETimespan::TicksPerMillisecond);
 	}
 
 	/**
-	 * Gets the milliseconds component of this time span.
+	 * Gets the fractional seconds (in nanosecond resolution).
 	 *
-	 * @return Milliseconds component.
-	 * @see GetTotalMilliseconds
+	 * @return Number of nanoseconds in fractional part.
+	 */
+	int32 GetFractionNano() const
+	{
+		return (int32)((Ticks % ETimespan::TicksPerSecond) * ETimespan::NanosecondsPerTick);
+	}
+
+	/**
+	 * Gets the hours component of this time span.
+	 *
+	 * @return Hours component.
+	 * @see GetTotalHours
 	 */
-	int32 GetMilliseconds() const
+	int32 GetHours() const
 	{
-		return (int32)((Ticks / ETimespan::TicksPerMillisecond) % 1000);
+		return (int32)((Ticks / ETimespan::TicksPerHour) % 24);
 	}
 
 	/**
-	 * Gets the minutes component of this time span.
+	 * Get the minutes component of this time span.
 	 *
 	 * @return Minutes component.
 	 * @see GetTotalMinutes
@@ -355,7 +417,7 @@ public:
 	}
 
 	/**
-	 * Gets the seconds component of this time span.
+	 * Get the seconds component of this time span.
 	 *
 	 * @return Seconds component.
 	 * @see GetTotalSeconds
@@ -366,7 +428,7 @@ public:
 	}
 
 	/**
-	 * Gets the number of ticks represented by this time span.
+	 * Get the number of ticks represented by this time span.
 	 *
 	 * @return Number of ticks.
 	 */
@@ -376,7 +438,7 @@ public:
 	}
 
 	/**
-	 * Gets the total number of days represented by this time span.
+	 * Get the total number of days represented by this time span.
 	 *
 	 * @return Number of days.
 	 * @see GetDays
@@ -387,7 +449,7 @@ public:
 	}
 
 	/**
-	 * Gets the total number of hours represented by this time span.
+	 * Get the total number of hours represented by this time span.
 	 *
 	 * @return Number of hours.
 	 * @see GetHours
@@ -398,10 +460,10 @@ public:
 	}
 
 	/**
-	 * Gets the total number of microseconds represented by this time span.
+	 * Get the total number of microseconds represented by this time span.
 	 *
 	 * @return Number of microseconds.
-	 * @see GetMicroseconds
+	 * @see GetFractionMicro
 	 */
 	double GetTotalMicroseconds() const
 	{
@@ -409,10 +471,10 @@ public:
 	}
 
 	/**
-	 * Gets the total number of milliseconds represented by this time span.
+	 * Get the total number of milliseconds represented by this time span.
 	 *
 	 * @return Number of milliseconds.
-	 * @see GetMilliseconds
+	 * @see GetFractionMilli
 	 */
 	double GetTotalMilliseconds() const
 	{
@@ -420,7 +482,7 @@ public:
 	}
 
 	/**
-	 * Gets the total number of minutes represented by this time span.
+	 * Get the total number of minutes represented by this time span.
 	 *
 	 * @return Number of minutes.
 	 * @see GetMinutes
@@ -431,7 +493,7 @@ public:
 	}
 
 	/**
-	 * Gets the total number of seconds represented by this time span.
+	 * Get the total number of seconds represented by this time span.
 	 *
 	 * @return Number of seconds.
 	 * @see GetSeconds
@@ -442,7 +504,7 @@ public:
 	}
 
 	/**
-	 * Imports a time span value from a text buffer.
+	 * Import a time span value from a text buffer.
 	 *
 	 * @param Buffer The text buffer to import from.
 	 * @param PortFlags Not used.
@@ -465,7 +527,7 @@ public:
 	}
 
 	/**
-	 * Serializes this time span from or into the specified archive.
+	 * Serialize this time span from or into the specified archive.
 	 *
 	 * @param Ar The archive to serialize from or into.
 	 * @return true on success, false otherwise.
@@ -473,10 +535,20 @@ public:
 	CORE_API bool Serialize(FArchive& Ar);
 
 	/**
-	 * Returns the string representation of this time span using a default format.
+	 * Return the string representation of this time span using a default format.
 	 *
 	 * The returned string has the following format:
-	 *		[-][d.]hh:mm:ss.fff
+	 *		p[d.]hh:mm:ss.fff
+	 *
+	 * Note that 'p' is the plus or minus sign, and the date component is
+	 * omitted for time spans that are shorter than one day.
+	 *
+	 * Examples:
+	 *      -42.15:11:36.457 (45 days, 15 hours, 11 minutes, 36.457 seconds in the past)
+	 *      +42.15:11:36.457 (45 days, 15 hours, 11 minutes, 36.457 seconds in the future)
+	 *      +15:11:36.457 (15 hours, 11 minutes, 36.457 seconds in the future)
+	 *      +00:11:36.457 (11 minutes, 36.457 seconds in the future)
+	 *      +00:00:36.457 (36.457 seconds in the future)
 	 *
 	 * @return String representation.
 	 * @see Parse
@@ -484,23 +556,22 @@ public:
 	CORE_API FString ToString() const;
 
 	/**
-	 * Converts this time span to its string representation.
+	 * Convert this time span to its string representation.
 	 *
 	 * The following formatting codes are available:
-	 *		%n - prints the minus sign (for negative time spans only)
-	 *		%N - prints the minus or plus sign (always)
-	 *		%d - prints the time span's days part
-	 *		%h - prints the time span's hours part (0..23)
-	 *		%m - prints the time span's minutes part (0..59)
-	 *		%s - prints the time span's seconds part (0..59)
-	 *		%f - prints the time span's milliseconds part (0..999)
-	 *		%D - prints the total number of days (without minus sign)
-	 *		%H - prints the total number of hours (without minus sign)
-	 *		%M - prints the total number of minutes (without minus sign)
-	 *		%S - prints the total number of seconds (without minus sign)
-	 *		%F - prints the total number of milliseconds (without minus sign)
-	 *
-	 * @param Format - The format of the returned string.
+	 *		%d - prints the days component
+	 *		%D - prints the zero-padded days component (00000000..10675199)
+	 *		%h - prints the zero-padded hours component (00..23)
+	 *		%m - prints the zero-padded minutes component (00..59)
+	 *		%s - prints the zero-padded seconds component (00..59)
+	 *		%f - prints the zero-padded fractional seconds (000..999)
+	 *		%u - prints the zero-padded fractional seconds (000000..999999)
+	 *		%n - prints the zero-padded fractional seconds (000000000..999999999)
+	 *
+	 * Depending on whether the time span is positive or negative, a plus or minus
+	 * sign character will always be added in front of the generated string.
+	 *
+	 * @param Format The format of the returned string.
 	 * @return String representation.
 	 * @see Parse
 	 */
@@ -509,61 +580,79 @@ public:
 public:
 
 	/**
-	 * Creates a time span that represents the specified number of days.
+	 * Create a time span that represents the specified number of days.
 	 *
 	 * @param Days The number of days.
 	 * @return Time span.
 	 * @see FromHours, FromMicroseconds, FromMilliseconds, FromMinutes, FromSeconds
 	 */
-	static CORE_API FTimespan FromDays(double Days);
+	static FTimespan FromDays(double Days)
+	{
+		return FTimespan(Days * ETimespan::TicksPerDay);
+	}
 
 	/**
-	 * Creates a time span that represents the specified number of hours.
+	 * Create a time span that represents the specified number of hours.
 	 *
 	 * @param Hours The number of hours.
 	 * @return Time span.
 	 * @see FromDays, FromMicroseconds, FromMilliseconds, FromMinutes, FromSeconds
 	 */
-	static CORE_API FTimespan FromHours(double Hours);
+	static FTimespan FromHours(double Hours)
+	{
+		return FTimespan(Hours * ETimespan::TicksPerHour);
+	}
 
 	/**
-	 * Creates a time span that represents the specified number of microseconds.
+	 * Create a time span that represents the specified number of microseconds.
 	 *
 	 * @param Microseconds The number of microseconds.
 	 * @return Time span.
 	 * @see FromDays, FromHours, FromMinutes, FromSeconds, FromMilliseconds
 	 */
-	static CORE_API FTimespan FromMicroseconds(double Microseconds);
+	static FTimespan FromMicroseconds(double Microseconds)
+	{
+		return FTimespan(Microseconds * ETimespan::TicksPerMicrosecond);
+	}
 
 	/**
-	 * Creates a time span that represents the specified number of milliseconds.
+	 * Create a time span that represents the specified number of milliseconds.
 	 *
 	 * @param Milliseconds The number of milliseconds.
 	 * @return Time span.
 	 * @see FromDays, FromHours, FromMicroseconds, FromMinutes, FromSeconds
 	 */
-	static CORE_API FTimespan FromMilliseconds(double Milliseconds);
+	static FTimespan FromMilliseconds(double Milliseconds)
+	{
+		return FTimespan(Milliseconds * ETimespan::TicksPerMillisecond);
+	}
 
 	/**
-	 * Creates a time span that represents the specified number of minutes.
+	 * Create a time span that represents the specified number of minutes.
 	 *
 	 * @param Minutes The number of minutes.
 	 * @return Time span.
 	 * @see FromDays, FromHours, FromMicroseconds, FromMilliseconds, FromSeconds
 	 */
-	static CORE_API FTimespan FromMinutes(double Minutes);
+	static FTimespan FromMinutes(double Minutes)
+	{
+		return FTimespan(Minutes * ETimespan::TicksPerMinute);
+	}
 
 	/**
-	 * Creates a time span that represents the specified number of seconds.
+	 * Create a time span that represents the specified number of seconds.
 	 *
 	 * @param Seconds The number of seconds.
 	 * @return Time span.
 	 * @see FromDays, FromHours, FromMicroseconds, FromMilliseconds, FromMinutes
 	 */
-	static CORE_API FTimespan FromSeconds(double Seconds);
+	static FTimespan FromSeconds(double Seconds)
+	{
+		return FTimespan(Seconds * ETimespan::TicksPerSecond);
+	}
 
 	/**
-	 * Returns the maximum time span value.
+	 * Return the maximum time span value.
 	 *
 	 * The maximum time span value is slightly more than 10,675,199 days.
 	 *
@@ -572,11 +661,11 @@ public:
 	 */
 	static FTimespan MaxValue()
 	{
-		return FTimespan(9223372036854775807);
+		return FTimespan(ETimespan::MaxTicks);
 	}
 
 	/**
-	 * Returns the minimum time span value.
+	 * Return the minimum time span value.
 	 *
 	 * The minimum time span value is slightly less than -10,675,199 days.
 	 *
@@ -585,14 +674,19 @@ public:
 	 */
 	static FTimespan MinValue()
 	{
-		return FTimespan(-9223372036854775807 - 1);
+		return FTimespan(ETimespan::MinTicks);
 	}
 
 	/**
-	 * Converts a string to a time span.
+	 * Convert a string to a time span.
 	 *
-	 * Currently, the string must be in the format written by FTimespan.ToString().
-	 * Other formats are not supported at this time.
+	 * The string must be in one of the following formats:
+	 *    p[d.]hh::mm::ss.fff
+	 *    p[d.]hh::mm::ss.uuuuuu
+	 *    p[d.]hh::mm::ss.nnnnnnnnn
+	 *
+	 * Note that 'p' is the plus or minus sign, and the date component may be
+	 * omitted for time spans that are shorter than one day.
 	 *
 	 * @param TimespanString The string to convert.
 	 * @param OutTimespan Will contain the parsed time span.
@@ -602,7 +696,24 @@ public:
 	static CORE_API bool Parse(const FString& TimespanString, FTimespan& OutTimespan);
 
 	/**
-	 * Returns the zero time span value.
+	 * Ratio between two time spans (handles zero values).
+	 *
+	 * @param Dividend The dividend.
+	 * @param Divisor The divisor.
+	 * @return The quotient, i.e. Dividend / Divisor.
+	 */
+	static double Ratio(FTimespan Dividend, FTimespan Divisor)
+	{
+		if (Divisor == FTimespan::Zero())
+		{
+			return 0.0;
+		}
+
+		return (double)Dividend.GetTicks() / (double)Divisor.GetTicks();
+	}
+
+	/**
+	 * Return the zero time span value.
 	 *
 	 * The zero time span value can be used in comparison operations with other time spans.
 	 *
@@ -619,39 +730,43 @@ public:
 	friend class UObject;
 
 	/**
-	 * Serializes the given time span from or into the specified archive.
+	 * Serialize the given time span from or into the specified archive.
 	 *
 	 * @param Ar The archive to serialize from or into.
 	 * @param Timespan The time span value to serialize.
 	 * @return The archive.
-	 *
-	 * @todo gmp: Figure out better include order in Core.h so this can be inlined.
 	 */
-	friend CORE_API FArchive& operator<<(FArchive& Ar, FTimespan& Timespan);
+	friend CORE_API FArchive& operator<<(FArchive& Ar, FTimespan& Timespan)
+	{
+		return Ar << Timespan.Ticks;
+	}
 
 	/**
-	 * Gets the hash for the specified time span.
+	 * Get the hash for the specified time span.
 	 *
 	 * @param Timespan The timespan to get the hash for.
 	 * @return Hash value.
 	 */
-	friend CORE_API uint32 GetTypeHash(const FTimespan& Timespan);
+	friend CORE_API uint32 GetTypeHash(const FTimespan& Timespan)
+	{
+		return GetTypeHash(Timespan.Ticks);
+	}
 
 protected:
 
 	/**
-	 * Assigns the specified components to this time span.
+	 * Assign the specified components to this time span.
 	 *
 	 * @param Days The days component.
 	 * @param Hours The hours component.
 	 * @param Minutes The minutes component.
 	 * @param Seconds The seconds component.
-	 * @param Milliseconds The milliseconds component.
-	 * @param Microseconds The microseconds component.
+	 * @param FractionNano The fractional seconds (in nanosecond resolution).
 	 */
-	void CORE_API Assign(int32 Days, int32 Hours, int32 Minutes, int32 Seconds, int32 Milliseconds, int32 Microseconds);
+	void CORE_API Assign(int32 Days, int32 Hours, int32 Minutes, int32 Seconds, int32 FractionNano);
 
 private:
+
 #ifdef COREUOBJECT_API
 	friend COREUOBJECT_API class UScriptStruct* Z_Construct_UScriptStruct_FTimespan();
 #else
@@ -660,13 +775,13 @@ private:
 
 private:
 
-	/** Holds the time span in 100 nanoseconds resolution. */
+	/** The time span value in 100 nanoseconds resolution. */
 	int64 Ticks;
 };
 
 
 /**
- * Pre-multiplies a time span with the given scalar.
+ * Pre-multiply a time span with the given scalar.
  *
  * @param Scalar The scalar to pre-multiply with.
  * @param Timespan The time span to multiply.
diff --git a/Engine/Source/Runtime/Core/Public/Windows/COMPointer.h b/Engine/Source/Runtime/Core/Public/Windows/COMPointer.h
--- a/Engine/Source/Runtime/Core/Public/Windows/COMPointer.h
+++ b/Engine/Source/Runtime/Core/Public/Windows/COMPointer.h
@@ -6,162 +6,249 @@
 
 
 /**
-* Smart COM object pointer.
-*/
+ * Smart COM object pointer.
+ */
 template<typename T>
 class TComPtr
 {
 public:
+
 	typedef T PointerType;
 
 public:
-	TComPtr() :	RawPointer(nullptr)
-	{
-	}
 
-	TComPtr(PointerType* const Source) : RawPointer(Source)
+	/** Default constructor (initialized to null). */
+	TComPtr()
+		:	RawPointer(NULL)
+	{ }
+
+	/**
+	 * Create and initialize a new instance.
+	 *
+	 * @param Object The object to point to.
+	 */
+	TComPtr(PointerType* const Object)
+		: RawPointer(Object)
 	{
-		if(RawPointer)
+		if (RawPointer)
 		{
 			RawPointer->AddRef();
 		}
 	}
 
-	TComPtr(const TComPtr<PointerType>& Source) : RawPointer(Source.RawPointer)
+	/**
+	 * Copy constructor.
+	 *
+	 * @param Other The instance to copy.
+	 */
+	TComPtr(const TComPtr<PointerType>& Other)
+		: RawPointer(Other.RawPointer)
 	{
-		if(RawPointer)
+		if (RawPointer)
 		{
 			RawPointer->AddRef();
 		}
 	}
 
-	TComPtr(TComPtr<PointerType>&& Source) : RawPointer(Source.RawPointer)
+	/**
+	 * Move constructor.
+	 *
+	 * @param Other The instance to move.
+	 */
+	TComPtr(TComPtr<PointerType>&& Other)
+		: RawPointer(Other.RawPointer)
 	{	
-		Source.RawPointer = nullptr;
+		Other.RawPointer = NULL;
 	}	
 
-	TComPtr<PointerType>& operator=(PointerType* const Source) 
+	/**
+	 * Assignment operator.
+	 *
+	 * @param Object The object to point to.
+	 */
+	TComPtr<PointerType>& operator=(PointerType* const Object) 
 	{
-		if(RawPointer != Source)
+		if (RawPointer != Object)
 		{
-			if (Source)
+			if (Object)
 			{
-				Source->AddRef();
+				Object->AddRef();
 			}
+
 			if (RawPointer)
 			{
 				RawPointer->Release();
 			}
-			RawPointer = Source;
+
+			RawPointer = Object;
 		}
+
 		return *this;
 	}
 
-	TComPtr<PointerType>& operator=(const TComPtr<PointerType>& Source) 
+	/**
+	 * Copy assignment operator.
+	 *
+	 * @param Other The instance to copy.
+	 */
+	TComPtr<PointerType>& operator=(const TComPtr<PointerType>& Other) 
 	{
-		if(RawPointer != Source.RawPointer)
+		if (RawPointer != Other.RawPointer)
 		{
-			if (Source.RawPointer)
+			if (Other.RawPointer)
 			{
-				Source->AddRef();
+				Other->AddRef();
 			}
+
 			if (RawPointer)
 			{
 				RawPointer->Release();
 			}
-			RawPointer = Source.RawPointer;
+
+			RawPointer = Other.RawPointer;
 		}
+
 		return *this;
 	}
 
-	TComPtr<PointerType>& operator=(TComPtr<PointerType>&& Source) 
+	/**
+	 * Move assignment operator.
+	 *
+	 * @param Other The instance to move.
+	 */
+	TComPtr<PointerType>& operator=(TComPtr<PointerType>&& Other) 
 	{			
-		if(RawPointer != Source.RawPointer)
+		if (RawPointer != Other.RawPointer)
 		{
-			if(RawPointer)
+			if (RawPointer)
 			{
 				RawPointer->Release();
 			}
-			RawPointer = Source.RawPointer;
-			Source.RawPointer = nullptr;
+
+			RawPointer = Other.RawPointer;
+			Other.RawPointer = NULL;
 		}
+
 		return *this;
 	}
 
+	/** Destructor. */
 	~TComPtr() 
 	{
-		if(RawPointer)
+		if (RawPointer)
 		{
 			RawPointer->Release();
 		}
 	}
 
-	PointerType** operator&()
+public:
+
+	FORCEINLINE PointerType** operator&()
 	{
 		return &(RawPointer);
 	}
 
-	PointerType* operator->() const 
+	FORCEINLINE PointerType* operator->() const 
 	{
-		check(RawPointer != nullptr);
+		check(RawPointer != NULL);
 		return RawPointer;
 	}
 
-	PointerType* Get() const
-	{
-		return RawPointer;
-	}
-	
-	bool operator==(PointerType* const InRawPointer) const
+	FORCEINLINE bool operator==(PointerType* const Object) const
 	{
-		return RawPointer == InRawPointer;
+		return RawPointer == Object;
 	}
 
-	bool operator!=(PointerType* const InRawPointer) const
+	FORCEINLINE bool operator!=(PointerType* const Object) const
 	{
-		return RawPointer != InRawPointer;
+		return RawPointer != Object;
 	}
 
-	operator PointerType*() const
+	FORCEINLINE operator PointerType*() const
 	{
 		return RawPointer;
 	}
 
-	void Reset()
+public:
+
+	/**
+	 * Set the pointer without adding a reference.
+	 *
+	 * @param InRawPointer The object to point to.
+	 * @see Detach
+	 */
+	void Attach(PointerType* Object)
 	{
-		if(RawPointer)
+		if (RawPointer)
 		{
 			RawPointer->Release();
-			RawPointer = nullptr;
 		}
+
+		RawPointer = Object;
+	}
+
+	/**
+	 * Reset the pointer without releasing a reference.
+	 *
+	 * @see Attach
+	 */
+	void Detach()
+	{
+		RawPointer = NULL;
 	}
 
-	// Set the pointer without adding a reference.
-	void Attach(PointerType* InRawPointer)
+	/**
+	 * Initialize this pointer from a COM interface to be queried.
+	 *
+	 * @param Riid The ID of the interface to be queried.
+	 * @param Unknown The object to query the interface from.
+	 * @return The result code of the query.
+	 */
+	HRESULT FromQueryInterface(REFIID Riid, IUnknown* Unknown)
 	{
-		if(RawPointer)
+		if (RawPointer)
 		{
 			RawPointer->Release();
+			RawPointer = NULL;
 		}
-		RawPointer = InRawPointer;
+
+		return Unknown->QueryInterface(Riid, reinterpret_cast<void**>(&(RawPointer)));
 	}
 
-	// Reset the pointer without releasing a reference.
-	void Detach()
+	/**
+	 * Get raw pointer to the object pointed to.
+	 *
+	 * @return Pointer to the object, or NULL if not valid.
+	 * @see IsValid
+	 */
+	FORCEINLINE PointerType* Get() const
+	{
+		return RawPointer;
+	}
+
+	/**
+	 * Whether this pointer is pointing to an actual object.
+	 *
+	 * @return true if the pointer is valid, false otherwise.
+	 * @see Get
+	 */
+	FORCEINLINE const bool IsValid() const
 	{
-		RawPointer = nullptr;
+		return (RawPointer != NULL);
 	}
 
-	HRESULT FromQueryInterface(REFIID riid, IUnknown* Unknown)
+	/** Reset this pointer to null. */
+	void Reset()
 	{
-		if(RawPointer)
+		if (RawPointer)
 		{
 			RawPointer->Release();
-			RawPointer = nullptr;
+			RawPointer = NULL;
 		}
-		return Unknown->QueryInterface(riid, reinterpret_cast<void**>(&(RawPointer)));
 	}
 
 private:
+
+	/** Pointer to the actual object, if any. */
 	PointerType* RawPointer;
 };
diff --git a/Engine/Source/Runtime/CoreUObject/Private/Serialization/BulkData.cpp b/Engine/Source/Runtime/CoreUObject/Private/Serialization/BulkData.cpp
--- a/Engine/Source/Runtime/CoreUObject/Private/Serialization/BulkData.cpp
+++ b/Engine/Source/Runtime/CoreUObject/Private/Serialization/BulkData.cpp
@@ -362,7 +362,7 @@ bool FUntypedBulkData::IsBulkDataLoaded() const
 
 bool FUntypedBulkData::IsAsyncLoadingComplete()
 {
-	return SerializeFuture.IsValid() == false || SerializeFuture.WaitFor(FTimespan(0));
+	return SerializeFuture.IsValid() == false || SerializeFuture.WaitFor(FTimespan::Zero());
 }
 
 /**
@@ -1334,7 +1334,7 @@ void FUntypedBulkData::WaitForAsyncLoading()
 {
 	check(SerializeFuture.IsValid());
 	DECLARE_SCOPE_CYCLE_COUNTER(TEXT("FUntypedBulkData::WaitForAsyncLoading"), STAT_UBD_WaitForAsyncLoading, STATGROUP_Memory);
-	while (!SerializeFuture.WaitFor(FTimespan(0, 0, 0, 0, 1000)))
+	while (!SerializeFuture.WaitFor(FTimespan::FromMilliseconds(1000.0)))
 	{
 		UE_LOG(LogSerialization, Warning, TEXT("Waiting for %s bulk data (%d) to be loaded longer than 1000ms"), *Filename, GetBulkDataSize());
 	}
diff --git a/Engine/Source/Runtime/CoreUObject/Private/UObject/Property.cpp b/Engine/Source/Runtime/CoreUObject/Private/UObject/Property.cpp
--- a/Engine/Source/Runtime/CoreUObject/Private/UObject/Property.cpp
+++ b/Engine/Source/Runtime/CoreUObject/Private/UObject/Property.cpp
@@ -263,7 +263,7 @@ struct TStructOpsTypeTraits<FTimespan> : public TStructOpsTypeTraitsBase2<FTimes
 	{
 		WithCopy = true,
 		WithExportTextItem = true,
-		WithImportTextItem = false, // @todo gmp: implement FTimespan::ImportTextItem
+		WithImportTextItem = true,
 		WithSerializer = true,
 		WithZeroConstructor = true,
 		WithIdenticalViaEquality = true,
@@ -271,7 +271,6 @@ struct TStructOpsTypeTraits<FTimespan> : public TStructOpsTypeTraitsBase2<FTimes
 };
 IMPLEMENT_STRUCT(Timespan);
 
-
 template<>
 struct TStructOpsTypeTraits<FSoftObjectPath> : public TStructOpsTypeTraitsBase2<FSoftObjectPath>
 {
diff --git a/Engine/Source/Runtime/CoreUObject/Public/UObject/NoExportTypes.h b/Engine/Source/Runtime/CoreUObject/Public/UObject/NoExportTypes.h
--- a/Engine/Source/Runtime/CoreUObject/Public/UObject/NoExportTypes.h
+++ b/Engine/Source/Runtime/CoreUObject/Public/UObject/NoExportTypes.h
@@ -292,7 +292,13 @@ enum class EUnit : uint8
 	/** Arbitrary multiplier */	
 	Multiplier,
 
-	/** Symbolic entry, not specifiable on meta data. */	
+
+	/** Percentage */
+	
+	Percentage,
+
+	/** Symbolic entry, not specifiable on meta data. */
+	
 	Unspecified
 };
 
diff --git a/Engine/Source/Runtime/Engine/Classes/Camera/CameraTypes.h b/Engine/Source/Runtime/Engine/Classes/Camera/CameraTypes.h
--- a/Engine/Source/Runtime/Engine/Classes/Camera/CameraTypes.h
+++ b/Engine/Source/Runtime/Engine/Classes/Camera/CameraTypes.h
@@ -120,6 +120,9 @@ struct FMinimalViewInfo
 	/** Combines this view with another one which will be weighted. Equals to this += OtherView * Weight. */
 	ENGINE_API void AddWeightedViewInfo(const FMinimalViewInfo& OtherView, const float& Weight);
 
+	/** Calculates the projection matrix using this view info's aspect ratio (regardless of bConstrainAspectRatio) */
+	ENGINE_API FMatrix CalculateProjectionMatrix() const;
+
 	/** Calculates the projection matrix (and potentially a constrained view rectangle) given a FMinimalViewInfo and partially configured projection data (must have the view rect already set) */
 	ENGINE_API static void CalculateProjectionMatrixGivenView(const FMinimalViewInfo& ViewInfo, TEnumAsByte<enum EAspectRatioAxisConstraint> AspectRatioAxisConstraint, class FViewport* Viewport, struct FSceneViewProjectionData& InOutProjectionData);
 };
diff --git a/Engine/Source/Runtime/Engine/Classes/Engine/EngineTypes.h b/Engine/Source/Runtime/Engine/Classes/Engine/EngineTypes.h
--- a/Engine/Source/Runtime/Engine/Classes/Engine/EngineTypes.h
+++ b/Engine/Source/Runtime/Engine/Classes/Engine/EngineTypes.h
@@ -476,6 +476,7 @@ enum EMaterialSamplerType
 	SAMPLERTYPE_DistanceFieldFont UMETA(DisplayName="Distance Field Font"),
 	SAMPLERTYPE_LinearColor UMETA(DisplayName = "Linear Color"),
 	SAMPLERTYPE_LinearGrayscale UMETA(DisplayName = "Linear Grayscale"),
+	SAMPLERTYPE_External UMETA(DisplayName = "External"),
 	SAMPLERTYPE_MAX,
 };
 
diff --git a/Engine/Source/Runtime/Engine/Classes/Engine/Texture.h b/Engine/Source/Runtime/Engine/Classes/Engine/Texture.h
--- a/Engine/Source/Runtime/Engine/Classes/Engine/Texture.h
+++ b/Engine/Source/Runtime/Engine/Classes/Engine/Texture.h
@@ -618,7 +618,7 @@ public:
 	/**
 	 * @return The material value type of this texture.
 	 */
-	virtual EMaterialValueType GetMaterialType() PURE_VIRTUAL(UTexture::GetMaterialType,return MCT_Texture;);
+	virtual EMaterialValueType GetMaterialType() const PURE_VIRTUAL(UTexture::GetMaterialType,return MCT_Texture;);
 
 	/**
 	 * Waits until all streaming requests for this texture has been fully processed.
@@ -736,6 +736,14 @@ public:
 	/** @return the height of the surface represented by the texture. */
 	virtual float GetSurfaceHeight() const PURE_VIRTUAL(UTexture::GetSurfaceHeight,return 0;);
 
+	/**
+	 * Access the GUID which defines this texture's resources externally through FExternalTextureRegistry
+	 */
+	virtual FGuid GetExternalTextureGuid() const
+	{
+		return FGuid();
+	}
+
 	//~ Begin UObject Interface.
 #if WITH_EDITOR
 	ENGINE_API virtual void PostEditChangeProperty(FPropertyChangedEvent& PropertyChangedEvent) override;
@@ -836,6 +844,14 @@ public:
 	 */
 	ENGINE_API static class UEnum* GetPixelFormatEnum();
 
+protected:
+
+#if WITH_EDITOR
+
+	/** Notify any loaded material instances that the texture has changed. */
+	ENGINE_API void NotifyMaterials();
+
+#endif //WITH_EDOTIR
 };
 
 /** 
diff --git a/Engine/Source/Runtime/Engine/Classes/Engine/Texture2D.h b/Engine/Source/Runtime/Engine/Classes/Engine/Texture2D.h
--- a/Engine/Source/Runtime/Engine/Classes/Engine/Texture2D.h
+++ b/Engine/Source/Runtime/Engine/Classes/Engine/Texture2D.h
@@ -174,7 +174,7 @@ public:
 	virtual float GetSurfaceWidth() const override { return GetSizeX(); }
 	virtual float GetSurfaceHeight() const override { return GetSizeY(); }
 	virtual FTextureResource* CreateResource() override;
-	virtual EMaterialValueType GetMaterialType() override { return MCT_Texture2D; }
+	virtual EMaterialValueType GetMaterialType() const override { return MCT_Texture2D; }
 	virtual void UpdateResource() override;
 	virtual float GetAverageBrightness(bool bIgnoreTrueBlack, bool bUseGrayscale) override;
 	virtual FTexturePlatformData** GetRunningPlatformData() override { return &PlatformData; }
diff --git a/Engine/Source/Runtime/Engine/Classes/Engine/Texture2DDynamic.h b/Engine/Source/Runtime/Engine/Classes/Engine/Texture2DDynamic.h
--- a/Engine/Source/Runtime/Engine/Classes/Engine/Texture2DDynamic.h
+++ b/Engine/Source/Runtime/Engine/Classes/Engine/Texture2DDynamic.h
@@ -35,7 +35,7 @@ class UTexture2DDynamic : public UTexture
 public:
 	//~ Begin UTexture Interface.
 	ENGINE_API virtual FTextureResource* CreateResource() override;
-	virtual EMaterialValueType GetMaterialType() override { return MCT_Texture2D; }
+	virtual EMaterialValueType GetMaterialType() const override { return MCT_Texture2D; }
 	ENGINE_API virtual float GetSurfaceWidth() const override;
 	ENGINE_API virtual float GetSurfaceHeight() const override;
 	//~ End UTexture Interface.
diff --git a/Engine/Source/Runtime/Engine/Classes/Engine/TextureCube.h b/Engine/Source/Runtime/Engine/Classes/Engine/TextureCube.h
--- a/Engine/Source/Runtime/Engine/Classes/Engine/TextureCube.h
+++ b/Engine/Source/Runtime/Engine/Classes/Engine/TextureCube.h
@@ -67,7 +67,7 @@ public:
 	virtual float GetSurfaceHeight() const override { return GetSizeY(); }
 	virtual FTextureResource* CreateResource() override;
 	virtual void UpdateResource() override;
-	virtual EMaterialValueType GetMaterialType() override { return MCT_TextureCube; }
+	virtual EMaterialValueType GetMaterialType() const override { return MCT_TextureCube; }
 	virtual FTexturePlatformData** GetRunningPlatformData() override { return &PlatformData; }
 	virtual TMap<FString, FTexturePlatformData*> *GetCookedPlatformData() override { return &CookedPlatformData; }
 	//~ End UTexture Interface
diff --git a/Engine/Source/Runtime/Engine/Classes/Engine/TextureRenderTarget.h b/Engine/Source/Runtime/Engine/Classes/Engine/TextureRenderTarget.h
--- a/Engine/Source/Runtime/Engine/Classes/Engine/TextureRenderTarget.h
+++ b/Engine/Source/Runtime/Engine/Classes/Engine/TextureRenderTarget.h
@@ -42,7 +42,7 @@ class UTextureRenderTarget : public UTexture
 
 	//~ Begin UTexture Interface
 	ENGINE_API virtual class FTextureResource* CreateResource() override;
-	ENGINE_API virtual EMaterialValueType GetMaterialType() override;
+	ENGINE_API virtual EMaterialValueType GetMaterialType() const override;
 	//~ End UTexture Interface
 };
 
diff --git a/Engine/Source/Runtime/Engine/Classes/Engine/TextureRenderTarget2D.h b/Engine/Source/Runtime/Engine/Classes/Engine/TextureRenderTarget2D.h
--- a/Engine/Source/Runtime/Engine/Classes/Engine/TextureRenderTarget2D.h
+++ b/Engine/Source/Runtime/Engine/Classes/Engine/TextureRenderTarget2D.h
@@ -156,7 +156,7 @@ class UTextureRenderTarget2D : public UTextureRenderTarget
 	virtual float GetSurfaceWidth() const override { return SizeX; }
 	virtual float GetSurfaceHeight() const override { return SizeY; }
 	ENGINE_API virtual FTextureResource* CreateResource() override;
-	ENGINE_API virtual EMaterialValueType GetMaterialType() override;
+	ENGINE_API virtual EMaterialValueType GetMaterialType() const override;
 	//~ End UTexture Interface.
 
 	//~ Begin UObject Interface
diff --git a/Engine/Source/Runtime/Engine/Classes/Engine/TextureRenderTargetCube.h b/Engine/Source/Runtime/Engine/Classes/Engine/TextureRenderTargetCube.h
--- a/Engine/Source/Runtime/Engine/Classes/Engine/TextureRenderTargetCube.h
+++ b/Engine/Source/Runtime/Engine/Classes/Engine/TextureRenderTargetCube.h
@@ -70,7 +70,7 @@ class UTextureRenderTargetCube : public UTextureRenderTarget
 	// and warns us. In this case, it is intentional, so we disable the warning:
 	virtual float GetSurfaceHeight()const  override { return SizeX; }	 //-V524
 	virtual FTextureResource* CreateResource() override;
-	virtual EMaterialValueType GetMaterialType() override;
+	virtual EMaterialValueType GetMaterialType() const override;
 	//~ End UTexture Interface.
 
 	EPixelFormat GetFormat() const
diff --git a/Engine/Source/Runtime/Engine/Classes/Kismet/KismetMathLibrary.h b/Engine/Source/Runtime/Engine/Classes/Kismet/KismetMathLibrary.h
--- a/Engine/Source/Runtime/Engine/Classes/Kismet/KismetMathLibrary.h
+++ b/Engine/Source/Runtime/Engine/Classes/Kismet/KismetMathLibrary.h
@@ -1260,10 +1260,18 @@ class ENGINE_API UKismetMathLibrary : public UBlueprintFunctionLibrary
 	UFUNCTION(BlueprintPure, Category="Math|Timespan", meta=(NativeMakeFunc))
 	static FTimespan MakeTimespan(int32 Days, int32 Hours, int32 Minutes, int32 Seconds, int32 Milliseconds);
 
+	/* Makes a Timespan struct */
+	UFUNCTION(BlueprintPure, Category="Math|Timespan", meta=(NativeMakeFunc))
+	static FTimespan MakeTimespan2(int32 Days, int32 Hours, int32 Minutes, int32 Seconds, int32 FractionNano);
+
 	/* Breaks a Timespan into its components */
 	UFUNCTION(BlueprintPure, Category="Math|Timespan", meta=(NativeBreakFunc))
 	static void BreakTimespan(FTimespan InTimespan, int32& Days, int32& Hours, int32& Minutes, int32& Seconds, int32& Milliseconds);
 
+	/* Breaks a Timespan into its components */
+	UFUNCTION(BlueprintPure, Category="Math|Timespan", meta=(NativeBreakFunc))
+	static void BreakTimespan2(FTimespan InTimespan, int32& Days, int32& Hours, int32& Minutes, int32& Seconds, int32& FractionNano);
+
 	/* Addition (A + B) */
 	UFUNCTION(BlueprintPure, meta=(DisplayName="Timespan + Timespan", CompactNodeTitle="+", Keywords="+ add plus"), Category="Math|Timespan")
 	static FTimespan Add_TimespanTimespan( FTimespan A, FTimespan B );
@@ -1276,6 +1284,10 @@ class ENGINE_API UKismetMathLibrary : public UBlueprintFunctionLibrary
 	UFUNCTION(BlueprintPure, meta=(DisplayName="Timespan * float", CompactNodeTitle="*", Keywords="* multiply"), Category="Math|Timespan")
 	static FTimespan Multiply_TimespanFloat( FTimespan A, float Scalar );
 
+	/* Scalar division (A * s) */
+	UFUNCTION(BlueprintPure, meta=(DisplayName="Timespan * float", CompactNodeTitle="/", Keywords="/ divide"), Category="Math|Timespan")
+	static FTimespan Divide_TimespanFloat( FTimespan A, float Scalar );
+
 	/* Returns true if the values are equal (A == B) */
 	UFUNCTION(BlueprintPure, meta=(DisplayName="Equal (Timespan)", CompactNodeTitle="==", Keywords="== equal"), Category="Math|Timespan")
 	static bool EqualEqual_TimespanTimespan( FTimespan A, FTimespan B );
diff --git a/Engine/Source/Runtime/Engine/Classes/Kismet/KismetMathLibrary.inl b/Engine/Source/Runtime/Engine/Classes/Kismet/KismetMathLibrary.inl
--- a/Engine/Source/Runtime/Engine/Classes/Kismet/KismetMathLibrary.inl
+++ b/Engine/Source/Runtime/Engine/Classes/Kismet/KismetMathLibrary.inl
@@ -1389,6 +1389,12 @@ FTimespan UKismetMathLibrary::Multiply_TimespanFloat( FTimespan A, float Scalar
 	return A * Scalar;
 }
 
+KISMET_MATH_FORCEINLINE
+FTimespan UKismetMathLibrary::Divide_TimespanFloat(FTimespan A, float Scalar)
+{
+	return A / Scalar;
+}
+
 KISMET_MATH_FORCEINLINE
 bool UKismetMathLibrary::EqualEqual_TimespanTimespan( FTimespan A, FTimespan B )
 {
@@ -1446,7 +1452,7 @@ int32 UKismetMathLibrary::GetHours( FTimespan A )
 KISMET_MATH_FORCEINLINE
 int32 UKismetMathLibrary::GetMilliseconds( FTimespan A )
 {
-	return A.GetMilliseconds();
+	return A.GetFractionMilli();
 }
 
 KISMET_MATH_FORCEINLINE
@@ -1506,12 +1512,7 @@ FTimespan UKismetMathLibrary::TimespanMinValue( )
 KISMET_MATH_INLINE
 float UKismetMathLibrary::TimespanRatio( FTimespan A, FTimespan B )
 {
-	if (B != FTimespan::Zero())
-	{
-		return (float)A.GetTicks() / (float)B.GetTicks();
-	}
-
-	return 0.0;
+	return FTimespan::Ratio(A, B);
 }
 
 KISMET_MATH_FORCEINLINE
diff --git a/Engine/Source/Runtime/Engine/Private/Camera/CameraStackTypes.cpp b/Engine/Source/Runtime/Engine/Private/Camera/CameraStackTypes.cpp
--- a/Engine/Source/Runtime/Engine/Private/Camera/CameraStackTypes.cpp
+++ b/Engine/Source/Runtime/Engine/Private/Camera/CameraStackTypes.cpp
@@ -72,47 +72,65 @@ void FMinimalViewInfo::AddWeightedViewInfo(const FMinimalViewInfo& OtherView, co
 	bUseFieldOfViewForLOD |= OtherViewWeighted.bUseFieldOfViewForLOD;
 }
 
+FMatrix FMinimalViewInfo::CalculateProjectionMatrix() const
+{
+	FMatrix ProjectionMatrix;
+
+	if (ProjectionMode == ECameraProjectionMode::Orthographic)
+	{
+		const float YScale = 1.0f / AspectRatio;
+
+		const float HalfOrthoWidth = OrthoWidth / 2.0f;
+		const float ScaledOrthoHeight = OrthoWidth / 2.0f * YScale;
+
+		const float NearPlane = OrthoNearClipPlane;
+		const float FarPlane = OrthoFarClipPlane;
+
+		const float ZScale = 1.0f / (FarPlane - NearPlane);
+		const float ZOffset = -NearPlane;
+
+		ProjectionMatrix = FReversedZOrthoMatrix(
+			HalfOrthoWidth,
+			ScaledOrthoHeight,
+			ZScale,
+			ZOffset
+			);
+	}
+	else
+	{
+		// Avoid divide by zero in the projection matrix calculation by clamping FOV
+		ProjectionMatrix = FReversedZPerspectiveMatrix(
+			FMath::Max(0.001f, FOV) * (float)PI / 360.0f,
+			AspectRatio,
+			1.0f,
+			GNearClippingPlane );
+	}
+
+	if (!OffCenterProjectionOffset.IsZero())
+	{
+		const float Left = -1.0f + OffCenterProjectionOffset.X;
+		const float Right = Left + 2.0f;
+		const float Bottom = -1.0f + OffCenterProjectionOffset.Y;
+		const float Top = Bottom + 2.0f;
+		ProjectionMatrix.M[2][0] = (Left + Right) / (Left - Right);
+		ProjectionMatrix.M[2][1] = (Bottom + Top) / (Bottom - Top);
+	}
+
+	return ProjectionMatrix;
+}
+
 void FMinimalViewInfo::CalculateProjectionMatrixGivenView(const FMinimalViewInfo& ViewInfo, TEnumAsByte<enum EAspectRatioAxisConstraint> AspectRatioAxisConstraint, FViewport* Viewport, FSceneViewProjectionData& InOutProjectionData)
 {
 	// Create the projection matrix (and possibly constrain the view rectangle)
 	if (ViewInfo.bConstrainAspectRatio)
-	{			
+	{
 		// Enforce a particular aspect ratio for the render of the scene. 
 		// Results in black bars at top/bottom etc.
 		InOutProjectionData.SetConstrainedViewRectangle(Viewport->CalculateViewExtents(ViewInfo.AspectRatio, InOutProjectionData.GetViewRect()));
 
-		//
-		if (ViewInfo.ProjectionMode == ECameraProjectionMode::Orthographic)
-		{
-			const float YScale = 1.0f / ViewInfo.AspectRatio;
-	
-			const float OrthoWidth = ViewInfo.OrthoWidth / 2.0f;
-			const float OrthoHeight = ViewInfo.OrthoWidth / 2.0f * YScale;
-
-			const float NearPlane = ViewInfo.OrthoNearClipPlane;
-			const float FarPlane = ViewInfo.OrthoFarClipPlane;
-
-			const float ZScale = 1.0f / (FarPlane - NearPlane);
-			const float ZOffset = -NearPlane;
-
-			InOutProjectionData.ProjectionMatrix = FReversedZOrthoMatrix(
-				OrthoWidth,
-				OrthoHeight,
-				ZScale,
-				ZOffset
-				);
-		}
-		else
-		{
-			// Avoid divide by zero in the projection matrix calculation by clamping FOV
-			InOutProjectionData.ProjectionMatrix = FReversedZPerspectiveMatrix(
-				FMath::Max(0.001f, ViewInfo.FOV) * (float)PI / 360.0f,
-				ViewInfo.AspectRatio,
-				1.0f,
-				GNearClippingPlane );
-		}
+		InOutProjectionData.ProjectionMatrix = ViewInfo.CalculateProjectionMatrix();
 	}
-	else 
+	else
 	{
 		// Avoid divide by zero in the projection matrix calculation by clamping FOV
 		float MatrixFOV = FMath::Max(0.001f, ViewInfo.FOV) * (float)PI / 360.0f;
@@ -177,4 +195,4 @@ void FMinimalViewInfo::CalculateProjectionMatrixGivenView(const FMinimalViewInfo
 		InOutProjectionData.ProjectionMatrix.M[2][0] = (Left + Right) / (Left - Right);
 		InOutProjectionData.ProjectionMatrix.M[2][1] = (Bottom + Top) / (Bottom - Top);
 	}
-}
+}
\ No newline at end of file
diff --git a/Engine/Source/Runtime/Engine/Private/EngineService.cpp b/Engine/Source/Runtime/Engine/Private/EngineService.cpp
--- a/Engine/Source/Runtime/Engine/Private/EngineService.cpp
+++ b/Engine/Source/Runtime/Engine/Private/EngineService.cpp
@@ -5,8 +5,8 @@
 #include "EngineGlobals.h"
 #include "Engine/Engine.h"
 #include "EngineServiceMessages.h"
-#include "Helpers/MessageEndpoint.h"
-#include "Helpers/MessageEndpointBuilder.h"
+#include "MessageEndpoint.h"
+#include "MessageEndpointBuilder.h"
 #include "Misc/NetworkVersion.h"
 
 
@@ -49,7 +49,7 @@ void FEngineService::SendNotification( const TCHAR* NotificationText, const FMes
 }
 
 
-void FEngineService::SendPong( const IMessageContextRef& Context )
+void FEngineService::SendPong( const TSharedRef<IMessageContext, ESPMode::ThreadSafe>& Context )
 {
 	if (MessageEndpoint.IsValid())
 	{
@@ -125,7 +125,7 @@ void FEngineService::SendPong( const IMessageContextRef& Context )
 /* FEngineService callbacks
  *****************************************************************************/
 
-void FEngineService::HandleAuthGrantMessage( const FEngineServiceAuthGrant& Message, const IMessageContextRef& Context )
+void FEngineService::HandleAuthGrantMessage( const FEngineServiceAuthGrant& Message, const TSharedRef<IMessageContext, ESPMode::ThreadSafe>& Context )
 {
 	if (AuthorizedUsers.Contains(Message.UserName))
 	{
@@ -140,7 +140,7 @@ void FEngineService::HandleAuthGrantMessage( const FEngineServiceAuthGrant& Mess
 }
 
 
-void FEngineService::HandleAuthDenyMessage( const FEngineServiceAuthDeny& Message, const IMessageContextRef& Context )
+void FEngineService::HandleAuthDenyMessage( const FEngineServiceAuthDeny& Message, const TSharedRef<IMessageContext, ESPMode::ThreadSafe>& Context )
 {
 	if (AuthorizedUsers.Contains(Message.UserName))
 	{
@@ -155,7 +155,7 @@ void FEngineService::HandleAuthDenyMessage( const FEngineServiceAuthDeny& Messag
 }
 
 
-void FEngineService::HandleExecuteCommandMessage( const FEngineServiceExecuteCommand& Message, const IMessageContextRef& Context )
+void FEngineService::HandleExecuteCommandMessage( const FEngineServiceExecuteCommand& Message, const TSharedRef<IMessageContext, ESPMode::ThreadSafe>& Context )
 {
 	if (AuthorizedUsers.Contains(Message.UserName))
 	{
@@ -177,13 +177,13 @@ void FEngineService::HandleExecuteCommandMessage( const FEngineServiceExecuteCom
 }
 
 
-void FEngineService::HandlePingMessage( const FEngineServicePing& Message, const IMessageContextRef& Context )
+void FEngineService::HandlePingMessage( const FEngineServicePing& Message, const TSharedRef<IMessageContext, ESPMode::ThreadSafe>& Context )
 {
 	SendPong(Context);
 }
 
 
-void FEngineService::HandleTerminateMessage( const FEngineServiceTerminate& Message, const IMessageContextRef& Context )
+void FEngineService::HandleTerminateMessage( const FEngineServiceTerminate& Message, const TSharedRef<IMessageContext, ESPMode::ThreadSafe>& Context )
 {
 	if (AuthorizedUsers.Contains(Message.UserName))
 	{
diff --git a/Engine/Source/Runtime/Engine/Private/EngineSessionManager.cpp b/Engine/Source/Runtime/Engine/Private/EngineSessionManager.cpp
--- a/Engine/Source/Runtime/Engine/Private/EngineSessionManager.cpp
+++ b/Engine/Source/Runtime/Engine/Private/EngineSessionManager.cpp
@@ -29,10 +29,10 @@ DEFINE_LOG_CATEGORY(LogEngineSessionManager);
 
 namespace SessionManagerDefs
 {
-	static const FTimespan SessionRecordExpiration(30, 0, 0, 0);	// 30 days
-	static const FTimespan SessionRecordTimeout(0, 3, 0);			// 3 minutes
-	static const FTimespan GlobalLockWaitTimeout(0, 0, 0, 0, 500);	// 1/2 second
-	static const int HeartbeatPeriodSeconds(60);				// 1 minute
+	static const FTimespan SessionRecordExpiration = FTimespan::FromDays(30.0);
+	static const FTimespan SessionRecordTimeout = FTimespan::FromMinutes(3.0);
+	static const FTimespan GlobalLockWaitTimeout = FTimespan::FromSeconds(0.5);
+	static const int HeartbeatPeriodSeconds(60);
 	static const FString DefaultUserActivity(TEXT("Unknown"));
 	static const FString StoreId(TEXT("Epic Games"));
 	static const FString RunningSessionToken(TEXT("Running"));
diff --git a/Engine/Source/Runtime/Engine/Private/ExternalTexture.cpp b/Engine/Source/Runtime/Engine/Private/ExternalTexture.cpp
--- a/Engine/Source/Runtime/Engine/Private/ExternalTexture.cpp
+++ b/Engine/Source/Runtime/Engine/Private/ExternalTexture.cpp
@@ -14,11 +14,10 @@ FExternalTextureRegistry& FExternalTextureRegistry::Get()
 	return *Singleton;
 }
 
-/* Register an external texture and its sampler state against a GUID */
-void FExternalTextureRegistry::RegisterExternalTexture(const FGuid& InGuid, FTextureRHIRef& InTextureRHI, FSamplerStateRHIRef& InSamplerStateRHI)
+/* Register an external texture, its sampler state and coordinate scale/bias against a GUID */
+void FExternalTextureRegistry::RegisterExternalTexture(const FGuid& InGuid, FTextureRHIRef& InTextureRHI, FSamplerStateRHIRef& InSamplerStateRHI, const FLinearColor& InCoordinateScaleRotation, const FLinearColor& InCoordinateOffset)
 {
-	//FPlatformMisc::LowLevelOutputDebugStringf(TEXT("RegisterExternalTexture %s\n"), *InGuid.ToString());
-	TextureEntries.Add(InGuid, FExternalTextureEntry(InTextureRHI, InSamplerStateRHI));
+	TextureEntries.Add(InGuid, FExternalTextureEntry(InTextureRHI, InSamplerStateRHI, InCoordinateScaleRotation, InCoordinateOffset));
 	for (const FMaterialRenderProxy* MaterialRenderProxy : ReferencingMaterialRenderProxies)
 	{
 		const_cast<FMaterialRenderProxy*>(MaterialRenderProxy)->CacheUniformExpressions();
@@ -41,11 +40,13 @@ void FExternalTextureRegistry::RemoveMaterialRenderProxyReference(const FMateria
 }
 
 /* Looks up an external texture for given a given GUID
- * @return false if the texture is not registered
- */
+	* @return false if the texture is not registered
+	*/
 bool FExternalTextureRegistry::GetExternalTexture(const FMaterialRenderProxy* MaterialRenderProxy, const FGuid& InGuid, FTextureRHIRef& OutTextureRHI, FSamplerStateRHIRef& OutSamplerStateRHI)
 {
-	if (MaterialRenderProxy)
+//	FPlatformMisc::LowLevelOutputDebugStringf(TEXT("GetExternalTexture: Guid = %s"), *InGuid.ToString());
+	// Only cache render proxies that have been initialized, since FMaterialRenderProxy::ReleaseDynamicRHI() is responsible for removing them
+	if (MaterialRenderProxy && MaterialRenderProxy->IsInitialized())
 	{
 		ReferencingMaterialRenderProxies.Add(MaterialRenderProxy);
 	}
@@ -53,14 +54,48 @@ bool FExternalTextureRegistry::GetExternalTexture(const FMaterialRenderProxy* Ma
 	FExternalTextureEntry* Entry = TextureEntries.Find(InGuid);
 	if (Entry)
 	{
-		//FPlatformMisc::LowLevelOutputDebugStringf(TEXT("GetExternalTexture %s found\n"),*InGuid.ToString());
+//		FPlatformMisc::LowLevelOutputDebugStringf(TEXT("GetExternalTexture: Found"));
 		OutTextureRHI = Entry->TextureRHI;
 		OutSamplerStateRHI = Entry->SamplerStateRHI;
 		return true;
 	}
 	else
 	{
-		//FPlatformMisc::LowLevelOutputDebugStringf(TEXT("GetExternalTexture %s not found\n"), *InGuid.ToString());
+//		FPlatformMisc::LowLevelOutputDebugStringf(TEXT("GetExternalTexture: NOT FOUND!"));
+		return false;
+	}
+}
+
+/* Looks up an texture coordinate scale rotation for given a given GUID
+ * @return false if the texture is not registered
+ */
+bool FExternalTextureRegistry::GetExternalTextureCoordinateScaleRotation(const FGuid& InGuid, FLinearColor& OutCoordinateScaleRotation)
+{
+	FExternalTextureEntry* Entry = TextureEntries.Find(InGuid);
+	if (Entry)
+	{
+		OutCoordinateScaleRotation = Entry->CoordinateScaleRotation;
+		return true;
+	}
+	else
+	{
+		return false;
+	}
+}
+
+/* Looks up an texture coordinate offset for given a given GUID
+* @return false if the texture is not registered
+*/
+bool FExternalTextureRegistry::GetExternalTextureCoordinateOffset(const FGuid& InGuid, FLinearColor& OutCoordinateOffset)
+{
+	FExternalTextureEntry* Entry = TextureEntries.Find(InGuid);
+	if (Entry)
+	{
+		OutCoordinateOffset = Entry->CoordinateOffset;
+		return true;
+	}
+	else
+	{
 		return false;
 	}
 }
diff --git a/Engine/Source/Runtime/Engine/Private/HighResScreenshot.cpp b/Engine/Source/Runtime/Engine/Private/HighResScreenshot.cpp
--- a/Engine/Source/Runtime/Engine/Private/HighResScreenshot.cpp
+++ b/Engine/Source/Runtime/Engine/Private/HighResScreenshot.cpp
@@ -7,8 +7,8 @@
 #include "UnrealClient.h"
 #include "Materials/Material.h"
 #include "Slate/SceneViewport.h"
-#include "Interfaces/IImageWrapper.h"
-#include "Interfaces/IImageWrapperModule.h"
+#include "IImageWrapper.h"
+#include "IImageWrapperModule.h"
 
 static TAutoConsoleVariable<int32> CVarSaveEXRCompressionQuality(
 	TEXT("r.SaveEXR.CompressionQuality"),
@@ -180,7 +180,7 @@ template<typename> struct FPixelTypeTraits {};
 
 template<> struct FPixelTypeTraits<FColor>
 {
-	static const ERGBFormat::Type SourceChannelLayout = ERGBFormat::BGRA;
+	static const ERGBFormat SourceChannelLayout = ERGBFormat::BGRA;
 
 	static FORCEINLINE bool IsWritingHDRImage(const bool)
 	{
@@ -190,7 +190,7 @@ template<> struct FPixelTypeTraits<FColor>
 
 template<> struct FPixelTypeTraits<FFloat16Color>
 {
-	static const ERGBFormat::Type SourceChannelLayout = ERGBFormat::RGBA;
+	static const ERGBFormat SourceChannelLayout = ERGBFormat::RGBA;
 
 	static FORCEINLINE bool IsWritingHDRImage(const bool bCaptureHDR)
 	{
@@ -201,7 +201,7 @@ template<> struct FPixelTypeTraits<FFloat16Color>
 
 template<> struct FPixelTypeTraits<FLinearColor>
 {
-	static const ERGBFormat::Type SourceChannelLayout = ERGBFormat::RGBA;
+	static const ERGBFormat SourceChannelLayout = ERGBFormat::RGBA;
 
 	static FORCEINLINE bool IsWritingHDRImage(const bool bCaptureHDR)
 	{
@@ -265,20 +265,20 @@ bool FHighResScreenshotConfig::SaveImage(const FString& File, const TArray<TPixe
 		ImageWriter->ImageWrapper.IsValid() &&
 		ImageWriter->ImageWrapper->SetRaw((void*)&Bitmap[0], sizeof(TPixelType)* x * y, x, y, Traits::SourceChannelLayout, BitsPerPixel))
 	{
-		ImageCompression::CompressionQuality LocalCompressionQuality = ImageCompression::Default;
+		EImageCompressionQuality LocalCompressionQuality = EImageCompressionQuality::Default;
 		
 		if(bIsWritingHDRImage && CVarSaveEXRCompressionQuality.GetValueOnAnyThread() == 0)
 		{
-			LocalCompressionQuality = ImageCompression::Uncompressed;
+			LocalCompressionQuality = EImageCompressionQuality::Uncompressed;
 		}
 
 		// Compress and write image
 		FArchive* Ar = FileManager->CreateFileWriter(Filename.GetCharArray().GetData());
 		if (Ar != nullptr)
 		{
-			const TArray<uint8>& CompressedData = ImageWriter->ImageWrapper->GetCompressed(LocalCompressionQuality);
-				int32 CompressedSize = CompressedData.Num();
-				Ar->Serialize((void*)CompressedData.GetData(), CompressedSize);
+			const TArray<uint8>& CompressedData = ImageWriter->ImageWrapper->GetCompressed((int32)LocalCompressionQuality);
+			int32 CompressedSize = CompressedData.Num();
+			Ar->Serialize((void*)CompressedData.GetData(), CompressedSize);
 			delete Ar;
 
 			bSuccess = true;
diff --git a/Engine/Source/Runtime/Engine/Private/ImageUtils.cpp b/Engine/Source/Runtime/Engine/Private/ImageUtils.cpp
--- a/Engine/Source/Runtime/Engine/Private/ImageUtils.cpp
+++ b/Engine/Source/Runtime/Engine/Private/ImageUtils.cpp
@@ -674,7 +674,7 @@ bool FImageUtils::ExportRenderTarget2DAsPNG(UTextureRenderTarget2D* TexRT, FArch
 
 		IImageWrapperModule& ImageWrapperModule = FModuleManager::Get().LoadModuleChecked<IImageWrapperModule>(TEXT("ImageWrapper"));
 
-		IImageWrapperPtr PNGImageWrapper = ImageWrapperModule.CreateImageWrapper(EImageFormat::PNG);
+		TSharedPtr<IImageWrapper> PNGImageWrapper = ImageWrapperModule.CreateImageWrapper(EImageFormat::PNG);
 
 		PNGImageWrapper->SetRaw(RawData.GetData(), RawData.GetAllocatedSize(), Size.X, Size.Y, ERGBFormat::BGRA, 8);
 
diff --git a/Engine/Source/Runtime/Engine/Private/Kismet/BlueprintPlatformLibrary.cpp b/Engine/Source/Runtime/Engine/Private/Kismet/BlueprintPlatformLibrary.cpp
--- a/Engine/Source/Runtime/Engine/Private/Kismet/BlueprintPlatformLibrary.cpp
+++ b/Engine/Source/Runtime/Engine/Private/Kismet/BlueprintPlatformLibrary.cpp
@@ -108,7 +108,7 @@ void UBlueprintPlatformLibrary::ScheduleLocalNotificationAtTime(const FDateTime&
 void UBlueprintPlatformLibrary::ScheduleLocalNotificationFromNow(int32 inSecondsFromNow, const FText& Title, const FText& Body, const FText& Action, const FString& ActivationEvent)
 {
 	FDateTime TargetTime = FDateTime::Now();
-	TargetTime += FTimespan(0, 0, inSecondsFromNow);
+	TargetTime += FTimespan::FromSeconds(inSecondsFromNow);
 
 	ScheduleLocalNotificationAtTime(TargetTime, true, Title, Body, Action, ActivationEvent);
 }
@@ -129,7 +129,7 @@ void UBlueprintPlatformLibrary::ScheduleLocalNotificationBadgeAtTime(const FDate
 void UBlueprintPlatformLibrary::ScheduleLocalNotificationBadgeFromNow(int32 inSecondsFromNow, const FString& ActivationEvent)
 {
 	FDateTime TargetTime = FDateTime::Now();
-	TargetTime += FTimespan(0, 0, inSecondsFromNow);
+	TargetTime += FTimespan::FromSeconds(inSecondsFromNow);
 
 	ScheduleLocalNotificationBadgeAtTime(TargetTime, true, ActivationEvent);
 }
diff --git a/Engine/Source/Runtime/Engine/Private/KismetMathLibrary.cpp b/Engine/Source/Runtime/Engine/Private/KismetMathLibrary.cpp
--- a/Engine/Source/Runtime/Engine/Private/KismetMathLibrary.cpp
+++ b/Engine/Source/Runtime/Engine/Private/KismetMathLibrary.cpp
@@ -613,18 +613,36 @@ void UKismetMathLibrary::BreakDateTime(FDateTime InDateTime, int32& Year, int32&
 
 FTimespan UKismetMathLibrary::MakeTimespan(int32 Days, int32 Hours, int32 Minutes, int32 Seconds, int32 Milliseconds)
 {
-	return FTimespan(Days, Hours, Minutes, Seconds, Milliseconds);
+	return FTimespan(Days, Hours, Minutes, Seconds, Milliseconds * 1000 * 1000);
 }
 
+
+FTimespan UKismetMathLibrary::MakeTimespan2(int32 Days, int32 Hours, int32 Minutes, int32 Seconds, int32 FractionNano)
+{
+	return FTimespan(Days, Hours, Minutes, Seconds, FractionNano);
+}
+
+
 void UKismetMathLibrary::BreakTimespan(FTimespan InTimespan, int32& Days, int32& Hours, int32& Minutes, int32& Seconds, int32& Milliseconds)
 {
 	Days = InTimespan.GetDays();
 	Hours = InTimespan.GetHours();
 	Minutes = InTimespan.GetMinutes();
 	Seconds = InTimespan.GetSeconds();
-	Milliseconds = InTimespan.GetMilliseconds();
+	Milliseconds = InTimespan.GetFractionMilli();
 }
 
+
+void UKismetMathLibrary::BreakTimespan2(FTimespan InTimespan, int32& Days, int32& Hours, int32& Minutes, int32& Seconds, int32& FractionNano)
+{
+	Days = InTimespan.GetDays();
+	Hours = InTimespan.GetHours();
+	Minutes = InTimespan.GetMinutes();
+	Seconds = InTimespan.GetSeconds();
+	FractionNano = InTimespan.GetFractionNano();
+}
+
+
 FTimespan UKismetMathLibrary::FromDays(float Days)
 {
 	if (Days < FTimespan::MinValue().GetTotalDays())
@@ -645,6 +663,7 @@ FTimespan UKismetMathLibrary::FromDays(float Days)
 	return FTimespan::FromDays(Days);
 }
 
+
 FTimespan UKismetMathLibrary::FromHours(float Hours)
 {
 	if (Hours < FTimespan::MinValue().GetTotalHours())
@@ -665,6 +684,7 @@ FTimespan UKismetMathLibrary::FromHours(float Hours)
 	return FTimespan::FromHours(Hours);
 }
 
+
 FTimespan UKismetMathLibrary::FromMinutes(float Minutes)
 {
 	if (Minutes < FTimespan::MinValue().GetTotalMinutes())
@@ -685,6 +705,7 @@ FTimespan UKismetMathLibrary::FromMinutes(float Minutes)
 	return FTimespan::FromMinutes(Minutes);
 }
 
+
 FTimespan UKismetMathLibrary::FromSeconds(float Seconds)
 {
 	if (Seconds < FTimespan::MinValue().GetTotalSeconds())
@@ -705,6 +726,7 @@ FTimespan UKismetMathLibrary::FromSeconds(float Seconds)
 	return FTimespan::FromSeconds(Seconds);
 }
 
+
 FTimespan UKismetMathLibrary::FromMilliseconds(float Milliseconds)
 {
 	if (Milliseconds < FTimespan::MinValue().GetTotalMilliseconds())
@@ -725,8 +747,9 @@ FTimespan UKismetMathLibrary::FromMilliseconds(float Milliseconds)
 	return FTimespan::FromMilliseconds(Milliseconds);
 }
 
-/* END Timespan functions */
 
+/* Rotator functions
+*****************************************************************************/
 
 FVector UKismetMathLibrary::GetForwardVector(FRotator InRot)
 {
diff --git a/Engine/Source/Runtime/Engine/Private/Materials/HLSLMaterialTranslator.h b/Engine/Source/Runtime/Engine/Private/Materials/HLSLMaterialTranslator.h
--- a/Engine/Source/Runtime/Engine/Private/Materials/HLSLMaterialTranslator.h
+++ b/Engine/Source/Runtime/Engine/Private/Materials/HLSLMaterialTranslator.h
@@ -1467,6 +1467,12 @@ protected:
 			return Errorf(TEXT("Operation not supported on a Texture"));
 		}
 
+		// External textures must have an external texture uniform expression
+		if ((Type & MCT_TextureExternal) && !UniformExpression->GetExternalTextureUniformExpression())
+		{
+			return Errorf(TEXT("Operation not supported on an external texture"));
+		}
+
 		if (Type == MCT_StaticBool)
 		{
 			return Errorf(TEXT("Operation not supported on a Static Bool"));
@@ -1547,6 +1553,8 @@ protected:
 		// Any code chunk can have a texture uniform expression (eg FMaterialUniformExpressionFlipBookTextureParameter),
 		// But a texture code chunk must have a texture uniform expression
 		check(!(CodeChunk.Type & MCT_Texture) || TextureUniformExpression || ExternalTextureUniformExpression);
+		// External texture samples must have a corresponding uniform expression
+		check(!(CodeChunk.Type & MCT_TextureExternal) || ExternalTextureUniformExpression);
 
 		TCHAR FormattedCode[MAX_SPRINTF]=TEXT("");
 		if(CodeChunk.Type == MCT_Float)
@@ -3678,6 +3686,48 @@ protected:
 		return AddUniformExpression(new FMaterialUniformExpressionExternalTexture(ExternalTextureGuid), MCT_TextureExternal, TEXT(""));
 	}
 
+	virtual int32 ExternalTexture(UTexture* InTexture, int32& TextureReferenceIndex) override
+	{
+		if (ShaderFrequency != SF_Pixel)
+		{
+			return INDEX_NONE;
+		}
+
+		TextureReferenceIndex = Material->GetReferencedTextures().Find(InTexture);
+		checkf(TextureReferenceIndex != INDEX_NONE, TEXT("Material expression called Compiler->ExternalTexture() without implementing UMaterialExpression::GetReferencedTexture properly"));
+
+		return AddUniformExpression(new FMaterialUniformExpressionExternalTexture(TextureReferenceIndex), MCT_TextureExternal, TEXT(""));
+	}
+
+	virtual int32 ExternalTextureParameter(FName ParameterName, UTexture* DefaultValue, int32& TextureReferenceIndex) override
+	{
+		if (ShaderFrequency != SF_Pixel)
+		{
+			return INDEX_NONE;
+		}
+
+		TextureReferenceIndex = Material->GetReferencedTextures().Find(DefaultValue);
+		checkf(TextureReferenceIndex != INDEX_NONE, TEXT("Material expression called Compiler->ExternalTextureParameter() without implementing UMaterialExpression::GetReferencedTexture properly"));
+		return AddUniformExpression(new FMaterialUniformExpressionExternalTextureParameter(ParameterName, TextureReferenceIndex), MCT_TextureExternal, TEXT(""));
+	}
+
+	virtual int32 ExternalTextureCoordinateScaleRotation(int32 TextureReferenceIndex, TOptional<FName> ParameterName) override
+	{
+		return AddUniformExpression(new FMaterialUniformExpressionExternalTextureCoordinateScaleRotation(TextureReferenceIndex, ParameterName), MCT_Float4, TEXT(""));
+	}
+	virtual int32 ExternalTextureCoordinateScaleRotation(const FGuid& ExternalTextureGuid) override
+	{
+		return AddUniformExpression(new FMaterialUniformExpressionExternalTextureCoordinateScaleRotation(ExternalTextureGuid), MCT_Float4, TEXT(""));
+	}
+	virtual int32 ExternalTextureCoordinateOffset(int32 TextureReferenceIndex, TOptional<FName> ParameterName) override
+	{
+		return AddUniformExpression(new FMaterialUniformExpressionExternalTextureCoordinateOffset(TextureReferenceIndex, ParameterName), MCT_Float4, TEXT(""));
+	}
+	virtual int32 ExternalTextureCoordinateOffset(const FGuid& ExternalTextureGuid) override
+	{
+		return AddUniformExpression(new FMaterialUniformExpressionExternalTextureCoordinateOffset(ExternalTextureGuid), MCT_Float4, TEXT(""));
+	}
+
 	virtual int32 GetTextureReferenceIndex(UTexture* TextureValue)
 	{
 		return Material->GetReferencedTextures().Find(TextureValue);
@@ -5218,6 +5268,20 @@ protected:
 		}
 	}
 
+	/**
+	 * Returns a float2 texture coordinate after 2x2 transform and offset applied
+	 *
+	 * @return	Code index
+	 */
+	virtual int32 RotateScaleOffsetTexCoords(int32 TexCoordCodeIndex, int32 RotationScale, int32 Offset) override
+	{
+		return AddCodeChunk(MCT_Float2,
+			TEXT("RotateScaleOffsetTexCoords(%s, %s, %s.xy)"),
+			*GetParameterCode(TexCoordCodeIndex),
+			*GetParameterCode(RotationScale),
+			*GetParameterCode(Offset));
+	}
+
 	/**
 	* Handles SpeedTree vertex animation (wind, smooth LOD)
 	*
diff --git a/Engine/Source/Runtime/Engine/Private/Materials/MaterialExpressions.cpp b/Engine/Source/Runtime/Engine/Private/Materials/MaterialExpressions.cpp
--- a/Engine/Source/Runtime/Engine/Private/Materials/MaterialExpressions.cpp
+++ b/Engine/Source/Runtime/Engine/Private/Materials/MaterialExpressions.cpp
@@ -355,6 +355,77 @@ bool CanConnectMaterialValueTypes(uint32 InputType, uint32 OutputType)
 	return false;
 }
 
+#if WITH_EDITOR
+
+/**
+ * Helper function that wraps the supplied texture coordinates in the necessary math to transform them for external textures
+ *
+ * @param Compiler                The compiler to add code to
+ * @param TexCoordCodeIndex       Index to the code chunk that supplies the vanilla texture coordinates
+ * @param TextureReferenceIndex   Index of the texture within the material (used to access the external texture transform at runtime)
+ * @param ParameterName           (Optional) Parameter name of the texture parameter that's assigned to the sample (used to access the external texture transform at runtime)
+ * @return Index to a new code chunk that supplies the transformed UV coordinates
+ */
+int32 CompileExternalTextureCoordinates(FMaterialCompiler* Compiler, int32 TexCoordCodeIndex, int32 TextureReferenceIndex, TOptional<FName> ParameterName = TOptional<FName>())
+{
+	int32 ScaleRotationCode = Compiler->ExternalTextureCoordinateScaleRotation(TextureReferenceIndex, ParameterName);
+	int32 OffsetCode = Compiler->ExternalTextureCoordinateOffset(TextureReferenceIndex, ParameterName);
+
+	return Compiler->RotateScaleOffsetTexCoords(TexCoordCodeIndex, ScaleRotationCode, OffsetCode);
+}
+
+/**
+ * Compile a texture sample taking into consideration external textures (which may use different sampling code in the shader on some platforms)
+ *
+ * @param Compiler                The compiler to add code to
+ * @param TexCoordCodeIndex       Index to the code chunk that supplies the vanilla texture coordinates
+ * @param SamplerType             The type of sampler that is to be used
+ * @param ParameterName           (Optional) Parameter name of the texture parameter that's assigned to the sample
+ * @param MipValue0Index          (Optional) Mip value (0) code index when mips are being used
+ * @param MipValue1Index          (Optional) Mip value (1) code index when mips are being used
+ * @param MipValueMode            (Optional) Texture MIP value mode
+ * @param SamplerSource           (Optional) Sampler source override
+ * @return Index to a new code chunk that samples the texture
+ */
+int32 CompileTextureSample(
+	FMaterialCompiler* Compiler,
+	UTexture* Texture,
+	int32 TexCoordCodeIndex,
+	EMaterialSamplerType SamplerType,
+	TOptional<FName> ParameterName = TOptional<FName>(),
+	int32 MipValue0Index=INDEX_NONE,
+	int32 MipValue1Index=INDEX_NONE,
+	ETextureMipValueMode MipValueMode=TMVM_None,
+	ESamplerSourceMode SamplerSource=SSM_FromTextureAsset
+	)
+{
+	int32 TextureReferenceIndex = INDEX_NONE;
+	int32 TextureCodeIndex = INDEX_NONE;
+	if (SamplerType == SAMPLERTYPE_External)
+	{
+		// External sampler, so generate the necessary external uniform expression based on whether we're using a parameter name or not
+		TextureCodeIndex = ParameterName.IsSet() ? Compiler->ExternalTextureParameter(ParameterName.GetValue(), Texture, TextureReferenceIndex) : Compiler->ExternalTexture(Texture, TextureReferenceIndex);
+
+		// External textures need an extra transform applied to the UV coordinates
+		TexCoordCodeIndex = CompileExternalTextureCoordinates(Compiler, TexCoordCodeIndex, TextureReferenceIndex, ParameterName);
+	}
+	else
+	{
+		TextureCodeIndex = ParameterName.IsSet() ? Compiler->TextureParameter(ParameterName.GetValue(), Texture, TextureReferenceIndex, SamplerSource) : Compiler->Texture(Texture, TextureReferenceIndex, SamplerSource, MipValueMode);
+	}
+
+	return Compiler->TextureSample(
+					TextureCodeIndex,
+					TexCoordCodeIndex,
+					SamplerType,
+					MipValue0Index,
+					MipValue1Index,
+					MipValueMode,
+					SamplerSource,
+					TextureReferenceIndex);
+}
+#endif
+
 UMaterialExpression::UMaterialExpression(const FObjectInitializer& ObjectInitializer)
 	: Super(ObjectInitializer)
 #if WITH_EDITORONLY_DATA
@@ -1219,6 +1290,11 @@ EMaterialSamplerType UMaterialExpressionTextureBase::GetSamplerTypeForTexture(co
 {
 	if (Texture)
 	{
+		if (Texture->GetMaterialType() == MCT_TextureExternal)
+		{
+			return SAMPLERTYPE_External;
+		}
+
 		switch (Texture->CompressionSettings)
 		{
 			case TC_Normalmap:
@@ -1465,10 +1541,24 @@ int32 UMaterialExpressionTextureSample::Compile(class FMaterialCompiler* Compile
 	if (Texture || TextureObject.Expression) // We deal with reroute textures later on in this function..
 	{
 		int32 TextureReferenceIndex = INDEX_NONE;
-		int32 TextureCodeIndex = TextureObject.Expression ? TextureObject.Compile(Compiler) : Compiler->Texture(Texture, TextureReferenceIndex, SamplerSource, MipValueMode);
+		int32 TextureCodeIndex = INDEX_NONE;
+
+		if (TextureObject.Expression)
+		{
+			TextureCodeIndex = TextureObject.Compile(Compiler);
+		}
+		else if (SamplerType == SAMPLERTYPE_External)
+		{
+			TextureCodeIndex = Compiler->ExternalTexture(Texture, TextureReferenceIndex);
+		}
+		else
+		{
+			TextureCodeIndex = Compiler->Texture(Texture, TextureReferenceIndex, SamplerSource, MipValueMode);
+		}
 
 		UTexture* EffectiveTexture = Texture;
-		EMaterialSamplerType EffectiveSamplerType = (EMaterialSamplerType)SamplerType;
+		EMaterialSamplerType EffectiveSamplerType = SamplerType;
+		TOptional<FName> EffectiveParameterName;
 		if (TextureObject.Expression)
 		{
 			UMaterialExpression* InputExpression = TextureObject.Expression;
@@ -1487,7 +1577,7 @@ int32 UMaterialExpressionTextureSample::Compile(class FMaterialCompiler* Compile
 				else if (OutputIndex >= 0)
 				{
 					uint32 OutputType = InputExpression->GetOutputType(OutputIndex);
-					if (OutputType != MCT_Texture2D && OutputType != MCT_TextureCube && OutputType != MCT_Texture)
+					if (OutputType != MCT_Texture2D && OutputType != MCT_TextureCube && OutputType != MCT_Texture && OutputType != MCT_TextureExternal)
 					{
 						return Compiler->Errorf(TEXT("TextureSample> Reroute not bound to proper texture type!"));
 					}
@@ -1527,6 +1617,7 @@ int32 UMaterialExpressionTextureSample::Compile(class FMaterialCompiler* Compile
 			{
 				EffectiveTexture = TextureObjectParameter->Texture;
 				EffectiveSamplerType = TextureObjectParameter->SamplerType;
+				EffectiveParameterName = TextureObjectParameter->ParameterName;
 			}
 
 			TextureReferenceIndex = Compiler->GetTextureReferenceIndex(EffectiveTexture);
@@ -1543,10 +1634,19 @@ int32 UMaterialExpressionTextureSample::Compile(class FMaterialCompiler* Compile
 				}
 			}
 
+			int32 CoordinateIndex = Coordinates.GetTracedInput().Expression ? Coordinates.Compile(Compiler) : Compiler->TextureCoordinate(ConstCoordinate, false, false);
+
+			// If the sampler type is an external texture, we have might have a scale/bias to apply to the UV coordinates.
+			// Generate that code for the TextureReferenceIndex here so we compile it using the correct texture based on possible reroute textures above
+			if (EffectiveSamplerType == SAMPLERTYPE_External)
+			{
+				CoordinateIndex = CompileExternalTextureCoordinates(Compiler, CoordinateIndex, TextureReferenceIndex, EffectiveParameterName);
+			}
+
 			return Compiler->TextureSample(
 				TextureCodeIndex,
-				Coordinates.GetTracedInput().Expression ? Coordinates.Compile(Compiler) : Compiler->TextureCoordinate(ConstCoordinate, false, false),
-				(EMaterialSamplerType)EffectiveSamplerType,
+				CoordinateIndex,
+				EffectiveSamplerType,
 				CompileMipValue0(Compiler),
 				CompileMipValue1(Compiler),
 				MipValueMode,
@@ -1719,19 +1819,16 @@ int32 UMaterialExpressionTextureSampleParameter::Compile(class FMaterialCompiler
 		return UMaterialExpressionTextureSample::Compile(Compiler, OutputIndex);
 	}
 
-	int32 MipValue0Index = CompileMipValue0(Compiler);
-	int32 MipValue1Index = CompileMipValue1(Compiler);
-	int32 TextureReferenceIndex = INDEX_NONE;
-
-	return Compiler->TextureSample(
-					Compiler->TextureParameter(ParameterName, Texture, TextureReferenceIndex, SamplerSource),
-					Coordinates.GetTracedInput().Expression ? Coordinates.Compile(Compiler) : Compiler->TextureCoordinate(ConstCoordinate, false, false),
-					(EMaterialSamplerType)SamplerType,
-					MipValue0Index,
-					MipValue1Index,
-					MipValueMode,
-					SamplerSource,
-					TextureReferenceIndex);
+	return CompileTextureSample(
+		Compiler,
+		Texture,
+		Coordinates.GetTracedInput().Expression ? Coordinates.Compile(Compiler) : Compiler->TextureCoordinate(ConstCoordinate, false, false),
+		SamplerType,
+		ParameterName,
+		CompileMipValue0(Compiler),
+		CompileMipValue1(Compiler),
+		MipValueMode,
+		SamplerSource);
 }
 
 void UMaterialExpressionTextureSampleParameter::GetCaption(TArray<FString>& OutCaptions) const
@@ -1839,7 +1936,7 @@ int32 UMaterialExpressionTextureObjectParameter::Compile(class FMaterialCompiler
 		return CompilerError(Compiler, GetRequirements());
 	}
 
-	return Compiler->TextureParameter(ParameterName, Texture);
+	return SamplerType == SAMPLERTYPE_External ? Compiler->ExternalTextureParameter(ParameterName, Texture) : Compiler->TextureParameter(ParameterName, Texture);
 }
 
 int32 UMaterialExpressionTextureObjectParameter::CompilePreview(class FMaterialCompiler* Compiler, int32 OutputIndex)
@@ -1850,7 +1947,7 @@ int32 UMaterialExpressionTextureObjectParameter::CompilePreview(class FMaterialC
 	}
 
 	// Preview the texture object by actually sampling it
-	return Compiler->TextureSample(Compiler->TextureParameter(ParameterName, Texture), Compiler->TextureCoordinate(0, false, false),(EMaterialSamplerType)SamplerType);
+	return CompileTextureSample(Compiler, Texture, Compiler->TextureCoordinate(0, false, false), SamplerType, ParameterName);
 }
 #endif // WITH_EDITOR
 
@@ -1915,7 +2012,7 @@ int32 UMaterialExpressionTextureObject::Compile(class FMaterialCompiler* Compile
 	{
 		return CompilerError(Compiler, TEXT("Requires valid texture"));
 	}
-	return Compiler->Texture(Texture);
+	return SamplerType == SAMPLERTYPE_External ? Compiler->ExternalTexture(Texture) : Compiler->Texture(Texture);
 }
 
 int32 UMaterialExpressionTextureObject::CompilePreview(class FMaterialCompiler* Compiler, int32 OutputIndex)
@@ -1925,8 +2022,7 @@ int32 UMaterialExpressionTextureObject::CompilePreview(class FMaterialCompiler*
 		return CompilerError(Compiler, TEXT("Requires valid texture"));
 	}
 
-	// Preview the texture object by actually sampling it
-	return Compiler->TextureSample(Compiler->Texture(Texture), Compiler->TextureCoordinate(0, false, false), UMaterialExpressionTextureBase::GetSamplerTypeForTexture( Texture ));
+	return CompileTextureSample(Compiler, Texture, Compiler->TextureCoordinate(0, false, false), UMaterialExpressionTextureBase::GetSamplerTypeForTexture( Texture ));
 }
 
 uint32 UMaterialExpressionTextureObject::GetOutputType(int32 OutputIndex)
@@ -2082,6 +2178,10 @@ bool UMaterialExpressionTextureSampleParameter2D::TextureIsValid( UTexture* InTe
 		{
 			Result = true;
 		}
+		if ( InTexture->GetMaterialType() == MCT_TextureExternal )
+		{
+			Result = true;
+		}
 	}
 	return Result;
 }
diff --git a/Engine/Source/Runtime/Engine/Private/Materials/MaterialUniformExpressions.cpp b/Engine/Source/Runtime/Engine/Private/Materials/MaterialUniformExpressions.cpp
--- a/Engine/Source/Runtime/Engine/Private/Materials/MaterialUniformExpressions.cpp
+++ b/Engine/Source/Runtime/Engine/Private/Materials/MaterialUniformExpressions.cpp
@@ -386,10 +386,25 @@ FUniformBufferRHIRef FUniformExpressionSet::CreateUniformBuffer(const FMaterialR
 			ESamplerSourceMode SourceMode;
 			Uniform2DTextureExpressions[ExpressionIndex]->GetTextureValue(MaterialRenderContext,MaterialRenderContext.Material,Value,SourceMode);
 
-			// gmartin: Trying to locate UE-23902
-			ensureMsgf(!Value || Value->IsValidLowLevel(), TEXT("Texture not valid! UE-23902! Parameter (%s)"),
-				   Uniform2DTextureExpressions[ExpressionIndex]->GetType() == &FMaterialUniformExpressionTextureParameter::StaticType ?
-				   	*((FMaterialUniformExpressionTextureParameter&)*Uniform2DTextureExpressions[ExpressionIndex]).GetParameterName().ToString() : TEXT("non-parameter"));
+			if (Value)
+			{
+				// Pre-application validity checks (explicit ensures to avoid needless string allocation)
+				const FMaterialUniformExpressionTextureParameter* TextureParameter = (Uniform2DTextureExpressions[ExpressionIndex]->GetType() == &FMaterialUniformExpressionTextureParameter::StaticType) ?
+					&static_cast<const FMaterialUniformExpressionTextureParameter&>(*Uniform2DTextureExpressions[ExpressionIndex]) : nullptr;
+
+				// gmartin: Trying to locate UE-23902
+				if (!Value->IsValidLowLevel())
+				{
+					ensureMsgf(false, TEXT("Texture not valid! UE-23902! Parameter (%s)"), TextureParameter ? *TextureParameter->GetParameterName().ToString() : TEXT("non-parameter"));
+				}
+
+				// Do not allow external textures to be applied to normal texture samplers
+				if (Value->GetMaterialType() == MCT_TextureExternal)
+				{
+					ensureMsgf(false, TEXT("External texture applied to a Texture2D sampler. This may work by chance on some platforms but is not portable. Please change sampler type to 'External'. Texture '%s' on parameter '%s' (slot %i) in material '%s'"),
+						*Value->GetPathName(), TextureParameter ? *TextureParameter->GetParameterName().ToString() : TEXT("non-parameter"), ExpressionIndex, *MaterialRenderContext.Material.GetFriendlyName());
+				}
+			}
 
 			if (Value && Value->Resource)
 			{
@@ -465,7 +480,7 @@ FUniformBufferRHIRef FUniformExpressionSet::CreateUniformBuffer(const FMaterialR
 		{
 			FTextureRHIRef TextureRHI;
 			FSamplerStateRHIRef SamplerStateRHI;
-			if (UniformExternalTextureExpressions[ExpressionIndex]->GetExternalTexture(MaterialRenderContext.MaterialRenderProxy, TextureRHI, SamplerStateRHI))
+			if (UniformExternalTextureExpressions[ExpressionIndex]->GetExternalTexture(MaterialRenderContext, TextureRHI, SamplerStateRHI))
 			{
 				*ResourceTable++ = TextureRHI;
 				*ResourceTable++ = SamplerStateRHI;
@@ -571,40 +586,96 @@ bool FMaterialUniformExpressionTexture::IsIdentical(const FMaterialUniformExpres
 	return TextureIndex == OtherTextureExpression->TextureIndex;
 }
 
-
-FMaterialUniformExpressionExternalTexture::FMaterialUniformExpressionExternalTexture()
+FMaterialUniformExpressionExternalTextureBase::FMaterialUniformExpressionExternalTextureBase(int32 InSourceTextureIndex)
+	: SourceTextureIndex(InSourceTextureIndex)
 {}
 
-FMaterialUniformExpressionExternalTexture::FMaterialUniformExpressionExternalTexture(const FGuid& InGuid) :
-	ExternalTextureGuid(InGuid)
-{
-}
-
-bool FMaterialUniformExpressionExternalTexture::GetExternalTexture(const FMaterialRenderProxy* MaterialRenderProxy, FTextureRHIRef& OutTextureRHI, FSamplerStateRHIRef& OutSamplerStateRHI)
+FMaterialUniformExpressionExternalTextureBase::FMaterialUniformExpressionExternalTextureBase(const FGuid& InGuid)
+	: SourceTextureIndex(INDEX_NONE)
+	, ExternalTextureGuid(InGuid)
 {
-	return FExternalTextureRegistry::Get().GetExternalTexture(MaterialRenderProxy, ExternalTextureGuid, OutTextureRHI, OutSamplerStateRHI);
 }
 
-void FMaterialUniformExpressionExternalTexture::Serialize(FArchive& Ar)
+void FMaterialUniformExpressionExternalTextureBase::Serialize(FArchive& Ar)
 {
+	Ar << SourceTextureIndex;
 	Ar << ExternalTextureGuid;
 }
 
-bool FMaterialUniformExpressionExternalTexture::IsIdentical(const FMaterialUniformExpression* OtherExpression) const
+bool FMaterialUniformExpressionExternalTextureBase::IsIdentical(const FMaterialUniformExpression* OtherExpression) const
 {
 	if (GetType() != OtherExpression->GetType())
 	{
 		return false;
 	}
-	FMaterialUniformExpressionExternalTexture* OtherExternalTextureExpression = (FMaterialUniformExpressionExternalTexture*)OtherExpression;
 
-	return ExternalTextureGuid == OtherExternalTextureExpression->ExternalTextureGuid;
+	const auto* Other = static_cast<const FMaterialUniformExpressionExternalTextureBase*>(OtherExpression);
+	return SourceTextureIndex == Other->SourceTextureIndex && ExternalTextureGuid == Other->ExternalTextureGuid;
 }
 
-FArchive& operator<<(FArchive& Ar, FMaterialUniformExpressionExternalTexture*& Ref)
+FGuid FMaterialUniformExpressionExternalTextureBase::ResolveExternalTextureGUID(const FMaterialRenderContext& Context, TOptional<FName> ParameterName) const
 {
-	Ar << (FMaterialUniformExpression*&)Ref;
-	return Ar;
+	// Use the compile-time GUID if it is set
+	if (ExternalTextureGuid.IsValid())
+	{
+		return ExternalTextureGuid;
+	}
+
+	const UTexture* TextureParameterObject = nullptr;
+	if (ParameterName.IsSet() && Context.MaterialRenderProxy && Context.MaterialRenderProxy->GetTextureValue(ParameterName.GetValue(), &TextureParameterObject, Context) && TextureParameterObject)
+	{
+		return TextureParameterObject->GetExternalTextureGuid();
+	}
+
+	// Otherwise attempt to use the texture index in the material, if it's valid
+	const UTexture* TextureObject = SourceTextureIndex != INDEX_NONE ? GetIndexedTexture(Context.Material, SourceTextureIndex) : nullptr;
+	if (TextureObject)
+	{
+		return TextureObject->GetExternalTextureGuid();
+	}
+
+	return FGuid();
+}
+
+bool FMaterialUniformExpressionExternalTexture::GetExternalTexture(const FMaterialRenderContext& Context, FTextureRHIRef& OutTextureRHI, FSamplerStateRHIRef& OutSamplerStateRHI) const
+{
+	check(IsInParallelRenderingThread());
+
+	FGuid GuidToLookup = ResolveExternalTextureGUID(Context);
+	return GuidToLookup.IsValid() && FExternalTextureRegistry::Get().GetExternalTexture(Context.MaterialRenderProxy, GuidToLookup, OutTextureRHI, OutSamplerStateRHI);
+}
+
+FMaterialUniformExpressionExternalTextureParameter::FMaterialUniformExpressionExternalTextureParameter()
+{}
+
+FMaterialUniformExpressionExternalTextureParameter::FMaterialUniformExpressionExternalTextureParameter(FName InParameterName, int32 InTextureIndex)
+	: Super(InTextureIndex)
+	, ParameterName(InParameterName)
+{}
+
+void FMaterialUniformExpressionExternalTextureParameter::Serialize(FArchive& Ar)
+{
+	Ar << ParameterName;
+	Super::Serialize(Ar);
+}
+
+bool FMaterialUniformExpressionExternalTextureParameter::GetExternalTexture(const FMaterialRenderContext& Context, FTextureRHIRef& OutTextureRHI, FSamplerStateRHIRef& OutSamplerStateRHI) const
+{
+	check(IsInParallelRenderingThread());
+
+	FGuid GuidToLookup = ResolveExternalTextureGUID(Context, ParameterName);
+	return GuidToLookup.IsValid() && FExternalTextureRegistry::Get().GetExternalTexture(Context.MaterialRenderProxy, GuidToLookup, OutTextureRHI, OutSamplerStateRHI);
+}
+
+bool FMaterialUniformExpressionExternalTextureParameter::IsIdentical(const FMaterialUniformExpression* OtherExpression) const
+{
+	if (GetType() != OtherExpression->GetType())
+	{
+		return false;
+	}
+
+	auto* Other = static_cast<const FMaterialUniformExpressionExternalTextureParameter*>(OtherExpression);
+	return ParameterName == Other->ParameterName && Super::IsIdentical(OtherExpression);
 }
 
 void FMaterialUniformExpressionVectorParameter::GetGameThreadNumberValue(const UMaterialInterface* SourceMaterialToCopyFrom, FLinearColor& OutValue) const
@@ -673,6 +744,83 @@ void FMaterialUniformExpressionScalarParameter::GetGameThreadNumberValue(const U
 	}
 }
 
+namespace
+{
+	void SerializeOptional(FArchive Ar, TOptional<FName> OptionalName)
+	{
+		bool bIsSet = OptionalName.IsSet();
+		Ar << bIsSet;
+
+		if (bIsSet)
+		{
+			if (!OptionalName.IsSet())
+			{
+				OptionalName.Emplace();
+			}
+
+			Ar << OptionalName.GetValue();
+		}
+	}
+}
+
+void FMaterialUniformExpressionExternalTextureCoordinateScaleRotation::Serialize(FArchive& Ar)
+{
+	// Write out the optional parameter name
+	SerializeOptional(Ar, ParameterName);
+
+	Super::Serialize(Ar);
+}
+
+bool FMaterialUniformExpressionExternalTextureCoordinateScaleRotation::IsIdentical(const FMaterialUniformExpression* OtherExpression) const
+{
+	if (GetType() != OtherExpression->GetType() || !Super::IsIdentical(OtherExpression))
+	{
+		return false;
+	}
+
+	const auto* Other = static_cast<const FMaterialUniformExpressionExternalTextureCoordinateScaleRotation*>(OtherExpression);
+	return ParameterName == Other->ParameterName;
+}
+
+void FMaterialUniformExpressionExternalTextureCoordinateScaleRotation::GetNumberValue(const FMaterialRenderContext& Context, FLinearColor& OutValue) const
+{
+	FGuid GuidToLookup = ResolveExternalTextureGUID(Context, ParameterName);
+	if (!GuidToLookup.IsValid() || !FExternalTextureRegistry::Get().GetExternalTextureCoordinateScaleRotation(GuidToLookup, OutValue))
+	{
+		OutValue = FLinearColor(1.f, 0.f, 0.f, 1.f);
+	}
+}
+
+void FMaterialUniformExpressionExternalTextureCoordinateOffset::Serialize(FArchive& Ar)
+{
+	// Write out the optional parameter name
+	SerializeOptional(Ar, ParameterName);
+
+	Super::Serialize(Ar);
+}
+
+bool FMaterialUniformExpressionExternalTextureCoordinateOffset::IsIdentical(const FMaterialUniformExpression* OtherExpression) const
+{
+	if (GetType() != OtherExpression->GetType() || !Super::IsIdentical(OtherExpression))
+	{
+		return false;
+	}
+
+	const auto* Other = static_cast<const FMaterialUniformExpressionExternalTextureCoordinateOffset*>(OtherExpression);
+	return ParameterName == Other->ParameterName;
+}
+
+void FMaterialUniformExpressionExternalTextureCoordinateOffset::GetNumberValue(const FMaterialRenderContext& Context, FLinearColor& OutValue) const
+{
+	FGuid GuidToLookup = ResolveExternalTextureGUID(Context, ParameterName);
+	if (!GuidToLookup.IsValid() || !FExternalTextureRegistry::Get().GetExternalTextureCoordinateOffset(GuidToLookup, OutValue))
+	{
+		OutValue.R = OutValue.G = 0;
+		OutValue.B = OutValue.A = 0;
+	}
+}
+
+
 IMPLEMENT_MATERIALUNIFORMEXPRESSION_TYPE(FMaterialUniformExpressionTexture);
 IMPLEMENT_MATERIALUNIFORMEXPRESSION_TYPE(FMaterialUniformExpressionConstant);
 IMPLEMENT_MATERIALUNIFORMEXPRESSION_TYPE(FMaterialUniformExpressionTime);
@@ -680,7 +828,11 @@ IMPLEMENT_MATERIALUNIFORMEXPRESSION_TYPE(FMaterialUniformExpressionRealTime);
 IMPLEMENT_MATERIALUNIFORMEXPRESSION_TYPE(FMaterialUniformExpressionVectorParameter);
 IMPLEMENT_MATERIALUNIFORMEXPRESSION_TYPE(FMaterialUniformExpressionScalarParameter);
 IMPLEMENT_MATERIALUNIFORMEXPRESSION_TYPE(FMaterialUniformExpressionTextureParameter);
+IMPLEMENT_MATERIALUNIFORMEXPRESSION_TYPE(FMaterialUniformExpressionExternalTextureBase);
 IMPLEMENT_MATERIALUNIFORMEXPRESSION_TYPE(FMaterialUniformExpressionExternalTexture);
+IMPLEMENT_MATERIALUNIFORMEXPRESSION_TYPE(FMaterialUniformExpressionExternalTextureParameter);
+IMPLEMENT_MATERIALUNIFORMEXPRESSION_TYPE(FMaterialUniformExpressionExternalTextureCoordinateScaleRotation);
+IMPLEMENT_MATERIALUNIFORMEXPRESSION_TYPE(FMaterialUniformExpressionExternalTextureCoordinateOffset);
 IMPLEMENT_MATERIALUNIFORMEXPRESSION_TYPE(FMaterialUniformExpressionFlipBookTextureParameter);
 IMPLEMENT_MATERIALUNIFORMEXPRESSION_TYPE(FMaterialUniformExpressionSine);
 IMPLEMENT_MATERIALUNIFORMEXPRESSION_TYPE(FMaterialUniformExpressionSquareRoot);
diff --git a/Engine/Source/Runtime/Engine/Private/Materials/MaterialUniformExpressions.h b/Engine/Source/Runtime/Engine/Private/Materials/MaterialUniformExpressions.h
--- a/Engine/Source/Runtime/Engine/Private/Materials/MaterialUniformExpressions.h
+++ b/Engine/Source/Runtime/Engine/Private/Materials/MaterialUniformExpressions.h
@@ -388,6 +388,27 @@ public:
 	}
 };
 
+
+/**
+ * An external texture parameter expression.
+ */
+class FMaterialUniformExpressionExternalTextureParameter: public FMaterialUniformExpressionExternalTexture
+{
+	typedef FMaterialUniformExpressionExternalTexture Super;
+	DECLARE_MATERIALUNIFORMEXPRESSION_TYPE(FMaterialUniformExpressionExternalTextureParameter);
+public:
+
+	FMaterialUniformExpressionExternalTextureParameter();
+	FMaterialUniformExpressionExternalTextureParameter(FName InParameterName, int32 InTextureIndex);
+
+	virtual void Serialize(FArchive& Ar) override;
+	virtual bool IsIdentical(const FMaterialUniformExpression* OtherExpression) const override;
+	virtual bool GetExternalTexture(const FMaterialRenderContext& Context, FTextureRHIRef& OutTextureRHI, FSamplerStateRHIRef& OutSamplerStateRHI) const override;
+
+private:
+	FName ParameterName;
+};
+
 /**
  */
 class FMaterialUniformExpressionSine: public FMaterialUniformExpression
@@ -1744,3 +1765,50 @@ private:
 	TRefCountPtr<FMaterialUniformExpressionTexture> TextureExpression;
 	int8 TextureProperty;
 };
+
+
+/**
+ * A uniform expression to lookup the UV coordinate rotation and scale for an external texture
+ */
+class FMaterialUniformExpressionExternalTextureCoordinateScaleRotation : public FMaterialUniformExpressionExternalTextureBase
+{
+	DECLARE_MATERIALUNIFORMEXPRESSION_TYPE(FMaterialUniformExpressionExternalTextureCoordinateScaleRotation);
+public:
+
+	FMaterialUniformExpressionExternalTextureCoordinateScaleRotation(){}
+	FMaterialUniformExpressionExternalTextureCoordinateScaleRotation(const FGuid& InGuid) : FMaterialUniformExpressionExternalTextureBase(InGuid) {}
+	FMaterialUniformExpressionExternalTextureCoordinateScaleRotation(int32 InSourceTextureIndex, TOptional<FName> InParameterName) : FMaterialUniformExpressionExternalTextureBase(InSourceTextureIndex), ParameterName(InParameterName) {}
+
+	virtual void Serialize(FArchive& Ar) override;
+	virtual bool IsIdentical(const FMaterialUniformExpression* OtherExpression) const override;
+	virtual void GetNumberValue(const FMaterialRenderContext& Context, FLinearColor& OutValue) const override;
+
+protected:
+	typedef FMaterialUniformExpressionExternalTextureBase Super;
+
+	/** Optional texture parameter name */
+	TOptional<FName> ParameterName;
+};
+
+/**
+ * A uniform expression to lookup the UV coordinate offset for an external texture
+ */
+class FMaterialUniformExpressionExternalTextureCoordinateOffset : public FMaterialUniformExpressionExternalTextureBase
+{
+	DECLARE_MATERIALUNIFORMEXPRESSION_TYPE(FMaterialUniformExpressionExternalTextureCoordinateOffset);
+public:
+
+	FMaterialUniformExpressionExternalTextureCoordinateOffset(){}
+	FMaterialUniformExpressionExternalTextureCoordinateOffset(const FGuid& InGuid) : FMaterialUniformExpressionExternalTextureBase(InGuid) {}
+	FMaterialUniformExpressionExternalTextureCoordinateOffset(int32 InSourceTextureIndex, TOptional<FName> InParameterName) : FMaterialUniformExpressionExternalTextureBase(InSourceTextureIndex), ParameterName(InParameterName) {}
+
+	virtual void Serialize(FArchive& Ar) override;
+	virtual bool IsIdentical(const FMaterialUniformExpression* OtherExpression) const override;
+	virtual void GetNumberValue(const FMaterialRenderContext& Context, FLinearColor& OutValue) const override;
+
+protected:
+	typedef FMaterialUniformExpressionExternalTextureBase Super;
+
+	/** Optional texture parameter name */
+	TOptional<FName> ParameterName;
+};
\ No newline at end of file
diff --git a/Engine/Source/Runtime/Engine/Private/MatineeUtils.cpp b/Engine/Source/Runtime/Engine/Private/MatineeUtils.cpp
--- a/Engine/Source/Runtime/Engine/Private/MatineeUtils.cpp
+++ b/Engine/Source/Runtime/Engine/Private/MatineeUtils.cpp
@@ -142,7 +142,7 @@ public:
 			BuildRedirectionTable();
 		}
 
-		if (FTrackRemapInfo* ClassRemapInfo = TrackRedirectMap.Find(TargetClass))
+		if (FTrackRemapInfo* ClassRemapInfo = TrackRedirectMap.Find(TargetClass->GetFName()))
 		{
 			// Scan thru prefixes to see if any match
 			const FString OldTrackName = TrackName.ToString();
@@ -185,7 +185,7 @@ private:
 
 				if (UClass* TargetClass = LoadClass<UObject>(NULL, *TargetClassName.ToString(), NULL, LOAD_None, NULL))
 				{
-					FTrackRemapInfo& ClassRemapInfo = TrackRedirectMap.FindOrAdd(TargetClass);
+					FTrackRemapInfo& ClassRemapInfo = TrackRedirectMap.FindOrAdd(TargetClass->GetFName());
 
 					FString OldFieldName;
 					FString NewFieldName;
@@ -212,13 +212,13 @@ private:
 	};
 private:
 	// A mapping from old track name to new track name, sorted by actor class being driven by the track.  Read in from MatineeTrackRedirects in ini files.
-	static TMultiMap<UClass*, FTrackRemapInfo> TrackRedirectMap;
+	static TMultiMap<FName, FTrackRemapInfo> TrackRedirectMap;
 
 	static bool bInitialized;
 };
 
 // Static members of FMatineeTrackRedirectionManager
-TMultiMap<UClass*, FMatineeTrackRedirectionManager::FTrackRemapInfo> FMatineeTrackRedirectionManager::TrackRedirectMap;
+TMultiMap<FName, FMatineeTrackRedirectionManager::FTrackRemapInfo> FMatineeTrackRedirectionManager::TrackRedirectMap;
 bool FMatineeTrackRedirectionManager::bInitialized = false;
 
 /////////////////////////////////////////////////////
diff --git a/Engine/Source/Runtime/Engine/Private/ShaderDerivedDataVersion.h b/Engine/Source/Runtime/Engine/Private/ShaderDerivedDataVersion.h
--- a/Engine/Source/Runtime/Engine/Private/ShaderDerivedDataVersion.h
+++ b/Engine/Source/Runtime/Engine/Private/ShaderDerivedDataVersion.h
@@ -11,5 +11,5 @@
 // In case of merge conflicts with DDC versions, you MUST generate a new GUID and set this new
 // guid as version
 
-#define GLOBALSHADERMAP_DERIVEDDATA_VER			TEXT("3E5071BDE6254726B84BA594D4D08422")
-#define MATERIALSHADERMAP_DERIVEDDATA_VER		TEXT("9415B3AB201F4D66BC34969101289C60")
+#define GLOBALSHADERMAP_DERIVEDDATA_VER			TEXT("605300842A8542F8A2AD72835F18F142")
+#define MATERIALSHADERMAP_DERIVEDDATA_VER		TEXT("A6F6D46EF9844AD39968049DB1F003F5")
