Fix CIS/circular dependency.

diff --git a/Engine/Source/Runtime/Core/Private/Misc/Timespan.cpp b/Engine/Source/Runtime/Core/Private/Misc/Timespan.cpp
--- a/Engine/Source/Runtime/Core/Private/Misc/Timespan.cpp
+++ b/Engine/Source/Runtime/Core/Private/Misc/Timespan.cpp
@@ -183,6 +183,21 @@ bool FTimespan::Parse(const FString& TimespanString, FTimespan& OutTimespan)
 }
 
 
+/* FTimespan friend functions
+ *****************************************************************************/
+
+FArchive& operator<<(FArchive& Ar, FTimespan& Timespan)
+{
+	return Ar << Timespan.Ticks;
+}
+
+
+uint32 GetTypeHash(const FTimespan& Timespan)
+{
+	return GetTypeHash(Timespan.Ticks);
+}
+
+
 /* FTimespan implementation
  *****************************************************************************/
 
diff --git a/Engine/Source/Runtime/Core/Public/Misc/Timespan.h b/Engine/Source/Runtime/Core/Public/Misc/Timespan.h
--- a/Engine/Source/Runtime/Core/Public/Misc/Timespan.h
+++ b/Engine/Source/Runtime/Core/Public/Misc/Timespan.h
@@ -3,10 +3,10 @@
 #pragma once
 
 #include "CoreTypes.h"
-#include "Containers/UnrealString.h"
 
 class FArchive;
 class FOutputDevice;
+class FString;
 class UObject;
 
 
@@ -736,10 +736,7 @@ public:
 	 * @param Timespan The time span value to serialize.
 	 * @return The archive.
 	 */
-	friend CORE_API FArchive& operator<<(FArchive& Ar, FTimespan& Timespan)
-	{
-		return Ar << Timespan.Ticks;
-	}
+	friend CORE_API FArchive& operator<<(FArchive& Ar, FTimespan& Timespan);
 
 	/**
 	 * Get the hash for the specified time span.
@@ -747,10 +744,7 @@ public:
 	 * @param Timespan The timespan to get the hash for.
 	 * @return Hash value.
 	 */
-	friend CORE_API uint32 GetTypeHash(const FTimespan& Timespan)
-	{
-		return GetTypeHash(Timespan.Ticks);
-	}
+	friend CORE_API uint32 GetTypeHash(const FTimespan& Timespan);
 
 protected:
 
