Copying //UE4/Dev-Enterprise to //UE4/Dev-Main (Source: //UE4/Dev-Enterprise @ 3622748)

============================
  MAJOR FEATURES & CHANGES
============================

Change 3576719 by JeanMichel.Dignard

	Made IntroSort the default sort.

	- Algo::Sort and ::Sort will call Algo::IntroSort
	- The sort version that was in Algo::Sort is now in Algo::LegacySort

Change 3604785 by [REDACTED_PERSON]

	Add Save/Load to ProfileVisualizer
	Expose ProfileVizualizer in Windows\Developer Tools\Debug menu

diff --git a/Engine/Source/Developer/TaskGraph/Private/SProfileVisualizer.cpp b/Engine/Source/Developer/TaskGraph/Private/SProfileVisualizer.cpp
--- a/Engine/Source/Developer/TaskGraph/Private/SProfileVisualizer.cpp
+++ b/Engine/Source/Developer/TaskGraph/Private/SProfileVisualizer.cpp
@@ -14,8 +14,22 @@
 #include "TaskGraphStyle.h"
 #include "SBarVisualizer.h"
 #include "SEventsTree.h"
+#include "STaskGraph.h"
+#include "SButton.h"
+#include "STextBlock.h"
+#include "App.h"
+#include "EngineVersion.h"
+#include "FileManager.h"
+#include "Framework/Notifications/NotificationManager.h"
+#include "Widgets/Notifications/SNotificationList.h"
+
+#if PLATFORM_DESKTOP && WITH_EDITOR
+#include "DesktopPlatformModule.h"
+#endif
+
+#include "GenericCommands.h"
 	
-void SProfileVisualizer::Construct( const FArguments& InArgs )
+void SProfileVisualizer::Construct(const FArguments& InArgs)
 {
 	ProfileData = InArgs._ProfileData;
 	ProfilerType = InArgs._ProfilerType;
@@ -28,6 +42,37 @@ void SProfileVisualizer::Construct( const FArguments& InArgs )
 		SNew(SVerticalBox)
 		+ SVerticalBox::Slot()
 		.AutoHeight()
+		[
+			SNew(SHorizontalBox)
+			+ SHorizontalBox::Slot()
+			.AutoWidth()
+			.Padding(0.0f)
+			[
+				SNew(SButton)
+				.OnClicked(this, &SProfileVisualizer::OnLoadClicked)
+				.Content()
+				[
+					SNew(STextBlock)
+					.Text(NSLOCTEXT("TaskGraph", "Load", "Load"))
+					.ToolTipText(NSLOCTEXT("TaskGraph", "Load_GPUTooltip", "Load GPU profiling data"))
+				]
+			]
+			+ SHorizontalBox::Slot()
+			.AutoWidth()
+			.Padding(0.0f)
+			[
+				SNew(SButton)
+				.OnClicked(this, &SProfileVisualizer::OnSaveClicked)
+				.Content()
+				[
+					SNew(STextBlock)
+					.Text(NSLOCTEXT("TaskGraph", "Save", "Save"))
+					.ToolTipText(NSLOCTEXT("TaskGraph", "Save_GPUTooltip", "Save the GPU profiling data"))
+				]
+			]
+		]
+		+ SVerticalBox::Slot()
+		.AutoHeight()
 		[
 			SNew(SBorder)
 			.Visibility(EVisibility::Visible)
@@ -165,3 +210,111 @@ void SProfileVisualizer::ShowGraphBarInEventsWindow( int32 WindowIndex )
 	EventsTree->HandleBarGraphExpansionChanged( SelectedBarGraph );
 }
 
+FReply SProfileVisualizer::OnSaveClicked()
+{
+#if PLATFORM_DESKTOP && WITH_EDITOR
+	// Message to display on completion
+	FText Message;
+
+	FString ProfileFilename = FPaths::ProjectLogDir();
+	ProfileFilename /= TEXT("profileViz");
+	ProfileFilename /= FString::Printf(
+		TEXT("%s-%i-%s.profViz"),
+		FApp::GetProjectName(),
+		FEngineVersion::Current().GetChangelist(),
+		*FDateTime::Now().ToString());
+
+	FArchive* ProfileFile = IFileManager::Get().CreateFileWriter(*ProfileFilename);
+
+	if (ProfileFile)
+	{
+		FVisualizerEvent::SaveVisualizerEventRecursively(ProfileFile, ProfileData);
+
+		// Close and delete archive.
+		ProfileFile->Close();
+		delete ProfileFile;
+		ProfileFile = NULL;
+
+		Message = NSLOCTEXT("TaskGraph", "ExportMessage", "Wrote profile data to file");
+	}
+	else
+	{
+		Message = NSLOCTEXT("TaskGraph", "ExportMessage", "Could not write profile data to file");
+	}
+
+	struct Local
+	{
+		static void NavigateToExportedFile(FString InGPUFilename, bool bInSuccessful)
+		{
+			InGPUFilename = FPaths::ConvertRelativePathToFull(InGPUFilename);
+			if (bInSuccessful)
+			{
+				FPlatformProcess::LaunchFileInDefaultExternalApplication(*InGPUFilename);
+			}
+			else
+			{
+				FPlatformProcess::ExploreFolder(*FPaths::GetPath(InGPUFilename));
+			}
+		}
+	};
+
+	FNotificationInfo Info(Message);
+	Info.Hyperlink = FSimpleDelegate::CreateStatic(&Local::NavigateToExportedFile, ProfileFilename, false);
+	Info.HyperlinkText = FText::FromString(ProfileFilename);
+	Info.bUseLargeFont = false;
+	Info.bFireAndForget = true;
+	Info.ExpireDuration = 8.0f;
+	FSlateNotificationManager::Get().AddNotification(Info);
+#endif //#if PLATFORM_DESKTOP && WITH_EDITOR
+
+	return FReply::Handled();
+
+}
+
+FReply SProfileVisualizer::OnLoadClicked()
+{
+#if PLATFORM_DESKTOP && WITH_EDITOR
+	// Prompt the user for the Filename
+	IDesktopPlatform* DesktopPlatform = FDesktopPlatformModule::Get();
+	if (DesktopPlatform != nullptr)
+	{
+		TArray<FString> OpenFilenames;
+		bool bOpened = DesktopPlatform->OpenFileDialog(
+			FSlateApplication::Get().FindBestParentWindowHandleForDialogs(nullptr),
+			NSLOCTEXT("UnrealEd", "Load", "Load Profile data").ToString(),
+			TEXT(""),
+			TEXT(""),
+			TEXT("Profile data (*.profViz) | *.profViz"),
+			EFileDialogFlags::None,
+			OpenFilenames);
+
+		if (bOpened)
+		{
+			FArchive* ProfileFile = IFileManager::Get().CreateFileReader(*OpenFilenames[0]);
+
+			if (ProfileFile)
+			{
+				TSharedPtr< FVisualizerEvent > InVisualizerData;
+				InVisualizerData = FVisualizerEvent::LoadVisualizerEvent(ProfileFile);
+								
+				static FName TaskGraphModule(TEXT("TaskGraph"));
+				if (FModuleManager::Get().IsModuleLoaded(TaskGraphModule))
+				{
+					IProfileVisualizerModule& ProfileVisualizer = FModuleManager::GetModuleChecked<IProfileVisualizerModule>(TaskGraphModule);
+
+					FText LoadedFileName = FText::AsCultureInvariant(ProfileFile->GetArchiveName());
+					ProfileVisualizer.DisplayProfileVisualizer(InVisualizerData, TEXT("Profile Data"), LoadedFileName);
+				}
+
+				// Close and delete archive.
+				ProfileFile->Close();
+				delete ProfileFile;
+				ProfileFile = NULL;
+			}
+		}
+	}
+#endif//#if PLATFORM_DESKTOP && WITH_EDITOR
+
+	return FReply::Handled();
+
+}
\ No newline at end of file
diff --git a/Engine/Source/Developer/TaskGraph/Private/SProfileVisualizer.h b/Engine/Source/Developer/TaskGraph/Private/SProfileVisualizer.h
--- a/Engine/Source/Developer/TaskGraph/Private/SProfileVisualizer.h
+++ b/Engine/Source/Developer/TaskGraph/Private/SProfileVisualizer.h
@@ -101,5 +101,11 @@ private:
 	 
 	/** Optional header message text color */
 	FLinearColor HeaderMessageTextColor;
+
+	/** We just clicked the save button */
+	FReply OnSaveClicked();
+
+	/** We just clicked the load button */
+	FReply OnLoadClicked();
 };
 
diff --git a/Engine/Source/Developer/TaskGraph/Private/STaskGraph.cpp b/Engine/Source/Developer/TaskGraph/Private/STaskGraph.cpp
--- a/Engine/Source/Developer/TaskGraph/Private/STaskGraph.cpp
+++ b/Engine/Source/Developer/TaskGraph/Private/STaskGraph.cpp
@@ -10,6 +10,10 @@
 #include "SProfileVisualizer.h"
 #include "Widgets/Docking/SDockTab.h"
 #include "TaskGraphStyle.h"
+#if WITH_EDITOR
+#include "WorkspaceMenuStructure.h"
+#include "WorkspaceMenuStructureModule.h"
+#endif //#if WITH_EDITOR
 
 /**
  * Creates Visualizer using Visualizer profile data format
@@ -17,22 +21,31 @@
  * @param ProfileData Visualizer data
  * @return Visualizer window
  */
-void MakeTaskGraphVisualizerWindow( TSharedPtr< FVisualizerEvent > ProfileData, const FText& WindowTitle, const FText& ProfilerType, const FText& HeaderMessageText = FText::GetEmpty(), const FLinearColor& HeaderMessageTextColor = FLinearColor::White )
+
+FName TaskGraphTabId("VisualizerSpawnPoint");
+TSharedRef<SDockTab> MakeTaskGraphVisualizerWindow( TSharedPtr< FVisualizerEvent > ProfileData, const FText& WindowTitle, const FText& ProfilerType, const FText& HeaderMessageText = FText::GetEmpty(), const FLinearColor& HeaderMessageTextColor = FLinearColor::White, const bool InsertTab = true )
 {
-	FGlobalTabmanager::Get()->InsertNewDocumentTab
+	TSharedRef<SDockTab> TaskGraphVizualizer = SNew(SDockTab)
+												.Label(WindowTitle)
+												.TabRole(ETabRole::NomadTab)
+												[
+													SNew(SProfileVisualizer)
+													.ProfileData(ProfileData)
+													.ProfilerType(ProfilerType)
+													.HeaderMessageText(HeaderMessageText)
+													.HeaderMessageTextColor(HeaderMessageTextColor)
+												];
+
+	if (InsertTab)
+	{
+		FGlobalTabmanager::Get()->InsertNewDocumentTab
 		(
-			"VisualizerSpawnPoint", FTabManager::ESearchPreference::RequireClosedTab,
-			SNew( SDockTab )
-			.Label( WindowTitle )
-			.TabRole( ETabRole::DocumentTab )
-			[
-				SNew( SProfileVisualizer )
-				.ProfileData( ProfileData )
-				.ProfilerType( ProfilerType )
-				.HeaderMessageText( HeaderMessageText )
-				.HeaderMessageTextColor( HeaderMessageTextColor )
-			]
+			TaskGraphTabId, FTabManager::ESearchPreference::RequireClosedTab,
+			TaskGraphVizualizer
 		);
+	}
+
+	return TaskGraphVizualizer;
 }
 
 /** Helper class that counts down when to unpause and stop movie. */
@@ -135,23 +148,6 @@ void DisplayProfileVisualizer(TSharedPtr< FVisualizerEvent > InProfileData, cons
 {
 	check( IsInGameThread() );
 
-	if ( !GHasRegisteredVisualizerLayout )
-	{
-		TSharedRef<FTabManager::FLayout> Layout = FTabManager::NewLayout( "Visualizer_Layout" )
-		->AddArea
-		(
-			FTabManager::NewArea(720, 768)
-			->Split
-			(
-				FTabManager::NewStack()
-				->AddTab("VisualizerSpawnPoint", ETabState::ClosedTab)
-			)
-		);
-
-		FGlobalTabmanager::Get()->RestoreFrom( Layout, TSharedPtr<SWindow>() );
-		GHasRegisteredVisualizerLayout = true;
-	}
-
 	FFormatNamedArguments Args;
 	Args.Add( TEXT("ProfilerType"), FText::FromString( InProfilerType ) );
 
@@ -171,9 +167,21 @@ public:
 	virtual void StartupModule() override
 	{
 		::InitProfileVisualizer();
+
+#if WITH_EDITOR
+		FGlobalTabmanager::Get()->RegisterNomadTabSpawner(TaskGraphTabId, FOnSpawnTab::CreateRaw(this, &FProfileVisualizerModule::SpawnProfileVizualizerTab))
+			.SetDisplayName(NSLOCTEXT("ProfileVisualizerModule", "TabTitle", "Profile Data Visualizer"))
+			.SetTooltipText(NSLOCTEXT("ProfileVisualizerModule", "TooltipText", "Open the Profile Data Visualizer tab."))
+			.SetGroup(WorkspaceMenu::GetMenuStructure().GetDeveloperToolsDebugCategory())
+			;
+#endif
+			
 	}
 	virtual void ShutdownModule() override
 	{
+#if WITH_EDITOR
+		FGlobalTabmanager::Get()->UnregisterNomadTabSpawner(TaskGraphTabId);
+#endif		
 		::ShutdownProfileVisualizer();
 	}
 	virtual void DisplayProfileVisualizer(TSharedPtr< FVisualizerEvent > InProfileData, const TCHAR* InProfilerType, const FText& HeaderMessageText, const FLinearColor& HeaderMessageTextColor) override
@@ -182,5 +190,13 @@ public:
 		::DisplayProfileVisualizer( InProfileData, InProfilerType, HeaderMessageText, HeaderMessageTextColor );
 #endif // WITH_EDITOR
 	}
+
+	TSharedRef<SDockTab> SpawnProfileVizualizerTab(const FSpawnTabArgs& Args)
+	{
+		TSharedPtr< FVisualizerEvent > InProfileData(new FVisualizerEvent(0., 0., 0., 0, "Dummy"));
+		
+		return MakeTaskGraphVisualizerWindow(InProfileData, FText::GetEmpty(), FText::GetEmpty(), FText::GetEmpty(), FLinearColor::White, false);
+	}
+
 };
 IMPLEMENT_MODULE(FProfileVisualizerModule, TaskGraph);
diff --git a/Engine/Source/Developer/TaskGraph/Private/VisualizerEvents.cpp b/Engine/Source/Developer/TaskGraph/Private/VisualizerEvents.cpp
--- a/Engine/Source/Developer/TaskGraph/Private/VisualizerEvents.cpp
+++ b/Engine/Source/Developer/TaskGraph/Private/VisualizerEvents.cpp
@@ -1,3 +1,79 @@
 // Copyright 1998-2017 Epic Games, Inc. All Rights Reserved.
 
 #include "VisualizerEvents.h"
+
+#if defined( _WIN64 )
+static_assert(sizeof(FVisualizerEvent) == 88, "FVisualizerEvent size has changed. Please update the save/load functions");
+#endif
+
+TSharedPtr< FVisualizerEvent > FVisualizerEvent::LoadVisualizerEventRecursively(FArchive *Ar, TSharedPtr< FVisualizerEvent > InParentEvent)
+{
+	double Start;
+	double Duration;
+	double DurationMs;
+	int32 Category;
+	FString EventName;
+	bool IsSelected;
+
+	*Ar << Start;
+	*Ar << Duration;
+	*Ar << DurationMs;
+	*Ar << Category;
+	*Ar << EventName;
+	*Ar << IsSelected;
+
+	TSharedPtr< FVisualizerEvent > VisualizerEvent(new FVisualizerEvent(Start, Duration,DurationMs,Category,EventName));
+	VisualizerEvent->ParentEvent = InParentEvent;
+
+	uint32 NumChildren;
+
+	Ar->SerializeInt(NumChildren, MAX_uint32);
+	VisualizerEvent->Children.Reserve(NumChildren);
+
+	for (uint32 i = 0; i < NumChildren; ++i)
+	{
+		TSharedPtr< FVisualizerEvent > ChildEvent = LoadVisualizerEventRecursively(Ar, VisualizerEvent);
+		VisualizerEvent->Children.Add(ChildEvent);
+	}
+
+	return VisualizerEvent;
+
+}
+
+TSharedPtr< FVisualizerEvent > FVisualizerEvent::LoadVisualizerEvent(FArchive *Ar)
+{
+	//FArchive Ar;
+	// Assumption: InProfileData contains only one (root) element. Otherwise an extra FVisualizerEvent root event is required.
+	TSharedPtr< FVisualizerEvent > DummyRoot;
+	// Recursively create visualizer event data.
+	TSharedPtr< FVisualizerEvent > StatEvents(LoadVisualizerEventRecursively( Ar, DummyRoot));
+	return StatEvents;
+}
+
+
+
+void FVisualizerEvent::SaveVisualizerEventRecursively(FArchive *Ar, TSharedPtr< FVisualizerEvent > VisualizerEvent )
+{
+	//Write all simple members
+	*Ar << VisualizerEvent->Start;
+	*Ar << VisualizerEvent->Duration;
+	*Ar << VisualizerEvent->DurationMs;
+	*Ar << VisualizerEvent->Category;
+	*Ar << VisualizerEvent->EventName;
+	*Ar << VisualizerEvent->IsSelected;
+
+	uint32 NumChildren;
+
+	NumChildren = VisualizerEvent->Children.Num();
+	Ar->SerializeInt(NumChildren, MAX_uint32);
+
+	for(uint32 i = 0; i < NumChildren; ++i)
+	{
+		TSharedPtr< FVisualizerEvent > Child;
+
+		Child = VisualizerEvent->Children[i];
+		SaveVisualizerEventRecursively(Ar, Child);
+	}
+}
+
+
diff --git a/Engine/Source/Developer/TaskGraph/Public/VisualizerEvents.h b/Engine/Source/Developer/TaskGraph/Public/VisualizerEvents.h
--- a/Engine/Source/Developer/TaskGraph/Public/VisualizerEvents.h
+++ b/Engine/Source/Developer/TaskGraph/Public/VisualizerEvents.h
@@ -76,6 +76,13 @@ struct FVisualizerEvent
 	{
 		ColorIndex = GetTypeHash( InEventName );
 	}
+
+	static TSharedPtr< FVisualizerEvent > LoadVisualizerEvent(FArchive *Ar);
+	static void SaveVisualizerEventRecursively(FArchive *Ar, TSharedPtr< FVisualizerEvent > VisualizerEvent);
+
+private:
+	static TSharedPtr< FVisualizerEvent > LoadVisualizerEventRecursively(FArchive *Ar, TSharedPtr< FVisualizerEvent > InParentEvent);
+
 };
 
 /** Array of graph events.*/
diff --git a/Engine/Source/Developer/TaskGraph/TaskGraph.Build.cs b/Engine/Source/Developer/TaskGraph/TaskGraph.Build.cs
--- a/Engine/Source/Developer/TaskGraph/TaskGraph.Build.cs
+++ b/Engine/Source/Developer/TaskGraph/TaskGraph.Build.cs
@@ -12,10 +12,20 @@ public class TaskGraph : ModuleRules
 				"CoreUObject",
 				"Slate",
 				"SlateCore",
-                "EditorStyle",
+				"EditorStyle",
 				"Engine",
-                "InputCore"
+				"InputCore"
 			}
 		);
+
+		if (Target.Type == TargetType.Editor)
+		{
+			PrivateDependencyModuleNames.AddRange(
+				new string[] {
+					"WorkspaceMenuStructure",
+					"DesktopPlatform"
+				}
+			);
+		}
 	}
 }
diff --git a/Engine/Source/Editor/GameProjectGeneration/Private/GameProjectUtils.cpp b/Engine/Source/Editor/GameProjectGeneration/Private/GameProjectUtils.cpp
--- a/Engine/Source/Editor/GameProjectGeneration/Private/GameProjectUtils.cpp
+++ b/Engine/Source/Editor/GameProjectGeneration/Private/GameProjectUtils.cpp
@@ -1248,6 +1248,8 @@ bool GameProjectUtils::GenerateProjectFromScratch(const FProjectInformation& InP
 			Descriptor.Modules.Add(FModuleDescriptor(*StartupModuleNames[Idx]));
 		}
 
+		Descriptor.bIsEnterpriseProject = InProjectInfo.bIsEnterpriseProject;
+
 		// Try to save it
 		FText LocalFailReason;
 		if(!Descriptor.Save(InProjectInfo.ProjectFilename, LocalFailReason))
@@ -1627,6 +1629,8 @@ bool GameProjectUtils::CreateProjectFromTemplate(const FProjectInformation& InPr
 		Project.EngineAssociation.Empty();
 		Project.EpicSampleNameHash = 0;
 
+		Project.bIsEnterpriseProject = InProjectInfo.bIsEnterpriseProject; // Force the enterprise flag to the value that was requested in the ProjectInfo.
+
 		// Fix up module names
 		const FString BaseSourceName = FPaths::GetBaseFilename(InProjectInfo.TemplateFile);
 		const FString BaseNewName = FPaths::GetBaseFilename(InProjectInfo.ProjectFilename);
diff --git a/Engine/Source/Editor/GameProjectGeneration/Private/SNewProjectWizard.cpp b/Engine/Source/Editor/GameProjectGeneration/Private/SNewProjectWizard.cpp
--- a/Engine/Source/Editor/GameProjectGeneration/Private/SNewProjectWizard.cpp
+++ b/Engine/Source/Editor/GameProjectGeneration/Private/SNewProjectWizard.cpp
@@ -7,6 +7,8 @@
 #include "Misc/MessageDialog.h"
 #include "HAL/FileManager.h"
 #include "Misc/App.h"
+#include "Misc/CommandLine.h"
+#include "Misc/Parse.h"
 #include "Widgets/SOverlay.h"
 #include "Layout/WidgetPath.h"
 #include "SlateOptMacros.h"
@@ -24,6 +26,7 @@
 #include "EditorStyleSet.h"
 #include "Editor.h"
 #include "Interfaces/IPluginManager.h"
+#include "Interfaces/IProjectManager.h"
 #include "ProjectDescriptor.h"
 #include "GameProjectGenerationLog.h"
 #include "GameProjectGenerationModule.h"
@@ -1447,6 +1450,19 @@ bool SNewProjectWizard::CreateProject( const FString& ProjectFile )
 	FProjectInformation ProjectInfo(ProjectFile, SelectedTemplate->bGenerateCode, bCopyStarterContent, SelectedTemplate->ProjectFile);
 	ProjectInfo.TargetedHardware = SelectedHardwareClassTarget;
 	ProjectInfo.DefaultGraphicsPerformance = SelectedGraphicsPreset;
+
+	const FProjectDescriptor* CurrentProject = IProjectManager::Get().GetCurrentProject();
+	if (CurrentProject != nullptr)
+	{
+		ProjectInfo.bIsEnterpriseProject = CurrentProject->bIsEnterpriseProject;
+	}
+	else
+	{
+		// Set the default value for the enterprise flag from the command line for now.
+		// This should be temporary until we implement a more generic approach.
+		ProjectInfo.bIsEnterpriseProject = FParse::Param(FCommandLine::Get(), TEXT("enterprise"));
+	}
+
 	if (!GameProjectUtils::CreateProject(ProjectInfo, FailReason, FailLog))
 	{
 		SOutputLogDialog::Open(LOCTEXT("CreateProject", "Create Project"), FailReason, FailLog, FText::GetEmpty());
diff --git a/Engine/Source/Editor/GameProjectGeneration/Public/GameProjectUtils.h b/Engine/Source/Editor/GameProjectGeneration/Public/GameProjectUtils.h
--- a/Engine/Source/Editor/GameProjectGeneration/Public/GameProjectUtils.h
+++ b/Engine/Source/Editor/GameProjectGeneration/Public/GameProjectUtils.h
@@ -19,6 +19,7 @@ struct FProjectInformation
 		, TemplateFile(MoveTemp(InTemplateFile))
 		, bShouldGenerateCode(bInGenerateCode)
 		, bCopyStarterContent(bInCopyStarterContent)
+		, bIsEnterpriseProject(false)
 		, TargetedHardware(EHardwareClass::Desktop)
 		, DefaultGraphicsPerformance(EGraphicsPreset::Maximum)
 	{
@@ -29,6 +30,7 @@ struct FProjectInformation
 
 	bool bShouldGenerateCode;
 	bool bCopyStarterContent;
+	bool bIsEnterpriseProject;
 
 	EHardwareClass::Type TargetedHardware;
 	EGraphicsPreset::Type DefaultGraphicsPerformance;
diff --git a/Engine/Source/Editor/UnrealEd/Private/Factories/EditorFactories.cpp b/Engine/Source/Editor/UnrealEd/Private/Factories/EditorFactories.cpp
--- a/Engine/Source/Editor/UnrealEd/Private/Factories/EditorFactories.cpp
+++ b/Engine/Source/Editor/UnrealEd/Private/Factories/EditorFactories.cpp
@@ -3880,7 +3880,10 @@ UObject* UTextureFactory::FactoryCreateBinary
 
 void UTextureFactory::ApplyAutoImportSettings(UTexture* Texture)
 {
-	FJsonObjectConverter::JsonObjectToUStruct(AutomatedImportSettings.ToSharedRef(), Texture->GetClass(), Texture, 0, CPF_InstancedReference);
+	if ( AutomatedImportSettings.IsValid() )
+	{
+		FJsonObjectConverter::JsonObjectToUStruct(AutomatedImportSettings.ToSharedRef(), Texture->GetClass(), Texture, 0, CPF_InstancedReference);
+	}
 }
 
 bool UTextureFactory::IsImportResolutionValid(int32 Width, int32 Height, bool bAllowNonPowerOfTwo, FFeedbackContext* Warn)
diff --git a/Engine/Source/Programs/AutomationTool/AutomationUtils/DeploymentContext.cs b/Engine/Source/Programs/AutomationTool/AutomationUtils/DeploymentContext.cs
--- a/Engine/Source/Programs/AutomationTool/AutomationUtils/DeploymentContext.cs
+++ b/Engine/Source/Programs/AutomationTool/AutomationUtils/DeploymentContext.cs
@@ -503,7 +503,8 @@ public class DeploymentContext //: ProjectParams
 		}
         else if (InputFile.HasExtension(".uplugin"))
         {
-            if (InputFile.IsUnderDirectory(EngineRoot))
+			DirectoryReference EnterpriseRoot = DirectoryReference.Combine(EngineRoot, "..", "Enterprise"); // Enterprise plugins aren't under the project additional plugin directories, so they shouldn't be remapped
+            if (InputFile.IsUnderDirectory(EngineRoot) || InputFile.IsUnderDirectory(EnterpriseRoot))
 			{
 				OutputFile = new StagedFileReference(InputFile.MakeRelativeTo(LocalRoot));
 			}
diff --git a/Engine/Source/Programs/UnrealBuildTool/ProjectFiles/VisualStudio/VCProject.cs b/Engine/Source/Programs/UnrealBuildTool/ProjectFiles/VisualStudio/VCProject.cs
--- a/Engine/Source/Programs/UnrealBuildTool/ProjectFiles/VisualStudio/VCProject.cs
+++ b/Engine/Source/Programs/UnrealBuildTool/ProjectFiles/VisualStudio/VCProject.cs
@@ -1299,6 +1299,10 @@ namespace UnrealBuildTool
 								DebugOptions += UProjectPath;
 								DebugOptions += " -skipcompile";
 							}
+							else if (TargetRulesObject.Type == TargetType.Editor && ProjectName == ProjectFileGenerator.EnterpriseProjectFileNameBase)
+							{
+								DebugOptions += " -enterprise";
+							}
 							else if (TargetRulesObject.Type == TargetType.Editor && ProjectName != "UE4")
 							{
 								DebugOptions += ProjectName;
diff --git a/Engine/Source/Runtime/Core/Private/Containers/Algo/AlgosTest.cpp b/Engine/Source/Runtime/Core/Private/Containers/Algo/AlgosTest.cpp
--- a/Engine/Source/Runtime/Core/Private/Containers/Algo/AlgosTest.cpp
+++ b/Engine/Source/Runtime/Core/Private/Containers/Algo/AlgosTest.cpp
@@ -9,9 +9,12 @@
 #include "Containers/Algo/Copy.h"
 #include "Containers/Algo/Heapify.h"
 #include "Containers/Algo/HeapSort.h"
+#include "Containers/Algo/IntroSort.h"
 #include "Containers/Algo/IsHeap.h"
 #include "Containers/Algo/IsSorted.h"
+#include "Containers/Algo/Sort.h"
 #include "Containers/Algo/Transform.h"
+#include "Templates/Greater.h"
 
 IMPLEMENT_SIMPLE_AUTOMATION_TEST(FAlgosTest, "System.Core.Misc.Algos", EAutomationTestFlags::ApplicationContextMask | EAutomationTestFlags::SmokeFilter)
 
@@ -301,5 +304,49 @@ bool FAlgosTest::RunTest(const FString& Parameters)
 		check(Algo::IsSorted(TestArray));
 	}
 
+	// intro sort
+	{
+		TArray<int> TestArray = TestData2;
+		Algo::IntroSort(TestArray);
+
+		check(Algo::IsSorted(TestArray));
+	}
+
+	// sort
+	{
+		// regular Sort
+		TArray<int> TestArray = TestData2;
+		Algo::Sort(TestArray);
+
+		check(Algo::IsSorted(TestArray));
+
+		// Sort with predicate
+		TestArray = TestData2;
+
+		TGreater<> Predicate;
+		Algo::Sort(TestArray, Predicate);
+
+		check(Algo::IsSorted(TestArray, Predicate));
+
+		// SortBy
+		TestArray = TestData2;
+
+		auto Projection = [](int Val) -> int
+		{
+			return Val % 1000; // will sort using the last 3 digits only
+		};
+
+		Algo::SortBy(TestArray, Projection);
+
+		check(Algo::IsSortedBy(TestArray, Projection));
+
+		// SortBy with predicate
+		TestArray = TestData2;
+
+		Algo::SortBy(TestArray, Projection, Predicate);
+
+		check(Algo::IsSortedBy(TestArray, Projection, Predicate));
+	}
+
 	return true;
 }
diff --git a/Engine/Source/Runtime/Core/Public/Containers/Algo/IntroSort.h b/Engine/Source/Runtime/Core/Public/Containers/Algo/IntroSort.h
--- a/Engine/Source/Runtime/Core/Public/Containers/Algo/IntroSort.h
+++ b/Engine/Source/Runtime/Core/Public/Containers/Algo/IntroSort.h
@@ -3,6 +3,7 @@
 #pragma once
 
 #include "Algo/Impl/BinaryHeap.h"
+#include "Math/UnrealMathUtility.h"
 #include "Templates/IdentityFunctor.h"
 #include "Templates/Invoke.h"
 #include "Templates/Less.h"
diff --git a/Engine/Source/Runtime/Core/Public/Containers/Algo/LegacySort.h b/Engine/Source/Runtime/Core/Public/Containers/Algo/LegacySort.h
new file mode 100644
--- /dev/null
+++ b/Engine/Source/Runtime/Core/Public/Containers/Algo/LegacySort.h
@@ -0,0 +1,164 @@
+// Copyright 1998-2017 Epic Games, Inc. All Rights Reserved.
+
+#pragma once
+
+#include "Templates/IdentityFunctor.h"
+#include "Templates/Invoke.h"
+#include "Templates/Less.h"
+#include "Templates/UnrealTemplate.h" // For GetData, GetNum
+
+
+namespace AlgoImpl
+{
+	/**
+	 * Sort elements using user defined predicate class. The sort is unstable, meaning that the ordering of equal items is not necessarily preserved.
+	 * This is the internal sorting function used by Sort overrides.
+	 * This used to be Algo::Sort and is now considered legacy.
+	 *
+	 * @param	First	pointer to the first element to sort
+	 * @param	Num		the number of items to sort
+	 * @param Predicate predicate class
+	 */
+	template <typename T, typename ProjectionType, typename PredicateType> 
+	void LegacySortInternal(T* First, int32 Num, ProjectionType Projection, PredicateType Predicate)
+	{
+		struct FStack
+		{
+			T* Min;
+			T* Max;
+		};
+
+		if( Num < 2 )
+		{
+			return;
+		}
+		FStack RecursionStack[32]={{First,First+Num-1}}, Current, Inner;
+		for( FStack* StackTop=RecursionStack; StackTop>=RecursionStack; --StackTop ) //-V625
+		{
+			Current = *StackTop;
+		Loop:
+			PTRINT Count = Current.Max - Current.Min + 1;
+			if( Count <= 8 )
+			{
+				// Use simple bubble-sort.
+				while( Current.Max > Current.Min )
+				{
+					T *Max, *Item;
+					for( Max=Current.Min, Item=Current.Min+1; Item<=Current.Max; Item++ )
+					{
+						if( Predicate( Invoke( Projection, *Max ), Invoke( Projection, *Item ) ) )
+						{
+							Max = Item;
+						}
+					}
+					Swap( *Max, *Current.Max-- );
+				}
+			}
+			else
+			{
+				// Grab middle element so sort doesn't exhibit worst-cast behavior with presorted lists.
+				Swap( Current.Min[Count/2], Current.Min[0] );
+
+				// Divide list into two halves, one with items <=Current.Min, the other with items >Current.Max.
+				Inner.Min = Current.Min;
+				Inner.Max = Current.Max+1;
+				for( ; ; )
+				{
+					while( ++Inner.Min<=Current.Max && !Predicate( Invoke( Projection, *Current.Min ), Invoke( Projection, *Inner.Min ) ) );
+					while( --Inner.Max> Current.Min && !Predicate( Invoke( Projection, *Inner.Max ), Invoke( Projection, *Current.Min ) ) );
+					if( Inner.Min>Inner.Max )
+					{
+						break;
+					}
+					Swap( *Inner.Min, *Inner.Max );
+				}
+				Swap( *Current.Min, *Inner.Max );
+
+				// Save big half and recurse with small half.
+				if( Inner.Max-1-Current.Min >= Current.Max-Inner.Min )
+				{
+					if( Current.Min+1 < Inner.Max )
+					{
+						StackTop->Min = Current.Min;
+						StackTop->Max = Inner.Max - 1;
+						StackTop++;
+					}
+					if( Current.Max>Inner.Min )
+					{
+						Current.Min = Inner.Min;
+						goto Loop;
+					}
+				}
+				else
+				{
+					if( Current.Max>Inner.Min )
+					{
+						StackTop->Min = Inner  .Min;
+						StackTop->Max = Current.Max;
+						StackTop++;
+					}
+					if( Current.Min+1<Inner.Max )
+					{
+						Current.Max = Inner.Max - 1;
+						goto Loop;
+					}
+				}
+			}
+		}
+	}
+}
+
+namespace Algo
+{
+	/**
+	 * Sort a range of elements using its operator<.  The sort is unstable.
+	 * This used to be Algo::Sort and is now considered legacy.
+	 *
+	 * @param  Range  The range to sort.
+	 */
+	template <typename RangeType>
+	FORCEINLINE void LegacySort(RangeType& Range)
+	{
+		AlgoImpl::LegacySortInternal(GetData(Range), GetNum(Range), FIdentityFunctor(), TLess<>());
+	}
+
+	/**
+	 * Sort a range of elements using a user-defined predicate class.  The sort is unstable.
+	 * This used to be Algo::Sort and is now considered legacy.
+	 *
+	 * @param  Range      The range to sort.
+	 * @param  Predicate  A binary predicate object used to specify if one element should precede another.
+	 */
+	template <typename RangeType, typename PredicateType>
+	FORCEINLINE void LegacySort(RangeType& Range, PredicateType Pred)
+	{
+		AlgoImpl::LegacySortInternal(GetData(Range), GetNum(Range), FIdentityFunctor(), MoveTemp(Pred));
+	}
+
+	/**
+	 * Sort a range of elements by a projection using the projection's operator<.  The sort is unstable.
+	 * This used to be Algo::Sort and is now considered legacy.
+	 *
+	 * @param  Range  The range to sort.
+	 * @param  Proj   The projection to sort by when applied to the element.
+	 */
+	template <typename RangeType, typename ProjectionType>
+	FORCEINLINE void LegacySortBy(RangeType& Range, ProjectionType Proj)
+	{
+		AlgoImpl::LegacySortInternal(GetData(Range), GetNum(Range), MoveTemp(Proj), TLess<>());
+	}
+
+	/**
+	 * Sort a range of elements by a projection using a user-defined predicate class.  The sort is unstable.
+	 * This used to be Algo::Sort and is now considered legacy.
+	 *
+	 * @param  Range      The range to sort.
+	 * @param  Proj       The projection to sort by when applied to the element.
+	 * @param  Predicate  A binary predicate object, applied to the projection, used to specify if one element should precede another.
+	 */
+	template <typename RangeType, typename ProjectionType, typename PredicateType>
+	FORCEINLINE void LegacySortBy(RangeType& Range, ProjectionType Proj, PredicateType Pred)
+	{
+		AlgoImpl::LegacySortInternal(GetData(Range), GetNum(Range), MoveTemp(Proj), MoveTemp(Pred));
+	}
+}
diff --git a/Engine/Source/Runtime/Core/Public/Containers/Algo/Sort.h b/Engine/Source/Runtime/Core/Public/Containers/Algo/Sort.h
--- a/Engine/Source/Runtime/Core/Public/Containers/Algo/Sort.h
+++ b/Engine/Source/Runtime/Core/Public/Containers/Algo/Sort.h
@@ -2,111 +2,9 @@
 
 #pragma once
 
-#include "Templates/IdentityFunctor.h"
-#include "Templates/Invoke.h"
-#include "Templates/Less.h"
-#include "Templates/UnrealTemplate.h" // For GetData, GetNum
+#include "Containers/Algo/IntroSort.h"
 
 
-namespace AlgoImpl
-{
-	/**
-	 * Sort elements using user defined predicate class. The sort is unstable, meaning that the ordering of equal items is not necessarily preserved.
-	 * This is the internal sorting function used by Sort overrides.
-	 *
-	 * @param	First	pointer to the first element to sort
-	 * @param	Num		the number of items to sort
-	 * @param Predicate predicate class
-	 */
-	template <typename T, typename ProjectionType, typename PredicateType> 
-	void SortInternal(T* First, int32 Num, ProjectionType Projection, PredicateType Predicate)
-	{
-		struct FStack
-		{
-			T* Min;
-			T* Max;
-		};
-	
-		if( Num < 2 )
-		{
-			return;
-		}
-		FStack RecursionStack[32]={{First,First+Num-1}}, Current, Inner;
-		for( FStack* StackTop=RecursionStack; StackTop>=RecursionStack; --StackTop ) //-V625
-		{
-			Current = *StackTop;
-		Loop:
-			PTRINT Count = Current.Max - Current.Min + 1;
-			if( Count <= 8 )
-			{
-				// Use simple bubble-sort.
-				while( Current.Max > Current.Min )
-				{
-					T *Max, *Item;
-					for( Max=Current.Min, Item=Current.Min+1; Item<=Current.Max; Item++ )
-					{
-						if( Predicate( Invoke( Projection, *Max ), Invoke( Projection, *Item ) ) )
-						{
-							Max = Item;
-						}
-					}
-					Swap( *Max, *Current.Max-- );
-				}
-			}
-			else
-			{
-				// Grab middle element so sort doesn't exhibit worst-cast behavior with presorted lists.
-				Swap( Current.Min[Count/2], Current.Min[0] );
-	
-				// Divide list into two halves, one with items <=Current.Min, the other with items >Current.Max.
-				Inner.Min = Current.Min;
-				Inner.Max = Current.Max+1;
-				for( ; ; )
-				{
-					while( ++Inner.Min<=Current.Max && !Predicate( Invoke( Projection, *Current.Min ), Invoke( Projection, *Inner.Min ) ) );
-					while( --Inner.Max> Current.Min && !Predicate( Invoke( Projection, *Inner.Max ), Invoke( Projection, *Current.Min ) ) );
-					if( Inner.Min>Inner.Max )
-					{
-						break;
-					}
-					Swap( *Inner.Min, *Inner.Max );
-				}
-				Swap( *Current.Min, *Inner.Max );
-	
-				// Save big half and recurse with small half.
-				if( Inner.Max-1-Current.Min >= Current.Max-Inner.Min )
-				{
-					if( Current.Min+1 < Inner.Max )
-					{
-						StackTop->Min = Current.Min;
-						StackTop->Max = Inner.Max - 1;
-						StackTop++;
-					}
-					if( Current.Max>Inner.Min )
-					{
-						Current.Min = Inner.Min;
-						goto Loop;
-					}
-				}
-				else
-				{
-					if( Current.Max>Inner.Min )
-					{
-						StackTop->Min = Inner  .Min;
-						StackTop->Max = Current.Max;
-						StackTop++;
-					}
-					if( Current.Min+1<Inner.Max )
-					{
-						Current.Max = Inner.Max - 1;
-						goto Loop;
-					}
-				}
-			}
-		}
-	}
-}
-
 namespace Algo
 {
 	/**
@@ -117,12 +15,11 @@ namespace Algo
 	template <typename RangeType>
 	FORCEINLINE void Sort(RangeType& Range)
 	{
-		AlgoImpl::SortInternal(GetData(Range), GetNum(Range), FIdentityFunctor(), TLess<>());
+		IntroSort(Range);
 	}
 
 	/**
 	 * Sort a range of elements using a user-defined predicate class.  The sort is unstable.
-	 * This is the internal sorting function used by Sort overrides.
 	 *
 	 * @param  Range      The range to sort.
 	 * @param  Predicate  A binary predicate object used to specify if one element should precede another.
@@ -130,7 +27,7 @@ namespace Algo
 	template <typename RangeType, typename PredicateType>
 	FORCEINLINE void Sort(RangeType& Range, PredicateType Pred)
 	{
-		AlgoImpl::SortInternal(GetData(Range), GetNum(Range), FIdentityFunctor(), MoveTemp(Pred));
+		IntroSort(Range, MoveTemp(Pred));
 	}
 
 	/**
@@ -142,12 +39,11 @@ namespace Algo
 	template <typename RangeType, typename ProjectionType>
 	FORCEINLINE void SortBy(RangeType& Range, ProjectionType Proj)
 	{
-		AlgoImpl::SortInternal(GetData(Range), GetNum(Range), MoveTemp(Proj), TLess<>());
+		IntroSortBy(Range, MoveTemp(Proj));
 	}
 
 	/**
 	 * Sort a range of elements by a projection using a user-defined predicate class.  The sort is unstable.
-	 * This is the internal sorting function used by Sort overrides.
 	 *
 	 * @param  Range      The range to sort.
 	 * @param  Proj       The projection to sort by when applied to the element.
@@ -156,6 +52,6 @@ namespace Algo
 	template <typename RangeType, typename ProjectionType, typename PredicateType>
 	FORCEINLINE void SortBy(RangeType& Range, ProjectionType Proj, PredicateType Pred)
 	{
-		AlgoImpl::SortInternal(GetData(Range), GetNum(Range), MoveTemp(Proj), MoveTemp(Pred));
+		IntroSortBy(Range, MoveTemp(Proj), MoveTemp(Pred));
 	}
 }
diff --git a/Engine/Source/Runtime/Core/Public/Templates/Sorting.h b/Engine/Source/Runtime/Core/Public/Templates/Sorting.h
--- a/Engine/Source/Runtime/Core/Public/Templates/Sorting.h
+++ b/Engine/Source/Runtime/Core/Public/Templates/Sorting.h
@@ -3,9 +3,10 @@
 #pragma once
 
 #include "CoreTypes.h"
+#include "Containers/Algo/BinarySearch.h"
+#include "Containers/Algo/Sort.h"
 #include "HAL/PlatformMath.h"
 #include "Templates/Less.h"
-#include "Algo/BinarySearch.h"
 
 /**
  * Helper class for dereferencing pointer types in Sort function
@@ -39,100 +40,31 @@ struct TDereferenceWrapper<T*, PREDICATE_CLASS>
 };
 
 /**
- * Sort elements using user defined predicate class. The sort is unstable, meaning that the ordering of equal items is not necessarily preserved.
- * This is the internal sorting function used by Sort overrides.
- *
- * @param	First	pointer to the first element to sort
- * @param	Num		the number of items to sort
- * @param Predicate predicate class
+ * Wraps a range into a container like interface to satisfy the GetData and GetNum global functions.
+ * We're not using TArrayView since it calls ::Sort creating a circular dependency.
  */
-template<class T, class PREDICATE_CLASS> 
-void SortInternal( T* First, const int32 Num, const PREDICATE_CLASS& Predicate )
+template <typename T>
+struct TArrayRange
 {
-	struct FStack
+	TArrayRange(T* InPtr, int32 InSize)
+		: Begin(InPtr)
+		, Size(InSize)
 	{
-		T* Min;
-		T* Max;
-	};
-
-	if( Num < 2 )
-	{
-		return;
 	}
-	FStack RecursionStack[32]={{First,First+Num-1}}, Current, Inner;
-	for( FStack* StackTop=RecursionStack; StackTop>=RecursionStack; --StackTop ) //-V625
-	{
-		Current = *StackTop;
-	Loop:
-		PTRINT Count = Current.Max - Current.Min + 1;
-		if( Count <= 8 )
-		{
-			// Use simple bubble-sort.
-			while( Current.Max > Current.Min )
-			{
-				T *Max, *Item;
-				for( Max=Current.Min, Item=Current.Min+1; Item<=Current.Max; Item++ )
-				{
-					if( Predicate( *Max, *Item ) )
-					{
-						Max = Item;
-					}
-				}
-				Exchange( *Max, *Current.Max-- );
-			}
-		}
-		else
-		{
-			// Grab middle element so sort doesn't exhibit worst-cast behavior with presorted lists.
-			Exchange( Current.Min[Count/2], Current.Min[0] );
 
-			// Divide list into two halves, one with items <=Current.Min, the other with items >Current.Max.
-			Inner.Min = Current.Min;
-			Inner.Max = Current.Max+1;
-			for( ; ; )
-			{
-				while( ++Inner.Min<=Current.Max && !Predicate( *Current.Min, *Inner.Min ) );
-				while( --Inner.Max> Current.Min && !Predicate( *Inner.Max, *Current.Min ) );
-				if( Inner.Min>Inner.Max )
-				{
-					break;
-				}
-				Exchange( *Inner.Min, *Inner.Max );
-			}
-			Exchange( *Current.Min, *Inner.Max );
+	T* GetData() const { return Begin; }
+	int32 Num() const { return Size; }
 
-			// Save big half and recurse with small half.
-			if( Inner.Max-1-Current.Min >= Current.Max-Inner.Min )
-			{
-				if( Current.Min+1 < Inner.Max )
-				{
-					StackTop->Min = Current.Min;
-					StackTop->Max = Inner.Max - 1;
-					StackTop++;
-				}
-				if( Current.Max>Inner.Min )
-				{
-					Current.Min = Inner.Min;
-					goto Loop;
-				}
-			}
-			else
-			{
-				if( Current.Max>Inner.Min )
-				{
-					StackTop->Min = Inner  .Min;
-					StackTop->Max = Current.Max;
-					StackTop++;
-				}
-				if( Current.Min+1<Inner.Max )
-				{
-					Current.Max = Inner.Max - 1;
-					goto Loop;
-				}
-			}
-		}
-	}
-}
+private:
+	T* Begin;
+	int32 Size;
+};
+
+template <typename T>
+struct TIsContiguousContainer< TArrayRange<T> >
+{
+	enum { Value = true };
+};
 
 /**
  * Sort elements using user defined predicate class. The sort is unstable, meaning that the ordering of equal items is not necessarily preserved.
@@ -144,7 +76,8 @@ void SortInternal( T* First, const int32 Num, const PREDICATE_CLASS& Predicate )
 template<class T, class PREDICATE_CLASS> 
 void Sort( T* First, const int32 Num, const PREDICATE_CLASS& Predicate )
 {
-	SortInternal( First, Num, TDereferenceWrapper<T, PREDICATE_CLASS>( Predicate ) );
+	TArrayRange<T> ArrayRange( First, Num );
+	Algo::Sort( ArrayRange, TDereferenceWrapper<T, PREDICATE_CLASS>( Predicate ) );
 }
 
 /**
@@ -157,7 +90,8 @@ void Sort( T* First, const int32 Num, const PREDICATE_CLASS& Predicate )
 template<class T, class PREDICATE_CLASS> 
 void Sort( T** First, const int32 Num, const PREDICATE_CLASS& Predicate )
 {
-	SortInternal( First, Num, TDereferenceWrapper<T*, PREDICATE_CLASS>( Predicate ) );
+	TArrayRange<T*> ArrayRange( First, Num );
+	Algo::Sort( ArrayRange, TDereferenceWrapper<T*, PREDICATE_CLASS>( Predicate ) );
 }
 
 /**
@@ -170,7 +104,8 @@ void Sort( T** First, const int32 Num, const PREDICATE_CLASS& Predicate )
 template<class T> 
 void Sort( T* First, const int32 Num )
 {
-	SortInternal( First, Num, TDereferenceWrapper<T, TLess<T> >( TLess<T>() ) );
+	TArrayRange<T> ArrayRange( First, Num );
+	Algo::Sort( ArrayRange, TDereferenceWrapper<T, TLess<T> >( TLess<T>() ) );
 }
 
 /**
@@ -182,7 +117,8 @@ void Sort( T* First, const int32 Num )
 template<class T> 
 void Sort( T** First, const int32 Num )
 {
-	SortInternal( First, Num, TDereferenceWrapper<T*, TLess<T> >( TLess<T>() ) );
+	TArrayRange<T*> ArrayRange( First, Num );
+	Algo::Sort( ArrayRange, TDereferenceWrapper<T*, TLess<T> >( TLess<T>() ) );
 }
 
 /**
