Integrate 3639296 from dev-console:

Disabled LLM in Test

diff --git a/Engine/Source/Runtime/Core/Public/HAL/LowLevelMemTracker.h b/Engine/Source/Runtime/Core/Public/HAL/LowLevelMemTracker.h
--- a/Engine/Source/Runtime/Core/Public/HAL/LowLevelMemTracker.h
+++ b/Engine/Source/Runtime/Core/Public/HAL/LowLevelMemTracker.h
@@ -19,7 +19,7 @@
 	#define LLM_SUPPORTED_PLATFORM (PLATFORM_XBOXONE || PLATFORM_PS4 || PLATFORM_WINDOWS)
 
 	// *** enable/disable LLM here ***
-	#define ENABLE_LOW_LEVEL_MEM_TRACKER (!UE_BUILD_SHIPPING && LLM_SUPPORTED_PLATFORM && WITH_ENGINE && 1)
+	#define ENABLE_LOW_LEVEL_MEM_TRACKER (!UE_BUILD_SHIPPING && !UE_BUILD_TEST && LLM_SUPPORTED_PLATFORM && WITH_ENGINE && 1)
 
 	// using asset tagging requires a significantly higher number of per-thread tags, so make it optional
 	// even if this is on, we still need to run with -llmtagset=assets because of the shear number of stat ids it makes
