Duplicate 3643786:
No need to create an NSString for FApplePlatformMisc::LowLevelOutputDebugString - we can just print it out as a UTF-8 string.

diff --git a/Engine/Source/Runtime/Core/Private/Apple/ApplePlatformMisc.cpp b/Engine/Source/Runtime/Core/Private/Apple/ApplePlatformMisc.cpp
--- a/Engine/Source/Runtime/Core/Private/Apple/ApplePlatformMisc.cpp
+++ b/Engine/Source/Runtime/Core/Private/Apple/ApplePlatformMisc.cpp
@@ -75,7 +75,7 @@ bool FApplePlatformMisc::DeleteStoredValue(const FString& InStoreId, const FStri
 void FApplePlatformMisc::LowLevelOutputDebugString( const TCHAR *Message )
 {
 	//NsLog will out to all iOS output consoles, instead of just the Xcode console.
-	NSLog(@"%@", [NSString stringWithUTF8String:TCHAR_TO_UTF8(Message)]);
+	NSLog(@"%s", TCHAR_TO_UTF8(Message));
 }
 
 const TCHAR* FApplePlatformMisc::GetSystemErrorMessage(TCHAR* OutBuffer, int32 BufferCount, int32 Error)
