Moved Set/Get/DeleteSetStoredValue API from FApplePlatformMisc back to IOSPlatformMisc. On Mac, for now, we need to keep using the old generic implementation based on KeyValueStore.ini file, for compatibility with older versions, crash reporter and the Launcher.

diff --git a/Engine/Source/Runtime/Core/Private/Apple/ApplePlatformMisc.cpp b/Engine/Source/Runtime/Core/Private/Apple/ApplePlatformMisc.cpp
--- a/Engine/Source/Runtime/Core/Private/Apple/ApplePlatformMisc.cpp
+++ b/Engine/Source/Runtime/Core/Private/Apple/ApplePlatformMisc.cpp
@@ -30,48 +30,6 @@ void FApplePlatformMisc::GetEnvironmentVariable(const TCHAR* VariableName, TCHAR
 	}
 }
 
-// Make sure that SetStoredValue and GetStoredValue generate the same key
-static NSString* MakeStoredValueKeyName(const FString& SectionName, const FString& KeyName)
-{
-	return [NSString stringWithFString:(SectionName + "/" + KeyName)];
-}
-
-bool FApplePlatformMisc::SetStoredValue(const FString& InStoreId, const FString& InSectionName, const FString& InKeyName, const FString& InValue)
-{
-	NSUserDefaults* UserSettings = [NSUserDefaults standardUserDefaults];
-
-	// convert input to an NSString
-	NSString* StoredValue = [NSString stringWithFString:InValue];
-
-	// store it
-	[UserSettings setObject:StoredValue forKey:MakeStoredValueKeyName(InSectionName, InKeyName)];
-
-	return true;
-}
-
-bool FApplePlatformMisc::GetStoredValue(const FString& InStoreId, const FString& InSectionName, const FString& InKeyName, FString& OutValue)
-{
-	NSUserDefaults* UserSettings = [NSUserDefaults standardUserDefaults];
-	
-	// get the stored NSString
-	NSString* StoredValue = [UserSettings objectForKey:MakeStoredValueKeyName(InSectionName, InKeyName)];
-
-	// if it was there, convert back to FString
-	if (StoredValue != nil)
-	{
-		OutValue = StoredValue;
-		return true;
-	}
-
-	return false;
-}
-
-bool FApplePlatformMisc::DeleteStoredValue(const FString& InStoreId, const FString& InSectionName, const FString& InKeyName)
-{
-	// No Implementation (currently only used by editor code so not needed on iOS)
-	return false;
-}
-
 void FApplePlatformMisc::LowLevelOutputDebugString( const TCHAR *Message )
 {
 	//NsLog will out to all iOS output consoles, instead of just the Xcode console.
diff --git a/Engine/Source/Runtime/Core/Private/IOS/IOSPlatformMisc.cpp b/Engine/Source/Runtime/Core/Private/IOS/IOSPlatformMisc.cpp
--- a/Engine/Source/Runtime/Core/Private/IOS/IOSPlatformMisc.cpp
+++ b/Engine/Source/Runtime/Core/Private/IOS/IOSPlatformMisc.cpp
@@ -916,6 +916,48 @@ void FIOSPlatformMisc::PlatformPreInit()
     signal(SIGPIPE, SIG_IGN);
 }
 
+// Make sure that SetStoredValue and GetStoredValue generate the same key
+static NSString* MakeStoredValueKeyName(const FString& SectionName, const FString& KeyName)
+{
+	return [NSString stringWithFString:(SectionName + "/" + KeyName)];
+}
+
+bool FIOSPlatformMisc::SetStoredValue(const FString& InStoreId, const FString& InSectionName, const FString& InKeyName, const FString& InValue)
+{
+	NSUserDefaults* UserSettings = [NSUserDefaults standardUserDefaults];
+
+	// convert input to an NSString
+	NSString* StoredValue = [NSString stringWithFString:InValue];
+
+	// store it
+	[UserSettings setObject:StoredValue forKey:MakeStoredValueKeyName(InSectionName, InKeyName)];
+
+	return true;
+}
+
+bool FIOSPlatformMisc::GetStoredValue(const FString& InStoreId, const FString& InSectionName, const FString& InKeyName, FString& OutValue)
+{
+	NSUserDefaults* UserSettings = [NSUserDefaults standardUserDefaults];
+
+	// get the stored NSString
+	NSString* StoredValue = [UserSettings objectForKey:MakeStoredValueKeyName(InSectionName, InKeyName)];
+
+	// if it was there, convert back to FString
+	if (StoredValue != nil)
+	{
+		OutValue = StoredValue;
+		return true;
+	}
+
+	return false;
+}
+
+bool FIOSPlatformMisc::DeleteStoredValue(const FString& InStoreId, const FString& InSectionName, const FString& InKeyName)
+{
+	// No Implementation (currently only used by editor code so not needed on iOS)
+	return false;
+}
+
 void FIOSPlatformMisc::SetGracefulTerminationHandler()
 {
     struct sigaction Action;
diff --git a/Engine/Source/Runtime/Core/Public/Apple/ApplePlatformMisc.h b/Engine/Source/Runtime/Core/Public/Apple/ApplePlatformMisc.h
--- a/Engine/Source/Runtime/Core/Public/Apple/ApplePlatformMisc.h
+++ b/Engine/Source/Runtime/Core/Public/Apple/ApplePlatformMisc.h
@@ -123,9 +123,6 @@ struct CORE_API FApplePlatformMisc : public FGenericPlatformMisc
 	static TArray<uint8> GetSystemFontBytes();
 	static FString GetDefaultLocale();
 	static FString GetDefaultLanguage();
-	static bool SetStoredValue(const FString& InStoreId, const FString& InSectionName, const FString& InKeyName, const FString& InValue);
-	static bool GetStoredValue(const FString& InStoreId, const FString& InSectionName, const FString& InKeyName, FString& OutValue);
-	static bool DeleteStoredValue(const FString& InStoreId, const FString& InSectionName, const FString& InKeyName);
 	static TArray<FString> GetPreferredLanguages();
 	static FString GetLocalCurrencyCode();
 	static FString GetLocalCurrencySymbol();
diff --git a/Engine/Source/Runtime/Core/Public/IOS/IOSPlatformMisc.h b/Engine/Source/Runtime/Core/Public/IOS/IOSPlatformMisc.h
--- a/Engine/Source/Runtime/Core/Public/IOS/IOSPlatformMisc.h
+++ b/Engine/Source/Runtime/Core/Public/IOS/IOSPlatformMisc.h
@@ -26,6 +26,9 @@ struct CORE_API FIOSPlatformMisc : public FApplePlatformMisc
 	static EAppReturnType::Type MessageBoxExt( EAppMsgType::Type MsgType, const TCHAR* Text, const TCHAR* Caption );
 	static void SetMemoryWarningHandler(void (* Handler)(const FGenericMemoryWarningContext& Context));
 	static bool HasPlatformFeature(const TCHAR* FeatureName);
+	static bool SetStoredValue(const FString& InStoreId, const FString& InSectionName, const FString& InKeyName, const FString& InValue);
+	static bool GetStoredValue(const FString& InStoreId, const FString& InSectionName, const FString& InKeyName, FString& OutValue);
+	static bool DeleteStoredValue(const FString& InStoreId, const FString& InSectionName, const FString& InKeyName);
 	static void GetValidTargetPlatforms(class TArray<class FString>& TargetPlatformNames);
 	static bool HasActiveWiFiConnection();
 
