fix for AppleARKit plugin errors
* we ignore this error, but still need to make it so the plugin doesn't init in IOS 9 and 10

diff --git a/Engine/Source/Runtime/Core/Public/IOS/IOSPlatformCompilerPreSetup.h b/Engine/Source/Runtime/Core/Public/IOS/IOSPlatformCompilerPreSetup.h
--- a/Engine/Source/Runtime/Core/Public/IOS/IOSPlatformCompilerPreSetup.h
+++ b/Engine/Source/Runtime/Core/Public/IOS/IOSPlatformCompilerPreSetup.h
@@ -145,6 +145,12 @@
 #pragma clang diagnostic ignored "-Wlogical-op-parentheses"
 #pragma clang diagnostic ignored "-Wconstant-logical-operand"
 
+
+// Apple LLVM 9.0 (Xcode 9.0) introduced ARKit in IOS 11 only,
+#if (__clang_major__ > 8)
+#pragma clang diagnostic ignored "-Wunguarded-availability-new"
+#endif
+
 // Apple LLVM 8.1.0 (Xcode 8.3) introduced -Wundefined-var-template
 #if (__clang_major__ > 8) || (__clang_major__ == 8 && __clang_minor__ >= 1)
 #pragma clang diagnostic ignored "-Wundefined-var-template"
