[Duplicate] CL 3666596 - Fix broken automated tests due to compiler optimizations.
 - Constants are replaced with volatile floats, so the compiler can't optimize out the results at compile time.
 - Removed duplicate FMath tests from PlatformTest.cpp. These tests are performed by FGenericPlatformMath::AutoTest() which is already being called in PlatformTest.cpp anyway.

diff --git a/Engine/Source/Runtime/Core/Private/GenericPlatform/GenericPlatformMath.cpp b/Engine/Source/Runtime/Core/Private/GenericPlatform/GenericPlatformMath.cpp
--- a/Engine/Source/Runtime/Core/Private/GenericPlatform/GenericPlatformMath.cpp
+++ b/Engine/Source/Runtime/Core/Private/GenericPlatform/GenericPlatformMath.cpp
@@ -83,30 +83,47 @@ void FGenericPlatformMath::FmodReportError(float X, float Y)
 }
 
 #if WITH_DEV_AUTOMATION_TESTS
-extern float TheCompilerDoesntKnowThisIsAlwaysZero;
+
+// Various constants used by the test code below. These are declared
+// volatile to prevent compiler optimizations from removing the checks.
+namespace CompilerHiddenConstants
+{
+	volatile float MinusOne = -1.0f;
+	volatile float Zero = 0.0f;
+	volatile float One = 1.0f;
+	volatile float Two = 2.0f;
+	volatile float Twelve = 12.0f;
+	volatile float Sixteen = 16.0f;
+	volatile float MinusOneE37 = -1.0e37f;
+	volatile float FloatMax = MAX_FLT;
+}
 
 
 void FGenericPlatformMath::AutoTest() 
 {
-	check(IsNaN(sqrtf(-1.0f)));
-	check(!IsFinite(sqrtf(-1.0f)));
-	check(!IsFinite(-1.0f/TheCompilerDoesntKnowThisIsAlwaysZero));
-	check(!IsFinite(1.0f/TheCompilerDoesntKnowThisIsAlwaysZero));
-	check(!IsNaN(-1.0f/TheCompilerDoesntKnowThisIsAlwaysZero));
-	check(!IsNaN(1.0f/TheCompilerDoesntKnowThisIsAlwaysZero));
-	check(!IsNaN(MAX_FLT));
-	check(IsFinite(MAX_FLT));
-	check(!IsNaN(0.0f));
-	check(IsFinite(0.0f));
-	check(!IsNaN(1.0f));
-	check(IsFinite(1.0f));
-	check(!IsNaN(-1.e37f));
-	check(IsFinite(-1.e37f));
-	check(FloorLog2(0) == 0);
-	check(FloorLog2(1) == 0);
-	check(FloorLog2(2) == 1);
-	check(FloorLog2(12) == 3);
-	check(FloorLog2(16) == 4);
+	{
+		using namespace CompilerHiddenConstants;
+
+		check(IsNaN(sqrtf(MinusOne)));
+		check(!IsFinite(sqrtf(MinusOne)));
+		check(!IsFinite(-1.0f / Zero));
+		check(!IsFinite(1.0f / Zero));
+		check(!IsNaN(-1.0f / Zero));
+		check(!IsNaN(1.0f / Zero));
+		check(!IsNaN(FloatMax));
+		check(IsFinite(FloatMax));
+		check(!IsNaN(Zero));
+		check(IsFinite(Zero));
+		check(!IsNaN(One));
+		check(IsFinite(One));
+		check(!IsNaN(MinusOneE37));
+		check(IsFinite(MinusOneE37));
+		check(FloorLog2(Zero) == 0);
+		check(FloorLog2(One) == 0);
+		check(FloorLog2(Two) == 1);
+		check(FloorLog2(Twelve) == 3);
+		check(FloorLog2(Sixteen) == 4);
+	}
 
 	{
 		// Shift test
diff --git a/Engine/Source/Runtime/Core/Private/Tests/HAL/PlatformTest.cpp b/Engine/Source/Runtime/Core/Private/Tests/HAL/PlatformTest.cpp
--- a/Engine/Source/Runtime/Core/Private/Tests/HAL/PlatformTest.cpp
+++ b/Engine/Source/Runtime/Core/Private/Tests/HAL/PlatformTest.cpp
@@ -11,8 +11,6 @@
 
 #if WITH_DEV_AUTOMATION_TESTS
 
-float TheCompilerDoesntKnowThisIsAlwaysZero = 0.0f;
-
 struct TestA
 {
 	virtual ~TestA() {}
@@ -54,25 +52,6 @@ bool FPlatformVerificationTest::RunTest (const FString& Parameters)
 #else
 	check(*(uint8*)&Test == 0x12);
 #endif
-	check(FMath::IsNaN(sqrtf(-1.0f)));
-	check(!FMath::IsFinite(sqrtf(-1.0f)));
-	check(!FMath::IsFinite(-1.0f/TheCompilerDoesntKnowThisIsAlwaysZero));
-	check(!FMath::IsFinite(1.0f/TheCompilerDoesntKnowThisIsAlwaysZero));
-	check(!FMath::IsNaN(-1.0f/TheCompilerDoesntKnowThisIsAlwaysZero));
-	check(!FMath::IsNaN(1.0f/TheCompilerDoesntKnowThisIsAlwaysZero));
-	check(!FMath::IsNaN(MAX_FLT));
-	check(FMath::IsFinite(MAX_FLT));
-	check(!FMath::IsNaN(0.0f));
-	check(FMath::IsFinite(0.0f));
-	check(!FMath::IsNaN(1.0f));
-	check(FMath::IsFinite(1.0f));
-	check(!FMath::IsNaN(-1.e37f));
-	check(FMath::IsFinite(-1.e37f));
-	check(FMath::FloorLog2(0) == 0);
-	check(FMath::FloorLog2(1) == 0);
-	check(FMath::FloorLog2(2) == 1);
-	check(FMath::FloorLog2(12) == 3);
-	check(FMath::FloorLog2(16) == 4);
 
 	FGenericPlatformMath::AutoTest();
 
