Hiding misleading declarations that cause Doxygen to source TMap to the wrong header file. Comments only, no code changes.

diff --git a/Engine/Source/Runtime/Core/Private/Internationalization/TextFormatArgumentModifier.h b/Engine/Source/Runtime/Core/Private/Internationalization/TextFormatArgumentModifier.h
--- a/Engine/Source/Runtime/Core/Private/Internationalization/TextFormatArgumentModifier.h
+++ b/Engine/Source/Runtime/Core/Private/Internationalization/TextFormatArgumentModifier.h
@@ -9,7 +9,9 @@
 
 struct FPrivateTextFormatArguments;
 
-template<typename KeyType,typename ValueType,typename SetAllocator ,typename KeyFuncs > class TMap;
+/// @cond DOXYGEN_WARNINGS
+template<typename KeyType, typename ValueType, typename SetAllocator, typename KeyFuncs > class TMap;
+/// @endcond
 
 /**
  * Plural form argument modifier.
diff --git a/Engine/Source/Runtime/Core/Public/Containers/ContainersFwd.h b/Engine/Source/Runtime/Core/Public/Containers/ContainersFwd.h
--- a/Engine/Source/Runtime/Core/Public/Containers/ContainersFwd.h
+++ b/Engine/Source/Runtime/Core/Public/Containers/ContainersFwd.h
@@ -2,6 +2,7 @@
 
 #pragma once
 
+/// @cond DOXYGEN_WARNINGS
 class FDefaultAllocator;
 class FDefaultSetAllocator;
 
@@ -14,3 +15,4 @@ template<typename KeyType, typename ValueType, typename SetAllocator = FDefaultS
 template<typename KeyType, typename ValueType, typename SetAllocator = FDefaultSetAllocator, typename KeyFuncs = TDefaultMapHashableKeyFuncs<KeyType, ValueType, true > > class TMultiMap;
 template <typename T = void > struct TLess;
 template<typename KeyType, typename ValueType, typename ArrayAllocator = FDefaultAllocator, typename SortPredicate = TLess<KeyType> > class TSortedMap;
+/// @endcond
