Remove deprecated aliases. Per Microsoft, the __declspec(deprecated) attribute is not valid here, even though it is ignored, and will cause an error in an upcoming version of the compiler.

diff --git a/Engine/Source/Runtime/Core/Public/Containers/StaticArray.h b/Engine/Source/Runtime/Core/Public/Containers/StaticArray.h
--- a/Engine/Source/Runtime/Core/Public/Containers/StaticArray.h
+++ b/Engine/Source/Runtime/Core/Public/Containers/StaticArray.h
@@ -119,16 +119,6 @@ private:
 	TArrayStorage Storage;
 };
 
-
-template<typename TElement>
-using TStaticArray2 = DEPRECATED(4.18, "Use TStaticArray instead.") TStaticArray<TElement, 2>;
-
-template<typename TElement>
-using TStaticArray3 = DEPRECATED(4.18, "Use TStaticArray instead.") TStaticArray<TElement, 3>;
-
-template<typename TElement>
-using TStaticArray4 = DEPRECATED(4.18, "Use TStaticArray instead.") TStaticArray<TElement, 4>;
-
 /** Creates a static array filled with the specified value. */
 template<typename TElement,uint32 NumElements>
 TStaticArray<TElement,NumElements> MakeUniformStaticArray(typename TCallTraits<TElement>::ParamType InValue)
