Fix PAK file encryption and signing keys not being initialized if the primary game module uses a shared engine PCH (which includes ModuleManager.h before the encryption/signing key macros are set).

diff --git a/Engine/Source/Programs/UnrealBuildTool/Configuration/TargetRules.cs b/Engine/Source/Programs/UnrealBuildTool/Configuration/TargetRules.cs
--- a/Engine/Source/Programs/UnrealBuildTool/Configuration/TargetRules.cs
+++ b/Engine/Source/Programs/UnrealBuildTool/Configuration/TargetRules.cs
@@ -1135,15 +1135,33 @@ namespace UnrealBuildTool
 			EncryptionAndSigning.CryptoSettings CryptoSettings = EncryptionAndSigning.ParseCryptoSettings(DirectoryReference.FromFile(ProjectFile), Platform);
 			if (CryptoSettings.IsAnyEncryptionEnabled())
 			{
-				ProjectDefinitions.Add("UE_ENCRYPTION_KEY=" + String.Join(",", CryptoSettings.EncryptionKey.Key.Select(x => String.Format("0x{0:X2}", x))));
+				ProjectDefinitions.Add(String.Format("IMPLEMENT_ENCRYPTION_KEY_REGISTRATION()=UE_REGISTER_ENCRYPTION_KEY({0})", FormatHexBytes(CryptoSettings.EncryptionKey.Key)));
 			}
+			else
+			{
+				ProjectDefinitions.Add("IMPLEMENT_ENCRYPTION_KEY_REGISTRATION()=");
+			}
+
 			if (CryptoSettings.bEnablePakSigning)
 			{
-				ProjectDefinitions.Add("UE_SIGNING_KEY_EXPONENT=" + String.Join(",", CryptoSettings.SigningKey.PublicKey.Exponent.Select(x => String.Format("0x{0:X2}", x))));
-				ProjectDefinitions.Add("UE_SIGNING_KEY_MODULUS=" + String.Join(",", CryptoSettings.SigningKey.PublicKey.Modulus.Select(x => String.Format("0x{0:X2}", x))));
+				ProjectDefinitions.Add(String.Format("IMPLEMENT_SIGNING_KEY_REGISTRATION()=UE_REGISTER_SIGNING_KEY(UE_LIST_ARGUMENT({0}), UE_LIST_ARGUMENT({1}))", FormatHexBytes(CryptoSettings.SigningKey.PublicKey.Exponent), FormatHexBytes(CryptoSettings.SigningKey.PublicKey.Modulus)));
+			}
+			else
+			{
+				ProjectDefinitions.Add("IMPLEMENT_SIGNING_KEY_REGISTRATION()=");
 			}
 		}
 
+		/// <summary>
+		/// Formats an array of bytes as a sequence of values
+		/// </summary>
+		/// <param name="Data">The data to convert into a string</param>
+		/// <returns>List of hexadecimal bytes</returns>
+		private static string FormatHexBytes(byte[] Data)
+		{
+			return String.Join(",", Data.Select(x => String.Format("0x{0:X2}", x)));
+		}
+
 		/// <summary>
 		/// Override any settings required for the selected target type
 		/// </summary>
diff --git a/Engine/Source/Runtime/Core/Public/Modules/ModuleManager.h b/Engine/Source/Runtime/Core/Public/Modules/ModuleManager.h
--- a/Engine/Source/Runtime/Core/Public/Modules/ModuleManager.h
+++ b/Engine/Source/Runtime/Core/Public/Modules/ModuleManager.h
@@ -759,57 +759,54 @@ class FDefaultGameModuleImpl
 	#define IMPLEMENT_FOREIGN_ENGINE_DIR() 
 #endif
 
+/**
+ * Macro for passing a list argument to a macro
+ */
+#define UE_LIST_ARGUMENT(...) __VA_ARGS__
+
 /**
  * Macro for registering signing keys for a project.
  */
-#ifdef UE_SIGNING_KEY_EXPONENT
-	#define IMPLEMENT_SIGNING_KEY_REGISTRATION() \
-		struct FSigningKeyRegistration \
+#define UE_REGISTER_SIGNING_KEY(ExponentValue, ModulusValue) \
+	struct FSigningKeyRegistration \
+	{ \
+		FSigningKeyRegistration() \
 		{ \
-			FSigningKeyRegistration() \
-			{ \
-				extern void RegisterSigningKeyCallback(void (*)(unsigned char OutExponent[64], unsigned char OutModulus[64])); \
-				RegisterSigningKeyCallback(&Callback); \
-			} \
-			static void Callback(unsigned char OutExponent[64], unsigned char OutModulus[64]) \
+			extern void RegisterSigningKeyCallback(void (*)(unsigned char OutExponent[64], unsigned char OutModulus[64])); \
+			RegisterSigningKeyCallback(&Callback); \
+		} \
+		static void Callback(unsigned char OutExponent[64], unsigned char OutModulus[64]) \
+		{ \
+			const unsigned char Exponent[64] = { ExponentValue }; \
+			const unsigned char Modulus[64] = { ModulusValue }; \
+			for(int ByteIdx = 0; ByteIdx < 64; ByteIdx++) \
 			{ \
-				const unsigned char Exponent[64] = { UE_SIGNING_KEY_EXPONENT }; \
-				const unsigned char Modulus[64] = { UE_SIGNING_KEY_MODULUS }; \
-				for(int ByteIdx = 0; ByteIdx < 64; ByteIdx++) \
-				{ \
-					OutExponent[ByteIdx] = Exponent[ByteIdx]; \
-					OutModulus[ByteIdx] = Modulus[ByteIdx]; \
-				} \
+				OutExponent[ByteIdx] = Exponent[ByteIdx]; \
+				OutModulus[ByteIdx] = Modulus[ByteIdx]; \
 			} \
-		} GSigningKeyRegistration;
-#else
-	#define IMPLEMENT_SIGNING_KEY_REGISTRATION()
-#endif
+		} \
+	} GSigningKeyRegistration;
 
 /**
  * Macro for registering encryption key for a project.
  */
-#ifdef UE_ENCRYPTION_KEY
-	#define IMPLEMENT_ENCRYPTION_KEY_REGISTRATION() \
-		struct FEncryptionKeyRegistration \
+#define UE_REGISTER_ENCRYPTION_KEY(...) \
+	struct FEncryptionKeyRegistration \
+	{ \
+		FEncryptionKeyRegistration() \
 		{ \
-			FEncryptionKeyRegistration() \
-			{ \
-				extern void RegisterEncryptionKeyCallback(void (*)(unsigned char OutKey[32])); \
-				RegisterEncryptionKeyCallback(&Callback); \
-			} \
-			static void Callback(unsigned char OutKey[32]) \
+			extern void RegisterEncryptionKeyCallback(void (*)(unsigned char OutKey[32])); \
+			RegisterEncryptionKeyCallback(&Callback); \
+		} \
+		static void Callback(unsigned char OutKey[32]) \
+		{ \
+			const unsigned char Key[32] = { __VA_ARGS__ }; \
+			for(int ByteIdx = 0; ByteIdx < 32; ByteIdx++) \
 			{ \
-				const unsigned char Key[32] = { UE_ENCRYPTION_KEY }; \
-				for(int ByteIdx = 0; ByteIdx < 32; ByteIdx++) \
-				{ \
-					OutKey[ByteIdx] = Key[ByteIdx]; \
-				} \
+				OutKey[ByteIdx] = Key[ByteIdx]; \
 			} \
-		} GEncryptionKeyRegistration;
-#else
-	#define IMPLEMENT_ENCRYPTION_KEY_REGISTRATION()
-#endif
+		} \
+	} GEncryptionKeyRegistration;
 
 /**
  * Macro for declaring the GIsDebugGame variable for monolithic development builds. NB: This define, and the UE_BUILD_DEVELOPMENT_WITH_DEBUGGAME defines like it, should NEVER be 
