UnrealMathTest Failing on Startup due to change in QuatToRot in UnrealMath.

In changelist 3873301 I changed how Quat's converted to Rotator's to make them more predictable for content creation workflows with the sequencer in particularly with the singularities at Pitch at 90,-90.

The UnrealMathTest directly copies code for that test and was not updated and so would fail with VS2017.  In VS2015 it seems that ClampAxis would convert the 360.0 degrees for the test but due to slight precision issues that don't show up in 2017. Correcting the test is the correct fix here.
Tested fix on both VS2015 and VS2017.

diff --git a/Engine/Source/Runtime/Core/Private/Tests/Math/UnrealMathTest.cpp b/Engine/Source/Runtime/Core/Private/Tests/Math/UnrealMathTest.cpp
--- a/Engine/Source/Runtime/Core/Private/Tests/Math/UnrealMathTest.cpp
+++ b/Engine/Source/Runtime/Core/Private/Tests/Math/UnrealMathTest.cpp
@@ -443,15 +443,15 @@ FORCENOINLINE FRotator TestQuaternionToRotator(const FQuat& Quat)
 	// and we don't want fast approximations to be used here.
 	if (SingularityTest < -SINGULARITY_THRESHOLD)
 	{
-		RotatorFromQuat.Pitch = 270.f;
-		RotatorFromQuat.Yaw = atan2f(YawY, YawX) * RAD_TO_DEG;
-		RotatorFromQuat.Roll = -RotatorFromQuat.Yaw - (2.f * atan2f(X, W) * RAD_TO_DEG);
+		RotatorFromQuat.Pitch = -90.f;
+		RotatorFromQuat.Yaw = -2.0f * FMath::Atan2(X, W) * RAD_TO_DEG;
+		RotatorFromQuat.Roll = 0.0f;
 	}
 	else if (SingularityTest > SINGULARITY_THRESHOLD)
 	{
 		RotatorFromQuat.Pitch = 90.f;
-		RotatorFromQuat.Yaw = atan2f(YawY, YawX) * RAD_TO_DEG;
-		RotatorFromQuat.Roll = RotatorFromQuat.Yaw - (2.f * atan2f(X, W) * RAD_TO_DEG);
+		RotatorFromQuat.Yaw = 2.0f * FMath::Atan2(X, W) * RAD_TO_DEG;
+		RotatorFromQuat.Roll = 0.0f;
 	}
 	else
 	{
