Disabling more benign PVS-Studio warnings.

diff --git a/Engine/Source/Developer/Localization/Private/LocTextHelperTests.cpp b/Engine/Source/Developer/Localization/Private/LocTextHelperTests.cpp
--- a/Engine/Source/Developer/Localization/Private/LocTextHelperTests.cpp
+++ b/Engine/Source/Developer/Localization/Private/LocTextHelperTests.cpp
@@ -179,7 +179,7 @@ bool FLocTextHelperTest::RunTest(const FString& Parameters)
 			return true;
 		}, true);
 
-		if (NumTestEntries != EnumeratedCount)
+		if (NumTestEntries != EnumeratedCount) //-V547
 		{
 			AddError(FString::Printf(TEXT("Failed to enumerate the expected number of translations for '%s'. Expected: %d. Actual: %d."), *InCultureName, NumTestEntries, EnumeratedCount));
 		}
diff --git a/Engine/Source/Editor/GameProjectGeneration/Private/SGameProjectDialog.cpp b/Engine/Source/Editor/GameProjectGeneration/Private/SGameProjectDialog.cpp
--- a/Engine/Source/Editor/GameProjectGeneration/Private/SGameProjectDialog.cpp
+++ b/Engine/Source/Editor/GameProjectGeneration/Private/SGameProjectDialog.cpp
@@ -187,7 +187,7 @@ void SGameProjectDialog::Construct( const FArguments& InArgs )
 	{
 		ShowProjectBrowser();
 	}
-	else if (ActiveTab == NewProjectTab)
+	else if (ActiveTab == NewProjectTab) //-V547
 	{
 		ShowNewProjectTab();
 	}
diff --git a/Engine/Source/Editor/GameProjectGeneration/Private/Tests/GameProjectAutomationTests.cpp b/Engine/Source/Editor/GameProjectGeneration/Private/Tests/GameProjectAutomationTests.cpp
--- a/Engine/Source/Editor/GameProjectGeneration/Private/Tests/GameProjectAutomationTests.cpp
+++ b/Engine/Source/Editor/GameProjectGeneration/Private/Tests/GameProjectAutomationTests.cpp
@@ -184,7 +184,7 @@ namespace GameProjectAutomationUtils
 						OutCreatedProjects++;
 						
 						// Now remove the files we just created (if required)
-						if(bRemoveCreatedProjects == true)						
+						if(bRemoveCreatedProjects == true) //-V547
 						{
 							FString RootFolder = FPaths::GetPath(DesiredProjectFilename);
 							GameProjectUtils::DeleteCreatedFiles(RootFolder, CreatedFiles);
diff --git a/Engine/Source/Editor/SceneOutliner/Private/SSceneOutliner.cpp b/Engine/Source/Editor/SceneOutliner/Private/SSceneOutliner.cpp
--- a/Engine/Source/Editor/SceneOutliner/Private/SSceneOutliner.cpp
+++ b/Engine/Source/Editor/SceneOutliner/Private/SSceneOutliner.cpp
@@ -2294,7 +2294,7 @@ namespace SceneOutliner
 
 		if ( !IsFilterActive() )
 		{
-			if (SelectedActorCount == 0)
+			if (SelectedActorCount == 0) //-V547
 			{
 				return FText::Format( LOCTEXT("ShowingAllActorsFmt", "{0} actors"), FText::AsNumber( TotalActorCount ) );
 			}
@@ -2307,7 +2307,7 @@ namespace SceneOutliner
 		{
 			return FText::Format( LOCTEXT("ShowingNoActorsFmt", "No matching actors ({0} total)"), FText::AsNumber( TotalActorCount ) );
 		}
-		else if (SelectedActorCount != 0)
+		else if (SelectedActorCount != 0) //-V547
 		{
 			return FText::Format( LOCTEXT("ShowingOnlySomeActorsSelectedFmt", "Showing {0} of {1} actors ({2} selected)"), FText::AsNumber( FilteredActorCount ), FText::AsNumber( TotalActorCount ), FText::AsNumber( SelectedActorCount ) );
 		}
diff --git a/Engine/Source/Editor/UnrealEd/Private/ParamParser.cpp b/Engine/Source/Editor/UnrealEd/Private/ParamParser.cpp
--- a/Engine/Source/Editor/UnrealEd/Private/ParamParser.cpp
+++ b/Engine/Source/Editor/UnrealEd/Private/ParamParser.cpp
@@ -191,7 +191,7 @@ bool GetFROTATOR
 		return true;
 	}
 
-	return (N > 0);
+	return true;
 }
 
 
diff --git a/Engine/Source/Runtime/Core/Private/Serialization/Archive.cpp b/Engine/Source/Runtime/Core/Private/Serialization/Archive.cpp
--- a/Engine/Source/Runtime/Core/Private/Serialization/Archive.cpp
+++ b/Engine/Source/Runtime/Core/Private/Serialization/Archive.cpp
@@ -541,7 +541,7 @@ void FArchive::SerializeCompressed( void* V, int64 Length, ECompressionFlags Fla
 		}
 		else
 		{
-			bHeaderWasValid = PackageFileTag.CompressedSize == PACKAGE_FILE_TAG;
+			bHeaderWasValid = PackageFileTag.CompressedSize == PACKAGE_FILE_TAG; //-V547
 		}
 
 		if (!bHeaderWasValid)
@@ -592,7 +592,7 @@ void FArchive::SerializeCompressed( void* V, int64 Length, ECompressionFlags Fla
 			// Read compressed data.
 			Serialize( CompressedBuffer, Chunk.CompressedSize );
 			// Decompress into dest pointer directly.
-			verify( FCompression::UncompressMemory( Flags, Dest, Chunk.UncompressedSize, CompressedBuffer, Chunk.CompressedSize, (Padding > 0) ? true : false, CompressionBitWindow ) );
+			verify( FCompression::UncompressMemory( Flags, Dest, Chunk.UncompressedSize, CompressedBuffer, Chunk.CompressedSize, (Padding > 0) ? true : false, CompressionBitWindow ) ); //-V547
 			// And advance it by read amount.
 			Dest += Chunk.UncompressedSize;
 		}
diff --git a/Engine/Source/Runtime/Core/Private/Windows/WindowsPlatformProcess.cpp b/Engine/Source/Runtime/Core/Private/Windows/WindowsPlatformProcess.cpp
--- a/Engine/Source/Runtime/Core/Private/Windows/WindowsPlatformProcess.cpp
+++ b/Engine/Source/Runtime/Core/Private/Windows/WindowsPlatformProcess.cpp
@@ -826,7 +826,7 @@ const TCHAR* FWindowsPlatformProcess::BaseDir()
 				{
 					if (Result[StringLength - 1] == TEXT('/') || Result[StringLength - 1] == TEXT('\\'))
 					{
-						if(--NumSubDirectories < 0)
+						if(--NumSubDirectories < 0) //-V547
 						{
 							break;
 						}
diff --git a/Engine/Source/Runtime/Core/Public/Windows/WindowsPlatformCodeAnalysis.h b/Engine/Source/Runtime/Core/Public/Windows/WindowsPlatformCodeAnalysis.h
--- a/Engine/Source/Runtime/Core/Public/Windows/WindowsPlatformCodeAnalysis.h
+++ b/Engine/Source/Runtime/Core/Public/Windows/WindowsPlatformCodeAnalysis.h
@@ -95,7 +95,7 @@
 	#pragma warning(disable : 6320) // warning C6320: Exception-filter expression is the constant EXCEPTION_EXECUTE_HANDLER. This might mask exceptions that were not intended to be handled.
 
 	//PVS-Studio settings:
-	//-V::505,542,581,601,668,677,686,688,690,703,704,711,719,720,728,730,735,751,1002
+	//-V::505,542,581,601,668,677,686,688,690,703,704,711,719,720,728,730,735,751,1002,1008
 	//-V:TRYCOMPRESSION:519,547
 	//-V:check(:501,547,560,605
 	//-V:checkf(:510
diff --git a/Engine/Source/Runtime/CoreUObject/Private/UObject/Class.cpp b/Engine/Source/Runtime/CoreUObject/Private/UObject/Class.cpp
--- a/Engine/Source/Runtime/CoreUObject/Private/UObject/Class.cpp
+++ b/Engine/Source/Runtime/CoreUObject/Private/UObject/Class.cpp
@@ -2159,7 +2159,7 @@ const TCHAR* UScriptStruct::ImportText(const TCHAR* InBuffer, void* Value, UObje
 	}
 	else
 	{
-		ErrorText->Logf(TEXT("%sImportText (%s): Missing opening parenthesis: %s"), ErrorCount++ > 0 ? LINE_TERMINATOR : TEXT(""), *StructName, InBuffer);
+		ErrorText->Logf(TEXT("%sImportText (%s): Missing opening parenthesis: %s"), ErrorCount++ > 0 ? LINE_TERMINATOR : TEXT(""), *StructName, InBuffer); //-V547
 		return nullptr;
 	}
 	return Buffer;
diff --git a/Engine/Source/Runtime/CoreUObject/Private/UObject/Property.cpp b/Engine/Source/Runtime/CoreUObject/Private/UObject/Property.cpp
--- a/Engine/Source/Runtime/CoreUObject/Private/UObject/Property.cpp
+++ b/Engine/Source/Runtime/CoreUObject/Private/UObject/Property.cpp
@@ -1175,7 +1175,7 @@ const TCHAR* UProperty::ImportSingleProperty( const TCHAR* Str, void* DestData,
 					}
 				}
 			}
-			else if (ArrayOp == ADO_RemoveIndex)
+			else if (ArrayOp == ADO_RemoveIndex) //-V547
 			{
 				SkipWhitespace(Str);
 				if(*Str++ != '(')
diff --git a/Engine/Source/Runtime/D3D12RHI/Private/D3D12Adapter.h b/Engine/Source/Runtime/D3D12RHI/Private/D3D12Adapter.h
--- a/Engine/Source/Runtime/D3D12RHI/Private/D3D12Adapter.h
+++ b/Engine/Source/Runtime/D3D12RHI/Private/D3D12Adapter.h
@@ -256,7 +256,7 @@ public:
 			ObjectType* NewObject = pfnCreationCore(GetDevice(1 << i));
 
 			// For AFR link up the resources so they can be implicitly destroyed
-			if (i > 0)
+			if (i > 0) //-V547
 			{
 				Previous->SetNextObject(NewObject);
 			}
diff --git a/Engine/Source/Runtime/D3D12RHI/Private/D3D12Device.cpp b/Engine/Source/Runtime/D3D12RHI/Private/D3D12Device.cpp
--- a/Engine/Source/Runtime/D3D12RHI/Private/D3D12Device.cpp
+++ b/Engine/Source/Runtime/D3D12RHI/Private/D3D12Device.cpp
@@ -95,7 +95,7 @@ void FD3D12Device::CreateCommandContexts()
 	{
 		FD3D12SubAllocatedOnlineHeap::SubAllocationDesc SubHeapDesc(&GlobalViewHeap, CurrentGlobalHeapOffset, DescriptorSuballocationPerContext);
 
-		const bool bIsDefaultContext = (i == 0);
+		const bool bIsDefaultContext = (i == 0); //-V547
 		const bool bIsAsyncComputeContext = true;
 		FD3D12CommandContext* NewCmdContext = GetOwningRHI()->CreateCommandContext(this, SubHeapDesc, bIsDefaultContext, bIsAsyncComputeContext);
 		CurrentGlobalHeapOffset += DescriptorSuballocationPerContext;
diff --git a/Engine/Source/Runtime/D3D12RHI/Private/Windows/WindowsD3D12Device.cpp b/Engine/Source/Runtime/D3D12RHI/Private/Windows/WindowsD3D12Device.cpp
--- a/Engine/Source/Runtime/D3D12RHI/Private/Windows/WindowsD3D12Device.cpp
+++ b/Engine/Source/Runtime/D3D12RHI/Private/Windows/WindowsD3D12Device.cpp
@@ -721,15 +721,15 @@ bool FD3D12DynamicRHI::RHIGetAvailableResolutions(FScreenResolutionArray& Resolu
 	int32 MinAllowableRefreshRate = 0;
 	int32 MaxAllowableRefreshRate = 10480;
 
-	if (MaxAllowableResolutionX == 0)
+	if (MaxAllowableResolutionX == 0) //-V547
 	{
 		MaxAllowableResolutionX = 10480;
 	}
-	if (MaxAllowableResolutionY == 0)
+	if (MaxAllowableResolutionY == 0) //-V547
 	{
 		MaxAllowableResolutionY = 10480;
 	}
-	if (MaxAllowableRefreshRate == 0)
+	if (MaxAllowableRefreshRate == 0) //-V547
 	{
 		MaxAllowableRefreshRate = 10480;
 	}
diff --git a/Engine/Source/Runtime/MovieScene/Private/Channels/MovieSceneFloatChannel.cpp b/Engine/Source/Runtime/MovieScene/Private/Channels/MovieSceneFloatChannel.cpp
--- a/Engine/Source/Runtime/MovieScene/Private/Channels/MovieSceneFloatChannel.cpp
+++ b/Engine/Source/Runtime/MovieScene/Private/Channels/MovieSceneFloatChannel.cpp
@@ -513,7 +513,7 @@ FKeyHandle AddKeyToChannel(FMovieSceneFloatChannel* Channel, FFrameNumber InFram
 	int32 ExistingIndex = ChannelInterface.FindKey(InFrameNumber);
 	if (ExistingIndex != INDEX_NONE)
 	{
-		FMovieSceneFloatValue& Value = ChannelInterface.GetValues()[ExistingIndex];
+		FMovieSceneFloatValue& Value = ChannelInterface.GetValues()[ExistingIndex]; //-V758
 		Value.Value = InValue;
 		switch (Interpolation)
 		{
diff --git a/Engine/Source/Runtime/OpenGLDrv/Private/Windows/OpenGLWindows.cpp b/Engine/Source/Runtime/OpenGLDrv/Private/Windows/OpenGLWindows.cpp
--- a/Engine/Source/Runtime/OpenGLDrv/Private/Windows/OpenGLWindows.cpp
+++ b/Engine/Source/Runtime/OpenGLDrv/Private/Windows/OpenGLWindows.cpp
@@ -710,15 +710,15 @@ bool PlatformGetAvailableResolutions(FScreenResolutionArray& Resolutions, bool b
 	int32 MinAllowableRefreshRate = 0;
 	int32 MaxAllowableRefreshRate = 10480;
 
-	if (MaxAllowableResolutionX == 0)
+	if (MaxAllowableResolutionX == 0) //-V547
 	{
 		MaxAllowableResolutionX = 10480;
 	}
-	if (MaxAllowableResolutionY == 0)
+	if (MaxAllowableResolutionY == 0) //-V547
 	{
 		MaxAllowableResolutionY = 10480;
 	}
-	if (MaxAllowableRefreshRate == 0)
+	if (MaxAllowableRefreshRate == 0) //-V547
 	{
 		MaxAllowableRefreshRate = 10480;
 	}
diff --git a/Engine/Source/Runtime/Windows/D3D11RHI/Private/Windows/WindowsD3D11ConstantBuffer.cpp b/Engine/Source/Runtime/Windows/D3D11RHI/Private/Windows/WindowsD3D11ConstantBuffer.cpp
--- a/Engine/Source/Runtime/Windows/D3D11RHI/Private/Windows/WindowsD3D11ConstantBuffer.cpp
+++ b/Engine/Source/Runtime/Windows/D3D11RHI/Private/Windows/WindowsD3D11ConstantBuffer.cpp
@@ -109,7 +109,7 @@ void FD3D11DynamicRHI::InitConstantBuffers()
 	{
 		uint32 Size = GConstantBufferSizes[BufferIndex];
 		uint32 SubBuffers = 1;
-		if(BufferIndex == GLOBAL_CONSTANT_BUFFER_INDEX)
+		if(BufferIndex == GLOBAL_CONSTANT_BUFFER_INDEX) //-V547
 		{
 			SubBuffers = 5;
 		}
diff --git a/Engine/Source/Runtime/Windows/D3D11RHI/Private/Windows/WindowsD3D11Device.cpp b/Engine/Source/Runtime/Windows/D3D11RHI/Private/Windows/WindowsD3D11Device.cpp
--- a/Engine/Source/Runtime/Windows/D3D11RHI/Private/Windows/WindowsD3D11Device.cpp
+++ b/Engine/Source/Runtime/Windows/D3D11RHI/Private/Windows/WindowsD3D11Device.cpp
@@ -1428,15 +1428,15 @@ bool FD3D11DynamicRHI::RHIGetAvailableResolutions(FScreenResolutionArray& Resolu
 	int32 MinAllowableRefreshRate = 0;
 	int32 MaxAllowableRefreshRate = 10480;
 
-	if (MaxAllowableResolutionX == 0)
+	if (MaxAllowableResolutionX == 0) //-V547
 	{
 		MaxAllowableResolutionX = 10480;
 	}
-	if (MaxAllowableResolutionY == 0)
+	if (MaxAllowableResolutionY == 0) //-V547
 	{
 		MaxAllowableResolutionY = 10480;
 	}
-	if (MaxAllowableRefreshRate == 0)
+	if (MaxAllowableRefreshRate == 0) //-V547
 	{
 		MaxAllowableRefreshRate = 10480;
 	}
@@ -1565,7 +1565,7 @@ bool FD3D11DynamicRHI::RHIGetAvailableResolutions(FScreenResolutionArray& Resolu
 		++CurrentOutput;
 
 	// TODO: Cap at 1 for default output
-	} while(CurrentOutput < 1);
+	} while(CurrentOutput < 1); //-V654
 
 	return true;
 }
