Removed spurious dependency on SDL2 from Core

(Edigrating CL 3957900 to Main)

diff --git a/Engine/Source/Runtime/Core/Core.Build.cs b/Engine/Source/Runtime/Core/Core.Build.cs
--- a/Engine/Source/Runtime/Core/Core.Build.cs
+++ b/Engine/Source/Runtime/Core/Core.Build.cs
@@ -107,13 +107,6 @@ public class Core : ModuleRules
 				"elftoolchain"
                 );
 
-			if (Target.Platform == UnrealTargetPlatform.Linux)
-			{
-				AddEngineThirdPartyPrivateStaticDependencies(Target,
-					"SDL2"
-					);
-			}
-
 			// Core uses dlopen()
 			PublicAdditionalLibraries.Add("dl");
         }
