Copying //UE4/Dev-Enterprise to //UE4/Dev-Main (Source: //UE4/Dev-Enterprise @ 3972172)

diff --git a/Engine/Source/Editor/UnrealEd/Private/Fbx/FbxSkeletalMeshImport.cpp b/Engine/Source/Editor/UnrealEd/Private/Fbx/FbxSkeletalMeshImport.cpp
--- a/Engine/Source/Editor/UnrealEd/Private/Fbx/FbxSkeletalMeshImport.cpp
+++ b/Engine/Source/Editor/UnrealEd/Private/Fbx/FbxSkeletalMeshImport.cpp
@@ -3604,7 +3604,7 @@ class FAsyncImportMorphTargetWork : public FNonAbandonableTask
 public:
 	FAsyncImportMorphTargetWork(FSkeletalMeshLODModel* InLODModel, const FReferenceSkeleton& InRefSkeleton, const FSkeletalMeshImportData& InBaseImportData, TArray<FVector>&& InMorphLODPoints,
 		FBXImportOptions* InImportOptions, TArray< FMorphTargetDelta >& InMorphDeltas, TArray<uint32>& InBaseIndexData, TArray< uint32 >& InBaseWedgePointIndices,
-		TMap<uint32,uint32>& InWedgePointToVertexIndexMap, const TMultiMap<int32, int32>& InOverlappingCorners,
+		TMap<uint32,uint32>& InWedgePointToVertexIndexMap, const FOverlappingCorners& InOverlappingCorners,
 		const TSet<uint32> InModifiedPoints, const TMultiMap< int32, int32 >& InWedgeToFaces, const FMeshDataBundle& InMeshDataBundle, const TArray<FVector>& InTangentZ)
 		: LODModel(InLODModel)
 		, RefSkeleton(InRefSkeleton)
@@ -3631,8 +3631,9 @@ public:
 		WasProcessed.Empty(MeshDataBundle.Indices.Num());
 		WasProcessed.AddZeroed(MeshDataBundle.Indices.Num());
 
-		TArray< int32 > OverlappingWedges;
 		TArray< int32 > WedgeFaces;
+		TArray< int32 > OverlappingWedgesDummy;
+		TArray< int32 > OtherOverlappingWedgesDummy;
 
 		// For each ModifiedPoints, reset the tangents for the affected wedges
 		for (int32 WedgeIdx = 0; WedgeIdx < MeshDataBundle.Indices.Num(); ++WedgeIdx)
@@ -3643,9 +3644,7 @@ public:
 			{
 				TangentZ[ WedgeIdx ] = FVector::ZeroVector;
 
-				OverlappingWedges.Reset();
-				OverlappingCorners.MultiFind( WedgeIdx, OverlappingWedges );
-				OverlappingWedges.Add( WedgeIdx );
+				const TArray<int32>& OverlappingWedges = FindIncludingNoOverlapping(OverlappingCorners, WedgeIdx, OverlappingWedgesDummy);
 
 				for ( const int32 OverlappingWedgeIndex : OverlappingWedges )
 				{
@@ -3667,9 +3666,7 @@ public:
 
 							TangentZ[ WedgeIndex ] = FVector::ZeroVector;
 
-							TArray< int32 > OtherOverlappingWedges;
-							OverlappingCorners.MultiFind( WedgeIndex, OtherOverlappingWedges );
-							OtherOverlappingWedges.Add( WedgeIndex );
+							const TArray<int32>& OtherOverlappingWedges = FindIncludingNoOverlapping(OverlappingCorners, WedgeIndex, OtherOverlappingWedgesDummy);
 
 							for ( const int32 OtherDupVert : OtherOverlappingWedges )
 							{
@@ -3789,6 +3786,22 @@ public:
 	}
 
 private:
+
+	const TArray<int32>& FindIncludingNoOverlapping(const FOverlappingCorners& Corners, int32 Key, TArray<int32>& NoOverlapping)
+	{
+		const TArray<int32>& Found = Corners.FindIfOverlapping(Key);
+		if (Found.Num() > 0)
+		{
+			return Found;
+		}
+		else
+		{
+			NoOverlapping.Reset(1);
+			NoOverlapping.Add(Key);
+			return NoOverlapping;
+		}
+	}
+
 	FSkeletalMeshLODModel* LODModel;
 	// @todo not thread safe
 	const FReferenceSkeleton& RefSkeleton;
@@ -3803,7 +3816,7 @@ private:
 	TArray< uint32 >& BaseWedgePointIndices;
 	TMap<uint32,uint32>& WedgePointToVertexIndexMap;
 
-	const TMultiMap<int32, int32>& OverlappingCorners;
+	const FOverlappingCorners& OverlappingCorners;
 	const TSet<uint32> ModifiedPoints;
 	const TMultiMap< int32, int32 >& WedgeToFaces;
 	const FMeshDataBundle& MeshDataBundle;
@@ -3920,7 +3933,7 @@ void UnFbx::FFbxImporter::ImportMorphTargetsInternal( TArray<FbxNode*>& SkelMesh
 		TangentOptions = (ETangentOptions::Type)(TangentOptions | ETangentOptions::UseMikkTSpace);
 	}
 
-	TMultiMap<int32, int32> OverlappingVertices;
+	FOverlappingCorners OverlappingVertices;
 	MeshUtilities.CalculateOverlappingCorners(MeshDataBundle.Vertices, MeshDataBundle.Indices, false, OverlappingVertices );
 
 	TArray<FVector> TangentZ;
diff --git a/Engine/Source/Editor/UnrealEd/Private/FbxExporter.h b/Engine/Source/Editor/UnrealEd/Private/FbxExporter.h
--- a/Engine/Source/Editor/UnrealEd/Private/FbxExporter.h
+++ b/Engine/Source/Editor/UnrealEd/Private/FbxExporter.h
@@ -59,13 +59,7 @@ public:
 	~FFbxExporter();
 	
 	//~ FGCObject
-	virtual void AddReferencedObjects(FReferenceCollector& Collector) override
-	{
-		if (ExportOptions != nullptr)
-		{
-			Collector.AddReferencedObject(ExportOptions);
-		}
-	}
+	virtual void AddReferencedObjects(FReferenceCollector& Collector) override;
 
 
 	/**
@@ -78,6 +72,11 @@ public:
 	*/
 	void FillExportOptions(bool BatchMode, bool bShowOptionDialog, const FString& FullPath, bool& OutOperationCanceled, bool& bOutExportAll);
 
+	/**
+	* Custom set of export options instead of UI dialog. For automation.
+	*/
+	void SetExportOptionsOverride(UFbxExportOption* OverrideOptions);
+
 	/**
 	 * Creates and readies an empty document for export.
 	 */
@@ -244,7 +243,8 @@ private:
 	/** Whether or not to export vertices unwelded */
 	static bool bStaticMeshExportUnWeldedVerts;
 
-	UFbxExportOption *ExportOptions;
+	UFbxExportOption *ExportOptionsUI;
+	UFbxExportOption *ExportOptionsOverride;
 
 	/** Adapter interface which allows ExportAnimTrack to act on both sequencer and matinee data. */
 	class IAnimTrackAdapter
@@ -465,6 +465,9 @@ private:
 	 * @param Label Property label.
 	 */
 	void CreateAnimatableUserProperty(FbxNode* Node, float Value, const char* Name, const char* Label);
+
+	/** Returns currently active FBX export options. Automation or UI diaog based options. */
+	UFbxExportOption* GetExportOptions();
 };
 
 
diff --git a/Engine/Source/Editor/UnrealEd/Private/FileHelpers.cpp b/Engine/Source/Editor/UnrealEd/Private/FileHelpers.cpp
--- a/Engine/Source/Editor/UnrealEd/Private/FileHelpers.cpp
+++ b/Engine/Source/Editor/UnrealEd/Private/FileHelpers.cpp
@@ -23,6 +23,7 @@
 #include "Engine/MapBuildDataRegistry.h"
 #include "Editor/EditorEngine.h"
 #include "ISourceControlModule.h"
+#include "SourceControlOperations.h"
 #include "Editor/UnrealEdEngine.h"
 #include "Settings/EditorLoadingSavingSettings.h"
 #include "Factories/Factory.h"
@@ -4114,4 +4115,9 @@ void UEditorLoadingAndSavingUtils::ExportScene(bool bExportSelectedActorsOnly)
 	FEditorFileUtils::Export(bExportSelectedActorsOnly);
 }
 
+void UEditorLoadingAndSavingUtils::UnloadPackages(const TArray<UPackage*>& PackagesToUnload, bool& bOutAnyPackagesUnloaded, FText& OutErrorMessage)
+{
+	bOutAnyPackagesUnloaded = PackageTools::UnloadPackages(PackagesToUnload, OutErrorMessage);
+}
+
 #undef LOCTEXT_NAMESPACE
diff --git a/Engine/Source/Editor/UnrealEd/Private/Kismet2/KismetDebugUtilities.cpp b/Engine/Source/Editor/UnrealEd/Private/Kismet2/KismetDebugUtilities.cpp
--- a/Engine/Source/Editor/UnrealEd/Private/Kismet2/KismetDebugUtilities.cpp
+++ b/Engine/Source/Editor/UnrealEd/Private/Kismet2/KismetDebugUtilities.cpp
@@ -4,6 +4,7 @@
 #include "Engine/BlueprintGeneratedClass.h"
 #include "GameFramework/Actor.h"
 #include "UObject/PropertyPortFlags.h"
+#include "UObject/TextProperty.h"
 #include "Widgets/DeclarativeSyntaxSupport.h"
 #include "Widgets/SWidget.h"
 #include "Layout/WidgetPath.h"
diff --git a/Engine/Source/Editor/UnrealEd/Private/MaterialGraphSchema.cpp b/Engine/Source/Editor/UnrealEd/Private/MaterialGraphSchema.cpp
--- a/Engine/Source/Editor/UnrealEd/Private/MaterialGraphSchema.cpp
+++ b/Engine/Source/Editor/UnrealEd/Private/MaterialGraphSchema.cpp
@@ -92,6 +92,9 @@ void FMaterialGraphSchemaAction_NewNode::SetFunctionInputType(UMaterialExpressio
 	case MCT_TextureCube:
 		FunctionInput->InputType = FunctionInput_TextureCube;
 		break;
+	case MCT_TextureExternal:
+		FunctionInput->InputType = FunctionInput_TextureExternal;
+		break;
 	case MCT_StaticBool:
 		FunctionInput->InputType = FunctionInput_StaticBool;
 		break;
diff --git a/Engine/Source/Editor/UnrealEd/Private/UnrealEdEngine.cpp b/Engine/Source/Editor/UnrealEd/Private/UnrealEdEngine.cpp
--- a/Engine/Source/Editor/UnrealEd/Private/UnrealEdEngine.cpp
+++ b/Engine/Source/Editor/UnrealEd/Private/UnrealEdEngine.cpp
@@ -15,6 +15,8 @@
 #include "Materials/Material.h"
 #include "Editor/EditorPerProjectUserSettings.h"
 #include "ISourceControlModule.h"
+#include "SourceControlHelpers.h"
+#include "SourceControlOperations.h"
 #include "Settings/EditorExperimentalSettings.h"
 #include "Settings/EditorLoadingSavingSettings.h"
 #include "ThumbnailRendering/ThumbnailManager.h"
diff --git a/Engine/Source/Editor/UnrealEd/Private/UnrealEdPrivatePCH.h b/Engine/Source/Editor/UnrealEd/Private/UnrealEdPrivatePCH.h
--- a/Engine/Source/Editor/UnrealEd/Private/UnrealEdPrivatePCH.h
+++ b/Engine/Source/Editor/UnrealEd/Private/UnrealEdPrivatePCH.h
@@ -618,15 +618,6 @@
 // From EditorStyle:
 #include "EditorStyleSet.h"
 
-// From SourceControl:
-#include "ISourceControlState.h"
-#include "SourceControlHelpers.h"
-#include "ISourceControlRevision.h"
-#include "ISourceControlProvider.h"
-#include "ISourceControlOperation.h"
-#include "SourceControlOperations.h"
-#include "ISourceControlModule.h"
-
 // From BlueprintGraph:
 #include "BlueprintNodeSignature.h"
 #include "K2Node.h"
diff --git a/Engine/Source/Editor/UnrealEd/Public/EditorReimportHandler.h b/Engine/Source/Editor/UnrealEd/Public/EditorReimportHandler.h
--- a/Engine/Source/Editor/UnrealEd/Public/EditorReimportHandler.h
+++ b/Engine/Source/Editor/UnrealEd/Public/EditorReimportHandler.h
@@ -64,9 +64,10 @@ public:
 	 * 2. Skip all asset that have a missing source file path
 	 * 3. Cancel the whole reimport command
 	 *
-	 * * @param	Objs	Objects to try reimporting
+	 * * @param	ToImportObjects		Objects to try reimporting
+	 * * @param	bShowNotification	True to show a notification when complete, false otherwise
 	 */
-	UNREALED_API virtual void ValidateAllSourceFileAndReimport(TArray<UObject*> &ToImportObjects);
+	UNREALED_API virtual void ValidateAllSourceFileAndReimport(TArray<UObject*> &ToImportObjects, bool bShowNotification = true);
 
 	/**
 	* Attempt to reimport multiple objects from its source by giving registered reimport
diff --git a/Engine/Source/Editor/UnrealEd/Public/FbxImporter.h b/Engine/Source/Editor/UnrealEd/Public/FbxImporter.h
--- a/Engine/Source/Editor/UnrealEd/Public/FbxImporter.h
+++ b/Engine/Source/Editor/UnrealEd/Public/FbxImporter.h
@@ -552,7 +552,7 @@ public:
 	 * @param bParseStatistics
 	 * @return bool
 	 */
-	bool OpenFile( FString Filename, bool bParseStatistics, bool bForSceneInfo = false );
+	bool OpenFile(FString Filename);
 	
 	/**
 	 * Import Fbx file.
diff --git a/Engine/Source/Editor/UnrealEd/Public/FileHelpers.h b/Engine/Source/Editor/UnrealEd/Public/FileHelpers.h
--- a/Engine/Source/Editor/UnrealEd/Public/FileHelpers.h
+++ b/Engine/Source/Editor/UnrealEd/Public/FileHelpers.h
@@ -5,6 +5,7 @@
 #include "CoreMinimal.h"
 #include "AssetData.h"
 #include "ISourceControlProvider.h"
+#include "UObject/TextProperty.h"
 #include "FileHelpers.generated.h"
 
 class ULevel;
@@ -153,6 +154,14 @@ public:
 	 */
 	UFUNCTION(BlueprintCallable, Category = "Editor Scripting | Editor Loading and Saving")
 	static UNREALED_API void ExportScene(bool bExportSelectedActorsOnly);
+
+	/**
+	* Unloads a list of packages
+	*
+	* @param PackagesToUnload Array of packages to unload.
+	*/
+	UFUNCTION(BlueprintCallable, Category = "Editor Scripting | Editor Loading and Saving")
+	static UNREALED_API void UnloadPackages(const TArray<UPackage*>& PackagesToUnload, bool& bOutAnyPackagesUnloaded, FText& OutErrorMessage);
 };
 
 
diff --git a/Engine/Source/Editor/UnrealEd/Public/UnrealEdSharedPCH.h b/Engine/Source/Editor/UnrealEd/Public/UnrealEdSharedPCH.h
--- a/Engine/Source/Editor/UnrealEd/Public/UnrealEdSharedPCH.h
+++ b/Engine/Source/Editor/UnrealEd/Public/UnrealEdSharedPCH.h
@@ -595,14 +595,6 @@
 // From EditorStyle:
 #include "EditorStyleSet.h"
 
-// From SourceControl:
-#include "ISourceControlState.h"
-#include "ISourceControlRevision.h"
-#include "SourceControlHelpers.h"
-#include "ISourceControlProvider.h"
-#include "ISourceControlOperation.h"
-#include "SourceControlOperations.h"
-
 // From BlueprintGraph:
 #include "BlueprintNodeSignature.h"
 #include "K2Node.h"
diff --git a/Engine/Source/Programs/AutomationTool/AutomationUtils/DeploymentContext.cs b/Engine/Source/Programs/AutomationTool/AutomationUtils/DeploymentContext.cs
--- a/Engine/Source/Programs/AutomationTool/AutomationUtils/DeploymentContext.cs
+++ b/Engine/Source/Programs/AutomationTool/AutomationUtils/DeploymentContext.cs
@@ -331,7 +331,7 @@ public class DeploymentContext //: ProjectParams
 		StageTargetConfigurations = new List<UnrealTargetConfiguration>(InTargetConfigurations);
 		StageTargets = new List<StageTarget>(InStageTargets);
 		StageExecutables = InStageExecutables;
-        IsCodeBasedProject = ProjectUtils.IsCodeBasedUProjectFile(RawProjectPath);
+        IsCodeBasedProject = ProjectUtils.IsCodeBasedUProjectFile(RawProjectPath, StageTargetConfigurations);
 		ShortProjectName = ProjectUtils.GetShortProjectName(RawProjectPath);
 		Stage = InStage;
 		Archive = InArchive;
diff --git a/Engine/Source/Programs/AutomationTool/AutomationUtils/ProjectParams.cs b/Engine/Source/Programs/AutomationTool/AutomationUtils/ProjectParams.cs
--- a/Engine/Source/Programs/AutomationTool/AutomationUtils/ProjectParams.cs
+++ b/Engine/Source/Programs/AutomationTool/AutomationUtils/ProjectParams.cs
@@ -1921,7 +1921,7 @@ namespace AutomationTool
 			}
 
             List<UnrealTargetPlatform> ClientTargetPlatformTypes = ClientTargetPlatforms.ConvertAll(x => x.Type).Distinct().ToList();
-            var Properties = ProjectUtils.GetProjectProperties(RawProjectPath, ClientTargetPlatformTypes, RunAssetNativization);
+            var Properties = ProjectUtils.GetProjectProperties(RawProjectPath, ClientTargetPlatformTypes, ClientConfigsToBuild, RunAssetNativization);
 
 			bIsCodeBasedProject = Properties.bIsCodeBasedProject;			
 			DetectedTargets = Properties.Targets;
diff --git a/Engine/Source/Programs/AutomationTool/AutomationUtils/ProjectUtils.cs b/Engine/Source/Programs/AutomationTool/AutomationUtils/ProjectUtils.cs
--- a/Engine/Source/Programs/AutomationTool/AutomationUtils/ProjectUtils.cs
+++ b/Engine/Source/Programs/AutomationTool/AutomationUtils/ProjectUtils.cs
@@ -86,7 +86,7 @@ namespace AutomationTool
 		/// </summary>
 		/// <param name="RawProjectPath">Full project path.</param>
 		/// <returns>Properties of the project.</returns>
-		public static ProjectProperties GetProjectProperties(FileReference RawProjectPath, List<UnrealTargetPlatform> ClientTargetPlatforms = null, bool AssetNativizationRequested = false)
+		public static ProjectProperties GetProjectProperties(FileReference RawProjectPath, List<UnrealTargetPlatform> ClientTargetPlatforms = null, List<UnrealTargetConfiguration> ClientTargetConfigurations = null, bool AssetNativizationRequested = false)
 		{
 			string ProjectKey = "UE4";
 			if (RawProjectPath != null)
@@ -96,7 +96,7 @@ namespace AutomationTool
 			ProjectProperties Properties;
 			if (PropertiesCache.TryGetValue(ProjectKey, out Properties) == false)
 			{
-                Properties = DetectProjectProperties(RawProjectPath, ClientTargetPlatforms, AssetNativizationRequested);
+                Properties = DetectProjectProperties(RawProjectPath, ClientTargetPlatforms, ClientTargetConfigurations, AssetNativizationRequested);
 				PropertiesCache.Add(ProjectKey, Properties);
 			}
 			return Properties;
@@ -107,9 +107,9 @@ namespace AutomationTool
 		/// </summary>
 		/// <param name="RawProjectPath">Full project path.</param>
 		/// <returns>True if the project is a UProject file with source code.</returns>
-		public static bool IsCodeBasedUProjectFile(FileReference RawProjectPath)
+		public static bool IsCodeBasedUProjectFile(FileReference RawProjectPath, List<UnrealTargetConfiguration> ClientTargetConfigurations = null)
 		{
-			return GetProjectProperties(RawProjectPath, null).bIsCodeBasedProject;
+			return GetProjectProperties(RawProjectPath, null, ClientTargetConfigurations).bIsCodeBasedProject;
 		}
 
 		/// <summary>
@@ -127,7 +127,7 @@ namespace AutomationTool
 			return ProjectClientBinariesPath;
 		}
 
-		private static bool RequiresTempTarget(FileReference RawProjectPath, List<UnrealTargetPlatform> ClientTargetPlatforms, bool AssetNativizationRequested)
+		private static bool RequiresTempTarget(FileReference RawProjectPath, List<UnrealTargetPlatform> ClientTargetPlatforms, List<UnrealTargetConfiguration> ClientTargetConfigurations, bool AssetNativizationRequested)
 		{
 			// check to see if we already have a Target.cs file
 			if (File.Exists (Path.Combine (Path.GetDirectoryName (RawProjectPath.FullName), "Source", RawProjectPath.GetFileNameWithoutExtension() + ".Target.cs")))
@@ -184,6 +184,20 @@ namespace AutomationTool
 				}
 			}
 
+			List<UnrealTargetConfiguration> TargetConfigurations = ClientTargetConfigurations;
+			if (TargetConfigurations == null || TargetConfigurations.Count < 1)
+			{
+				// No client target configurations, add all in
+				TargetConfigurations = new List<UnrealTargetConfiguration>();
+				foreach (UnrealTargetConfiguration TargetConfigurationType in Enum.GetValues(typeof(UnrealTargetConfiguration)))
+				{
+					if (TargetConfigurationType != UnrealTargetConfiguration.Unknown)
+					{
+						TargetConfigurations.Add(TargetConfigurationType);
+					}
+				}
+			}
+
 			// Change the working directory to be the Engine/Source folder. We are running from Engine/Binaries/DotNET
 			DirectoryReference oldCWD = DirectoryReference.GetCurrentDirectory();
 			try
@@ -212,7 +226,11 @@ namespace AutomationTool
 					// find if there are any plugins enabled or disabled which differ from the default
 					foreach(PluginInfo Plugin in AvailablePlugins)
 					{
-						bool bPluginEnabledForProject = Plugins.IsPluginEnabledForProject(Plugin, Project, TargetPlatformType, TargetType.Game);
+						bool bPluginEnabledForProject = false;
+						foreach (UnrealTargetConfiguration TargetConfigType in TargetConfigurations)
+						{
+							bPluginEnabledForProject |= Plugins.IsPluginEnabledForProject(Plugin, Project, TargetPlatformType, TargetConfigType, TargetRules.TargetType.Game);
+						}
 						if ((bPluginEnabledForProject != Plugin.EnabledByDefault) || (bPluginEnabledForProject && Plugin.Descriptor.bInstalled))
 						{
 							// NOTE: this code was only marking plugins that compiled for the platform to upgrade to code project, however
@@ -296,7 +314,7 @@ namespace AutomationTool
 		/// </summary>
 		/// <param name="RawProjectPath">Full project path.</param>
 		/// <returns>Project properties.</returns>
-        private static ProjectProperties DetectProjectProperties(FileReference RawProjectPath, List<UnrealTargetPlatform> ClientTargetPlatforms, bool AssetNativizationRequested)
+        private static ProjectProperties DetectProjectProperties(FileReference RawProjectPath, List<UnrealTargetPlatform> ClientTargetPlatforms, List<UnrealTargetConfiguration> ClientTargetConfigurations, bool AssetNativizationRequested)
 		{
 			var Properties = new ProjectProperties();
 			Properties.RawProjectPath = RawProjectPath;
@@ -306,7 +324,7 @@ namespace AutomationTool
 			if (RawProjectPath != null)
 			{
                 string TempTargetDir = CommandUtils.CombinePaths(Path.GetDirectoryName(RawProjectPath.FullName), "Intermediate", "Source");
-                if (RequiresTempTarget(RawProjectPath, ClientTargetPlatforms, AssetNativizationRequested))
+                if (RequiresTempTarget(RawProjectPath, ClientTargetPlatforms, ClientTargetConfigurations, AssetNativizationRequested))
 				{
 					GenerateTempTarget(RawProjectPath);
 					Properties.bWasGenerated = true;
diff --git a/Engine/Source/Programs/AutomationTool/Scripts/BuildPluginCommand.Automation.cs b/Engine/Source/Programs/AutomationTool/Scripts/BuildPluginCommand.Automation.cs
--- a/Engine/Source/Programs/AutomationTool/Scripts/BuildPluginCommand.Automation.cs
+++ b/Engine/Source/Programs/AutomationTool/Scripts/BuildPluginCommand.Automation.cs
@@ -174,7 +174,7 @@ class BuildPlugin : BuildCommand
 				bool bBuildDeveloperTools = (TargetType == TargetType.Editor || TargetType == TargetType.Program);
 				bool bBuildEditor = (TargetType == TargetType.Editor);
 				bool bBuildRequiresCookedData = (TargetType != TargetType.Editor && TargetType != TargetType.Program);
-				if (Module.IsCompiledInConfiguration(Platform, TargetType, bBuildDeveloperTools, bBuildEditor, bBuildRequiresCookedData))
+				if (Module.IsCompiledInConfiguration(Platform, Configuration, TargetType, bBuildDeveloperTools, bBuildEditor, bBuildRequiresCookedData))
 				{
 					bCompilePlatform = true;
 				}
diff --git a/Engine/Source/Programs/UnrealBuildTool/Configuration/UEBuildModule.cs b/Engine/Source/Programs/UnrealBuildTool/Configuration/UEBuildModule.cs
--- a/Engine/Source/Programs/UnrealBuildTool/Configuration/UEBuildModule.cs
+++ b/Engine/Source/Programs/UnrealBuildTool/Configuration/UEBuildModule.cs
@@ -461,12 +461,19 @@ namespace UnrealBuildTool
 			SystemIncludePaths.UnionWith(PublicSystemIncludePaths);
 			Definitions.AddRange(PublicDefinitions);
 
-			// Add the import of export declaration for the module
+			// Add the import or export declaration for the module
 			if(Rules.Type == ModuleRules.ModuleType.CPlusPlus)
 			{
 				if(Rules.Target.LinkType == TargetLinkType.Monolithic)
 				{
-					Definitions.Add(ModuleApiDefine + "=");
+					if (Rules.Target.bShouldCompileAsDLL && Rules.Target.bHasExports)
+					{
+						Definitions.Add(ModuleApiDefine + "=DLLEXPORT");
+					}
+					else
+					{
+						Definitions.Add(ModuleApiDefine + "=");
+					}
 				}
 				else if(Binary == null || SourceBinary != Binary)
 				{
diff --git a/Engine/Source/Programs/UnrealBuildTool/Configuration/UEBuildTarget.cs b/Engine/Source/Programs/UnrealBuildTool/Configuration/UEBuildTarget.cs
--- a/Engine/Source/Programs/UnrealBuildTool/Configuration/UEBuildTarget.cs
+++ b/Engine/Source/Programs/UnrealBuildTool/Configuration/UEBuildTarget.cs
@@ -3034,7 +3034,7 @@ namespace UnrealBuildTool
 					{
 						foreach (ModuleDescriptor ModuleDescriptor in Plugin.Descriptor.Modules)
 						{
-							if (ModuleDescriptor.IsCompiledInConfiguration(Platform, TargetType, bAllowDeveloperModules && Rules.bBuildDeveloperTools, Rules.bBuildEditor, Rules.bBuildRequiresCookedData))
+							if (ModuleDescriptor.IsCompiledInConfiguration(Platform, Configuration, TargetType, bAllowDeveloperModules && Rules.bBuildDeveloperTools, Rules.bBuildEditor, Rules.bBuildRequiresCookedData))
 							{
 								FileReference ModuleFileName = RulesAssembly.GetModuleFileName(ModuleDescriptor.Name);
 								if (!ModuleFileName.ContainsAnyNames(ExcludeFoldersArray, Plugin.Directory))
@@ -3367,7 +3367,7 @@ namespace UnrealBuildTool
 				{
 					if (Plugin.EnabledByDefault && !ReferencedNames.Contains(Plugin.Name))
 					{
-						if (Plugin.Descriptor.bCanContainContent || Plugin.Descriptor.Modules.Any(x => x.IsCompiledInConfiguration(Platform, TargetType, Rules.bBuildDeveloperTools, Rules.bBuildEditor, Rules.bBuildRequiresCookedData)))
+						if (Plugin.Descriptor.bCanContainContent || Plugin.Descriptor.Modules.Any(x => x.IsCompiledInConfiguration(Platform, Configuration, TargetType, Rules.bBuildDeveloperTools, Rules.bBuildEditor, Rules.bBuildRequiresCookedData)))
 						{
 							ReferencedNames.Add(Plugin.Name);
 
@@ -3447,7 +3447,7 @@ namespace UnrealBuildTool
 			else
 			{
 				// Check if the plugin is required for this platform
-				if (!Reference.IsEnabledForPlatform(Platform) || !Reference.IsEnabledForTarget(TargetType))
+				if(!Reference.IsEnabledForPlatform(Platform) || !Reference.IsEnabledForTargetConfiguration(Configuration) || !Reference.IsEnabledForTarget(TargetType))
 				{
 					Log.TraceVerbose("Ignoring plugin '{0}' (referenced via {1}) for platform/configuration", Reference.Name, ReferenceChain);
 					return null;
@@ -3502,7 +3502,7 @@ namespace UnrealBuildTool
 				{
 					foreach (ModuleDescriptor ModuleInfo in Info.Descriptor.Modules)
 					{
-						if (ModuleInfo.IsCompiledInConfiguration(Platform, TargetType, Rules.bBuildDeveloperTools, Rules.bBuildEditor, Rules.bBuildRequiresCookedData))
+						if (ModuleInfo.IsCompiledInConfiguration(Platform, Configuration, TargetType, Rules.bBuildDeveloperTools, Rules.bBuildEditor, Rules.bBuildRequiresCookedData))
 						{
 							UEBuildModuleCPP Module = FindOrCreateCppModuleByName(ModuleInfo.Name, PluginReferenceChain);
 							if (!Instance.Modules.Contains(Module))
diff --git a/Engine/Source/Programs/UnrealBuildTool/System/ModuleDescriptor.cs b/Engine/Source/Programs/UnrealBuildTool/System/ModuleDescriptor.cs
--- a/Engine/Source/Programs/UnrealBuildTool/System/ModuleDescriptor.cs
+++ b/Engine/Source/Programs/UnrealBuildTool/System/ModuleDescriptor.cs
@@ -152,6 +152,16 @@ namespace UnrealBuildTool
 		/// </summary>
 		public TargetType[] BlacklistTargets;
 
+		/// <summary>
+		/// List of allowed target configurations
+		/// </summary>
+		public UnrealTargetConfiguration[] WhitelistTargetConfigurations;
+
+		/// <summary>
+		/// List of disallowed target configurations
+		/// </summary>
+		public UnrealTargetConfiguration[] BlacklistTargetConfigurations;
+
 		/// <summary>
 		/// List of additional dependencies for building this module.
 		/// </summary>
@@ -207,6 +217,18 @@ namespace UnrealBuildTool
 				Module.BlacklistTargets = BlacklistTargets;
 			}
 
+			UnrealTargetConfiguration[] WhitelistTargetConfigurations;
+			if (InObject.TryGetEnumArrayField<UnrealTargetConfiguration>("WhitelistTargetConfigurations", out WhitelistTargetConfigurations))
+			{
+				Module.WhitelistTargetConfigurations = WhitelistTargetConfigurations;
+			}
+
+			UnrealTargetConfiguration[] BlacklistTargetConfigurations;
+			if (InObject.TryGetEnumArrayField<UnrealTargetConfiguration>("BlacklistTargetConfigurations", out BlacklistTargetConfigurations))
+			{
+				Module.BlacklistTargetConfigurations = BlacklistTargetConfigurations;
+			}
+
 			string[] AdditionalDependencies;
 			if (InObject.TryGetStringArrayField("AdditionalDependencies", out AdditionalDependencies))
 			{
@@ -262,6 +284,24 @@ namespace UnrealBuildTool
 				}
 				Writer.WriteArrayEnd();
 			}
+			if (WhitelistTargetConfigurations != null && WhitelistTargetConfigurations.Length > 0)
+			{
+				Writer.WriteArrayStart("WhitelistTargetConfigurations");
+				foreach (UnrealTargetConfiguration WhitelistTargetConfiguration in WhitelistTargetConfigurations)
+				{
+					Writer.WriteValue(WhitelistTargetConfiguration.ToString());
+				}
+				Writer.WriteArrayEnd();
+			}
+			if (BlacklistTargetConfigurations != null && BlacklistTargetConfigurations.Length > 0)
+			{
+				Writer.WriteArrayStart("BlacklistTargetConfigurations");
+				foreach (UnrealTargetConfiguration BlacklistTargetConfiguration in BlacklistTargetConfigurations)
+				{
+					Writer.WriteValue(BlacklistTargetConfiguration.ToString());
+				}
+				Writer.WriteArrayEnd();
+			}
 			if (AdditionalDependencies != null && AdditionalDependencies.Length > 0)
 			{
 				Writer.WriteArrayStart("AdditionalDependencies");
@@ -293,15 +333,16 @@ namespace UnrealBuildTool
 			}
 		}
 
-        /// <summary>
-        /// Determines whether the given plugin module is part of the current build.
-        /// </summary>
-        /// <param name="Platform">The platform being compiled for</param>
-        /// <param name="TargetType">The type of the target being compiled</param>
-        /// <param name="bBuildDeveloperTools">Whether the configuration includes developer tools (typically UEBuildConfiguration.bBuildDeveloperTools for UBT callers)</param>
-        /// <param name="bBuildEditor">Whether the configuration includes the editor (typically UEBuildConfiguration.bBuildEditor for UBT callers)</param>
-        /// <param name="bBuildRequiresCookedData">Whether the configuration requires cooked content (typically UEBuildConfiguration.bBuildRequiresCookedData for UBT callers)</param>
-        public bool IsCompiledInConfiguration(UnrealTargetPlatform Platform, TargetType TargetType, bool bBuildDeveloperTools, bool bBuildEditor, bool bBuildRequiresCookedData)
+		/// <summary>
+		/// Determines whether the given plugin module is part of the current build.
+		/// </summary>
+		/// <param name="Platform">The platform being compiled for</param>
+		/// <param name="TargetConfiguration">The target configuration being compiled for</param>
+		/// <param name="TargetType">The type of the target being compiled</param>
+		/// <param name="bBuildDeveloperTools">Whether the configuration includes developer tools (typically UEBuildConfiguration.bBuildDeveloperTools for UBT callers)</param>
+		/// <param name="bBuildEditor">Whether the configuration includes the editor (typically UEBuildConfiguration.bBuildEditor for UBT callers)</param>
+		/// <param name="bBuildRequiresCookedData">Whether the configuration requires cooked content (typically UEBuildConfiguration.bBuildRequiresCookedData for UBT callers)</param>
+		public bool IsCompiledInConfiguration(UnrealTargetPlatform Platform, UnrealTargetConfiguration TargetConfiguration, TargetType TargetType, bool bBuildDeveloperTools, bool bBuildEditor, bool bBuildRequiresCookedData)
 		{
 			// Check the platform is whitelisted
 			if (WhitelistPlatforms != null && WhitelistPlatforms.Length > 0 && !WhitelistPlatforms.Contains(Platform))
@@ -327,6 +368,18 @@ namespace UnrealBuildTool
 				return false;
 			}
 
+			// Check the target configuration is whitelisted
+			if (WhitelistTargetConfigurations != null && WhitelistTargetConfigurations.Length > 0 && !WhitelistTargetConfigurations.Contains(TargetConfiguration))
+			{
+				return false;
+			}
+
+			// Check the target configuration is not blacklisted
+			if (BlacklistTargetConfigurations != null && BlacklistTargetConfigurations.Contains(TargetConfiguration))
+			{
+				return false;
+			}
+
 			// Check the module is compatible with this target.
 			switch (Type)
 			{
diff --git a/Engine/Source/Programs/UnrealBuildTool/System/PluginReferenceDescriptor.cs b/Engine/Source/Programs/UnrealBuildTool/System/PluginReferenceDescriptor.cs
--- a/Engine/Source/Programs/UnrealBuildTool/System/PluginReferenceDescriptor.cs
+++ b/Engine/Source/Programs/UnrealBuildTool/System/PluginReferenceDescriptor.cs
@@ -51,6 +51,16 @@ namespace UnrealBuildTool
 		/// </summary>
 		public UnrealTargetPlatform[] BlacklistPlatforms;
 
+		/// <summary>
+		/// If enabled, list of target configurations for which the plugin should be enabled (or all target configurations if blank).
+		/// </summary>
+		public UnrealTargetConfiguration[] WhitelistTargetConfigurations;
+
+		/// <summary>
+		/// If enabled, list of target configurations for which the plugin should be disabled.
+		/// </summary>
+		public UnrealTargetConfiguration[] BlacklistTargetConfigurations;
+
 		/// <summary>
 		/// If enabled, list of targets for which the plugin should be enabled (or all targets if blank).
 		/// </summary>
@@ -108,7 +118,15 @@ namespace UnrealBuildTool
 			{
 				Writer.WriteEnumArrayField("BlacklistPlatforms", BlacklistPlatforms);
 			}
-			if(WhitelistTargets != null && WhitelistTargets.Length > 0)
+			if (WhitelistTargetConfigurations != null && WhitelistTargetConfigurations.Length > 0)
+			{
+				Writer.WriteEnumArrayField("WhitelistTargetConfigurations", WhitelistTargetConfigurations);
+			}
+			if (BlacklistTargetConfigurations != null && BlacklistTargetConfigurations.Length > 0)
+			{
+				Writer.WriteEnumArrayField("BlacklistTargetConfigurations", BlacklistTargetConfigurations);
+			}
+			if (WhitelistTargets != null && WhitelistTargets.Length > 0)
 			{
 				Writer.WriteEnumArrayField("WhitelistTargets", WhitelistTargets);
 			}
@@ -155,6 +173,8 @@ namespace UnrealBuildTool
 			RawObject.TryGetStringField("MarketplaceURL", out Descriptor.MarketplaceURL);
 			RawObject.TryGetEnumArrayField<UnrealTargetPlatform>("WhitelistPlatforms", out Descriptor.WhitelistPlatforms);
 			RawObject.TryGetEnumArrayField<UnrealTargetPlatform>("BlacklistPlatforms", out Descriptor.BlacklistPlatforms);
+			RawObject.TryGetEnumArrayField<UnrealTargetConfiguration>("WhitelistTargetConfigurations", out Descriptor.WhitelistTargetConfigurations);
+			RawObject.TryGetEnumArrayField<UnrealTargetConfiguration>("BlacklistTargetConfigurations", out Descriptor.BlacklistTargetConfigurations);
 			RawObject.TryGetEnumArrayField<TargetType>("WhitelistTargets", out Descriptor.WhitelistTargets);
 			RawObject.TryGetEnumArrayField<TargetType>("BlacklistTargets", out Descriptor.BlacklistTargets);
 			RawObject.TryGetEnumArrayField<UnrealTargetPlatform>("SupportedTargetPlatforms", out Descriptor.SupportedTargetPlatforms);
@@ -183,6 +203,28 @@ namespace UnrealBuildTool
 			return true;
 		}
 
+		/// <summary>
+		/// Determines if this reference enables the plugin for a given target configuration
+		/// </summary>
+		/// <param name="TargetConfiguration">The target configuration to check</param>
+		/// <returns>True if the plugin should be enabled</returns>
+		public bool IsEnabledForTargetConfiguration(UnrealTargetConfiguration TargetConfiguration)
+		{
+			if (!bEnabled)
+			{
+				return false;
+			}
+			if (WhitelistTargetConfigurations != null && WhitelistTargetConfigurations.Length > 0 && !WhitelistTargetConfigurations.Contains(TargetConfiguration))
+			{
+				return false;
+			}
+			if (BlacklistTargetConfigurations != null && BlacklistTargetConfigurations.Contains(TargetConfiguration))
+			{
+				return false;
+			}
+			return true;
+		}
+
 		/// <summary>
 		/// Determines if this reference enables the plugin for a given target
 		/// </summary>
diff --git a/Engine/Source/Programs/UnrealBuildTool/System/Plugins.cs b/Engine/Source/Programs/UnrealBuildTool/System/Plugins.cs
--- a/Engine/Source/Programs/UnrealBuildTool/System/Plugins.cs
+++ b/Engine/Source/Programs/UnrealBuildTool/System/Plugins.cs
@@ -363,9 +363,10 @@ namespace UnrealBuildTool
 		/// <param name="Project">The project to check</param>
 		/// <param name="Plugin">Information about the plugin</param>
 		/// <param name="Platform">The target platform</param>
+		/// <param name="TargetConfiguration">The target configuration</param>
 		/// <param name="Target"></param>
 		/// <returns>True if the plugin should be enabled for this project</returns>
-		public static bool IsPluginEnabledForProject(PluginInfo Plugin, ProjectDescriptor Project, UnrealTargetPlatform Platform, TargetType Target)
+		public static bool IsPluginEnabledForProject(PluginInfo Plugin, ProjectDescriptor Project, UnrealTargetPlatform Platform, UnrealTargetConfiguration TargetConfiguration, TargetType Target)
 		{
 			bool bEnabled = Plugin.EnabledByDefault;
 			if (Project != null && Project.Plugins != null)
@@ -374,7 +375,7 @@ namespace UnrealBuildTool
 				{
 					if (String.Compare(PluginReference.Name, Plugin.Name, true) == 0 && !PluginReference.bOptional)
 					{
-						bEnabled = PluginReference.IsEnabledForPlatform(Platform) && PluginReference.IsEnabledForTarget(Target);
+						bEnabled = PluginReference.IsEnabledForPlatform(Platform) && PluginReference.IsEnabledForTargetConfiguration(TargetConfiguration) && PluginReference.IsEnabledForTarget(Target);
 					}
 				}
 			}
diff --git a/Engine/Source/Programs/UnrealBuildTool/System/UProjectInfo.cs b/Engine/Source/Programs/UnrealBuildTool/System/UProjectInfo.cs
--- a/Engine/Source/Programs/UnrealBuildTool/System/UProjectInfo.cs
+++ b/Engine/Source/Programs/UnrealBuildTool/System/UProjectInfo.cs
@@ -223,9 +223,10 @@ namespace UnrealBuildTool
 		/// <param name="Project">The project to check</param>
 		/// <param name="PluginName">Name of the plugin to check</param>
 		/// <param name="Platform">The target platform</param>
+		/// <param name="TargetConfiguration">The target configuration</param>
 		/// <param name="Target"></param>
 		/// <returns>True if the plugin should be enabled for this project</returns>
-		public static bool IsPluginEnabledForProject(string PluginName, ProjectDescriptor Project, UnrealTargetPlatform Platform, TargetType Target)
+		public static bool IsPluginEnabledForProject(string PluginName, ProjectDescriptor Project, UnrealTargetPlatform Platform, UnrealTargetConfiguration TargetConfiguration, TargetType Target)
 		{
 			bool bEnabled = false;
 			if (Project != null && Project.Plugins != null)
@@ -234,9 +235,7 @@ namespace UnrealBuildTool
 				{
 					if (String.Compare(PluginReference.Name, PluginName, true) == 0)
 					{
-						// start with whether it's enabled by default
-						bEnabled = PluginReference.bEnabled;
-						bEnabled = PluginReference.IsEnabledForPlatform(Platform) && PluginReference.IsEnabledForTarget(Target);
+						bEnabled = PluginReference.bEnabled && PluginReference.IsEnabledForPlatform(Platform) && PluginReference.IsEnabledForTargetConfiguration(TargetConfiguration) && PluginReference.IsEnabledForTarget(Target);
 						break;
 					}
 				}
diff --git a/Engine/Source/Programs/UnrealHeaderTool/Private/HeaderParser.cpp b/Engine/Source/Programs/UnrealHeaderTool/Private/HeaderParser.cpp
--- a/Engine/Source/Programs/UnrealHeaderTool/Private/HeaderParser.cpp
+++ b/Engine/Source/Programs/UnrealHeaderTool/Private/HeaderParser.cpp
@@ -6392,6 +6392,7 @@ void FHeaderParser::CompileFunctionDeclaration(FClasses& AllClasses)
 	}
 
 	// Peek ahead to look for a CORE_API style DLL import/export token if present
+	FString APIMacroIfPresent;
 	{
 		FToken Token;
 		if (GetToken(Token, true))
@@ -6411,6 +6412,8 @@ void FHeaderParser::CompileFunctionDeclaration(FClasses& AllClasses)
 					}
 					FuncInfo.FunctionFlags |= FUNC_RequiredAPI;
 					FuncInfo.FunctionExportFlags |= FUNCEXPORT_RequiredAPI;
+
+					APIMacroIfPresent = RequiredAPIMacroIfPresent;
 				}
 			}
 
@@ -6421,6 +6424,12 @@ void FHeaderParser::CompileFunctionDeclaration(FClasses& AllClasses)
 		}
 	}
 
+	// Look for static again, in case there was an ENGINE_API token first
+	if (!APIMacroIfPresent.IsEmpty() && MatchIdentifier(TEXT("static")))
+	{
+		FError::Throwf(TEXT("Unexpected API macro '%s'. Did you mean to put '%s' after the static keyword?"), *APIMacroIfPresent, *APIMacroIfPresent);
+	}
+
 	// Look for virtual again, in case there was an ENGINE_API token first
 	if (MatchIdentifier(TEXT("virtual")))
 	{
diff --git a/Engine/Source/Runtime/Core/Private/Internationalization/ICUInternationalization.cpp b/Engine/Source/Runtime/Core/Private/Internationalization/ICUInternationalization.cpp
--- a/Engine/Source/Runtime/Core/Private/Internationalization/ICUInternationalization.cpp
+++ b/Engine/Source/Runtime/Core/Private/Internationalization/ICUInternationalization.cpp
@@ -24,6 +24,8 @@ THIRD_PARTY_INCLUDES_END
 
 DEFINE_LOG_CATEGORY_STATIC(LogICUInternationalization, Log, All);
 
+static_assert(sizeof(UChar) == 2, "UChar (from ICU) is assumed to always be 2-bytes!");
+
 namespace
 {
 	struct FICUOverrides
diff --git a/Engine/Source/Runtime/Core/Private/Internationalization/ICUUtilities.cpp b/Engine/Source/Runtime/Core/Private/Internationalization/ICUUtilities.cpp
--- a/Engine/Source/Runtime/Core/Private/Internationalization/ICUUtilities.cpp
+++ b/Engine/Source/Runtime/Core/Private/Internationalization/ICUUtilities.cpp
@@ -9,7 +9,48 @@ THIRD_PARTY_INCLUDES_END
 
 namespace ICUUtilities
 {
-	FStringConverter::FStringConverter()
+	void FStringConverterImpl_NativeUTF16::ConvertString(const TCHAR* Source, const int32 SourceStartIndex, const int32 SourceLen, icu::UnicodeString& Destination, const bool ShouldNullTerminate)
+	{
+		if (SourceLen > 0)
+		{
+			// Get the internal buffer of the string, we're going to write to it directly
+			UChar* InternalStringBuffer = Destination.getBuffer(SourceLen + 1);
+
+			// Copy into the string buffer and optionally terminate
+			FMemory::Memcpy(InternalStringBuffer, Source + SourceStartIndex, SourceLen * sizeof(TCHAR));
+			if (ShouldNullTerminate)
+			{
+				InternalStringBuffer[SourceLen] = 0;
+			}
+
+			// Size it back down to the correct size and release our lock on the string buffer
+			Destination.releaseBuffer(SourceLen);
+		}
+		else
+		{
+			Destination.remove();
+		}
+	}
+
+	void FStringConverterImpl_NativeUTF16::ConvertString(const icu::UnicodeString& Source, const int32 SourceStartIndex, const int32 SourceLen, FString& Destination)
+	{
+		if (SourceLen > 0)
+		{
+			// Get the internal buffer of the string, we're going to write to it directly
+			TArray<TCHAR>& InternalStringBuffer = Destination.GetCharArray();
+			InternalStringBuffer.SetNumUninitialized(SourceLen + 1);
+
+			// Copy into the string buffer and terminate
+			FMemory::Memcpy(InternalStringBuffer.GetData(), Source.getBuffer() + SourceStartIndex, SourceLen * sizeof(TCHAR));
+			InternalStringBuffer[SourceLen] = 0;
+		}
+		else
+		{
+			Destination.Reset();
+		}
+	}
+
+	FStringConverterImpl_ConvertToUnicodeString::FStringConverterImpl_ConvertToUnicodeString()
 		: ICUConverter(nullptr)
 	{
 		UErrorCode ICUStatus = U_ZERO_ERROR;
@@ -17,17 +58,12 @@ namespace ICUUtilities
 		check(U_SUCCESS(ICUStatus));
 	}
 
-	FStringConverter::~FStringConverter()
+	FStringConverterImpl_ConvertToUnicodeString::~FStringConverterImpl_ConvertToUnicodeString()
 	{
 		ucnv_close(ICUConverter);
 	}
 
-	void FStringConverter::ConvertString(const FString& Source, icu::UnicodeString& Destination, const bool ShouldNullTerminate)
-	{
-		ConvertString(*Source, 0, Source.Len(), Destination, ShouldNullTerminate);
-	}
-
-	void FStringConverter::ConvertString(const TCHAR* Source, const int32 SourceStartIndex, const int32 SourceLen, icu::UnicodeString& Destination, const bool ShouldNullTerminate)
+	void FStringConverterImpl_ConvertToUnicodeString::ConvertString(const TCHAR* Source, const int32 SourceStartIndex, const int32 SourceLen, icu::UnicodeString& Destination, const bool ShouldNullTerminate)
 	{
 		if (SourceLen > 0)
 		{
@@ -60,36 +96,17 @@ namespace ICUUtilities
 		}
 	}
 
-	icu::UnicodeString FStringConverter::ConvertString(const FString& Source, const bool ShouldNullTerminate)
-	{
-		icu::UnicodeString Destination;
-		ConvertString(Source, Destination, ShouldNullTerminate);
-		return Destination;
-	}
-
-	icu::UnicodeString FStringConverter::ConvertString(const TCHAR* Source, const int32 SourceStartIndex, const int32 SourceLen, const bool ShouldNullTerminate)
+	void FStringConverterImpl_ConvertToUnicodeString::ConvertString(const icu::UnicodeString& Source, const int32 SourceStartIndex, const int32 SourceLen, FString& Destination)
 	{
-		icu::UnicodeString Destination;
-		ConvertString(Source, SourceStartIndex, SourceLen, Destination, ShouldNullTerminate);
-		return Destination;
-	}
-
-	void FStringConverter::ConvertString(const icu::UnicodeString& Source, FString& Destination)
-	{
-		return ConvertString(Source, 0, Source.length(), Destination);
-	}
-
-	void FStringConverter::ConvertString(const icu::UnicodeString& Source, const int32 SourceStartIndex, const int32 SourceLen, FString& Destination)
-	{
-		if (Source.length() > 0)
+		if (SourceLen > 0)
 		{
 			UErrorCode ICUStatus = U_ZERO_ERROR;
 
 			ucnv_reset(ICUConverter);
-			
+
 			// Get the internal buffer of the string, we're going to use it as scratch space
 			TArray<TCHAR>& InternalStringBuffer = Destination.GetCharArray();
-				
+
 			// Work out the maximum size required and resize the buffer so it can hold enough data
 			const int32_t DestinationCapacityBytes = UCNV_GET_MAX_BYTES_FOR_STRING(SourceLen, ucnv_getMaxCharSize(ICUConverter));
 			const int32 DestinationCapacityTCHARs = DestinationCapacityBytes / sizeof(TCHAR);
@@ -105,10 +122,44 @@ namespace ICUUtilities
 		}
 		else
 		{
-			Destination.Empty();
+			Destination.Reset();
 		}
 	}
 
+	void FStringConverter::ConvertString(const FString& Source, icu::UnicodeString& Destination, const bool ShouldNullTerminate)
+	{
+		ConvertString(*Source, 0, Source.Len(), Destination, ShouldNullTerminate);
+	}
+
+	void FStringConverter::ConvertString(const TCHAR* Source, const int32 SourceStartIndex, const int32 SourceLen, icu::UnicodeString& Destination, const bool ShouldNullTerminate)
+	{
+		Impl.ConvertString(Source, SourceStartIndex, SourceLen, Destination, ShouldNullTerminate);
+	}
+
+	icu::UnicodeString FStringConverter::ConvertString(const FString& Source, const bool ShouldNullTerminate)
+	{
+		icu::UnicodeString Destination;
+		ConvertString(Source, Destination, ShouldNullTerminate);
+		return Destination;
+	}
+
+	icu::UnicodeString FStringConverter::ConvertString(const TCHAR* Source, const int32 SourceStartIndex, const int32 SourceLen, const bool ShouldNullTerminate)
+	{
+		icu::UnicodeString Destination;
+		ConvertString(Source, SourceStartIndex, SourceLen, Destination, ShouldNullTerminate);
+		return Destination;
+	}
+
+	void FStringConverter::ConvertString(const icu::UnicodeString& Source, FString& Destination)
+	{
+		return ConvertString(Source, 0, Source.length(), Destination);
+	}
+
+	void FStringConverter::ConvertString(const icu::UnicodeString& Source, const int32 SourceStartIndex, const int32 SourceLen, FString& Destination)
+	{
+		Impl.ConvertString(Source, SourceStartIndex, SourceLen, Destination);
+	}
+
 	FString FStringConverter::ConvertString(const icu::UnicodeString& Source)
 	{
 		FString Destination;
diff --git a/Engine/Source/Runtime/Core/Private/Internationalization/ICUUtilities.h b/Engine/Source/Runtime/Core/Private/Internationalization/ICUUtilities.h
--- a/Engine/Source/Runtime/Core/Private/Internationalization/ICUUtilities.h
+++ b/Engine/Source/Runtime/Core/Private/Internationalization/ICUUtilities.h
@@ -13,6 +13,39 @@ THIRD_PARTY_INCLUDES_END
 
 namespace ICUUtilities
 {
+	/**
+	 * Implementation of the string converter that can copy FString to icu::UnicodeString directly since the native string format for this platform is already UTF-16 (as used by ICU)
+	 * Note: Do not use this type directly! Use the FStringConverterImpl typedef, which will be set correctly for your platform
+	 */
+	class FStringConverterImpl_NativeUTF16
+	{
+	public:
+		void ConvertString(const TCHAR* Source, const int32 SourceStartIndex, const int32 SourceLen, icu::UnicodeString& Destination, const bool ShouldNullTerminate);
+		void ConvertString(const icu::UnicodeString& Source, const int32 SourceStartIndex, const int32 SourceLen, FString& Destination);
+	};
+
+	/**
+	 * Implementation of the string converter that can copy FString to icu::UnicodeString via an ICU converter since the native string format for this platform is not UTF-16 (as used by ICU)
+	 * Note: Do not use this type directly! Use the FStringConverterImpl typedef, which will be set correctly for your platform
+	 */
+	class FStringConverterImpl_ConvertToUnicodeString
+	{
+	public:
+		FStringConverterImpl_ConvertToUnicodeString();
+		~FStringConverterImpl_ConvertToUnicodeString();
+
+		void ConvertString(const TCHAR* Source, const int32 SourceStartIndex, const int32 SourceLen, icu::UnicodeString& Destination, const bool ShouldNullTerminate);
+		void ConvertString(const icu::UnicodeString& Source, const int32 SourceStartIndex, const int32 SourceLen, FString& Destination);
+
+	private:
+		UConverter* ICUConverter;
+	};
+
+	/** Work out the best string converter to use based upon our native platform string traits */
+	template <bool IsUnicode, size_t TCHARSize> struct FStringConverterImpl_PlatformSpecific { typedef FStringConverterImpl_ConvertToUnicodeString Type; };
+	template <> struct FStringConverterImpl_PlatformSpecific<true, 2> { typedef FStringConverterImpl_NativeUTF16 Type; }; // A unicode encoding with a wchar_t size of 2 bytes is assumed to be UTF-16
+	typedef FStringConverterImpl_PlatformSpecific<FPlatformString::IsUnicodeEncoded, sizeof(TCHAR)>::Type FStringConverterImpl;
+
 	/** 
 	 * An object that can convert between FString and icu::UnicodeString
 	 * Note: This object is not thread-safe.
@@ -20,8 +53,9 @@ namespace ICUUtilities
 	class FStringConverter
 	{
 	public:
-		FStringConverter();
-		~FStringConverter();
+		FStringConverter() = default;
+		FStringConverter(const FStringConverter&) = delete;
+		FStringConverter& operator=(const FStringConverter&) = delete;
 
 		/** Convert FString -> icu::UnicodeString */
 		void ConvertString(const FString& Source, icu::UnicodeString& Destination, const bool ShouldNullTerminate = true);
@@ -36,11 +70,7 @@ namespace ICUUtilities
 		FString ConvertString(const icu::UnicodeString& Source, const int32 SourceStartIndex, const int32 SourceLen);
 
 	private:
-		/** Non-copyable */
-		FStringConverter(const FStringConverter&);
-		FStringConverter& operator=(const FStringConverter&);
-
-		UConverter* ICUConverter;
+		FStringConverterImpl Impl;
 	};
 
 	/** Convert FString -> icu::UnicodeString */
diff --git a/Engine/Source/Runtime/Core/Private/Misc/FileHelper.cpp b/Engine/Source/Runtime/Core/Private/Misc/FileHelper.cpp
--- a/Engine/Source/Runtime/Core/Private/Misc/FileHelper.cpp
+++ b/Engine/Source/Runtime/Core/Private/Misc/FileHelper.cpp
@@ -26,7 +26,7 @@ static const FString InvalidFilenames[] = {
 -----------------------------------------------------------------------------*/
 
 /**
- * Load a binary file to a dynamic array.
+ * Load a binary file to a dynamic array with two uninitialized bytes at end as padding.
  */
 bool FFileHelper::LoadFileToArray( TArray<uint8>& Result, const TCHAR* Filename, uint32 Flags )
 {
@@ -42,7 +42,8 @@ bool FFileHelper::LoadFileToArray( TArray<uint8>& Result, const TCHAR* Filename,
 		return false;
 	}
 	int64 TotalSize = Reader->TotalSize();
-	Result.Reset( TotalSize );
+	// Allocate slightly larger than file size to avoid re-allocation when caller null terminates file buffer
+	Result.Reset( TotalSize + 2 );
 	Result.AddUninitialized( TotalSize );
 	Reader->Serialize(Result.GetData(), Result.Num());
 	bool Success = Reader->Close();
@@ -514,7 +515,7 @@ bool FFileHelper::LoadANSITextFileToStrings(const TCHAR* InFilename, IFileManage
 		int32 Size = TextFile->TotalSize();
 		// read the file
 		TArray<uint8> Buffer;
-		Buffer.Empty(Size);
+		Buffer.Empty(Size + 1);
 		Buffer.AddUninitialized(Size);
 		TextFile->Serialize(Buffer.GetData(), Size);
 		// zero terminate it
diff --git a/Engine/Source/Runtime/Core/Private/Serialization/Archive.cpp b/Engine/Source/Runtime/Core/Private/Serialization/Archive.cpp
--- a/Engine/Source/Runtime/Core/Private/Serialization/Archive.cpp
+++ b/Engine/Source/Runtime/Core/Private/Serialization/Archive.cpp
@@ -24,6 +24,7 @@
 #include "Misc/NetworkVersion.h"
 #include "Interfaces/ITargetPlatform.h"
 #include "Serialization/CompressedChunkInfo.h"
+#include "Serialization/ArchiveSerializedPropertyChain.h"
 
 /*-----------------------------------------------------------------------------
 	FArchive implementation.
@@ -36,6 +37,8 @@ FArchive::FArchive()
 #endif
 	CustomVersionContainer = nullptr;
 
+	SerializedPropertyChain = nullptr;
+
 #if USE_STABLE_LOCALIZATION_KEYS
 	LocalizationNamespacePtr = nullptr;
 #endif // USE_STABLE_LOCALIZATION_KEYS
@@ -54,6 +57,9 @@ FArchive::FArchive(const FArchive& ArchiveToCopy)
 
 	CopyTrivialFArchiveStatusMembers(ArchiveToCopy);
 
+	SerializedPropertyChain = nullptr;
+	SetSerializedPropertyChain(ArchiveToCopy.SerializedPropertyChain, ArchiveToCopy.SerializedProperty);
+
 	// Don't know why this is set to false, but this is what the original copying code did
 	ArIsFilterEditorOnly  = false;
 
@@ -76,6 +82,8 @@ FArchive& FArchive::operator=(const FArchive& ArchiveToCopy)
 #endif
 	CopyTrivialFArchiveStatusMembers(ArchiveToCopy);
 
+	SetSerializedPropertyChain(ArchiveToCopy.SerializedPropertyChain, ArchiveToCopy.SerializedProperty);
+
 	// Don't know why this is set to false, but this is what the original copying code did
 	ArIsFilterEditorOnly  = false;
 
@@ -104,6 +112,8 @@ FArchive::~FArchive()
 {
 	delete CustomVersionContainer;
 
+	delete SerializedPropertyChain;
+
 #if USE_STABLE_LOCALIZATION_KEYS
 	delete LocalizationNamespacePtr;
 #endif // USE_STABLE_LOCALIZATION_KEYS
@@ -151,17 +161,20 @@ void FArchive::Reset()
 	ArIsNetArchive						= false;
 	ArCustomPropertyList				= nullptr;
 	ArUseCustomPropertyList				= false;
-	CookingTargetPlatform = nullptr;
-	SerializedProperty = nullptr;
-#if WITH_EDITORONLY_DATA
-	EditorOnlyPropertyStack = 0;
-#endif
+	CookingTargetPlatform				= nullptr;
+	SerializedProperty					= nullptr;
+
+	delete SerializedPropertyChain;
+	SerializedPropertyChain = nullptr;
+
 #if USE_STABLE_LOCALIZATION_KEYS
 	SetBaseLocalizationNamespace(FString());
 #endif // USE_STABLE_LOCALIZATION_KEYS
+
 #if WITH_EDITOR
 	ArDebugSerializationFlags			= 0;
 #endif
+
 	// Reset all custom versions to the current registered versions.
 	ResetCustomVersions();
 }
@@ -205,10 +218,7 @@ void FArchive::CopyTrivialFArchiveStatusMembers(const FArchive& ArchiveToCopy)
 	ArCustomPropertyList				 = ArchiveToCopy.ArCustomPropertyList;
 	ArUseCustomPropertyList				 = ArchiveToCopy.ArUseCustomPropertyList;
 	CookingTargetPlatform                = ArchiveToCopy.CookingTargetPlatform;
-	SerializedProperty = ArchiveToCopy.SerializedProperty;
-#if WITH_EDITORONLY_DATA
-	EditorOnlyPropertyStack = ArchiveToCopy.EditorOnlyPropertyStack;
-#endif
+	SerializedProperty					 = ArchiveToCopy.SerializedProperty;
 #if USE_STABLE_LOCALIZATION_KEYS
 	SetBaseLocalizationNamespace(ArchiveToCopy.GetBaseLocalizationNamespace());
 #endif // USE_STABLE_LOCALIZATION_KEYS
@@ -225,6 +235,93 @@ FString FArchive::GetArchiveName() const
 	return TEXT("FArchive");
 }
 
+void FArchive::GetSerializedPropertyChain(TArray<class UProperty*>& OutProperties) const
+{
+	if (SerializedPropertyChain)
+	{
+		const int32 NumProperties = SerializedPropertyChain->GetNumProperties();
+		OutProperties.Reserve(NumProperties);
+
+		for (int32 PropertyIndex = 0; PropertyIndex < NumProperties; ++PropertyIndex)
+		{
+			OutProperties.Add(SerializedPropertyChain->GetPropertyFromStack(PropertyIndex));
+		}
+	}
+}
+
+void FArchive::SetSerializedPropertyChain(const FArchiveSerializedPropertyChain* InSerializedPropertyChain, class UProperty* InSerializedPropertyOverride)
+{
+	if (InSerializedPropertyChain && InSerializedPropertyChain->GetNumProperties() > 0)
+	{
+		if (!SerializedPropertyChain)
+		{
+			SerializedPropertyChain = new FArchiveSerializedPropertyChain();
+		}
+		*SerializedPropertyChain = *InSerializedPropertyChain;
+	}
+	else
+	{
+		delete SerializedPropertyChain;
+		SerializedPropertyChain = nullptr;
+	}
+
+	if (InSerializedPropertyOverride)
+	{
+		SerializedProperty = InSerializedPropertyOverride;
+	}
+	else if (SerializedPropertyChain && SerializedPropertyChain->GetNumProperties() > 0)
+	{
+		SerializedProperty = SerializedPropertyChain->GetPropertyFromStack(0);
+	}
+	else
+	{
+		SerializedProperty = nullptr;
+	}
+}
+
+void FArchive::PushSerializedProperty(class UProperty* InProperty, const bool bIsEditorOnlyProperty)
+{
+	if (InProperty)
+	{
+		// Push this property into the chain
+		if (!SerializedPropertyChain)
+		{
+			SerializedPropertyChain = new FArchiveSerializedPropertyChain();
+		}
+		SerializedPropertyChain->PushProperty(InProperty, bIsEditorOnlyProperty);
+
+		// Update the serialized property pointer with the new head
+		SerializedProperty = InProperty;
+	}
+}
+
+void FArchive::PopSerializedProperty(class UProperty* InProperty, const bool bIsEditorOnlyProperty)
+{
+	if (InProperty)
+	{
+		// Pop this property from the chain
+		check(SerializedPropertyChain);
+		SerializedPropertyChain->PopProperty(InProperty, bIsEditorOnlyProperty);
+
+		// Update the serialized property pointer with the new head
+		if (SerializedPropertyChain->GetNumProperties() > 0)
+		{
+			SerializedProperty = SerializedPropertyChain->GetPropertyFromStack(0);
+		}
+		else
+		{
+			SerializedProperty = nullptr;
+		}
+	}
+}
+
+#if WITH_EDITORONLY_DATA
+bool FArchive::IsEditorOnlyPropertyOnTheStack() const
+{
+	return SerializedPropertyChain && SerializedPropertyChain->HasEditorOnlyProperty();
+}
+#endif
+
 #if USE_STABLE_LOCALIZATION_KEYS
 void FArchive::SetBaseLocalizationNamespace(const FString& InLocalizationNamespace)
 {
diff --git a/Engine/Source/Runtime/Core/Public/Math/Interval.h b/Engine/Source/Runtime/Core/Public/Math/Interval.h
--- a/Engine/Source/Runtime/Core/Public/Math/Interval.h
+++ b/Engine/Source/Runtime/Core/Public/Math/Interval.h
@@ -8,13 +8,29 @@
 #include "Math/NumericLimits.h"
 #include "Math/UnrealMathUtility.h"
 
+/**
+ * Type traits for Arithmetic interval.
+ */
+template <typename ElementType> struct TIntervalTraits
+{
+	static_assert(TIsArithmetic<ElementType>::Value, "Incompatible TInterval element type.");
+
+	static ElementType Max()
+	{
+		return TNumericLimits<ElementType>::Max();
+	}
+
+	static ElementType Lowest()
+	{
+		return TNumericLimits<ElementType>::Lowest();
+	}
+};
+
 /**
  * Template for numeric interval
  */
 template<typename ElementType> struct TInterval
 {
-	static_assert(TIsArithmetic<ElementType>::Value, "Interval can be used only with numeric types");
-	
 	/** Holds the lower bound of the interval. */
 	ElementType Min;
 	
@@ -29,8 +45,8 @@ public:
 	 * The interval is invalid
 	 */
 	TInterval()
-		: Min(TNumericLimits<ElementType>::Max())
-		, Max(TNumericLimits<ElementType>::Lowest())
+		: Min(TIntervalTraits<ElementType>::Max())
+		, Max(TIntervalTraits<ElementType>::Lowest())
 	{ }
 
     /**
@@ -91,7 +107,7 @@ public:
 	 *
 	 * @return false when interval is invalid, true otherwise
 	 */
-	ElementType IsValid() const
+	bool IsValid() const
 	{
 		return (Min <= Max);
 	}
diff --git a/Engine/Source/Runtime/Core/Public/Misc/FileHelper.h b/Engine/Source/Runtime/Core/Public/Misc/FileHelper.h
--- a/Engine/Source/Runtime/Core/Public/Misc/FileHelper.h
+++ b/Engine/Source/Runtime/Core/Public/Misc/FileHelper.h
@@ -38,7 +38,7 @@ struct CORE_API FFileHelper
 	static void BufferToString( FString& Result, const uint8* Buffer, int32 Size );
 
 	/**
-	 * Load a binary file to a dynamic array.
+	 * Load a binary file to a dynamic array with two uninitialized bytes at end as padding.
 	 *
 	 * @param Result    Receives the contents of the file
 	 * @param Filename  The file to read
diff --git a/Engine/Source/Runtime/Core/Public/Misc/ITransaction.h b/Engine/Source/Runtime/Core/Public/Misc/ITransaction.h
--- a/Engine/Source/Runtime/Core/Public/Misc/ITransaction.h
+++ b/Engine/Source/Runtime/Core/Public/Misc/ITransaction.h
@@ -27,6 +27,143 @@ public:
 };
 
 
+/** Delta-change information for an object that was transacted */
+struct FTransactionObjectDeltaChange
+{
+	FTransactionObjectDeltaChange()
+		: bHasNameChange(false)
+		, bHasOuterChange(false)
+		, bHasPendingKillChange(false)
+		, bHasNonPropertyChanges(false)
+	{
+	}
+
+	bool HasChanged() const
+	{
+		return bHasNameChange || bHasOuterChange || bHasPendingKillChange || bHasNonPropertyChanges || ChangedProperties.Num() > 0;
+	}
+
+	/** True if the object name has changed */
+	bool bHasNameChange : 1;
+	/** True of the object outer has changed */
+	bool bHasOuterChange : 1;
+	/** True if the object "pending kill" state has changed */
+	bool bHasPendingKillChange : 1;
+	/** True if the object has changes other than property changes (may be caused by custom serialization) */
+	bool bHasNonPropertyChanges : 1;
+	/** Array of properties that have changed on the object */
+	TArray<FName> ChangedProperties;
+};
+
+
+/** Different kinds of actions that can trigger a transaction object event */
+enum class ETransactionObjectEventType : uint8
+{
+	/** This event was caused by an undo/redo operation */
+	UndoRedo,
+	/** This event was caused by a transaction being finalized within the transaction system */
+	Finalized,
+	/** This event was caused by a transaction snapshot. Several of these may be generated in the case of an interactive change */
+	Snapshot,
+};
+
+
+/**
+ * Transaction object events.
+ *
+ * Transaction object events are used to notify objects when they are transacted in some way.
+ * This mostly just means that an object has had an undo/redo applied to it, however an event is also triggered
+ * when the object has been finalized as part of a transaction (allowing you to detect object changes).
+ */
+class FTransactionObjectEvent
+{
+public:
+	FTransactionObjectEvent(const ETransactionObjectEventType InEventType, const FTransactionObjectDeltaChange& InDeltaChange, const TSharedPtr<ITransactionObjectAnnotation>& InAnnotation, const FName InOriginalObjectName, const FName InOriginalObjectPathName, const FName InOriginalObjectOuterPathName)
+		: EventType(InEventType)
+		, DeltaChange(InDeltaChange)
+		, Annotation(InAnnotation)
+		, OriginalObjectName(InOriginalObjectName)
+		, OriginalObjectPathName(InOriginalObjectPathName)
+		, OriginalObjectOuterPathName(InOriginalObjectOuterPathName)
+	{
+	}
+
+	/** What kind of action caused this event? */
+	ETransactionObjectEventType GetEventType() const
+	{
+		return EventType;
+	}
+
+	/** Was the pending kill state of this object changed? (implies non-property changes) */
+	bool HasPendingKillChange() const
+	{
+		return DeltaChange.bHasPendingKillChange;
+	}
+
+	/** Was the name of this object changed? (implies non-property changes) */
+	bool HasNameChange() const
+	{
+		return DeltaChange.bHasNameChange;
+	}
+
+	/** Get the original name of this object */
+	FName GetOriginalObjectName() const
+	{
+		return OriginalObjectName;
+	}
+
+	/** Get the original path name of this object */
+	FName GetOriginalObjectPathName() const
+	{
+		return OriginalObjectPathName;
+	}
+
+	/** Was the outer of this object changed? (implies non-property changes) */
+	bool HasOuterChange() const
+	{
+		return DeltaChange.bHasOuterChange;
+	}
+
+	/** Get the original outer path name of this object */
+	FName GetOriginalObjectOuterPathName() const
+	{
+		return OriginalObjectOuterPathName;
+	}
+
+	/** Were any non-property changes made to the object? */
+	bool HasNonPropertyChanges() const
+	{
+		return DeltaChange.bHasNameChange || DeltaChange.bHasOuterChange || DeltaChange.bHasPendingKillChange || DeltaChange.bHasNonPropertyChanges;
+	}
+
+	/** Were any property changes made to the object? */
+	bool HasPropertyChanges() const
+	{
+		return DeltaChange.ChangedProperties.Num() > 0;
+	}
+
+	/** Get the list of changed properties. Each entry is actually a chain of property names (root -> leaf) separated by a dot, eg) "ObjProp.StructProp". */
+	const TArray<FName>& GetChangedProperties() const
+	{
+		return DeltaChange.ChangedProperties;
+	}
+
+	/** Get the annotation object associated with the object being transacted (if any). */
+	TSharedPtr<const ITransactionObjectAnnotation> GetAnnotation() const
+	{
+		return Annotation;
+	}
+
+private:
+	ETransactionObjectEventType EventType;
+	FTransactionObjectDeltaChange DeltaChange;
+	TSharedPtr<ITransactionObjectAnnotation> Annotation;
+	FName OriginalObjectName;
+	FName OriginalObjectPathName;
+	FName OriginalObjectOuterPathName;
+};
+
+
 /**
  * Interface for transactions.
  *
@@ -37,6 +174,9 @@ class ITransaction
 {
 public:
 
+	/** Called when this transaction is completed to finalize the transaction */
+	virtual void Finalize( ) = 0;
+
 	/** Applies the transaction. */
 	virtual void Apply( ) = 0;
 
@@ -69,4 +209,11 @@ public:
 	 * @param Object The primary object to set.
 	 */
 	virtual void SetPrimaryObject( UObject* Object ) = 0;
+
+	/**
+	 * Snapshots a UObject within the transaction.
+	 *
+	 * @param Object The object to snapshot.
+	 */
+	virtual void SnapshotObject( UObject* Object ) = 0;
 };
diff --git a/Engine/Source/Runtime/Core/Public/Misc/Timespan.h b/Engine/Source/Runtime/Core/Public/Misc/Timespan.h
--- a/Engine/Source/Runtime/Core/Public/Misc/Timespan.h
+++ b/Engine/Source/Runtime/Core/Public/Misc/Timespan.h
@@ -4,6 +4,7 @@
 
 #include "CoreTypes.h"
 #include "Misc/AssertionMacros.h"
+#include "Math/Interval.h"
 
 class FArchive;
 class FOutputDevice;
@@ -774,6 +775,19 @@ private:
 	int64 Ticks;
 };
 
+template <>
+struct TIntervalTraits<FTimespan>
+{
+	static FTimespan Max()
+	{
+		return FTimespan::MaxValue();
+	}
+
+	static FTimespan Lowest()
+	{
+		return FTimespan::MinValue();
+	}
+};
 
 /**
  * Pre-multiply a time span with the given scalar.
diff --git a/Engine/Source/Runtime/Core/Public/Serialization/Archive.h b/Engine/Source/Runtime/Core/Public/Serialization/Archive.h
--- a/Engine/Source/Runtime/Core/Public/Serialization/Archive.h
+++ b/Engine/Source/Runtime/Core/Public/Serialization/Archive.h
@@ -19,6 +19,7 @@
 class FCustomVersionContainer;
 class ITargetPlatform;
 struct FUntypedBulkData;
+struct FArchiveSerializedPropertyChain;
 template<class TEnum> class TEnumAsByte;
 typedef TFunction<bool (double RemainingTime)> FExternalReadCallback;
 
@@ -1245,7 +1246,8 @@ public:
 	}
 
 	/**
-	 * Sets the property that is currently being serialized
+	 * Overrides the property that is currently being serialized
+	 * @note: You likely want to call PushSerializedProperty/PopSerializedProperty instead
 	 * 
 	 * @param InProperty Pointer to the property that is currently being serialized
 	 */
@@ -1254,41 +1256,63 @@ public:
 		SerializedProperty = InProperty;
 	}
 
-	/** 
-	 * Adds external read dependency 
+	/**
+	 * Gets the property that is currently being serialized
 	 *
-	 * @return true if dependency has been added, false if Archive does not support them
+	 * @return Pointer to the property that is currently being serialized
 	 */
-	virtual bool AttachExternalReadDependency(FExternalReadCallback& ReadCallback) { return false; };
-
-#if WITH_EDITORONLY_DATA
-	/** Pushes editor-only marker to the stack of currently serialized properties */
-	FORCEINLINE void PushEditorOnlyProperty()
+	FORCEINLINE class UProperty* GetSerializedProperty() const
 	{
-		EditorOnlyPropertyStack++;
+		return SerializedProperty;
 	}
-	/** Pops editor-only marker from the stack of currently serialized properties */
-	FORCEINLINE void PopEditorOnlyProperty()
+
+	/**
+	 * Gets the chain of properties that are currently being serialized
+	 * @note This populates the array in stack order, so the 0th entry in the array is the top of the stack of properties
+	 */
+	void GetSerializedPropertyChain(TArray<class UProperty*>& OutProperties) const;
+
+	/**
+	 * Get the raw serialized property chain for this archive
+	 * @note Accessing this directly can avoid an array allocation depending on your use-case
+	 */
+	FORCEINLINE const FArchiveSerializedPropertyChain* GetSerializedPropertyChain() const
 	{
-		EditorOnlyPropertyStack--;
-		check(EditorOnlyPropertyStack >= 0);
+		return SerializedPropertyChain;
 	}
+
+	/**
+	 * Set the raw serialized property chain for this archive, optionally overriding the serialized property too (or null to use the head of the property chain)
+	 */
+	void SetSerializedPropertyChain(const FArchiveSerializedPropertyChain* InSerializedPropertyChain, class UProperty* InSerializedPropertyOverride = nullptr);
+
+	/**
+	 * Push a property that is currently being serialized onto the property stack
+	 * 
+	 * @param InProperty			Pointer to the property that is currently being serialized
+	 * @param bIsEditorOnlyProperty True if the property is editor only (call UProperty::IsEditorOnlyProperty to work this out, as the archive can't since it can't access CoreUObject types)
+	 */
+	void PushSerializedProperty(class UProperty* InProperty, const bool bIsEditorOnlyProperty);
+
+	/**
+	 * Pop a property that was previously being serialized off the property stack
+	 * 
+	 * @param InProperty			Pointer to the property that was previously being serialized
+	 * @param bIsEditorOnlyProperty True if the property is editor only (call UProperty::IsEditorOnlyProperty to work this out, as the archive can't since it can't access CoreUObject types)
+	 */
+	void PopSerializedProperty(class UProperty* InProperty, const bool bIsEditorOnlyProperty);
+
+#if WITH_EDITORONLY_DATA
 	/** Returns true if the stack of currently serialized properties contains an editor-only property */
-	FORCEINLINE bool IsEditorOnlyPropertyOnTheStack() const
-	{
-		return EditorOnlyPropertyStack > 0;
-	}
+	bool IsEditorOnlyPropertyOnTheStack() const;
 #endif
 
-	/**
-	 * Gets the property that is currently being serialized
+	/** 
+	 * Adds external read dependency 
 	 *
-	 * @return Pointer to the property that is currently being serialized
+	 * @return true if dependency has been added, false if Archive does not support them
 	 */
-	FORCEINLINE class UProperty* GetSerializedProperty() const
-	{
-		return SerializedProperty;
-	}
+	virtual bool AttachExternalReadDependency(FExternalReadCallback& ReadCallback) { return false; };
 
 #if USE_STABLE_LOCALIZATION_KEYS
 	/**
@@ -1578,10 +1602,8 @@ private:
 	/** Holds the pointer to the property that is currently being serialized */
 	class UProperty* SerializedProperty;
 
-#if WITH_EDITORONLY_DATA
-	/** Non-zero if on the current stack of serialized properties there's an editor-only property. Needs to live in FArchive because of SerializedProperty. */
-	int32 EditorOnlyPropertyStack;
-#endif
+	/** Holds the chain of properties that are currently being serialized */
+	FArchiveSerializedPropertyChain* SerializedPropertyChain;
 
 #if USE_STABLE_LOCALIZATION_KEYS
 	/**
diff --git a/Engine/Source/Runtime/Core/Public/Windows/PostWindowsApi.h b/Engine/Source/Runtime/Core/Public/Windows/PostWindowsApi.h
--- a/Engine/Source/Runtime/Core/Public/Windows/PostWindowsApi.h
+++ b/Engine/Source/Runtime/Core/Public/Windows/PostWindowsApi.h
@@ -30,6 +30,8 @@ THIRD_PARTY_INCLUDES_END
 #undef CreateDesktop
 #undef GetMessage
 #undef GetCommandLine
+#undef GetProp
+#undef SetProp
 
 // Undefine all the atomics. AllowWindowsPlatformAtomics/HideWindowsPlatformAtomics temporarily defining these macros.
 #if !PLATFORM_XBOXONE
diff --git a/Engine/Source/Runtime/CoreUObject/Private/Serialization/ArchiveUObject.cpp b/Engine/Source/Runtime/CoreUObject/Private/Serialization/ArchiveUObject.cpp
--- a/Engine/Source/Runtime/CoreUObject/Private/Serialization/ArchiveUObject.cpp
+++ b/Engine/Source/Runtime/CoreUObject/Private/Serialization/ArchiveUObject.cpp
@@ -130,22 +130,21 @@ FArchive& FObjectAndNameAsStringProxyArchive::operator<<(FWeakObjectPtr& Obj)
 	return *this;
 }
 
-#if WITH_EDITORONLY_DATA
-void FSerializedPropertyScope::PushEditorOnlyProperty()
+void FSerializedPropertyScope::PushProperty()
 {
-	if (Property && Property->IsEditorOnlyProperty())
+	if (Property)
 	{
-		Ar.PushEditorOnlyProperty();
+		Ar.PushSerializedProperty(Property, Property->IsEditorOnlyProperty());
 	}
 }
-void FSerializedPropertyScope::PopEditorOnlyProperty()
+
+void FSerializedPropertyScope::PopProperty()
 {
-	if (Property && Property->IsEditorOnlyProperty())
+	if (Property)
 	{
-		Ar.PopEditorOnlyProperty();
+		Ar.PopSerializedProperty(Property, Property->IsEditorOnlyProperty());
 	}
 }
-#endif
 
 void FArchiveReplaceObjectRefBase::SerializeObject(UObject* ObjectToSerialize)
 {
diff --git a/Engine/Source/Runtime/CoreUObject/Private/UObject/LinkerSave.cpp b/Engine/Source/Runtime/CoreUObject/Private/UObject/LinkerSave.cpp
--- a/Engine/Source/Runtime/CoreUObject/Private/UObject/LinkerSave.cpp
+++ b/Engine/Source/Runtime/CoreUObject/Private/UObject/LinkerSave.cpp
@@ -252,7 +252,7 @@ void FLinkerSave::Serialize( void* V, int64 Length )
 {
 #if WITH_EDITOR
 	Saver->ArDebugSerializationFlags = ArDebugSerializationFlags;
-	Saver->SetSerializedProperty(GetSerializedProperty());
+	Saver->SetSerializedPropertyChain(GetSerializedPropertyChain(), GetSerializedProperty());
 #endif
 	Saver->Serialize( V, Length );
 }
diff --git a/Engine/Source/Runtime/CoreUObject/Private/UObject/Obj.cpp b/Engine/Source/Runtime/CoreUObject/Private/UObject/Obj.cpp
--- a/Engine/Source/Runtime/CoreUObject/Private/UObject/Obj.cpp
+++ b/Engine/Source/Runtime/CoreUObject/Private/UObject/Obj.cpp
@@ -336,6 +336,13 @@ void UObject::PostEditChange(void)
 void UObject::PostEditChangeProperty(FPropertyChangedEvent& PropertyChangedEvent)
 {
 	FCoreUObjectDelegates::OnObjectPropertyChanged.Broadcast(this, PropertyChangedEvent);
+
+	// Snapshot the transaction buffer for this object if this was from an interactive change
+	// This allows listeners to be notified of intermediate changes of state
+	if (PropertyChangedEvent.ChangeType == EPropertyChangeType::Interactive)
+	{
+		SnapshotTransactionBuffer(this);
+	}
 }
 
 
@@ -529,6 +536,10 @@ void UObject::PostEditUndo(TSharedPtr<ITransactionObjectAnnotation> TransactionA
 	UObject::PostEditUndo();
 }
 
+void UObject::PostTransacted(const FTransactionObjectEvent& TransactionEvent)
+{
+	FCoreUObjectDelegates::OnObjectTransacted.Broadcast(this, TransactionEvent);
+}
 
 bool UObject::IsSelectedInEditor() const
 {
@@ -1665,6 +1676,26 @@ void UObject::GetAssetRegistryTags(TArray<FAssetRegistryTag>& OutTags) const
 	}
 
 	FAssetRegistryTag::GetAssetRegistryTagsFromSearchableProperties(this, OutTags);
+
+#if WITH_EDITOR
+	// Check if there's a UMetaData for this object that has tags that are requested in the settings to be transferred to the Asset Registry
+	const TSet<FName>& MetaDataTagsForAR = GetMetaDataTagsForAssetRegistry();
+	if (MetaDataTagsForAR.Num() > 0)
+	{
+		TMap<FName, FString>* MetaDataMap = UMetaData::GetMapForObject(this);
+		if (MetaDataMap)
+		{
+			for (TMap<FName, FString>::TConstIterator It(*MetaDataMap); It; ++It)
+			{
+				FName Tag = It->Key;
+				if (!Tag.IsNone() && MetaDataTagsForAR.Contains(Tag))
+				{
+					OutTags.Add(FAssetRegistryTag(Tag, It->Value, UObject::FAssetRegistryTag::TT_Alphabetical));
+				}
+			}
+		}
+	}
+#endif // WITH_EDITOR
 }
 
 const FName& UObject::SourceFileTagName()
@@ -1674,6 +1705,13 @@ const FName& UObject::SourceFileTagName()
 }
 
 #if WITH_EDITOR
+static TSet<FName> MetaDataTagsForAssetRegistry;
+
+TSet<FName>& UObject::GetMetaDataTagsForAssetRegistry()
+{
+	return MetaDataTagsForAssetRegistry;
+}
+
 void UObject::GetAssetRegistryTagMetadata(TMap<FName, FAssetRegistryTagMetadata>& OutMetadata) const
 {
 	OutMetadata.Add(FPrimaryAssetId::PrimaryAssetTypeTag,
diff --git a/Engine/Source/Runtime/CoreUObject/Private/UObject/UObjectGlobals.cpp b/Engine/Source/Runtime/CoreUObject/Private/UObject/UObjectGlobals.cpp
--- a/Engine/Source/Runtime/CoreUObject/Private/UObject/UObjectGlobals.cpp
+++ b/Engine/Source/Runtime/CoreUObject/Private/UObject/UObjectGlobals.cpp
@@ -110,6 +110,7 @@ FCoreUObjectDelegates::FOnPreObjectPropertyChanged FCoreUObjectDelegates::OnPreO
 FCoreUObjectDelegates::FOnObjectPropertyChanged FCoreUObjectDelegates::OnObjectPropertyChanged;
 TSet<UObject*> FCoreUObjectDelegates::ObjectsModifiedThisFrame;
 FCoreUObjectDelegates::FOnObjectModified FCoreUObjectDelegates::OnObjectModified;
+FCoreUObjectDelegates::FOnObjectTransacted FCoreUObjectDelegates::OnObjectTransacted;
 FCoreUObjectDelegates::FOnAssetLoaded FCoreUObjectDelegates::OnAssetLoaded;
 FCoreUObjectDelegates::FOnObjectSaved FCoreUObjectDelegates::OnObjectSaved;
 #endif // WITH_EDITOR
@@ -2147,7 +2148,6 @@ UObject* StaticDuplicateObjectEx( FObjectDuplicationParameters& Parameters )
 	return DupRootObject;
 }
 
-
 /**
  * Save a copy of this object into the transaction buffer if we are currently recording into
  * one (undo/redo). If bMarkDirty is true, will also mark the package as needing to be saved.
@@ -2185,6 +2185,26 @@ bool SaveToTransactionBuffer(UObject* Object, bool bMarkDirty)
 	return bSavedToTransactionBuffer;
 }
 
+/**
+ * Causes the transaction system to emit a snapshot event for the given object if the following conditions are met:
+ *  a) The object is currently transacting.
+ *  b) The object has changed since it started transacting.
+ *
+ * @param	Object		object to snapshot.
+ */
+void SnapshotTransactionBuffer(UObject* Object)
+{
+	// Neither PIE world objects nor script packages should end up in the transaction buffer. Additionally, in order
+	// to save a copy of the object, we must have a transactor and the object must be transactional.
+	const bool IsTransactional = Object->HasAnyFlags(RF_Transactional);
+	const bool IsNotPIEOrContainsScriptObject = (Object->GetOutermost()->HasAnyPackageFlags(PKG_PlayInEditor | PKG_ContainsScript) == false);
+
+	if (GUndo && IsTransactional && IsNotPIEOrContainsScriptObject)
+	{
+		GUndo->SnapshotObject(Object);
+	}
+}
+
 /**
  * Check for StaticAllocateObject error; only for use with the editor, make or other commandlets.
  * 
diff --git a/Engine/Source/Runtime/CoreUObject/Public/Serialization/SerializedPropertyScope.h b/Engine/Source/Runtime/CoreUObject/Public/Serialization/SerializedPropertyScope.h
--- a/Engine/Source/Runtime/CoreUObject/Public/Serialization/SerializedPropertyScope.h
+++ b/Engine/Source/Runtime/CoreUObject/Public/Serialization/SerializedPropertyScope.h
@@ -13,31 +13,24 @@ class COREUOBJECT_API FSerializedPropertyScope
 {
 	FArchive& Ar;
 	UProperty* Property;
-	UProperty* OldProperty;
-#if WITH_EDITORONLY_DATA
-	void PushEditorOnlyProperty();
-	void PopEditorOnlyProperty();
-#endif
+	void PushProperty();
+	void PopProperty();
 public:
 	FSerializedPropertyScope(FArchive& InAr, UProperty* InProperty, const UProperty* OnlyIfOldProperty = nullptr)
 		: Ar(InAr)
 		, Property(InProperty)
 	{
-		OldProperty = Ar.GetSerializedProperty();
-		if (!OnlyIfOldProperty || OldProperty == OnlyIfOldProperty)
+		if (!OnlyIfOldProperty || Ar.GetSerializedProperty() == OnlyIfOldProperty)
 		{
-			Ar.SetSerializedProperty(Property);
+			PushProperty();
+		}
+		else
+		{
+			Property = nullptr;
 		}
-		
-#if WITH_EDITORONLY_DATA
-		PushEditorOnlyProperty();
-#endif
 	}
 	~FSerializedPropertyScope()
 	{
-#if WITH_EDITORONLY_DATA
-		PopEditorOnlyProperty();
-#endif
-		Ar.SetSerializedProperty(OldProperty);
+		PopProperty();
 	}
 };
diff --git a/Engine/Source/Runtime/CoreUObject/Public/UObject/Object.h b/Engine/Source/Runtime/CoreUObject/Public/UObject/Object.h
--- a/Engine/Source/Runtime/CoreUObject/Public/UObject/Object.h
+++ b/Engine/Source/Runtime/CoreUObject/Public/UObject/Object.h
@@ -13,6 +13,7 @@ class FConfigCacheIni;
 class FEditPropertyChain;
 class ITargetPlatform;
 class ITransactionObjectAnnotation;
+class FTransactionObjectEvent;
 struct FFrame;
 struct FObjectInstancingGraph;
 struct FPropertyChangedChainEvent;
@@ -324,14 +325,22 @@ public:
 	/** Called after applying a transaction to the object in cases where transaction annotation was provided. Default implementation simply calls PostEditChange. */
 	virtual void PostEditUndo(TSharedPtr<ITransactionObjectAnnotation> TransactionAnnotation);
 
+	/**
+	 * Called after the object has been transacted in some way.
+	 * TransactionEvent describes what actually happened.
+	 * @note Unlike PostEditUndo (which is called for any object in the transaction), this is only called on objects that are actually changed by the transaction.
+	 */
+	virtual void PostTransacted(const FTransactionObjectEvent& TransactionEvent);
+
+private:
 	/**
 	 * Test the selection state of a UObject
 	 *
 	 * @return		true if the object is selected, false otherwise.
 	 * @todo UE4 this doesn't belong here, but it doesn't belong anywhere else any better
 	 */
-private:
 	virtual bool IsSelectedInEditor() const;
+
 public:
 #endif // WITH_EDITOR
 
@@ -669,6 +678,10 @@ public:
 
 	/** Gathers a collection of asset registry tag metadata */
 	virtual void GetAssetRegistryTagMetadata(TMap<FName, FAssetRegistryTagMetadata>& OutMetadata) const;
+
+	/** The metadata tags to be transferred from the UMetaData to the Asset Registry */
+	static TSet<FName>& GetMetaDataTagsForAssetRegistry();
+
 #endif
 
 	/** Returns true if this object is considered an asset. */
diff --git a/Engine/Source/Runtime/CoreUObject/Public/UObject/ObjectMacros.h b/Engine/Source/Runtime/CoreUObject/Public/UObject/ObjectMacros.h
--- a/Engine/Source/Runtime/CoreUObject/Public/UObject/ObjectMacros.h
+++ b/Engine/Source/Runtime/CoreUObject/Public/UObject/ObjectMacros.h
@@ -1003,7 +1003,7 @@ namespace UM
 		/// [ClassMetadata] For BehaviorTree nodes indicates that the class is deprecated and will display a warning when compiled.
 		DeprecatedNode,
 
-		/// [ClassMetadata] [FunctionMetadata] Used in conjunction with DeprecatedNode or DeprecatedFunction to customize the warning message displayed to the user.
+		/// [ClassMetadata] [PropertyMetadata] [FunctionMetadata] Used in conjunction with DeprecatedNode, DeprecatedProperty, or DeprecatedFunction to customize the warning message displayed to the user.
 		DeprecationMessage,
 
 		/// [ClassMetadata] [PropertyMetadata] [FunctionMetadata] The name to display for this class, property, or function instead of auto-generating it from the name.
@@ -1090,6 +1090,12 @@ namespace UM
 		/// [PropertyMetadata] Used by FDirectoryPath properties. Indicates that the path will be picked using the Slate-style directory picker inside the game Content dir.
 		ContentDir,
 
+		/// [PropertyMetadata] This property is deprecated, any blueprint references to it cause a compilation warning.
+		DeprecatedProperty,
+
+		/// [ClassMetadata] [PropertyMetadata] [FunctionMetadata] Used in conjunction with DeprecatedNode, DeprecatedProperty, or DeprecatedFunction to customize the warning message displayed to the user.
+		// DeprecationMessage, (Commented out so as to avoid duplicate name with version in the Class section, but still show in the property section)
+
 		/// [ClassMetadata] [PropertyMetadata] [FunctionMetadata] The name to display for this class, property, or function instead of auto-generating it from the name.
 		// DisplayName, (Commented out so as to avoid duplicate name with version in the Class section, but still show in the property section)
 
@@ -1277,6 +1283,14 @@ namespace UM
 		/// [PropertyMetadata] [FunctionMetadata] Flag set on a property or function to prevent it being exported to a scripting language.
 		//ScriptNoExport, (Commented out so as to avoid duplicate name with version in the Property section, but still show in the function section)
 
+		/// [FunctionMetadata] Flags a static function taking a struct as its first argument so that it "hoists" the function to be a method of the struct when exporting it to a scripting language.
+		/// The value is optional, and may specify a name override for the struct method.
+		ScriptMethod,
+
+		/// [FunctionMetadata] Flags a static function taking a struct as its first argument so that it "hoists" the function to be an mathematical operator of the struct when exporting it to a scripting language.
+		/// The value describes the kind of operator using C++ operator syntax (+, +=, -,-=, *, *=, /, /=, %, %=, &, &=, |, |=, ^, ^=, >>, >>=, <<, <<=).
+		ScriptMathOp,
+
 		/// [FunctionMetadata] For BlueprintCallable functions indicates that the parameter pin should be hidden from the user's view.
 		HidePin,
 
diff --git a/Engine/Source/Runtime/CoreUObject/Public/UObject/UObjectGlobals.h b/Engine/Source/Runtime/CoreUObject/Public/UObject/UObjectGlobals.h
--- a/Engine/Source/Runtime/CoreUObject/Public/UObject/UObjectGlobals.h
+++ b/Engine/Source/Runtime/CoreUObject/Public/UObject/UObjectGlobals.h
@@ -528,6 +528,14 @@ void GlobalSetProperty( const TCHAR* Value, UClass* Class, UProperty* Property,
  */
 COREUOBJECT_API bool SaveToTransactionBuffer(UObject* Object, bool bMarkDirty);
 
+/**
+ * Causes the transaction system to emit a snapshot event for the given object if the following conditions are met:
+ *  a) The object is currently transacting.
+ *  b) The object has changed since it started transacting.
+ *
+ * @param	Object		object to snapshot.
+ */
+COREUOBJECT_API void SnapshotTransactionBuffer(UObject* Object);
 
 /**
  * Check for StaticAllocateObject error; only for use with the editor, make or other commandlets.
@@ -1875,6 +1883,10 @@ struct COREUOBJECT_API FCoreUObjectDelegates
 		}
 	}
 
+	/** Callback for an object being transacted */
+	DECLARE_MULTICAST_DELEGATE_TwoParams(FOnObjectTransacted, UObject*, const class FTransactionObjectEvent&);
+	static FOnObjectTransacted OnObjectTransacted;
+
 	/** Callback for when an asset is saved. This is called from UObject::PreSave before it is actually written to disk, for every object saved */
 	DECLARE_MULTICAST_DELEGATE_OneParam(FOnObjectSaved, UObject*);
 	static FOnObjectSaved OnObjectSaved;
diff --git a/Engine/Source/Runtime/D3D12RHI/Private/D3D12Viewport.cpp b/Engine/Source/Runtime/D3D12RHI/Private/D3D12Viewport.cpp
--- a/Engine/Source/Runtime/D3D12RHI/Private/D3D12Viewport.cpp
+++ b/Engine/Source/Runtime/D3D12RHI/Private/D3D12Viewport.cpp
@@ -6,6 +6,7 @@
 
 #include "D3D12RHIPrivate.h"
 #include "RenderCore.h"
+#include "Engine/RendererSettings.h"
 
 namespace D3D12RHI
 {
@@ -676,7 +677,8 @@ FViewportRHIRef FD3D12DynamicRHI::RHICreateViewport(void* WindowHandle, uint32 S
 	// Use a default pixel format if none was specified	
 	if (PreferredPixelFormat == EPixelFormat::PF_Unknown)
 	{
-		PreferredPixelFormat = EPixelFormat::PF_A2B10G10R10;
+		static const auto CVarDefaultBackBufferPixelFormat = IConsoleManager::Get().FindTConsoleVariableDataInt(TEXT("r.DefaultBackBufferPixelFormat"));
+		PreferredPixelFormat = EDefaultBackBufferPixelFormat::Convert2PixelFormat(CVarDefaultBackBufferPixelFormat->GetValueOnGameThread());
 	}
 
 	FD3D12Viewport* RenderingViewport = new FD3D12Viewport(&GetAdapter(), (HWND)WindowHandle, SizeX, SizeY, bIsFullscreen, PreferredPixelFormat);
@@ -699,7 +701,8 @@ void FD3D12DynamicRHI::RHIResizeViewport(FViewportRHIParamRef ViewportRHI, uint3
 	// Use a default pixel format if none was specified	
 	if (PreferredPixelFormat == EPixelFormat::PF_Unknown)
 	{
-		PreferredPixelFormat = EPixelFormat::PF_A2B10G10R10;
+		static const auto CVarDefaultBackBufferPixelFormat = IConsoleManager::Get().FindTConsoleVariableDataInt(TEXT("r.DefaultBackBufferPixelFormat"));
+		PreferredPixelFormat = EDefaultBackBufferPixelFormat::Convert2PixelFormat(CVarDefaultBackBufferPixelFormat->GetValueOnGameThread());
 	}
 
 	FD3D12Viewport* Viewport = FD3D12DynamicRHI::ResourceCast(ViewportRHI);
diff --git a/Engine/Source/Runtime/Engine/Classes/EditorFramework/AssetImportData.h b/Engine/Source/Runtime/Engine/Classes/EditorFramework/AssetImportData.h
--- a/Engine/Source/Runtime/Engine/Classes/EditorFramework/AssetImportData.h
+++ b/Engine/Source/Runtime/Engine/Classes/EditorFramework/AssetImportData.h
@@ -103,16 +103,22 @@ public:
 	void UpdateFilenameOnly(const FString& InPath, int32 Index);
 
 	/** Helper function to return the first filename stored in this data. The resulting filename will be absolute (ie, not relative to the asset).  */
+#if WITH_EDITOR
+	UFUNCTION(BlueprintCallable, meta=(DisplayName="GetFirstFilename", ScriptName="GetFirstFilename"), Category="AssetImportData")
+	FString K2_GetFirstFilename() const;
+#endif
 	FString GetFirstFilename() const;
 
 	/** Const access to the source file data */
 	const FAssetImportInfo& GetSourceData() const { return SourceData; }
 
 	/** Extract all the (resolved) filenames from this data  */
-	void ExtractFilenames(TArray<FString>& AbsoluteFilenames) const;
-
-	/** Extract all the (resolved) filenames from this data  */
+#if WITH_EDITOR
+	UFUNCTION(BlueprintCallable, meta=(DisplayName="ExtractFilenames", ScriptName="ExtractFilenames"), Category="AssetImportData")
+	TArray<FString> K2_ExtractFilenames() const;
+#endif
 	TArray<FString> ExtractFilenames() const;
+	void ExtractFilenames(TArray<FString>& AbsoluteFilenames) const;
 
 	/** Resolve a filename that is relative to either the specified package, BaseDir() or absolute */
 	static FString ResolveImportFilename(const FString& InRelativePath, const UPackage* Outermost);
diff --git a/Engine/Source/Runtime/Engine/Classes/Engine/AssetManager.h b/Engine/Source/Runtime/Engine/Classes/Engine/AssetManager.h
--- a/Engine/Source/Runtime/Engine/Classes/Engine/AssetManager.h
+++ b/Engine/Source/Runtime/Engine/Classes/Engine/AssetManager.h
@@ -11,6 +11,8 @@
 #include "GenericPlatform/GenericPlatformChunkInstall.h"
 #include "AssetManager.generated.h"
 
+DECLARE_LOG_CATEGORY_EXTERN(LogAssetManager, Log, All);
+
 /** Defined in C++ file */
 struct FPrimaryAssetTypeData;
 struct FPrimaryAssetData;
diff --git a/Engine/Source/Runtime/Engine/Classes/Engine/AssetManagerSettings.h b/Engine/Source/Runtime/Engine/Classes/Engine/AssetManagerSettings.h
--- a/Engine/Source/Runtime/Engine/Classes/Engine/AssetManagerSettings.h
+++ b/Engine/Source/Runtime/Engine/Classes/Engine/AssetManagerSettings.h
@@ -93,7 +93,18 @@ public:
 	/** Redirect from /game/assetpath to /game/assetpathnew */
 	UPROPERTY(config, EditAnywhere, Category = "Redirects")
 	TArray<FAssetManagerRedirect> AssetPathRedirects;
+
+	/** The metadata tags to be transferred to the Asset Registry. */
+	UPROPERTY(config, EditAnywhere, AdvancedDisplay, Category = "Asset Registry", DisplayName = "Metadata Tags For Asset Registry")
+	TSet<FName> MetaDataTagsForAssetRegistry;
+
 #if WITH_EDITOR
+	virtual void PostInitProperties() override;
+	virtual void PreEditChange(UProperty* PropertyAboutToChange) override;
 	virtual void PostEditChangeProperty(FPropertyChangedEvent& PropertyChangedEvent) override;
+
+private:
+	void ApplyMetaDataTagsSettings();
+	void ClearMetaDataTagsSettings();
 #endif
 };
\ No newline at end of file
diff --git a/Engine/Source/Runtime/Engine/Classes/Engine/Engine.h b/Engine/Source/Runtime/Engine/Classes/Engine/Engine.h
--- a/Engine/Source/Runtime/Engine/Classes/Engine/Engine.h
+++ b/Engine/Source/Runtime/Engine/Classes/Engine/Engine.h
@@ -35,6 +35,7 @@ class FSceneViewExtensions;
 class IStereoRendering;
 class SViewport;
 class UEditorEngine;
+class UEngineCustomTimeStep;
 class UGameUserSettings;
 class UGameViewportClient;
 class ULocalPlayer;
@@ -1248,6 +1249,12 @@ public:
 	UPROPERTY(config, EditAnywhere, Category=Framerate, meta=(UIMin=0, UIMax=200, EditCondition="!bUseFixedFrameRate"))
 	FFloatRange SmoothedFrameRateRange;
 
+private:
+	/** Control how the Engine process the Framerate/Timestep */
+	UPROPERTY(transient)
+	UEngineCustomTimeStep* CustomTimeStep;
+
+public:
 	/** 
 	 * Whether we should check for more than N pawns spawning in a single frame.  
 	 * Basically, spawning pawns and all of their attachments can be slow.  And on consoles it
@@ -1933,6 +1940,13 @@ public:
 	 */
 	void UpdateTimeAndHandleMaxTickRate();
 
+	/**
+	 * Set the CustomTimeStep that will control the Engine Framerate/Timestep
+	 *
+	 * @return	true if the CustomTimeStep was properly initialized
+	 */
+	bool SetCustomTimeStep(UEngineCustomTimeStep* InCustomTimeStep);
+
 	/** Executes the deferred commands **/
 	void TickDeferredCommands();
 
diff --git a/Engine/Source/Runtime/Engine/Classes/Engine/EngineCustomTimeStep.h b/Engine/Source/Runtime/Engine/Classes/Engine/EngineCustomTimeStep.h
new file mode 100644
--- /dev/null
+++ b/Engine/Source/Runtime/Engine/Classes/Engine/EngineCustomTimeStep.h
@@ -0,0 +1,21 @@
+// Copyright 1998-2018 Epic Games, Inc. All Rights Reserved.
+
+#pragma once
+
+#include "CoreMinimal.h"
+#include "UObject/Object.h"
+#include "UObject/ObjectMacros.h"
+#include "EngineCustomTimeStep.generated.h"
+
+UCLASS(abstract)
+class ENGINE_API UEngineCustomTimeStep : public UObject
+{
+	GENERATED_BODY()
+
+public:
+	virtual bool Initialize(class UEngine* InEngine) PURE_VIRTUAL(UEngineCustomTimeStep::Initialize, return false;);
+	virtual void Shutdown(class UEngine* InEngine) PURE_VIRTUAL(UEngineCustomTimeStep::Shutdown, );
+
+	virtual void UpdateTimeStep(class UEngine* InEngine) PURE_VIRTUAL(UEngineCustomTimeStep::UpdateTimeStep, );
+};
+
diff --git a/Engine/Source/Runtime/Engine/Classes/Engine/EngineTypes.h b/Engine/Source/Runtime/Engine/Classes/Engine/EngineTypes.h
--- a/Engine/Source/Runtime/Engine/Classes/Engine/EngineTypes.h
+++ b/Engine/Source/Runtime/Engine/Classes/Engine/EngineTypes.h
@@ -815,6 +815,9 @@ namespace EWorldType
 		/** A preview world for a game */
 		GamePreview,
 
+		/** A minimal RPC world for a game */
+		GameRPC,
+
 		/** An editor world that was loaded but not currently being edited in the level editor */
 		Inactive
 	};
@@ -3395,7 +3398,7 @@ enum EPhysicalSurface
 
 
 /** Describes how often this component is allowed to move. */
-UENUM()
+UENUM(BlueprintType)
 namespace EComponentMobility
 {
 	enum Type
diff --git a/Engine/Source/Runtime/Engine/Classes/Engine/GameInstance.h b/Engine/Source/Runtime/Engine/Classes/Engine/GameInstance.h
--- a/Engine/Source/Runtime/Engine/Classes/Engine/GameInstance.h
+++ b/Engine/Source/Runtime/Engine/Classes/Engine/GameInstance.h
@@ -176,6 +176,12 @@ public:
 	/* Called to initialize the game instance for standalone instances of the game */
 	void InitializeStandalone();
 
+	/* Called to initialize the game instance with a minimal world suitable for basic network RPC */
+	void InitializeForMinimalNetRPC(const FName InPackageName);
+
+	/** Static util function used by InitializeForMinimalNetRPC and LoadMap to create the minimal world suitable for basic network RPC */
+	static void CreateMinimalNetRPCWorld(const FName InPackageName, UPackage*& OutWorldPackage, UWorld*& OutWorld);
+
 #if WITH_EDITOR
 	/* Called to initialize the game instance for PIE instances of the game */
 	virtual FGameInstancePIEResult InitializeForPlayInEditor(int32 PIEInstanceIndex, const FGameInstancePIEParameters& Params);
diff --git a/Engine/Source/Runtime/Engine/Classes/Engine/RendererSettings.h b/Engine/Source/Runtime/Engine/Classes/Engine/RendererSettings.h
--- a/Engine/Source/Runtime/Engine/Classes/Engine/RendererSettings.h
+++ b/Engine/Source/Runtime/Engine/Classes/Engine/RendererSettings.h
@@ -7,6 +7,7 @@
 #include "Engine/EngineTypes.h"
 #include "Engine/Scene.h"
 #include "Engine/DeveloperSettings.h"
+#include "PixelFormat.h"
 
 #include "RendererSettings.generated.h"
 
@@ -90,6 +91,32 @@ namespace EEarlyZPass
 	};
 }
 
+/**
+ * Enumerates available options for alpha channel through post processing. The renderer will always generate premultiplied RGBA
+ * with alpha as translucency (0 = fully opaque; 1 = fully translucent).
+ */
+UENUM()
+namespace EAlphaChannelMode
+{
+	enum Type
+	{
+		/** Disabled, reducing GPU cost to the minimum. (default). */
+		Disabled = 0 UMETA(DisplayName = "Disabled"),
+
+		/** Maintain alpha channel only within linear color space. Tonemapper won't output alpha channel. */
+		LinearColorSpaceOnly = 1 UMETA(DisplayName="Linear color space only"),
+
+		/** Maintain alpha channel within linear color space, but also pass it through the tonemapper.
+		 *
+		 * CAUTION: Passing the alpha channel through the tonemapper can unevitably lead to pretty poor compositing quality as
+		 * opposed to linear color space compositing, especially on purely additive pixels bloom can generate. This settings is
+		 * exclusively targeting broadcast industry in case of hardware unable to do linear color space compositing and
+		 * tonemapping.
+		 */
+		AllowThroughTonemapper = 2 UMETA(DisplayName="Allow through tonemapper"),
+	};
+}
+
 /** used by FPostProcessSettings AutoExposure*/
 UENUM()
 namespace EAutoExposureMethodUI
@@ -106,6 +133,26 @@ namespace EAutoExposureMethodUI
 	};
 }
 
+/** used by DefaultBackBufferPixelFormat*/
+UENUM()
+namespace EDefaultBackBufferPixelFormat
+{
+	enum Type
+	{
+		DBBPF_B8G8R8A8 = 0	UMETA(DisplayName = "8bit RGBA"),
+		DBBPF_A16B16G16R16	UMETA(DisplayName = "16bit RGBA"),
+		DBBPF_FloatRGB		UMETA(DisplayName = "Float RGB"),
+		DBBPF_FloatRGBA		UMETA(DisplayName = "Float RGBA"),
+		DBBPF_A2B10G10R10	UMETA(DisplayName = "10bit RGB, 2bit Alpha"),
+		DBBPF_MAX,
+	};
+}
+
+namespace EDefaultBackBufferPixelFormat
+{
+	ENGINE_API EPixelFormat Convert2PixelFormat(int32 InDefaultBackBufferPixelFormat);
+}
+
 /**
  * Rendering settings.
  */
@@ -276,9 +323,9 @@ class ENGINE_API URendererSettings : public UDeveloperSettings
 
 	UPROPERTY(config, EditAnywhere, Category = Postprocessing, meta = (
 		ConsoleVariable = "r.PostProcessing.PropagateAlpha", DisplayName = "Enable alpha channel support in post processing (experimental).",
-		ToolTip = "Whether alpha channel should be supported in post processing chain. Still experimental: works only with Temporal AA, Motion Blur, Circle Depth Of Field. This option also force disable the separate translucency.",
+		ToolTip = "Configures alpha channel support in renderer's post processing chain. Still experimental: works only with Temporal AA, Motion Blur, Circle Depth Of Field. This option also force disable the separate translucency.",
 		ConfigRestartRequired = true))
-	uint32 bEnableAlphaChannelInPostProcessing : 1;
+	TEnumAsByte<EAlphaChannelMode::Type> bEnableAlphaChannelInPostProcessing;
 
 	UPROPERTY(config, EditAnywhere, Category = Postprocessing, meta = (
 		ConsoleVariable = "r.UsePreExposure", DisplayName = "Apply Pre-exposure before writing to the scene color",
@@ -342,6 +389,12 @@ class ENGINE_API URendererSettings : public UDeveloperSettings
 		ToolTip = "Which units to use for spot lights"))
 	ELightUnits DefaultSpotLightUnits;
 
+	UPROPERTY(config, EditAnywhere, AdvancedDisplay, Category = DefaultSettings, meta = (
+		ConsoleVariable = "r.DefaultBackBufferPixelFormat",DisplayName = "Frame Buffer Pixel Format",
+		ToolTip = "Pixel format used for back buffer, when not specified",
+		ConfigRestartRequired=true))
+	TEnumAsByte<EDefaultBackBufferPixelFormat::Type> DefaultBackBufferPixelFormat;
+
 	UPROPERTY(config, EditAnywhere, Category=Optimizations, meta=(
 		ConsoleVariable="r.Shadow.UnbuiltPreviewInGame",DisplayName="Render Unbuilt Preview Shadows in game",
 		ToolTip="Whether to render unbuilt preview shadows in game.  When enabled and lighting is not built, expensive preview shadows will be rendered in game.  When disabled, lighting in game and editor won't match which can appear to be a bug."))
diff --git a/Engine/Source/Runtime/Engine/Classes/Engine/StaticMesh.h b/Engine/Source/Runtime/Engine/Classes/Engine/StaticMesh.h
--- a/Engine/Source/Runtime/Engine/Classes/Engine/StaticMesh.h
+++ b/Engine/Source/Runtime/Engine/Classes/Engine/StaticMesh.h
@@ -868,6 +868,12 @@ public:
 
 #if WITH_EDITOR
 
+	/**
+	 * Sets a Material given a Material Index
+	 */
+	UFUNCTION(BlueprintCallable, Category = "StaticMesh")
+	ENGINE_API void SetMaterial(int32 MaterialIndex, UMaterialInterface* NewMaterial);
+
 	/**
 	 * Returns true if LODs of this static mesh may share texture lightmaps.
 	 */
