Unreal Math: Added Ray/Plane intersection function
- New FMath::RayPlaneIntersection() static function

diff --git a/Engine/Source/Runtime/Core/Public/Math/Plane.h b/Engine/Source/Runtime/Core/Public/Math/Plane.h
--- a/Engine/Source/Runtime/Core/Public/Math/Plane.h
+++ b/Engine/Source/Runtime/Core/Public/Math/Plane.h
@@ -284,6 +284,16 @@ public:
 /* FMath inline functions
  *****************************************************************************/
 
+inline FVector FMath::RayPlaneIntersection( const FVector& RayOrigin, const FVector& RayDirection, const FPlane& Plane )
+{
+	const FVector PlaneNormal = FVector( Plane.X, Plane.Y, Plane.Z );
+	const FVector PlaneOrigin = PlaneNormal * Plane.W;
+
+	const float Distance = FVector::DotProduct( ( PlaneOrigin - RayOrigin ), PlaneNormal ) / FVector::DotProduct( RayDirection, PlaneNormal );
+	return RayOrigin + RayDirection * Distance;
+}
+
+
 inline FVector FMath::LinePlaneIntersection
 	(
 	const FVector &Point1,
diff --git a/Engine/Source/Runtime/Core/Public/Math/UnrealMathUtility.h b/Engine/Source/Runtime/Core/Public/Math/UnrealMathUtility.h
--- a/Engine/Source/Runtime/Core/Public/Math/UnrealMathUtility.h
+++ b/Engine/Source/Runtime/Core/Public/Math/UnrealMathUtility.h
@@ -943,6 +943,17 @@ struct FMath : public FPlatformMath
 
 	// Geometry intersection 
 
+	/**
+	 * Find the intersection of a ray and a plane.  The ray has a start point with an infinite length.  Assumes that the
+	 * line and plane do indeed intersect; you must make sure they're not parallel before calling.
+	 *
+	 * @param RayOrigin	The start point of the ray
+	 * @param RayDirection	The direction the ray is pointing (normalized vector)
+	 * @param Plane	The plane to intersect with
+	 *
+	 * @return The point of intersection between the ray and the plane.
+	 */
+	static FVector RayPlaneIntersection( const FVector& RayOrigin, const FVector& RayDirection, const FPlane& Plane );
 
 	/**
 	 * Find the intersection of a line and an offset plane. Assumes that the
@@ -971,6 +982,7 @@ struct FMath : public FPlatformMath
 	 */
 	static FVector LinePlaneIntersection( const FVector &Point1, const FVector &Point2, const FPlane  &Plane);
 
+
 	// @parma InOutScissorRect should be set to View.ViewRect before the call
 	// @return 0: light is not visible, 1:use scissor rect, 2: no scissor rect needed
 	static CORE_API uint32 ComputeProjectedSphereScissorRect(struct FIntRect& InOutScissorRect, FVector SphereOrigin, float Radius, FVector ViewOrigin, const FMatrix& ViewMatrix, const FMatrix& ProjMatrix);
