Added an explicit constructor to FLinearColor which takes a FVector4.

diff --git a/Engine/Source/Runtime/Core/Private/Math/Color.cpp b/Engine/Source/Runtime/Core/Private/Math/Color.cpp
--- a/Engine/Source/Runtime/Core/Private/Math/Color.cpp
+++ b/Engine/Source/Runtime/Core/Private/Math/Color.cpp
@@ -6,6 +6,7 @@
 
 #include "Math/Color.h"
 #include "Math/Vector.h"
+#include "Math/Vector4.h"
 #include "Math/Float16Color.h"
 
 // Common colors.
@@ -54,6 +55,13 @@ FLinearColor::FLinearColor(const FVector& Vector) :
 	A(1.0f)
 {}
 
+FLinearColor::FLinearColor(const FVector4& Vector) :
+	R(Vector.X),
+	G(Vector.Y),
+	B(Vector.Z),
+	A(Vector.W)
+{}
+
 FLinearColor::FLinearColor(const FFloat16Color& C)
 {
 	R = C.R.GetFloat();
diff --git a/Engine/Source/Runtime/Core/Public/Math/Color.h b/Engine/Source/Runtime/Core/Public/Math/Color.h
--- a/Engine/Source/Runtime/Core/Public/Math/Color.h
+++ b/Engine/Source/Runtime/Core/Public/Math/Color.h
@@ -53,6 +53,8 @@ struct FLinearColor
 	CORE_API FLinearColor(const FColor& Color);
 
 	CORE_API FLinearColor(const FVector& Vector);
+
+	CORE_API explicit FLinearColor(const FVector4& Vector);
 	
 	CORE_API explicit FLinearColor(const FFloat16Color& C);
 
