FString::Printf() will now report a deprecation message instead of a hard error when passed a pointer for the formatting string.

diff --git a/Engine/Source/Runtime/Core/Public/Containers/UnrealString.h b/Engine/Source/Runtime/Core/Public/Containers/UnrealString.h
--- a/Engine/Source/Runtime/Core/Public/Containers/UnrealString.h
+++ b/Engine/Source/Runtime/Core/Public/Containers/UnrealString.h
@@ -1363,7 +1363,7 @@ public:
 	 * @returns FString object that was constructed using format and additional parameters.
 	 */
 	template <typename FmtType, typename... Types>
-	static FString Printf(const FmtType& Fmt, Types... Args)
+	static typename TEnableIf<TIsArrayOrRefOfType<FmtType, TCHAR>::Value, FString>::Type Printf(const FmtType& Fmt, Types... Args)
 	{
 		static_assert(TIsArrayOrRefOfType<FmtType, TCHAR>::Value, "Formatting string must be a TCHAR array.");
 		static_assert(TAnd<TIsValidVariadicFunctionArg<Types>...>::Value, "Invalid argument(s) passed to FString::Printf");
@@ -1371,6 +1371,18 @@ public:
 		return PrintfImpl(Fmt, Args...);
 	}
 
+	template <typename FmtType, typename... Types>
+	DEPRECATED(4.20, "The formatting string must now be a TCHAR string literal.")
+	static typename TEnableIf<!TIsArrayOrRefOfType<FmtType, TCHAR>::Value, FString>::Type Printf(const FmtType& Fmt, Types... Args)
+	{
+		// NOTE: When this deprecated function is removed, the return type of the overload above
+		//       should be set to simply FString.
+
+		static_assert(TAnd<TIsValidVariadicFunctionArg<Types>...>::Value, "Invalid argument(s) passed to FString::Printf");
+
+		return PrintfImpl(Fmt, Args...);
+	}
+
 private:
 	static FString VARARGS PrintfImpl(const TCHAR* Fmt, ...);
 public:
