ensureMsgf() will now report a deprecation message instead of a hard error when passed a pointer for the formatting string.

diff --git a/Engine/Source/Runtime/Core/Public/Misc/AssertionMacros.h b/Engine/Source/Runtime/Core/Public/Misc/AssertionMacros.h
--- a/Engine/Source/Runtime/Core/Public/Misc/AssertionMacros.h
+++ b/Engine/Source/Runtime/Core/Public/Misc/AssertionMacros.h
@@ -5,6 +5,7 @@
 #include "CoreTypes.h"
 #include "HAL/PlatformMisc.h"
 #include "Templates/AndOrNot.h"
+#include "Templates/EnableIf.h"
 #include "Templates/IsArrayOrRefOfType.h"
 #include "Templates/IsValidVariadicFunctionArg.h"
 #include "Misc/VarArgs.h"
@@ -85,7 +86,7 @@ public:
 
 	/** Failed assertion handler.  Warning: May be called at library startup time. */
 	template <typename FmtType, typename... Types>
-	static FORCEINLINE bool OptionallyLogFormattedEnsureMessageReturningFalse(bool bLog, const ANSICHAR* Expr, const ANSICHAR* File, int32 Line, const FmtType& FormattedMsg, Types... Args)
+	static FORCEINLINE typename TEnableIf<TIsArrayOrRefOfType<FmtType, TCHAR>::Value, bool>::Type OptionallyLogFormattedEnsureMessageReturningFalse(bool bLog, const ANSICHAR* Expr, const ANSICHAR* File, int32 Line, const FmtType& FormattedMsg, Types... Args)
 	{
 		static_assert(TIsArrayOrRefOfType<FmtType, TCHAR>::Value, "Formatting string must be a TCHAR array.");
 		static_assert(TAnd<TIsValidVariadicFunctionArg<Types>...>::Value, "Invalid argument(s) passed to ensureMsgf");
@@ -93,6 +94,18 @@ public:
 		return OptionallyLogFormattedEnsureMessageReturningFalseImpl(bLog, Expr, File, Line, FormattedMsg, Args...);
 	}
 
+	template <typename FmtType, typename... Types>
+	DEPRECATED(4.20, "The formatting string must now be a TCHAR string literal.")
+	static FORCEINLINE typename TEnableIf<!TIsArrayOrRefOfType<FmtType, TCHAR>::Value, bool>::Type OptionallyLogFormattedEnsureMessageReturningFalse(bool bLog, const ANSICHAR* Expr, const ANSICHAR* File, int32 Line, const FmtType& FormattedMsg, Types... Args)
+	{
+		// NOTE: When this deprecated function is removed, the return type of the overload above
+		//       should be set to simply bool.
+
+		static_assert(TAnd<TIsValidVariadicFunctionArg<Types>...>::Value, "Invalid argument(s) passed to ensureMsgf");
+
+		return OptionallyLogFormattedEnsureMessageReturningFalseImpl(bLog, Expr, File, Line, FormattedMsg, Args...);
+	}
+
 #endif // DO_CHECK || DO_GUARD_SLOW
 
 	/**
