Revert changes made in CL# 3889457, which broke some existing names.

diff --git a/Engine/Source/Runtime/Core/Private/UObject/UnrealNames.cpp b/Engine/Source/Runtime/Core/Private/UObject/UnrealNames.cpp
--- a/Engine/Source/Runtime/Core/Private/UObject/UnrealNames.cpp
+++ b/Engine/Source/Runtime/Core/Private/UObject/UnrealNames.cpp
@@ -328,7 +328,7 @@ FString FName::NameToDisplayString( const FString& InDisplayName, const bool bIs
 		}
 
 		// If the current character is upper case or a digit, and the previous character wasn't, then we need to insert a space if there wasn't one previously
-		if( (bUpperCase || bIsDigit) && (!bInARun || ((CharIndex + 1 < Chars.Num() && FChar::IsLower(Chars[CharIndex + 1])) && !bIsDigit)) && !bWasOpenParen)
+		if( (bUpperCase || bIsDigit) && !bInARun && !bWasOpenParen)
 		{
 			if( !bWasSpace && OutDisplayName.Len() > 0 )
 			{
