#DEV-VR - Wrapping android functions with USE_ANDROID_JNI for Lumin

diff --git a/Engine/Source/Runtime/Core/Private/Android/AndroidMisc.cpp b/Engine/Source/Runtime/Core/Private/Android/AndroidMisc.cpp
--- a/Engine/Source/Runtime/Core/Private/Android/AndroidMisc.cpp
+++ b/Engine/Source/Runtime/Core/Private/Android/AndroidMisc.cpp
@@ -938,9 +938,11 @@ void FAndroidMisc::PrepareMobileHaptics(EMobileHapticsType Type)
 
 void FAndroidMisc::TriggerMobileHaptics()
 {
+#if USE_ANDROID_JNI
 	extern void AndroidThunkCpp_Vibrate(int32 Duration);
 	// tiny little vibration
 	AndroidThunkCpp_Vibrate(10);
+#endif
 }
 
 void FAndroidMisc::ReleaseMobileHaptics()
@@ -950,8 +952,10 @@ void FAndroidMisc::ReleaseMobileHaptics()
 
 void FAndroidMisc::ShareURL(const FString& URL, const FText& Description, int32 LocationHintX, int32 LocationHintY)
 {
+#if USE_ANDROID_JNI
 	extern void AndroidThunkCpp_ShareURL(const FString& URL, const FText& Description, const FText& SharePrompt, int32 LocationHintX, int32 LocationHintY);
 	AndroidThunkCpp_ShareURL(URL, Description, NSLOCTEXT("AndroidMisc", "ShareURL", "Share URL"), LocationHintX, LocationHintY);
+#endif
 }
 
 
@@ -1822,6 +1826,7 @@ bool FAndroidMisc::AreHeadPhonesPluggedIn()
 
 ENetworkConnectionType FAndroidMisc::GetNetworkConnectionType()
 {
+#if USE_ANDROID_JNI
 	extern int32 AndroidThunkCpp_GetNetworkConnectionType();
 
 	switch (AndroidThunkCpp_GetNetworkConnectionType())
@@ -1834,6 +1839,7 @@ ENetworkConnectionType FAndroidMisc::GetNetworkConnectionType()
 		case ANDROIDTHUNK_CONNECTION_TYPE_WIMAX:			return ENetworkConnectionType::WiMAX;
 		case ANDROIDTHUNK_CONNECTION_TYPE_BLUETOOTH:		return ENetworkConnectionType::Bluetooth;
 	}
+#endif
 	return ENetworkConnectionType::Unknown;
 }
 
