Fixed barycentric computation case causing a check. Should really just output an invalid sentinel value. Warns about failures but no longer crashes.

diff --git a/Engine/Source/Runtime/Core/Private/Math/UnrealMath.cpp b/Engine/Source/Runtime/Core/Private/Math/UnrealMath.cpp
--- a/Engine/Source/Runtime/Core/Private/Math/UnrealMath.cpp
+++ b/Engine/Source/Runtime/Core/Private/Math/UnrealMath.cpp
@@ -1964,8 +1964,12 @@ FVector FMath::ComputeBaryCentric2D(const FVector& Point, const FVector& A, cons
 	// Compute the normal of the triangle
 	const FVector TriNorm = (B-A) ^ (C-A);
 
-	//check collinearity of A,B,C
-	check(TriNorm.SizeSquared() > SMALL_NUMBER && "Collinear points in FMath::ComputeBaryCentric2D()");
+	// Check the size of the triangle is reasonable (TriNorm.Size() will be twice the triangle area)
+	if(TriNorm.SizeSquared() <= SMALL_NUMBER)
+	{
+		UE_LOG(LogUnrealMath, Warning, TEXT("Small triangle detected in FMath::ComputeBaryCentric2D(), can't compute valid barycentric coordinate."));
+		return FVector(0.0f, 0.0f, 0.0f);
+	}
 
 	const FVector N = TriNorm.GetSafeNormal();
 
