TBA: Add a fast path for serializing to binary archives that don't require callbacks when leaving slots.

#rnx

diff --git a/Engine/Source/Runtime/Core/Private/Serialization/Formatters/BinaryArchiveFormatter.cpp b/Engine/Source/Runtime/Core/Private/Serialization/Formatters/BinaryArchiveFormatter.cpp
--- a/Engine/Source/Runtime/Core/Private/Serialization/Formatters/BinaryArchiveFormatter.cpp
+++ b/Engine/Source/Runtime/Core/Private/Serialization/Formatters/BinaryArchiveFormatter.cpp
@@ -12,3 +12,7 @@ FBinaryArchiveFormatter::~FBinaryArchiveFormatter()
 {
 }
 
+bool FBinaryArchiveFormatter::RequiresStructuralMetadata() const
+{
+	return false;
+}
diff --git a/Engine/Source/Runtime/Core/Private/Serialization/StructuredArchive.cpp b/Engine/Source/Runtime/Core/Private/Serialization/StructuredArchive.cpp
--- a/Engine/Source/Runtime/Core/Private/Serialization/StructuredArchive.cpp
+++ b/Engine/Source/Runtime/Core/Private/Serialization/StructuredArchive.cpp
@@ -57,6 +57,11 @@ FStructuredArchiveChildReader::~FStructuredArchiveChildReader()
 
 FStructuredArchive::FStructuredArchive(FArchiveFormatterType& InFormatter)
 	: Formatter(InFormatter)
+#if DO_GUARD_SLOW
+	, bRequiresStructuralMetadata(true)
+#else
+	, bRequiresStructuralMetadata(InFormatter.RequiresStructuralMetadata())
+#endif
 	, NextElementId(RootElementId + 1)
 	, CurrentSlotElementId(INDEX_NONE)
 {
@@ -114,26 +119,29 @@ void FStructuredArchive::EnterSlot(int32 ElementId, EElementType ElementType)
 
 void FStructuredArchive::LeaveSlot()
 {
-	switch (CurrentScope.Top().Type)
+	if (bRequiresStructuralMetadata)
 	{
-	case EElementType::Record:
-		Formatter.LeaveField();
-		break;
-	case EElementType::Array:
-		Formatter.LeaveArrayElement();
+		switch (CurrentScope.Top().Type)
+		{
+		case EElementType::Record:
+			Formatter.LeaveField();
+			break;
+		case EElementType::Array:
+			Formatter.LeaveArrayElement();
 #if DO_GUARD_SLOW
-		CurrentContainer.Top()->Index++;
+			CurrentContainer.Top()->Index++;
 #endif
-		break;
-	case EElementType::Stream:
-		Formatter.LeaveStreamElement();
-		break;
-	case EElementType::Map:
-		Formatter.LeaveMapElement();
+			break;
+		case EElementType::Stream:
+			Formatter.LeaveStreamElement();
+			break;
+		case EElementType::Map:
+			Formatter.LeaveMapElement();
 #if DO_GUARD_SLOW
-		CurrentContainer.Top()->Index++;
+			CurrentContainer.Top()->Index++;
 #endif
-		break;
+			break;
+		}
 	}
 }
 
@@ -144,46 +152,54 @@ void FStructuredArchive::SetScope(int32 Depth, int32 ElementId)
 	checkf(CurrentSlotElementId == INDEX_NONE || GetUnderlyingArchive().IsLoading(), TEXT("Cannot change scope until having written a value to the current slot"));
 
 	// Roll back to the correct scope
-	for(int32 CurrentDepth = CurrentScope.Num() - 1; CurrentDepth > Depth; CurrentDepth--)
+	if (bRequiresStructuralMetadata)
 	{
-		// Leave the current element
-		const FElement& Element = CurrentScope[CurrentDepth];
-		switch (Element.Type)
+		for (int32 CurrentDepth = CurrentScope.Num() - 1; CurrentDepth > Depth; CurrentDepth--)
 		{
-		case EElementType::Record:
-			Formatter.LeaveRecord();
+			// Leave the current element
+			const FElement& Element = CurrentScope[CurrentDepth];
+			switch (Element.Type)
+			{
+			case EElementType::Record:
+				Formatter.LeaveRecord();
 #if DO_GUARD_SLOW
-			delete CurrentContainer.Pop(false);
+				delete CurrentContainer.Pop(false);
 #endif
-			break;
-		case EElementType::Array:
+				break;
+			case EElementType::Array:
 #if DO_GUARD_SLOW
-			checkf(GetUnderlyingArchive().IsLoading() || CurrentContainer.Top()->Index == CurrentContainer.Top()->Count, TEXT("Incorrect number of elements serialized in array"));
+				checkf(GetUnderlyingArchive().IsLoading() || CurrentContainer.Top()->Index == CurrentContainer.Top()->Count, TEXT("Incorrect number of elements serialized in array"));
 #endif
-			Formatter.LeaveArray();
+				Formatter.LeaveArray();
 #if DO_GUARD_SLOW
-			delete CurrentContainer.Pop(false);
+				delete CurrentContainer.Pop(false);
 #endif
-			break;
-		case EElementType::Stream:
-			Formatter.LeaveStream();
-			break;
-		case EElementType::Map:
+				break;
+			case EElementType::Stream:
+				Formatter.LeaveStream();
+				break;
+			case EElementType::Map:
 #if DO_GUARD_SLOW
-			checkf(CurrentContainer.Top()->Index == CurrentContainer.Top()->Count, TEXT("Incorrect number of elements serialized in map"));
+				checkf(CurrentContainer.Top()->Index == CurrentContainer.Top()->Count, TEXT("Incorrect number of elements serialized in map"));
 #endif
-			Formatter.LeaveMap();
+				Formatter.LeaveMap();
 #if DO_GUARD_SLOW
-			delete CurrentContainer.Pop(false);
+				delete CurrentContainer.Pop(false);
 #endif
-			break;
-		}
+				break;
+			}
 
-		// Remove the element from the stack
-		CurrentScope.RemoveAt(CurrentDepth, 1, false);
+			// Remove the element from the stack
+			CurrentScope.RemoveAt(CurrentDepth, 1, false);
 
-		// Leave the slot containing it
-		LeaveSlot();
+			// Leave the slot containing it
+			LeaveSlot();
+		}
+	}
+	else
+	{
+		// Remove all the top elements from the stack
+		CurrentScope.RemoveAt(Depth + 1, CurrentScope.Num() - (Depth + 1));
 	}
 }
 
diff --git a/Engine/Source/Runtime/Core/Private/Serialization/StructuredArchiveFormatter.cpp b/Engine/Source/Runtime/Core/Private/Serialization/StructuredArchiveFormatter.cpp
--- a/Engine/Source/Runtime/Core/Private/Serialization/StructuredArchiveFormatter.cpp
+++ b/Engine/Source/Runtime/Core/Private/Serialization/StructuredArchiveFormatter.cpp
@@ -11,6 +11,11 @@ FStructuredArchiveFormatter::~FStructuredArchiveFormatter()
 {
 }
 
+bool FStructuredArchiveFormatter::RequiresStructuralMetadata() const
+{
+	return true;
+}
+
 //////////// Functions ////////////
 
 #if WITH_TEXT_ARCHIVE_SUPPORT
diff --git a/Engine/Source/Runtime/Core/Public/Serialization/Formatters/BinaryArchiveFormatter.h b/Engine/Source/Runtime/Core/Public/Serialization/Formatters/BinaryArchiveFormatter.h
--- a/Engine/Source/Runtime/Core/Public/Serialization/Formatters/BinaryArchiveFormatter.h
+++ b/Engine/Source/Runtime/Core/Public/Serialization/Formatters/BinaryArchiveFormatter.h
@@ -12,6 +12,7 @@ public:
 	FBinaryArchiveFormatter(FArchive& InInner);
 	virtual ~FBinaryArchiveFormatter();
 
+	virtual bool RequiresStructuralMetadata() const override;
 	virtual FArchive& GetUnderlyingArchive() override;
 
 	virtual void EnterRecord() override;
diff --git a/Engine/Source/Runtime/Core/Public/Serialization/StructuredArchive.h b/Engine/Source/Runtime/Core/Public/Serialization/StructuredArchive.h
--- a/Engine/Source/Runtime/Core/Public/Serialization/StructuredArchive.h
+++ b/Engine/Source/Runtime/Core/Public/Serialization/StructuredArchive.h
@@ -334,6 +334,11 @@ private:
 	*/
 	FArchiveFormatterType& Formatter;
 
+	/**
+	 * Whether the formatter requires structural metadata. This allows optimizing the path for binary archives in editor builds.
+	 */
+	const bool bRequiresStructuralMetadata;
+
 #if WITH_TEXT_ARCHIVE_SUPPORT
 	enum class EElementType : unsigned char
 	{
diff --git a/Engine/Source/Runtime/Core/Public/Serialization/StructuredArchiveFormatter.h b/Engine/Source/Runtime/Core/Public/Serialization/StructuredArchiveFormatter.h
--- a/Engine/Source/Runtime/Core/Public/Serialization/StructuredArchiveFormatter.h
+++ b/Engine/Source/Runtime/Core/Public/Serialization/StructuredArchiveFormatter.h
@@ -61,6 +61,8 @@ class CORE_API FStructuredArchiveFormatter
 {
 public:
 	virtual ~FStructuredArchiveFormatter();
+
+	virtual bool RequiresStructuralMetadata() const;
 	
 	virtual FArchive& GetUnderlyingArchive() = 0;
 	virtual FStructuredArchiveFormatter* CreateSubtreeReader() { return this; }
