Add a new allocator (TNonRelocatableInlineAllocator) which eliminates conditional logic from GetAllocation(), at the expense of not being trivially relocatable. This provides a significant performance increase for allocators that are dereferenced frequently, but prevents instances of it being stored inside other allocators (which are expected to be trivially relocatable by default).

diff --git a/Engine/Source/Runtime/Core/Public/Containers/ContainerAllocationPolicies.h b/Engine/Source/Runtime/Core/Public/Containers/ContainerAllocationPolicies.h
--- a/Engine/Source/Runtime/Core/Public/Containers/ContainerAllocationPolicies.h
+++ b/Engine/Source/Runtime/Core/Public/Containers/ContainerAllocationPolicies.h
@@ -505,7 +505,7 @@ public:
 				else
 				{
 					// Reallocate the indirect data for the new size.
-					SecondaryData.ResizeAllocation(PreviousNumElements,NumElements,NumBytesPerElement);
+					SecondaryData.ResizeAllocation(PreviousNumElements, NumElements, NumBytesPerElement);
 				}
 			}
 		}
@@ -515,21 +515,21 @@ public:
 			// If the elements use less space than the inline allocation, only use the inline allocation as slack.
 			return NumElements <= NumInlineElements ?
 				NumInlineElements :
-								  SecondaryData.CalculateSlackReserve(NumElements, NumBytesPerElement);
+				SecondaryData.CalculateSlackReserve(NumElements, NumBytesPerElement);
 		}
 		FORCEINLINE int32 CalculateSlackShrink(int32 NumElements, int32 NumAllocatedElements, int32 NumBytesPerElement) const
 		{
 			// If the elements use less space than the inline allocation, only use the inline allocation as slack.
 			return NumElements <= NumInlineElements ?
-			NumInlineElements :
-							  SecondaryData.CalculateSlackShrink(NumElements, NumAllocatedElements, NumBytesPerElement);
+				NumInlineElements :
+				SecondaryData.CalculateSlackShrink(NumElements, NumAllocatedElements, NumBytesPerElement);
 		}
 		FORCEINLINE int32 CalculateSlackGrow(int32 NumElements, int32 NumAllocatedElements, int32 NumBytesPerElement) const
 		{
 			// If the elements use less space than the inline allocation, only use the inline allocation as slack.
 			return NumElements <= NumInlineElements ?
-			NumInlineElements :
-							  SecondaryData.CalculateSlackGrow(NumElements, NumAllocatedElements, NumBytesPerElement);
+				NumInlineElements :
+				SecondaryData.CalculateSlackGrow(NumElements, NumAllocatedElements, NumBytesPerElement);
 		}
 
 		SIZE_T GetAllocatedSize(int32 NumAllocatedElements, SIZE_T NumBytesPerElement) const
@@ -568,6 +568,147 @@ struct TAllocatorTraits<TInlineAllocator<NumInlineElements, SecondaryAllocator>>
 	enum { SupportsMove = TAllocatorTraits<SecondaryAllocator>::SupportsMove };
 };
 
+/**
+ * Implements a variant of TInlineAllocator with a secondary heap allocator that is allowed to store a pointer to its inline elements.
+ * This allows caching a pointer to the elements which avoids any conditional logic in GetAllocation(), but prevents the allocator being trivially relocatable.
+ * All UE4 allocators typically rely on elements being trivially relocatable, so instances of this allocator cannot be used in other containers.
+ */
+template <uint32 NumInlineElements>
+class TNonRelocatableInlineAllocator
+{
+public:
+
+	enum { NeedsElementType = true };
+	enum { RequireRangeCheck = true };
+
+	template<typename ElementType>
+	class ForElementType
+	{
+	public:
+
+		/** Default constructor. */
+		ForElementType()
+			: Data(GetInlineElements())
+		{
+		}
+
+		/**
+		 * Moves the state of another allocator into this one.
+		 * Assumes that the allocator is currently empty, i.e. memory may be allocated but any existing elements have already been destructed (if necessary).
+		 * @param Other - The allocator to move the state from.  This allocator should be left in a valid empty state.
+		 */
+		FORCEINLINE void MoveToEmpty(ForElementType& Other)
+		{
+			checkSlow(this != &Other);
+
+			if (HasAllocation())
+			{
+				FMemory::Free(Data);
+			}
+
+			if (Other.HasAllocation())
+			{
+				Data = Other.Data;
+				Other.Data = nullptr;
+			}
+			else
+			{
+				Data = GetInlineElements();
+				RelocateConstructItems<ElementType>(GetInlineElements(), Other.GetInlineElements(), NumInlineElements);
+			}
+		}
+
+		// FContainerAllocatorInterface
+		FORCEINLINE ElementType* GetAllocation() const
+		{
+			return Data;
+		}
+
+		void ResizeAllocation(int32 PreviousNumElements,int32 NumElements,SIZE_T NumBytesPerElement)
+		{
+			// Check if the new allocation will fit in the inline data area.
+			if(NumElements <= NumInlineElements)
+			{
+				// If the old allocation wasn't in the inline data area, relocate it into the inline data area.
+				if(HasAllocation())
+				{
+					RelocateConstructItems<ElementType>(GetInlineElements(), Data, PreviousNumElements);
+					FMemory::Free(Data);
+					Data = GetInlineElements();
+				}
+			}
+			else
+			{
+				if (HasAllocation())
+				{
+					// Reallocate the indirect data for the new size.
+					Data = (ElementType*)FMemory::Realloc(Data, NumElements*NumBytesPerElement);
+				}
+				else
+				{
+					// Allocate new indirect memory for the data.
+					Data = (ElementType*)FMemory::Realloc(nullptr, NumElements*NumBytesPerElement);
+
+					// Move the data out of the inline data area into the new allocation.
+					RelocateConstructItems<ElementType>(Data, GetInlineElements(), PreviousNumElements);
+				}
+			}
+		}
+
+		FORCEINLINE int32 CalculateSlackReserve(int32 NumElements, SIZE_T NumBytesPerElement) const
+		{
+			// If the elements use less space than the inline allocation, only use the inline allocation as slack.
+			return HasAllocation() ? DefaultCalculateSlackReserve(NumElements, NumBytesPerElement, true) : NumInlineElements;
+		}
+
+		FORCEINLINE int32 CalculateSlackShrink(int32 NumElements, int32 NumAllocatedElements, int32 NumBytesPerElement) const
+		{
+			// If the elements use less space than the inline allocation, only use the inline allocation as slack.
+			return HasAllocation() ? DefaultCalculateSlackShrink(NumElements, NumAllocatedElements, NumBytesPerElement, true) : NumInlineElements;
+		}
+
+		FORCEINLINE int32 CalculateSlackGrow(int32 NumElements, int32 NumAllocatedElements, int32 NumBytesPerElement) const
+		{
+			// If the elements use less space than the inline allocation, only use the inline allocation as slack.
+			return HasAllocation() ? DefaultCalculateSlackGrow(NumElements, NumAllocatedElements, NumBytesPerElement, true) : NumInlineElements;
+		}
+
+		SIZE_T GetAllocatedSize(int32 NumAllocatedElements, SIZE_T NumBytesPerElement) const
+		{
+			return HasAllocation()? (NumAllocatedElements * NumBytesPerElement) : 0;
+		}
+
+		FORCEINLINE bool HasAllocation() const
+		{
+			return Data != GetInlineElements();
+		}
+
+	private:
+		ForElementType(const ForElementType&) = delete;
+		ForElementType& operator=(const ForElementType&) = delete;
+
+		/** The data is allocated through the indirect allocation policy if more than NumInlineElements is needed. */
+		ElementType* Data;
+
+		/** The data is stored in this array if less than NumInlineElements is needed. */
+		TTypeCompatibleBytes<ElementType> InlineData[NumInlineElements];
+
+		/** @return the base of the aligned inline element data */
+		FORCEINLINE ElementType* GetInlineElements() const
+		{
+			return (ElementType*)InlineData;
+		}
+	};
+
+	typedef void ForAnyElementType;
+};
+
+template <uint32 NumInlineElements>
+struct TAllocatorTraits<TNonRelocatableInlineAllocator<NumInlineElements>> : TAllocatorTraitsBase<TNonRelocatableInlineAllocator<NumInlineElements>>
+{
+	enum { SupportsMove = true };
+};
+
 /**
  * The fixed allocation policy allocates up to a specified number of elements in the same allocation as the container.
  * It's like the inline allocator, except it doesn't provide secondary storage when the inline storage has been filled.
