TBA: Use the trivially relocatable allocator with structured archives, to reduce time taken allocating memory in the constructor.

#rnx

diff --git a/Engine/Source/Runtime/Core/Public/Serialization/StructuredArchive.h b/Engine/Source/Runtime/Core/Public/Serialization/StructuredArchive.h
--- a/Engine/Source/Runtime/Core/Public/Serialization/StructuredArchive.h
+++ b/Engine/Source/Runtime/Core/Public/Serialization/StructuredArchive.h
@@ -384,7 +384,7 @@ private:
 	 * Tracks the current stack of objects being written. Used by SetScope() to ensure that scopes are always closed correctly in the underlying formatter,
 	 * and to make sure that the archive is always written in a forwards-only way (ie. writing to an element id that is not in scope will assert)
 	 */
-	TArray<FElement> CurrentScope;
+	TArray<FElement, TNonRelocatableInlineAllocator<32>> CurrentScope;
 
 #if DO_GUARD_SLOW
 	/**
