Move checks for unique field names behind an additional define, to avoid performance hit when compiling blueprints.

#rnx

diff --git a/Engine/Source/Runtime/Core/Private/Serialization/StructuredArchive.cpp b/Engine/Source/Runtime/Core/Private/Serialization/StructuredArchive.cpp
--- a/Engine/Source/Runtime/Core/Private/Serialization/StructuredArchive.cpp
+++ b/Engine/Source/Runtime/Core/Private/Serialization/StructuredArchive.cpp
@@ -9,11 +9,16 @@
 //////////// FStructuredArchive::FContainer ////////////
 
 #if DO_GUARD_SLOW
+
+#define CHECK_UNIQUE_FIELD_NAMES 0
+
 struct FStructuredArchive::FContainer
 {
 	int Index;
 	int Count;
+#if CHECK_UNIQUE_FIELD_NAMES
 	TSet<FString> KeyNames;
+#endif
 
 	FContainer(int InCount)
 		: Index(0)
@@ -385,12 +390,14 @@ FStructuredArchive::FSlot FStructuredArchive::FRecord::EnterField(FArchiveFieldN
 	Ar.CurrentSlotElementId = Ar.NextElementId++;
 
 #if DO_GUARD_SLOW
+#if CHECK_UNIQUE_FIELD_NAMES
 	if (!Ar.GetUnderlyingArchive().IsLoading())
 	{
 		FContainer& Container = *Ar.CurrentContainer.Top();
 		checkf(!Container.KeyNames.Contains(Name.Name), TEXT("Multiple keys called '%s' serialized into record"), Name.Name);
 		Container.KeyNames.Add(Name.Name);
 	}
+#endif
 #endif
 
 	Ar.Formatter.EnterField(Name);
@@ -423,12 +430,14 @@ TOptional<FStructuredArchive::FSlot> FStructuredArchive::FRecord::TryEnterField(
 	Ar.SetScope(Depth, ElementId);
 
 #if DO_GUARD_SLOW
+#if CHECK_UNIQUE_FIELD_NAMES
 	if (!GetUnderlyingArchive().IsLoading())
 	{
 		FContainer& Container = *Ar.CurrentContainer.Top();
 		checkf(!Container.KeyNames.Contains(Name.Name), TEXT("Multiple keys called '%s' serialized into record"), Name.Name);
 		Container.KeyNames.Add(Name.Name);
 	}
+#endif
 #endif
 
 	if (Ar.Formatter.TryEnterField(Name, bEnterWhenWriting))
@@ -485,23 +494,27 @@ FStructuredArchive::FSlot FStructuredArchive::FMap::EnterElement(FString& Name)
 	Ar.CurrentSlotElementId = Ar.NextElementId++;
 
 #if DO_GUARD_SLOW
+#if CHECK_UNIQUE_FIELD_NAMES
 	if(Ar.GetUnderlyingArchive().IsSaving())
 	{
 		FContainer& Container = *Ar.CurrentContainer.Top();
 		checkf(!Container.KeyNames.Contains(Name), TEXT("Multiple keys called '%s' serialized into record"), *Name);
 		Container.KeyNames.Add(Name);
 	}
+#endif
 #endif
 
 	Ar.Formatter.EnterMapElement(Name);
 
 #if DO_GUARD_SLOW
+#if CHECK_UNIQUE_FIELD_NAMES
 	if(Ar.GetUnderlyingArchive().IsLoading())
 	{
 		FContainer& Container = *Ar.CurrentContainer.Top();
 		checkf(!Container.KeyNames.Contains(Name), TEXT("Multiple keys called '%s' serialized into record"), *Name);
 		Container.KeyNames.Add(Name);
 	}
+#endif
 #endif
 
 	return FSlot(Ar, Depth, Ar.CurrentSlotElementId);
