Linux CIS Fixes for TOptional

#trivial

diff --git a/Engine/Source/Runtime/Core/Public/Misc/Optional.h b/Engine/Source/Runtime/Core/Public/Misc/Optional.h
--- a/Engine/Source/Runtime/Core/Public/Misc/Optional.h
+++ b/Engine/Source/Runtime/Core/Public/Misc/Optional.h
@@ -6,8 +6,8 @@
 #include "Misc/AssertionMacros.h"
 #include "Templates/TypeCompatibleBytes.h"
 #include "Templates/UnrealTemplate.h"
+#include "Serialization/Archive.h"
 
-class FArchive;
 
 /**
  * When we have an optional value IsSet() returns true, and GetValue() is meaningful.
@@ -148,11 +148,11 @@ public:
 
 	friend FArchive& operator<<(FArchive& Ar, TOptional& Optional)
 	{
-		bool bIsSet = Optional.bIsSet;
-		Ar << bIsSet;
+		bool bOptionalIsSet = Optional.bIsSet;
+		Ar << bOptionalIsSet;
 		if (Ar.IsLoading())
 		{
-			if (bIsSet)
+			if (bOptionalIsSet)
 			{
 				if (!Optional.bIsSet)
 				{
@@ -167,7 +167,7 @@ public:
 		}
 		else
 		{
-			if (bIsSet)
+			if (bOptionalIsSet)
 			{
 				Ar << Optional.GetValue();
 			}
