Missing header in Monolithic builds

Fixed indent issues, was using spaces vs tabs

diff --git a/Engine/Source/Runtime/Core/Private/Unix/UnixPlatformCrashContext.cpp b/Engine/Source/Runtime/Core/Private/Unix/UnixPlatformCrashContext.cpp
--- a/Engine/Source/Runtime/Core/Private/Unix/UnixPlatformCrashContext.cpp
+++ b/Engine/Source/Runtime/Core/Private/Unix/UnixPlatformCrashContext.cpp
@@ -25,6 +25,7 @@
 #include "HAL/PlatformMallocCrash.h"
 #include "Unix/UnixPlatformRunnableThread.h"
 #include "HAL/ExceptionHandling.h"
+#include "Stats/Stats.h"
 
 #include "HAL/ThreadHeartBeat.h"
 
@@ -303,7 +304,7 @@ namespace UnixCrashReporterTracker
 
 	bool Tick(float DeltaTime)
 	{
-        QUICK_SCOPE_CYCLE_COUNTER(STAT_UnixCrashReporterTracker_Tick);
+		QUICK_SCOPE_CYCLE_COUNTER(STAT_UnixCrashReporterTracker_Tick);
 
 		if (!FPlatformProcess::IsProcRunning(CurrentlyRunningCrashReporter))
 		{
