Reverted change that delays crash reporting initialization for CEF on Mac

diff --git a/Engine/Source/Runtime/Core/Private/Mac/MacPlatformMisc.cpp b/Engine/Source/Runtime/Core/Private/Mac/MacPlatformMisc.cpp
--- a/Engine/Source/Runtime/Core/Private/Mac/MacPlatformMisc.cpp
+++ b/Engine/Source/Runtime/Core/Private/Mac/MacPlatformMisc.cpp
@@ -1435,7 +1435,7 @@ void FMacPlatformMisc::SetCrashHandler(void (* CrashHandler)(const FGenericCrash
 	
 	GCrashHandlerPointer = CrashHandler;
 
-	if(!FMacApplicationInfo::CrashReporter && !GCrashMalloc && !CrashHandler)
+	if (!FMacApplicationInfo::CrashReporter && !GCrashMalloc)
 	{
 		// Configure the crash handler malloc zone to reserve some VM space for itself
 		GCrashMalloc = new FMacMallocCrashHandler( 128 * 1024 * 1024 );
