Fixed InitializeAvailableCultures missing some languages

diff --git a/Engine/Source/Runtime/Core/Private/Internationalization/ICUInternationalization.cpp b/Engine/Source/Runtime/Core/Private/Internationalization/ICUInternationalization.cpp
--- a/Engine/Source/Runtime/Core/Private/Internationalization/ICUInternationalization.cpp
+++ b/Engine/Source/Runtime/Core/Private/Internationalization/ICUInternationalization.cpp
@@ -351,6 +351,20 @@ void FICUInternationalization::InitializeAvailableCultures()
 		}
 	}
 
+	// getAvailableLocales doesn't always cover all languages that ICU supports, so we spin that list too and add any that were missed
+	// Note: getISOLanguages returns the start of an array of const char* null-terminated strings, with a null string signifying the end of the array
+	for (const char* const* AvailableLanguages = icu::Locale::getISOLanguages(); *AvailableLanguages; ++AvailableLanguages)
+	{
+		FString LanguageCode = UTF8_TO_TCHAR(*AvailableLanguages);
+		LanguageCode.ToLowerInline();
+
+		// Only care about 2-letter codes
+		if (LanguageCode.Len() == 2)
+		{
+			AppendCultureData(LanguageCode, FString(), FString());
+		}
+	}
+
 	// Also add our invariant culture if it wasn't found when processing the ICU locales
 	if (!AllAvailableCulturesMap.Contains(TEXT("en-US-POSIX")))
 	{
