Cache files that are invalid or the wrong case sensitivity

diff --git a/Engine/Source/Runtime/Core/Private/Unix/UnixPlatformFile.cpp b/Engine/Source/Runtime/Core/Private/Unix/UnixPlatformFile.cpp
--- a/Engine/Source/Runtime/Core/Private/Unix/UnixPlatformFile.cpp
+++ b/Engine/Source/Runtime/Core/Private/Unix/UnixPlatformFile.cpp
@@ -4,6 +4,7 @@
 #include "GenericPlatform/GenericPlatformProcess.h"
 #include "HAL/PlatformTime.h"
 #include "Containers/StringConv.h"
+#include "Containers/LruCache.h"
 #include "Logging/LogMacros.h"
 #include "Misc/Paths.h"
 #include <sys/file.h>
@@ -268,6 +269,73 @@ private:
 	friend class FUnixFileRegistry;
 };
 
+extern int32 CORE_API GMaxNumberFileMappingCache;
+namespace
+{
+	double MaxInvalidCacheTime = 0.5; // 500ms
+	struct FileEntry
+	{
+		FString File;
+		bool bInvalid = false;
+		double CacheTime = 0.0f;
+		bool IsInvalid() const
+		{
+			double Current = FPlatformTime::Seconds();
+			return bInvalid && Current - CacheTime <= MaxInvalidCacheTime;
+		}
+	};
+	struct FileMapCache
+	{
+		virtual const FileEntry* Find(const FString& Key) = 0;
+		virtual void AddEntry(const FString& Key, const FString& Elem) = 0;
+		virtual void Invalidate(const FString& Key) = 0;
+	};
+	class FileMapCacheDummy : public FileMapCache
+	{
+	public:
+		const FileEntry* Find(const FString& Key) override
+		{
+			return nullptr;
+		}
+		void AddEntry(const FString& Key, const FString& Elem) override
+		{ }
+		void Invalidate(const FString& Key) override
+		{ }
+	};
+	class FileMapCacheDefault : public FileMapCache
+	{
+	public:
+		FileMapCacheDefault()
+			: Cache(GMaxNumberFileMappingCache)
+		{
+		}
+		const FileEntry* Find(const FString& Key) override
+		{
+			FScopeLock ScopeLock(&Mutex);
+			return Cache.FindAndTouch(Key);
+		}
+		void AddEntry(const FString& Key, const FString& Elem) override
+		{
+			FScopeLock ScopeLock(&Mutex);
+			Cache.Add(Key, {Elem, Elem.IsEmpty(), FPlatformTime::Seconds()});
+		}
+		void Invalidate(const FString& Key) override
+		{
+			FScopeLock ScopeLock(&Mutex);
+			Cache.Remove(Key);
+		}
+	private:
+		FCriticalSection Mutex;
+		TLruCache<FString, FileEntry> Cache;
+	};
+	FileMapCache& GetFileMapCache()
+	{
+		static FileMapCacheDefault DefaultCache;
+		static FileMapCacheDummy DummyCache;
+		return GMaxNumberFileMappingCache > 0 ? *static_cast<FileMapCache*>(&DefaultCache) : *static_cast<FileMapCache*>(&DummyCache);
+	}
+}
+
 /**
  * A class to handle case insensitive file opening. This is a band-aid, non-performant approach,
  * without any caching.
@@ -446,16 +514,39 @@ public:
 		}
 		else
 		{
-			// perform a case-insensitive search from /
+			FileMapCache& Cache = GetFileMapCache();
+			const FileEntry* Entry = Cache.Find(PossiblyWrongFilename);
 
-			int MaxPathComponents = CountPathComponents(PossiblyWrongFilename);
-			if (MaxPathComponents > 0)
+			if (Entry != nullptr)
+			{
+				if (Entry->IsInvalid())
+				{
+					bFound = false;
+				}
+				else
+				{
+					ExistingFilename = Entry->File;
+					bFound = true;
+				}
+			}
+			else
 			{
-				FString FoundFilename(TEXT("/"));	// start with root
-				bFound = MapFileRecursively(PossiblyWrongFilename, 0, MaxPathComponents, FoundFilename);
-				if (bFound)
+				// perform a case-insensitive search from /
+				int MaxPathComponents = CountPathComponents(PossiblyWrongFilename);
+				if (MaxPathComponents > 0)
 				{
-					ExistingFilename = FoundFilename;
+					FString FoundFilename(TEXT("/"));	// start with root
+					bFound = MapFileRecursively(PossiblyWrongFilename, 0, MaxPathComponents, FoundFilename);
+					if (bFound)
+					{
+						ExistingFilename = FoundFilename;
+						Cache.AddEntry(PossiblyWrongFilename, ExistingFilename);
+					}
+					else
+					{
+						// Cache a failed to find entry, we'll look again if the call in is greater then MaxInvalidCacheTime from this cache point
+						Cache.AddEntry(PossiblyWrongFilename, "");
+					}
 				}
 			}
 		}
@@ -589,6 +680,8 @@ bool FUnixPlatformFile::DeleteFile(const TCHAR* Filename)
 		return false;
 	}
 
+	GetFileMapCache().Invalidate(IntendedFilename);
+
 	// removing mapped file is too dangerous
 	if (IntendedFilename != CaseSensitiveFilename)
 	{
@@ -618,12 +711,15 @@ bool FUnixPlatformFile::IsReadOnly(const TCHAR* Filename)
 bool FUnixPlatformFile::MoveFile(const TCHAR* To, const TCHAR* From)
 {
 	FString CaseSensitiveFilename;
-	if (!GCaseInsensMapper.MapCaseInsensitiveFile(NormalizeFilename(From, true), CaseSensitiveFilename))
+	FString IntendedFilename = NormalizeFilename(From, true);
+	if (!GCaseInsensMapper.MapCaseInsensitiveFile(IntendedFilename, CaseSensitiveFilename))
 	{
 		// could not find the file
 		return false;
 	}
 
+	GetFileMapCache().Invalidate(IntendedFilename);
+
 	int32 Result = rename(TCHAR_TO_UTF8(*CaseSensitiveFilename), TCHAR_TO_UTF8(*NormalizeFilename(To, true)));
 	if (Result == -1 && errno == EXDEV)
 	{
@@ -768,6 +864,9 @@ IFileHandle* FUnixPlatformFile::OpenWrite(const TCHAR* Filename, bool bAppend, b
 		Flags |= O_WRONLY;
 	}
 
+	// We may have cached this as an invalid file, so lets just remove a newly created file from the cache
+	GetFileMapCache().Invalidate(FString(Filename));
+
 	// create directories if needed.
 	if (!CreateDirectoriesFromPath(Filename))
 	{
@@ -850,6 +949,8 @@ bool FUnixPlatformFile::DeleteDirectory(const TCHAR* Directory)
 		return false;
 	}
 
+	GetFileMapCache().Invalidate(IntendedFilename);
+
 	// removing mapped directory is too dangerous
 	if (IntendedFilename != CaseSensitiveFilename)
 	{
diff --git a/Engine/Source/Runtime/Core/Private/Unix/UnixPlatformMemory.cpp b/Engine/Source/Runtime/Core/Private/Unix/UnixPlatformMemory.cpp
--- a/Engine/Source/Runtime/Core/Private/Unix/UnixPlatformMemory.cpp
+++ b/Engine/Source/Runtime/Core/Private/Unix/UnixPlatformMemory.cpp
@@ -45,6 +45,18 @@ bool CORE_API GSuppressDwarfParsing    = false;
 // Used in UnixPlatformStackwalk to skip the crash handling callstack frames.
 bool CORE_API GFullCrashCallstack = false;
 
+// Used to set the maximum number of file mappings.
+#if UE_EDITOR
+int32 CORE_API GMaxNumberFileMappingCache = 10000;
+#else
+int32 CORE_API GMaxNumberFileMappingCache = 100;
+#endif
+namespace
+{
+	// The max allowed to be set for the caching
+	const int32 MaximumAllowedMaxNumFileMappingCache = 1000000;
+}
+
 void FUnixPlatformMemory::Init()
 {
 	FGenericPlatformMemory::Init();
@@ -137,6 +149,13 @@ class FMalloc* FUnixPlatformMemory::BaseAllocator()
 					GFullCrashCallstack = true;
 				}
 
+				const char FileMapCacheCmd[] = "-filemapcachesize=";
+				if (const char* Cmd = FCStringAnsi::Stristr(Arg, FileMapCacheCmd))
+				{
+					int32 Max = FCStringAnsi::Atoi(Cmd + sizeof(FileMapCacheCmd) - 1);
+					GMaxNumberFileMappingCache = FMath::Clamp(Max, 0, MaximumAllowedMaxNumFileMappingCache);
+				}
+
 #if UE_USE_MALLOC_REPLAY_PROXY
 				if (FCStringAnsi::Stricmp(Arg, "-mallocsavereplay") == 0)
 				{
diff --git a/Engine/Source/Runtime/Core/Private/Unix/UnixPlatformMisc.cpp b/Engine/Source/Runtime/Core/Private/Unix/UnixPlatformMisc.cpp
--- a/Engine/Source/Runtime/Core/Private/Unix/UnixPlatformMisc.cpp
+++ b/Engine/Source/Runtime/Core/Private/Unix/UnixPlatformMisc.cpp
@@ -134,6 +134,7 @@ void FUnixPlatformMisc::PlatformInit()
 	UE_LOG(LogInit, Log, TEXT(" -ansimalloc - use malloc()/free() from libc (useful for tools like valgrind and electric fence)"));
 	UE_LOG(LogInit, Log, TEXT(" -jemalloc - use jemalloc for all memory allocation"));
 	UE_LOG(LogInit, Log, TEXT(" -binnedmalloc - use binned malloc  for all memory allocation"));
+	UE_LOG(LogInit, Log, TEXT(" -filemapcachesize=NUMBER - set the size for case-sensitive file mapping cache"));
 
 	// [RCL] FIXME: this should be printed in specific modules, if at all
 	UE_LOG(LogInit, Log, TEXT(" -httpproxy=ADDRESS:PORT - redirects HTTP requests to a proxy (only supported if compiled with libcurl)"));
