Avoid printing when in a signal handler.
Put that off until the end

diff --git a/Engine/Source/Runtime/Core/Private/Unix/UnixPlatformCrashContext.cpp b/Engine/Source/Runtime/Core/Private/Unix/UnixPlatformCrashContext.cpp
--- a/Engine/Source/Runtime/Core/Private/Unix/UnixPlatformCrashContext.cpp
+++ b/Engine/Source/Runtime/Core/Private/Unix/UnixPlatformCrashContext.cpp
@@ -120,11 +120,15 @@ void FUnixCrashContext::InitFromEnsureHandler(const TCHAR* EnsureMessage, const
 	}
 }
 
+volatile sig_atomic_t GEnteredSignalHandler = 0;
+
 /**
  * Handles graceful termination. Gives time to exit gracefully, but second signal will quit immediately.
  */
 void GracefulTerminationHandler(int32 Signal, siginfo_t* Info, void* Context)
 {
+	GEnteredSignalHandler = 1;
+
 	// do not flush logs at this point; this can result in a deadlock if the signal was received while we were holding lock in the malloc (flushing allocates memory)
 	if( !GIsRequestingExit )
 	{
@@ -134,6 +138,8 @@ void GracefulTerminationHandler(int32 Signal, siginfo_t* Info, void* Context)
 	{
 		FPlatformMisc::RequestExit(true);
 	}
+
+	GEnteredSignalHandler = 0;
 }
 
 void CreateExceptionInfoString(int32 Signal, siginfo_t* Info, ucontext_t *Context)
diff --git a/Engine/Source/Runtime/Core/Private/Unix/UnixPlatformMisc.cpp b/Engine/Source/Runtime/Core/Private/Unix/UnixPlatformMisc.cpp
--- a/Engine/Source/Runtime/Core/Private/Unix/UnixPlatformMisc.cpp
+++ b/Engine/Source/Runtime/Core/Private/Unix/UnixPlatformMisc.cpp
@@ -148,8 +148,24 @@ void FUnixPlatformMisc::PlatformInit()
 	}
 }
 
+volatile sig_atomic_t GDeferedExitLogging = 0;
+
 void FUnixPlatformMisc::PlatformTearDown()
 {
+	// We requested to close from a signal so we couldnt print.
+	if (GDeferedExitLogging)
+	{
+		uint8 OverriddenErrorLevel = 0;
+		if (FPlatformMisc::HasOverriddenReturnCode(&OverriddenErrorLevel))
+		{
+			UE_LOG(LogCore, Log, TEXT("FUnixPlatformMisc::RequestExit(bForce=false, ReturnCode=%d)"), OverriddenErrorLevel);
+		}
+		else
+		{
+			UE_LOG(LogCore, Log, TEXT("FUnixPlatformMisc::RequestExit(false)"));
+		}
+	}
+
 	FPlatformProcess::CeaseBeingFirstInstance();
 }
 
@@ -188,13 +204,25 @@ void FUnixPlatformMisc::LowLevelOutputDebugString(const TCHAR *Message)
 	fprintf(stderr, "%ls", Message);	// there's no good way to implement that really
 }
 
+extern volatile sig_atomic_t GEnteredSignalHandler;
 uint8 GOverriddenReturnCode = 0;
 bool GHasOverriddenReturnCode = false;
 
 void FUnixPlatformMisc::RequestExit(bool Force)
 {
-	UE_LOG(LogCore, Log,  TEXT("FUnixPlatformMisc::RequestExit(%i)"), Force );
-	if( Force )
+	if (GEnteredSignalHandler)
+	{
+		// Still log something but use a signal-safe function
+		const ANSICHAR ExitMsg[] = "FUnixPlatformMisc::RequestExit\n";
+		write(STDOUT_FILENO, ExitMsg, sizeof(ExitMsg));
+		GDeferedExitLogging = 1;
+	}
+	else
+	{
+		UE_LOG(LogCore, Log,  TEXT("FUnixPlatformMisc::RequestExit(%i)"), Force );
+	}
+
+	if(Force)
 	{
 		// Force immediate exit. Cannot call abort() here, because abort() raises SIGABRT which we treat as crash
 		// (to prevent other, particularly third party libs, from quitting without us noticing)
@@ -209,13 +237,31 @@ void FUnixPlatformMisc::RequestExit(bool Force)
 		}
 	}
 
-	// Tell the platform specific code we want to exit cleanly from the main loop.
-	FGenericPlatformMisc::RequestExit(Force);
+	if (GEnteredSignalHandler)
+	{
+		// Lets set our selfs to request exit as the generic platform request exit could log
+		GIsRequestingExit = 1;
+	}
+	else
+	{
+		// Tell the platform specific code we want to exit cleanly from the main loop.
+		FGenericPlatformMisc::RequestExit(Force);
+	}
 }
 
 void FUnixPlatformMisc::RequestExitWithStatus(bool Force, uint8 ReturnCode)
 {
-	UE_LOG(LogCore, Log, TEXT("FUnixPlatformMisc::RequestExit(bForce=%s, ReturnCode=%d)"), Force ? TEXT("true") : TEXT("false"), ReturnCode);
+	if (GEnteredSignalHandler)
+	{
+		// Still log something but use a signal-safe function
+		const ANSICHAR ExitMsg[] = "FUnixPlatformMisc::RequestExitWithStatus\n";
+		write(STDOUT_FILENO, ExitMsg, sizeof(ExitMsg));
+		GDeferedExitLogging = 1;
+	}
+	else
+	{
+		UE_LOG(LogCore, Log, TEXT("FUnixPlatformMisc::RequestExit(bForce=%s, ReturnCode=%d)"), Force ? TEXT("true") : TEXT("false"), ReturnCode);
+	}
 
 	GOverriddenReturnCode = ReturnCode;
 	GHasOverriddenReturnCode = true;
