Fix for adding a new module to a project.

diff --git a/Engine/Source/Runtime/Core/Private/Modules/ModuleManager.cpp b/Engine/Source/Runtime/Core/Private/Modules/ModuleManager.cpp
--- a/Engine/Source/Runtime/Core/Private/Modules/ModuleManager.cpp
+++ b/Engine/Source/Runtime/Core/Private/Modules/ModuleManager.cpp
@@ -466,7 +466,17 @@ IModuleInterface* FModuleManager::LoadModuleWithFailureReason(const FName InModu
 
 			if (ModuleInfo->Filename.IsEmpty())
 			{
-				UE_LOG(LogModuleManager, Warning, TEXT("No filename provided for module %s"), *InModuleName.ToString());
+				TMap<FName, FString> ModulePathMap;
+				FindModulePaths(*InModuleName.ToString(), ModulePathMap);
+
+				if (ModulePathMap.Num() != 1)
+				{
+					UE_LOG(LogModuleManager, Warning, TEXT("ModuleManager: Unable to load module '%s'  - %d instances of that module name found."), *InModuleName.ToString(), ModulePathMap.Num());
+					OutFailureReason = EModuleLoadResult::FileNotFound;
+					return nullptr;
+				}
+
+				ModuleInfo->Filename = MoveTemp(TMap<FName, FString>::TIterator(ModulePathMap).Value());
 			}
 
 			// Determine which file to load for this module.
@@ -1070,7 +1080,7 @@ bool FModuleManager::LoadModuleWithCallback( const FName InModuleName, FOutputDe
 
 void FModuleManager::MakeUniqueModuleFilename( const FName InModuleName, FString& UniqueSuffix, FString& UniqueModuleFileName ) const
 {
-	auto Module = FindModuleChecked(InModuleName);
+	TSharedRef<const FModuleInfo, ESPMode::ThreadSafe> Module = FindModuleChecked(InModuleName);
 
 	IFileManager& FileManager = IFileManager::Get();
 
@@ -1082,7 +1092,7 @@ void FModuleManager::MakeUniqueModuleFilename( const FName InModuleName, FString
 		const FString ModuleName = InModuleName.ToString();
 		const int32 MatchPos = Module->OriginalFilename.Find(ModuleName, ESearchCase::IgnoreCase, ESearchDir::FromEnd);
 
-		if (ensure(MatchPos != INDEX_NONE))
+		if (MatchPos != INDEX_NONE)
 		{
 			const int32 SuffixPos = MatchPos + ModuleName.Len();
 			UniqueModuleFileName = FString::Printf( TEXT( "%s-%s%s" ),
