Refactor LargeMemoryWriter code so that the DuplicateDataArchive could reuse the large alloc managment.

diff --git a/Engine/Source/Runtime/Core/Private/Serialization/LargeMemoryData.cpp b/Engine/Source/Runtime/Core/Private/Serialization/LargeMemoryData.cpp
new file mode 100644
--- /dev/null
+++ b/Engine/Source/Runtime/Core/Private/Serialization/LargeMemoryData.cpp
@@ -0,0 +1,90 @@
+// Copyright 1998-2018 Epic Games, Inc. All Rights Reserved.
+
+#include "Serialization/LargeMemoryData.h"
+#include "Logging/LogMacros.h"
+#include "CoreGlobals.h"
+
+/*----------------------------------------------------------------------------
+	FLargeMemoryData
+----------------------------------------------------------------------------*/
+
+FLargeMemoryData::FLargeMemoryData(const int64 PreAllocateBytes)
+	: Data(nullptr)
+	, NumBytes(FMath::Max<int64>(0, PreAllocateBytes))
+	, MaxBytes(0)
+{
+	GrowBuffer();
+}
+
+FLargeMemoryData::~FLargeMemoryData()
+{
+	if (Data)
+	{
+		FMemory::Free(Data);
+	}
+}
+
+bool FLargeMemoryData::Write(void* InData, int64 InOffset, int64 InNum)
+{
+	// Allow InData to be null if InNum == 0.
+	if (InOffset >= 0 && (InData ? InNum >= 0 : InNum == 0))
+	{
+		// Grow the buffer to the offset position even if InNum == 0.
+		NumBytes = FMath::Max<int64>(NumBytes, InOffset + InNum);
+		if (NumBytes > MaxBytes)
+		{
+			GrowBuffer();
+		}
+
+		if (InNum)
+		{
+			FMemory::Memcpy(&Data[InOffset], InData, InNum);
+		}
+
+		return true;
+	}
+	else
+	{
+		return false;
+	}
+}
+
+bool FLargeMemoryData::Read(void* OutData, int64 InOffset, int64 InNum) const
+{
+	// Allow OutData to be null if InNum == 0.
+	if (InOffset >= 0 && (OutData ? InNum >= 0 : InNum == 0) && (InOffset + InNum <= NumBytes))
+	{
+		if (InNum)
+		{
+			FMemory::Memcpy(OutData, &Data[InOffset], InNum);
+		}
+
+		return true;
+	}
+	else
+	{
+		return false;
+	}
+}
+
+void FLargeMemoryData::ReleaseOwnership()
+{
+	Data = nullptr;
+	NumBytes = 0;
+	MaxBytes = 0;
+}
+
+void FLargeMemoryData::GrowBuffer()
+{
+	// Allocate slack proportional to the buffer size. Min 64 KB
+	MaxBytes = FMath::Max<int64>(64 * 1024, FMemory::QuantizeSize(NumBytes + 3 * NumBytes / 8 + 16));
+
+	if (Data)
+	{
+		Data = (uint8*)FMemory::Realloc(Data, MaxBytes);
+	}
+	else
+	{
+		Data = (uint8*)FMemory::Malloc(MaxBytes);
+	}
+}
diff --git a/Engine/Source/Runtime/Core/Private/Serialization/LargeMemoryWriter.cpp b/Engine/Source/Runtime/Core/Private/Serialization/LargeMemoryWriter.cpp
--- a/Engine/Source/Runtime/Core/Private/Serialization/LargeMemoryWriter.cpp
+++ b/Engine/Source/Runtime/Core/Private/Serialization/LargeMemoryWriter.cpp
@@ -10,37 +10,19 @@
 
 FLargeMemoryWriter::FLargeMemoryWriter(const int64 PreAllocateBytes, bool bIsPersistent, const TCHAR* InFilename)
 	: FMemoryArchive()
-	, Data(nullptr)
-	, NumBytes(0)
-	, MaxBytes(0)
+	, Data(PreAllocateBytes)
 	, ArchiveName(InFilename ? InFilename : TEXT("FLargeMemoryWriter"))
 {
 	this->SetIsSaving(true);
 	this->SetIsPersistent(bIsPersistent);
-	GrowBuffer(PreAllocateBytes);
 }
 
 void FLargeMemoryWriter::Serialize(void* InData, int64 Num)
 {
-	UE_CLOG(!Data, LogSerialization, Fatal, TEXT("Tried to serialize data to an FLargeMemoryWriter that was already released. Archive name: %s."), *ArchiveName);
-	
-	const int64 NumBytesToAdd = Offset + Num - NumBytes;
-	if (NumBytesToAdd > 0)
-	{
-		const int64 NewByteCount = NumBytes + NumBytesToAdd;
-		if (NewByteCount > MaxBytes)
-		{
-			GrowBuffer(NewByteCount);
-		}
-
-		NumBytes = NewByteCount;
-	}
-
-	check((Offset + Num) <= NumBytes);
+	UE_CLOG(!Data.HasData(), LogSerialization, Fatal, TEXT("Tried to serialize data to an FLargeMemoryWriter that was already released. Archive name: %s."), *ArchiveName);
 
-	if (Num)
+	if (Data.Write(InData, Offset, Num))
 	{
-		FMemory::Memcpy(&Data[Offset], InData, Num);
 		Offset += Num;
 	}
 }
@@ -50,52 +32,8 @@ FString FLargeMemoryWriter::GetArchiveName() const
 	return ArchiveName;
 }
 
-int64 FLargeMemoryWriter::TotalSize()
-{
-	return NumBytes;
-}
-
 uint8* FLargeMemoryWriter::GetData() const
 {
-	UE_CLOG(!Data, LogSerialization, Warning, TEXT("Tried to get written data from an FLargeMemoryWriter that was already released. Archive name: %s."), *ArchiveName);
-
-	return Data;
+	UE_CLOG(!Data.HasData(), LogSerialization, Warning, TEXT("Tried to get written data from an FLargeMemoryWriter that was already released. Archive name: %s."), *ArchiveName);
+	return const_cast<uint8*>(Data.GetData());
 }
-
-void FLargeMemoryWriter::ReleaseOwnership()
-{
-	Data = nullptr;
-	NumBytes = 0;
-	MaxBytes = 0;
-}
-
-FLargeMemoryWriter::~FLargeMemoryWriter()
-{
-	if (Data)
-	{
-		FMemory::Free(Data);
-	}
-}
-
-void FLargeMemoryWriter::GrowBuffer(const int64 DesiredBytes)
-{
-	int64 NewBytes = 64 * 1024; // Initial alloc size
-
-	if (MaxBytes || DesiredBytes > NewBytes)
-	{
-		// Allocate slack proportional to the buffer size
-		NewBytes = FMemory::QuantizeSize(DesiredBytes + 3 * DesiredBytes / 8 + 16);
-	}
-
-	if (Data)
-	{
-		Data = (uint8*)FMemory::Realloc(Data, NewBytes);
-	}
-	else
-	{
-		Data = (uint8*)FMemory::Malloc(NewBytes);
-	}
-
-	MaxBytes = NewBytes;
-}
-
diff --git a/Engine/Source/Runtime/Core/Public/Serialization/LargeMemoryData.h b/Engine/Source/Runtime/Core/Public/Serialization/LargeMemoryData.h
new file mode 100644
--- /dev/null
+++ b/Engine/Source/Runtime/Core/Public/Serialization/LargeMemoryData.h
@@ -0,0 +1,74 @@
+// Copyright 1998-2018 Epic Games, Inc. All Rights Reserved.
+
+#pragma once
+
+#include "CoreTypes.h"
+
+/**
+* Data storage for the large memory reader and writer.
+*/
+
+class CORE_API FLargeMemoryData
+{
+public:
+
+	explicit FLargeMemoryData(const int64 PreAllocateBytes = 0);
+	~FLargeMemoryData();
+
+	/** Write data at the given offset. Returns true if the data was written. */
+	bool Write(void* InData, int64 InOffset, int64 InNum);
+
+	/** Append data at the given offset. */
+	FORCEINLINE void Append(void* InData, int64 InNum)
+	{
+		Write(InData, NumBytes, InNum);
+	}
+
+	/** Read data at the given offset. Returns true if the data was read. */
+	bool Read(void* OutData, int64 InOffset, int64 InNum) const;
+
+	/** Gets the size of the data written. */
+	FORCEINLINE int64 GetSize() const
+	{
+		return NumBytes;
+	}
+
+	/** Returns the written data.  */
+	FORCEINLINE uint8* GetData()
+	{
+		return Data;
+	}
+
+	/** Returns the written data.  */
+	FORCEINLINE const uint8* GetData() const
+	{
+		return Data;
+	}
+
+	/** Releases ownership of the written data. */
+	void ReleaseOwnership();
+
+	/** Check whether data is allocated or if the ownership was released. */
+	bool HasData() const 
+	{
+		return Data != nullptr;
+	}
+
+private:
+
+	/** Non-copyable */
+	FLargeMemoryData(const FLargeMemoryData&) = delete;
+	FLargeMemoryData& operator=(const FLargeMemoryData&) = delete;
+
+	/** Memory owned by this archive. Ownership can be released by calling ReleaseOwnership() */
+	uint8* Data;
+
+	/** Number of bytes currently written to our data buffer */
+	int64 NumBytes;
+
+	/** Number of bytes currently allocated for our data buffer */
+	int64 MaxBytes;
+
+	/** Resizes the data buffer to at least NumBytes with some slack */
+	void GrowBuffer();
+};
diff --git a/Engine/Source/Runtime/Core/Public/Serialization/LargeMemoryWriter.h b/Engine/Source/Runtime/Core/Public/Serialization/LargeMemoryWriter.h
--- a/Engine/Source/Runtime/Core/Public/Serialization/LargeMemoryWriter.h
+++ b/Engine/Source/Runtime/Core/Public/Serialization/LargeMemoryWriter.h
@@ -6,6 +6,7 @@
 #include "Containers/UnrealString.h"
 #include "UObject/NameTypes.h"
 #include "Serialization/MemoryArchive.h"
+#include "Serialization/LargeMemoryData.h"
 
 /**
 * Archive for storing a large amount of arbitrary data to memory
@@ -29,7 +30,10 @@ public:
 	/**
 	 * Gets the total size of the data written
 	 */
-	virtual int64 TotalSize() override;
+	virtual int64 TotalSize() override
+	{
+		return Data.GetSize();
+	}
 	
 	/**
 	 * Returns the written data. To release this archive's ownership of the data, call ReleaseOwnership()
@@ -39,32 +43,20 @@ public:
 	/** 
 	 * Releases ownership of the written data
 	 */
-	void ReleaseOwnership();
-	
-	/**
-	 * Destructor
-	 */
-	~FLargeMemoryWriter();
+	FORCEINLINE void ReleaseOwnership()
+	{
+		Data.ReleaseOwnership();
+	}
 
 private:
 
+	FLargeMemoryData Data;
+
 	/** Non-copyable */
 	FLargeMemoryWriter(const FLargeMemoryWriter&) = delete;
 	FLargeMemoryWriter& operator=(const FLargeMemoryWriter&) = delete;
 
-	/** Memory owned by this archive. Ownership can be released by calling ReleaseOwnership() */
-	uint8* Data;
-
-	/** Number of bytes currently written to our data buffer */
-	int64 NumBytes;
-
-	/** Number of bytes currently allocated for our data buffer */
-	int64 MaxBytes;
-
-	/** Resizes the data buffer to at least the desired new size with some slack */
-	void GrowBuffer(const int64 DesiredBytes);
 
 	/** Archive name, used for debugging, by default set to NAME_None. */
 	const FString ArchiveName;
 };
-
diff --git a/Engine/Source/Runtime/CoreUObject/Private/Serialization/DuplicateDataReader.cpp b/Engine/Source/Runtime/CoreUObject/Private/Serialization/DuplicateDataReader.cpp
--- a/Engine/Source/Runtime/CoreUObject/Private/Serialization/DuplicateDataReader.cpp
+++ b/Engine/Source/Runtime/CoreUObject/Private/Serialization/DuplicateDataReader.cpp
@@ -20,7 +20,7 @@
  * @param	InDuplicatedObjects		map of original object to copy of that object
  * @param	InObjectData			object data to read from
  */
-FDuplicateDataReader::FDuplicateDataReader( class FUObjectAnnotationSparse<FDuplicatedObject,false>& InDuplicatedObjects ,const TArray<uint8>& InObjectData, uint32 InPortFlags, UObject* InDestOuter )
+FDuplicateDataReader::FDuplicateDataReader( class FUObjectAnnotationSparse<FDuplicatedObject,false>& InDuplicatedObjects ,const FLargeMemoryData& InObjectData, uint32 InPortFlags, UObject* InDestOuter )
 	: DuplicatedObjectAnnotation(InDuplicatedObjects)
 	, ObjectData(InObjectData)
 	, Offset(0)
diff --git a/Engine/Source/Runtime/CoreUObject/Private/Serialization/DuplicateDataWriter.cpp b/Engine/Source/Runtime/CoreUObject/Private/Serialization/DuplicateDataWriter.cpp
--- a/Engine/Source/Runtime/CoreUObject/Private/Serialization/DuplicateDataWriter.cpp
+++ b/Engine/Source/Runtime/CoreUObject/Private/Serialization/DuplicateDataWriter.cpp
@@ -20,7 +20,7 @@
  * @param	InApplyFlags			the flags that should always be set on the duplicated objects (regardless of whether they're set on the source)
  * @param	InInstanceGraph			the instancing graph to use when creating the duplicate objects.
  */
-FDuplicateDataWriter::FDuplicateDataWriter( FUObjectAnnotationSparse<FDuplicatedObject,false>& InDuplicatedObjects ,TArray<uint8>& InObjectData,UObject* SourceObject,
+FDuplicateDataWriter::FDuplicateDataWriter( FUObjectAnnotationSparse<FDuplicatedObject,false>& InDuplicatedObjects, FLargeMemoryData& InObjectData, UObject* SourceObject,
 	UObject* DestObject, EObjectFlags InFlagMask, EObjectFlags InApplyFlags, EInternalObjectFlags InInternalFlagMask, EInternalObjectFlags InApplyInternalFlags, FObjectInstancingGraph* InInstanceGraph, uint32 InPortFlags)
 : DuplicatedObjectAnnotation(InDuplicatedObjects)
 , ObjectData(InObjectData)
diff --git a/Engine/Source/Runtime/CoreUObject/Private/UObject/UObjectGlobals.cpp b/Engine/Source/Runtime/CoreUObject/Private/UObject/UObjectGlobals.cpp
--- a/Engine/Source/Runtime/CoreUObject/Private/UObject/UObjectGlobals.cpp
+++ b/Engine/Source/Runtime/CoreUObject/Private/UObject/UObjectGlobals.cpp
@@ -2043,7 +2043,7 @@ UObject* StaticDuplicateObjectEx( FObjectDuplicationParameters& Parameters )
 														);
 	}
 
-	TArray<uint8> ObjectData;
+	FLargeMemoryData ObjectData;
 
 	FUObjectAnnotationSparse<FDuplicatedObject,false>  DuplicatedObjectAnnotation;
 
diff --git a/Engine/Source/Runtime/CoreUObject/Public/Serialization/DuplicatedDataReader.h b/Engine/Source/Runtime/CoreUObject/Public/Serialization/DuplicatedDataReader.h
--- a/Engine/Source/Runtime/CoreUObject/Public/Serialization/DuplicatedDataReader.h
+++ b/Engine/Source/Runtime/CoreUObject/Public/Serialization/DuplicatedDataReader.h
@@ -6,6 +6,7 @@
 #include "Serialization/ArchiveUObject.h"
 #include "UObject/UObjectAnnotation.h"
 #include "Serialization/DuplicatedObject.h"
+#include "Serialization/LargeMemoryData.h"
 
 /*----------------------------------------------------------------------------
 	FDuplicateDataReader.
@@ -19,8 +20,8 @@ class FDuplicateDataReader : public FArchiveUObject
 private:
 
 	class FUObjectAnnotationSparse<FDuplicatedObject,false>&	DuplicatedObjectAnnotation;
-	const TArray<uint8>&					ObjectData;
-	int32									Offset;
+	const FLargeMemoryData&					ObjectData;
+	int64									Offset;
 
 	//~ Begin FArchive Interface.
 
@@ -33,15 +34,14 @@ private:
 
 	virtual void Serialize(void* Data,int64 Num)
 	{
-		if(Num)
+		if (ObjectData.Read(Data, Offset, Num))
 		{
-			if (!(Offset + Num <= ObjectData.Num()))
-			{
-				SerializeFail();
-			}
-			FMemory::Memcpy(Data,&ObjectData[Offset],Num);
 			Offset += Num;
 		}
+		else
+		{
+			SerializeFail();
+		}
 	}
 
 	virtual void Seek(int64 InPos)
@@ -64,7 +64,7 @@ public:
 	}
 	virtual int64 TotalSize()
 	{
-		return ObjectData.Num();
+		return ObjectData.GetSize();
 	}
 
 	/**
@@ -73,5 +73,5 @@ public:
 	 * @param	InDuplicatedObjectAnnotation		Annotation for storing a mapping from source to duplicated object
 	 * @param	InObjectData					Object data to read from
 	 */
-	FDuplicateDataReader( FUObjectAnnotationSparse<FDuplicatedObject,false>& InDuplicatedObjectAnnotation, const TArray<uint8>& InObjectData, uint32 InPortFlags, UObject* InDestOuter );
-};
+	FDuplicateDataReader( FUObjectAnnotationSparse<FDuplicatedObject,false>& InDuplicatedObjectAnnotation, const FLargeMemoryData& InObjectData, uint32 InPortFlags, UObject* InDestOuter );
+};
\ No newline at end of file
diff --git a/Engine/Source/Runtime/CoreUObject/Public/Serialization/DuplicatedDataWriter.h b/Engine/Source/Runtime/CoreUObject/Public/Serialization/DuplicatedDataWriter.h
--- a/Engine/Source/Runtime/CoreUObject/Public/Serialization/DuplicatedDataWriter.h
+++ b/Engine/Source/Runtime/CoreUObject/Public/Serialization/DuplicatedDataWriter.h
@@ -6,6 +6,7 @@
 #include "Serialization/ArchiveUObject.h"
 #include "UObject/UObjectAnnotation.h"
 #include "Serialization/DuplicatedObject.h"
+#include "Serialization/LargeMemoryData.h"
 
 struct FObjectInstancingGraph;
 
@@ -20,8 +21,8 @@ class FDuplicateDataWriter : public FArchiveUObject
 private:
 
 	class FUObjectAnnotationSparse<FDuplicatedObject,false>&	DuplicatedObjectAnnotation;
-	TArray<uint8>&							ObjectData;
-	int64										Offset;
+	FLargeMemoryData&						ObjectData;
+	int64									Offset;
 	EObjectFlags							FlagMask;
 	EObjectFlags							ApplyFlags;
 	EInternalObjectFlags InternalFlagMask;
@@ -41,15 +42,8 @@ private:
 
 	virtual void Serialize(void* Data,int64 Num)
 	{
-		// Don't try to add/memcpy zero sized items
-		if (Data != NULL && Num > 0)
+		if (ObjectData.Write(Data, Offset, Num))
 		{
-			if (Num + Offset > ObjectData.Num())
-			{
-				ObjectData.AddUninitialized(Num + Offset - ObjectData.Num());
-			}
-
-			FMemory::Memcpy(&ObjectData[Offset],Data,Num);
 			Offset += Num;
 		}
 	}
@@ -83,8 +77,9 @@ public:
 
 	virtual int64 TotalSize()
 	{
-		return ObjectData.Num();
+		return ObjectData.GetSize();
 	}
+
 	TArray<UObject*>	UnserializedObjects;
 
 	/**
@@ -110,7 +105,7 @@ public:
 	 */
 	FDuplicateDataWriter(
 		FUObjectAnnotationSparse<FDuplicatedObject, false>& InDuplicatedObjects,
-		TArray<uint8>& InObjectData,
+		FLargeMemoryData& InObjectData,
 		UObject* SourceObject,
 		UObject* DestObject,
 		EObjectFlags InFlagMask,
