Fix incorrect changelist being used in FEngineVersion::CompatibleWith.

diff --git a/Engine/Source/Runtime/Core/Private/Misc/EngineVersion.cpp b/Engine/Source/Runtime/Core/Private/Misc/EngineVersion.cpp
--- a/Engine/Source/Runtime/Core/Private/Misc/EngineVersion.cpp
+++ b/Engine/Source/Runtime/Core/Private/Misc/EngineVersion.cpp
@@ -11,7 +11,7 @@
 FEngineVersion FEngineVersion::CurrentVersion(ENGINE_MAJOR_VERSION, ENGINE_MINOR_VERSION, ENGINE_PATCH_VERSION, BuildSettings::GetCurrentChangelist() | (BuildSettings::IsLicenseeVersion()? (1U << 31) : 0), BuildSettings::GetBranchName());
 
 // Version which this engine maintains strict API and package compatibility with. By default, we always maintain compatibility with the current major/minor version, unless we're built at a different changelist.
-FEngineVersion FEngineVersion::CompatibleWithVersion(ENGINE_MAJOR_VERSION, ENGINE_MINOR_VERSION, 0, BuildSettings::GetCurrentChangelist() | (BuildSettings::IsLicenseeVersion()? (1U << 31) : 0), BuildSettings::GetBranchName());
+FEngineVersion FEngineVersion::CompatibleWithVersion(ENGINE_MAJOR_VERSION, ENGINE_MINOR_VERSION, 0, BuildSettings::GetCompatibleChangelist() | (BuildSettings::IsLicenseeVersion()? (1U << 31) : 0), BuildSettings::GetBranchName());
 
 
 FEngineVersionBase::FEngineVersionBase()
