Copied cl 4124584 from dev-enterprise

Fixed high precision tangents lightmass baking issue.
- FPackedRGBA16N::ToFVector was failing to unpack Z and W because the intrinsic was only reading the first 32 bits on SSE platforms.

diff --git a/Engine/Source/Runtime/Core/Public/Math/UnrealMathSSE.h b/Engine/Source/Runtime/Core/Public/Math/UnrealMathSSE.h
--- a/Engine/Source/Runtime/Core/Public/Math/UnrealMathSSE.h
+++ b/Engine/Source/Runtime/Core/Public/Math/UnrealMathSSE.h
@@ -998,24 +998,24 @@ FORCEINLINE void VectorStoreURGB10A2N(const VectorRegister& Vec, void* Ptr)
 }
 
 /**
-* Loads packed RGBA16(4 bytes) from unaligned memory and converts them into 4 FLOATs.
-* IMPORTANT: You need to call VectorResetFloatRegisters() before using scalar FLOATs after you've used this intrinsic!
-*
-* @param Ptr			Unaligned memory pointer to the RGBA16(8 bytes).
-* @return				VectorRegister with 4 FLOATs loaded from Ptr.
-*/
-#define VectorLoadURGBA16N( Ptr ) _mm_cvtepi32_ps(_mm_unpacklo_epi16(_mm_cvtsi32_si128(*(int32*)Ptr), _mm_setzero_si128()))
+ * Loads packed RGBA16(8 bytes) from unaligned memory and converts them into 4 FLOATs.
+ * IMPORTANT: You need to call VectorResetFloatRegisters() before using scalar FLOATs after you've used this intrinsic!
+ *
+ * @param Ptr			Unaligned memory pointer to the RGBA16(8 bytes).
+ * @return				VectorRegister with 4 FLOATs loaded from Ptr.
+ */
+#define VectorLoadURGBA16N( Ptr ) _mm_cvtepi32_ps(_mm_unpacklo_epi16(_mm_cvtsi64_si128(*(int64*)Ptr), _mm_setzero_si128()))
 
 /**
-* Loads packed signed RGBA16(4 bytes) from unaligned memory and converts them into 4 FLOATs.
-* IMPORTANT: You need to call VectorResetFloatRegisters() before using scalar FLOATs after you've used this intrinsic!
-*
-* @param Ptr			Unaligned memory pointer to the RGBA16(8 bytes).
-* @return				VectorRegister with 4 FLOATs loaded from Ptr.
-*/
+ * Loads packed signed RGBA16(8 bytes) from unaligned memory and converts them into 4 FLOATs.
+ * IMPORTANT: You need to call VectorResetFloatRegisters() before using scalar FLOATs after you've used this intrinsic!
+ *
+ * @param Ptr			Unaligned memory pointer to the RGBA16(8 bytes).
+ * @return				VectorRegister with 4 FLOATs loaded from Ptr.
+ */
 FORCEINLINE VectorRegister VectorLoadSRGBA16N(const void* Ptr)
 {
-	auto Temp = _mm_unpacklo_epi16(_mm_cvtsi32_si128(*(int32*)Ptr), _mm_setzero_si128());
+	auto Temp = _mm_unpacklo_epi16(_mm_cvtsi64_si128(*(int64*)Ptr), _mm_setzero_si128());
 	auto Mask = _mm_cmpgt_epi32(Temp, _mm_set1_epi32(32767));
 	auto Comp = _mm_and_si128(Mask, _mm_set1_epi32(~32767));
 	return _mm_cvtepi32_ps(_mm_or_si128(Comp, Temp));
