Build fix for x86, replaced _mm_cvtsi64_si128 with _mm_loadl_epi64 in UnrealMathSSE

diff --git a/Engine/Source/Runtime/Core/Public/Math/UnrealMathSSE.h b/Engine/Source/Runtime/Core/Public/Math/UnrealMathSSE.h
--- a/Engine/Source/Runtime/Core/Public/Math/UnrealMathSSE.h
+++ b/Engine/Source/Runtime/Core/Public/Math/UnrealMathSSE.h
@@ -1004,7 +1004,7 @@ FORCEINLINE void VectorStoreURGB10A2N(const VectorRegister& Vec, void* Ptr)
  * @param Ptr			Unaligned memory pointer to the RGBA16(8 bytes).
  * @return				VectorRegister with 4 FLOATs loaded from Ptr.
  */
-#define VectorLoadURGBA16N( Ptr ) _mm_cvtepi32_ps(_mm_unpacklo_epi16(_mm_cvtsi64_si128(*(int64*)Ptr), _mm_setzero_si128()))
+#define VectorLoadURGBA16N( Ptr ) _mm_cvtepi32_ps(_mm_unpacklo_epi16(_mm_loadl_epi64((const __m128i*)Ptr), _mm_setzero_si128()))
 
 /**
  * Loads packed signed RGBA16(8 bytes) from unaligned memory and converts them into 4 FLOATs.
@@ -1015,7 +1015,7 @@ FORCEINLINE void VectorStoreURGB10A2N(const VectorRegister& Vec, void* Ptr)
  */
 FORCEINLINE VectorRegister VectorLoadSRGBA16N(const void* Ptr)
 {
-	auto Temp = _mm_unpacklo_epi16(_mm_cvtsi64_si128(*(int64*)Ptr), _mm_setzero_si128());
+	auto Temp = _mm_unpacklo_epi16(_mm_loadl_epi64((const __m128i*)Ptr), _mm_setzero_si128());
 	auto Mask = _mm_cmpgt_epi32(Temp, _mm_set1_epi32(32767));
 	auto Comp = _mm_and_si128(Mask, _mm_set1_epi32(~32767));
 	return _mm_cvtepi32_ps(_mm_or_si128(Comp, Temp));
