Fixed crash when FiB discovered a string table reference

diff --git a/Engine/Source/Editor/Kismet/Private/ImaginaryBlueprintData.cpp b/Engine/Source/Editor/Kismet/Private/ImaginaryBlueprintData.cpp
--- a/Engine/Source/Editor/Kismet/Private/ImaginaryBlueprintData.cpp
+++ b/Engine/Source/Editor/Kismet/Private/ImaginaryBlueprintData.cpp
@@ -1,7 +1,8 @@
 // Copyright 1998-2018 Epic Games, Inc. All Rights Reserved.
 #include "ImaginaryBlueprintData.h"
 #include "Misc/ScopeLock.h"
-
+#include "Async/Async.h"
+#include "Internationalization/StringTableCore.h"
 
 #define LOCTEXT_NAMESPACE "FindInBlueprints"
 
@@ -10,11 +11,39 @@
 
 FText FSearchableValueInfo::GetDisplayText(const TMap<int32, FText>& InLookupTable) const
 {
+	FText Result;
 	if (!DisplayText.IsEmpty() || LookupTableKey == -1)
 	{
-		return DisplayText;
+		Result = DisplayText;
+	}
+	else
+	{
+		Result = FindInBlueprintsHelpers::AsFText(LookupTableKey, InLookupTable);
 	}
-	return FindInBlueprintsHelpers::AsFText(LookupTableKey, InLookupTable);
+
+	if (Result.IsFromStringTable() && FTextInspector::GetSourceString(Result) == &FStringTableEntry::GetPlaceholderSourceString() && !IsInGameThread())
+	{
+		// String Table asset references in FiB may be unresolved as we can't load the asset on the search thread
+		// To solve this we send a request to the game thread to get the display string and wait for the result, this will trigger the asset to load
+		FName TableId;
+		FString Key;
+		if (FTextInspector::GetTableIdAndKey(Result, TableId, Key) && IStringTableEngineBridge::IsStringTableFromAsset(TableId))
+		{
+			TPromise<bool> Promise;
+
+			// Run the request on the game thread, filling the promise when done
+			AsyncTask(ENamedThreads::GameThread, [Result, &Promise]()
+			{
+				Result.ToString(); // Trigger the asset load
+				Promise.SetValue(true); // Signal completion
+			});
+
+			// Get the promise value to block until the AsyncTask has completed
+			Promise.GetFuture().Get();
+		}
+	}
+
+	return Result;
 }
 
 ////////////////////////////
diff --git a/Engine/Source/Runtime/Core/Private/Internationalization/TextHistory.cpp b/Engine/Source/Runtime/Core/Private/Internationalization/TextHistory.cpp
--- a/Engine/Source/Runtime/Core/Private/Internationalization/TextHistory.cpp
+++ b/Engine/Source/Runtime/Core/Private/Internationalization/TextHistory.cpp
@@ -1131,6 +1131,7 @@ bool FTextHistory_Transform::GetHistoricNumericData(const FText& InText, FHistor
 FTextHistory_StringTableEntry::FTextHistory_StringTableEntry(FName InTableId, FString&& InKey)
 	: TableId(InTableId)
 	, Key(MoveTemp(InKey))
+	, bStringTableAssetPendingLoad(false)
 {
 	GetStringTableEntry();
 }
@@ -1217,11 +1218,16 @@ void FTextHistory_StringTableEntry::Serialize(FArchive& Ar)
 		Ar << Key;
 
 		// String Table assets should already have been created via dependency loading when using the EDL (although they may not be fully loaded yet)
-		FStringTableRedirects::RedirectTableIdAndKey(TableId, Key, GEventDrivenLoaderEnabled ? EStringTableLoadingPolicy::Find : EStringTableLoadingPolicy::FindOrLoad);
+		const bool bIsAsset = IStringTableEngineBridge::IsStringTableFromAsset(TableId);
+		const EStringTableLoadingPolicy LoadingPolicy = (!bIsAsset || !IsInGameThread() || GEventDrivenLoaderEnabled) ? EStringTableLoadingPolicy::Find : EStringTableLoadingPolicy::FindOrLoad;
+		FStringTableRedirects::RedirectTableIdAndKey(TableId, Key, LoadingPolicy);
 
 		// Re-cache the pointer
 		StringTableEntry.Reset();
-		GetStringTableEntry();
+		FStringTableEntryConstPtr StringTableEntryPin = GetStringTableEntry(LoadingPolicy == EStringTableLoadingPolicy::Find);
+
+		// If we couldn't load a string table asset because this wasn't the game thread, defer the loading request until we're able to process it
+		bStringTableAssetPendingLoad = !StringTableEntryPin.IsValid() && bIsAsset && !IsInGameThread() && !GEventDrivenLoaderEnabled;
 	}
 	else if (Ar.IsSaving())
 	{
@@ -1250,11 +1256,11 @@ void FTextHistory_StringTableEntry::SerializeForDisplayString(FArchive& Ar, FTex
 	}
 }
 
-FStringTableEntryConstPtr FTextHistory_StringTableEntry::GetStringTableEntry() const
+FStringTableEntryConstPtr FTextHistory_StringTableEntry::GetStringTableEntry(const bool InSilent) const
 {
 	FStringTableEntryConstPtr StringTableEntryPin = StringTableEntry.Pin();
 
-	bool bSuppressMissingEntryWarning = false;
+	bool bSuppressMissingEntryWarning = InSilent;
 
 	if (!StringTableEntryPin.IsValid() || !StringTableEntryPin->IsOwned())
 	{
@@ -1264,12 +1270,29 @@ FStringTableEntryConstPtr FTextHistory_StringTableEntry::GetStringTableEntry() c
 		StringTableEntryPin = StringTableEntry.Pin();
 		if (!StringTableEntryPin.IsValid() || !StringTableEntryPin->IsOwned())
 		{
+			if (bStringTableAssetPendingLoad && IsInGameThread())
+			{
+				// This path should never be taken when EDL is enabled
+				check(!GEventDrivenLoaderEnabled);
+
+				// Attempt to load the string table asset now
+				FTextHistory_StringTableEntry* NonConstThis = const_cast<FTextHistory_StringTableEntry*>(this);
+				FStringTableRedirects::RedirectTableIdAndKey(NonConstThis->TableId, NonConstThis->Key, EStringTableLoadingPolicy::FindOrLoad);
+
+				// We always set this to true even if the load failed
+				bStringTableAssetPendingLoad = false;
+			}
+			else
+			{
+				bSuppressMissingEntryWarning |= bStringTableAssetPendingLoad;
+			}
+
 			StringTableEntryPin.Reset();
 
 			FStringTableConstPtr StringTable = FStringTableRegistry::Get().FindStringTable(TableId);
 			if (StringTable.IsValid())
 			{
-				bSuppressMissingEntryWarning = !StringTable->IsLoaded();
+				bSuppressMissingEntryWarning |= !StringTable->IsLoaded();
 				StringTableEntryPin = StringTable->FindEntry(Key);
 			}
 
diff --git a/Engine/Source/Runtime/Core/Private/Internationalization/TextHistory.h b/Engine/Source/Runtime/Core/Private/Internationalization/TextHistory.h
--- a/Engine/Source/Runtime/Core/Private/Internationalization/TextHistory.h
+++ b/Engine/Source/Runtime/Core/Private/Internationalization/TextHistory.h
@@ -496,7 +496,7 @@ private:
 	FTextHistory_StringTableEntry& operator=(FTextHistory_StringTableEntry&);
 
 	/** Get the string table pointer, potentially re-caching it if it's missing or stale */
-	FStringTableEntryConstPtr GetStringTableEntry() const;
+	FStringTableEntryConstPtr GetStringTableEntry(const bool InSilent = false) const;
 
 	/** The string table ID being referenced */
 	FName TableId;
@@ -504,6 +504,9 @@ private:
 	/** The key within the string table being referenced */
 	FString Key;
 
+	/** True if the string table asset referenced is pending load because we weren't able to load it during Serialize */
+	mutable bool bStringTableAssetPendingLoad;
+
 	/** Cached string table entry pointer */
 	mutable FStringTableEntryConstWeakPtr StringTableEntry;
 
diff --git a/Engine/Source/Runtime/Core/Private/Internationalization/TextLocalizationManager.cpp b/Engine/Source/Runtime/Core/Private/Internationalization/TextLocalizationManager.cpp
--- a/Engine/Source/Runtime/Core/Private/Internationalization/TextLocalizationManager.cpp
+++ b/Engine/Source/Runtime/Core/Private/Internationalization/TextLocalizationManager.cpp
@@ -5,6 +5,7 @@
 #include "Internationalization/ILocalizedTextSource.h"
 #include "Internationalization/LocalizationResourceTextSource.h"
 #include "Internationalization/PolyglotTextSource.h"
+#include "Internationalization/StringTableRegistry.h"
 #include "GenericPlatform/GenericPlatformFile.h"
 #include "HAL/FileManager.h"
 #include "HAL/LowLevelMemTracker.h"
@@ -297,6 +298,9 @@ void BeginInitTextLocalization()
 	// twice since FInternationalization::Initialize sets the culture.
 	FInternationalization::Get();
 
+	// Make sure the String Table Registry is initialized as it may trigger module loads
+	FStringTableRegistry::Get();
+
 	FInternationalization::Get().OnCultureChanged().AddRaw(&(FTextLocalizationManager::Get()), &FTextLocalizationManager::OnCultureChanged);
 }
 
diff --git a/Engine/Source/Runtime/Core/Public/Internationalization/StringTableCore.h b/Engine/Source/Runtime/Core/Public/Internationalization/StringTableCore.h
--- a/Engine/Source/Runtime/Core/Public/Internationalization/StringTableCore.h
+++ b/Engine/Source/Runtime/Core/Public/Internationalization/StringTableCore.h
@@ -189,11 +189,18 @@ public:
 		}
 	}
 
+	/** Is this string table from an asset? */
+	static bool IsStringTableFromAsset(const FName InTableId)
+	{
+		return InstancePtr && InstancePtr->IsStringTableFromAssetImpl(InTableId);
+	}
+
 protected:
 	virtual ~IStringTableEngineBridge() {}
 
 	virtual void RedirectAndLoadStringTableAssetImpl(FName& InOutTableId, const EStringTableLoadingPolicy InLoadingPolicy) = 0;
 	virtual void CollectStringTableAssetReferencesImpl(const FName InTableId, FArchive& InAr) = 0;
+	virtual bool IsStringTableFromAssetImpl(const FName InTableId) = 0;
 
 	/** Singleton instance, populated by the derived type */
 	static IStringTableEngineBridge* InstancePtr;
diff --git a/Engine/Source/Runtime/Engine/Private/Internationalization/StringTable.cpp b/Engine/Source/Runtime/Engine/Private/Internationalization/StringTable.cpp
--- a/Engine/Source/Runtime/Engine/Private/Internationalization/StringTable.cpp
+++ b/Engine/Source/Runtime/Engine/Private/Internationalization/StringTable.cpp
@@ -82,7 +82,7 @@ private:
 	//~ IStringTableEngineInterop interface
 	virtual void RedirectAndLoadStringTableAssetImpl(FName& InOutTableId, const EStringTableLoadingPolicy InLoadingPolicy) override
 	{
-		FSoftObjectPath StringTableAssetReference = GetAssetReference(InOutTableId);
+		const FSoftObjectPath StringTableAssetReference = GetAssetReference(InOutTableId);
 		if (StringTableAssetReference.IsValid())
 		{
 			UStringTable* StringTableAsset = Cast<UStringTable>(StringTableAssetReference.ResolveObject());
@@ -111,6 +111,12 @@ private:
 		InAr << StringTableAsset;
 	}
 
+	virtual bool IsStringTableFromAssetImpl(const FName InTableId) override
+	{
+		const FSoftObjectPath StringTableAssetReference = GetAssetReference(InTableId);
+		return StringTableAssetReference.IsValid();
+	}
+
 	static FSoftObjectPath GetAssetReference(const FName InTableId)
 	{
 		const FString StringTableAssetName = InTableId.ToString();
