When archive decompression fails, the archive name will now also be logged with the assert message.

diff --git a/Engine/Source/Runtime/Core/Private/Serialization/Archive.cpp b/Engine/Source/Runtime/Core/Private/Serialization/Archive.cpp
--- a/Engine/Source/Runtime/Core/Private/Serialization/Archive.cpp
+++ b/Engine/Source/Runtime/Core/Private/Serialization/Archive.cpp
@@ -691,7 +691,8 @@ void FArchive::SerializeCompressed( void* V, int64 Length, ECompressionFlags Fla
 			// Read compressed data.
 			Serialize( CompressedBuffer, Chunk.CompressedSize );
 			// Decompress into dest pointer directly.
-			verify( FCompression::UncompressMemory( Flags, Dest, Chunk.UncompressedSize, CompressedBuffer, Chunk.CompressedSize, (Padding > 0) ? true : false, CompressionBitWindow ) ); //-V547
+			bool bUncompressMemorySucceeded = FCompression::UncompressMemory( Flags, Dest, Chunk.UncompressedSize, CompressedBuffer, Chunk.CompressedSize, (Padding > 0) ? true : false, CompressionBitWindow ); //-V547
+			verifyf(bUncompressMemorySucceeded, TEXT("Failed to uncompress data in %s. Check log for details."), *GetArchiveName());
 			// And advance it by read amount.
 			Dest += Chunk.UncompressedSize;
 		}
