Copied from Dev-Editor CL 4114042

Fixed bug in TBitArray. In operator=, the invariant that MaxBits be a multiple of NumBitsPerDWORD is broken. This can cause a crash when a copied TBitArray has its allocation grown, because NumElements is not greater than NumAllocatedElements.

diff --git a/Engine/Source/Runtime/Core/Public/Containers/BitArray.h b/Engine/Source/Runtime/Core/Public/Containers/BitArray.h
--- a/Engine/Source/Runtime/Core/Public/Containers/BitArray.h
+++ b/Engine/Source/Runtime/Core/Public/Containers/BitArray.h
@@ -226,11 +226,10 @@ public:
 		}
 
 		Empty(Copy.Num());
-		NumBits = MaxBits = Copy.NumBits;
+		NumBits = Copy.NumBits;
 		if(NumBits)
 		{
 			const int32 NumDWORDs = FMath::DivideAndRoundUp(MaxBits, NumBitsPerDWORD);
-			Realloc(0);
 			FMemory::Memcpy(GetData(),Copy.GetData(),NumDWORDs * sizeof(uint32));
 		}
 		return *this;
