Log which symbolicator being used

diff --git a/Engine/Source/Runtime/Core/Private/Unix/UnixPlatformMisc.cpp b/Engine/Source/Runtime/Core/Private/Unix/UnixPlatformMisc.cpp
--- a/Engine/Source/Runtime/Core/Private/Unix/UnixPlatformMisc.cpp
+++ b/Engine/Source/Runtime/Core/Private/Unix/UnixPlatformMisc.cpp
@@ -102,6 +102,9 @@ void UnixPlatform_UpdateCacheLineSize()
 	}
 }
 
+// Init'ed in UnixPlatformMemory for now. Once the old crash symbolicator is gone remove this
+extern bool CORE_API GUseNewCrashSymbolicator;
+
 void FUnixPlatformMisc::PlatformInit()
 {
 	// install a platform-specific signal handler
@@ -111,6 +114,15 @@ void FUnixPlatformMisc::PlatformInit()
 	bool bFirstInstance = FPlatformProcess::IsFirstInstance();
 	bool bIsNullRHI = !FApp::CanEverRender();
 
+	if (GUseNewCrashSymbolicator)
+	{
+		UE_LOG(LogInit, Log, TEXT("Using custom *.sym to symbolicate"));
+	}
+	else
+	{
+		UE_LOG(LogInit, Log, TEXT("Using libdwarf/libelf to symbolicate"));
+	}
+
 	UE_LOG(LogInit, Log, TEXT("Unix hardware info:"));
 	UE_LOG(LogInit, Log, TEXT(" - we are %sthe first instance of this executable"), bFirstInstance ? TEXT("") : TEXT("not "));
 	UE_LOG(LogInit, Log, TEXT(" - this process' id (pid) is %d, parent process' id (ppid) is %d"), static_cast< int32 >(getpid()), static_cast< int32 >(getppid()));
