Logs were getting spammed from Shader Compiler

diff --git a/Engine/Source/Runtime/Core/Private/HAL/UnrealMemory.cpp b/Engine/Source/Runtime/Core/Private/HAL/UnrealMemory.cpp
--- a/Engine/Source/Runtime/Core/Private/HAL/UnrealMemory.cpp
+++ b/Engine/Source/Runtime/Core/Private/HAL/UnrealMemory.cpp
@@ -378,7 +378,7 @@ int FMemory_GCreateMalloc_ThreadUnsafe()
 #endif
 
 // On Mac it's too early to log here in some cases. For example GMalloc may be created during initialization of a third party dylib on load, before CoreFoundation is initialized
-#if !PLATFORM_MAC
+#if !PLATFORM_DESKTOP
 	// by this point, we assume we can log
 	double SizeInMb = ProgramSize / (1024.0 * 1024.0);
 	FPlatformMisc::LowLevelOutputDebugStringf(TEXT("Used memory before allocating anything was %.2fMB\n"), SizeInMb);
