Handle invalid cases

diff --git a/Engine/Source/Runtime/Core/Private/Unix/UnixPlatformStackWalk.cpp b/Engine/Source/Runtime/Core/Private/Unix/UnixPlatformStackWalk.cpp
--- a/Engine/Source/Runtime/Core/Private/Unix/UnixPlatformStackWalk.cpp
+++ b/Engine/Source/Runtime/Core/Private/Unix/UnixPlatformStackWalk.cpp
@@ -1483,11 +1483,7 @@ int32 FUnixPlatformStackWalk::GetProcessModuleCount()
 	int Size = 0;
 	dl_iterate_phdr(NumberOfDynamicLibrariesCallback, &Size);
 
-	return 0;
-
-	// This will disable portable callstacks for Linux. Need to wait for the new symbolicator to enable
-	// otherwise we end up loading dwarf in while an ensure happens causing huge hitches/slowdowns
-	/* return Size; */
+	return Size;
 }
 
 namespace
@@ -1549,14 +1545,15 @@ namespace
 
 int32 FUnixPlatformStackWalk::GetProcessModuleSignatures(FStackWalkModuleInfo *ModuleSignatures, const int32 ModuleSignaturesSize)
 {
+	if (ModuleSignatures == nullptr || ModuleSignaturesSize == 0)
+	{
+		return 0;
+	}
+
 	ProcessModuleSignatures Signatures{ModuleSignatures, ModuleSignaturesSize, 0};
 	dl_iterate_phdr(CollectModuleSignatures, &Signatures);
 
-	return 0;
-
-	// This will disable portable callstacks for Linux. Need to wait for the new symbolicator to enable
-	// otherwise we end up loading dwarf in while an ensure happens causing huge hitches/slowdowns
-	/* return Signatures.Index; */
+	return Signatures.Index;
 }
 
 static FCriticalSection EnsureLock;
