Only use with new callstack symbolicator

diff --git a/Engine/Source/Runtime/Core/Private/Unix/UnixPlatformStackWalk.cpp b/Engine/Source/Runtime/Core/Private/Unix/UnixPlatformStackWalk.cpp
--- a/Engine/Source/Runtime/Core/Private/Unix/UnixPlatformStackWalk.cpp
+++ b/Engine/Source/Runtime/Core/Private/Unix/UnixPlatformStackWalk.cpp
@@ -1480,6 +1480,13 @@ namespace
 
 int32 FUnixPlatformStackWalk::GetProcessModuleCount()
 {
+	// If we are not using the new crash symbolicator we want to avoid generating portable callstacks these
+	// as they may cause dwarf/elf symbols to be loaded dynamically which can cause hitching during an ensure
+	if (!GUseNewCrashSymbolicator)
+	{
+		return 0;
+	}
+
 	int Size = 0;
 	dl_iterate_phdr(NumberOfDynamicLibrariesCallback, &Size);
 
@@ -1545,7 +1552,7 @@ namespace
 
 int32 FUnixPlatformStackWalk::GetProcessModuleSignatures(FStackWalkModuleInfo *ModuleSignatures, const int32 ModuleSignaturesSize)
 {
-	if (ModuleSignatures == nullptr || ModuleSignaturesSize == 0)
+	if (ModuleSignatures == nullptr || ModuleSignaturesSize == 0 || !GUseNewCrashSymbolicator)
 	{
 		return 0;
 	}
