ModuleManager now explicitly looks up the module manifest when there is no binary for the filename it has been given - this fixes issues with hot reloaded new components.

diff --git a/Engine/Source/Runtime/Core/Private/Modules/ModuleManager.cpp b/Engine/Source/Runtime/Core/Private/Modules/ModuleManager.cpp
--- a/Engine/Source/Runtime/Core/Private/Modules/ModuleManager.cpp
+++ b/Engine/Source/Runtime/Core/Private/Modules/ModuleManager.cpp
@@ -464,7 +464,7 @@ IModuleInterface* FModuleManager::LoadModuleWithFailureReason(const FName InModu
 
 			UE_LOG(LogModuleManager, Verbose, TEXT("ModuleManager: Load Module '%s' DLL '%s'"), *InModuleName.ToString(), *ModuleInfo->Filename);
 
-			if (ModuleInfo->Filename.IsEmpty())
+			if (ModuleInfo->Filename.IsEmpty() || !FPaths::FileExists(ModuleInfo->Filename))
 			{
 				TMap<FName, FString> ModulePathMap;
 				FindModulePaths(*InModuleName.ToString(), ModulePathMap);
