Apply fix from CL# 4101171 to Release-4.20
#rnx

diff --git a/Engine/Source/Runtime/Core/Private/FramePro/FrameProProfiler.cpp b/Engine/Source/Runtime/Core/Private/FramePro/FrameProProfiler.cpp
--- a/Engine/Source/Runtime/Core/Private/FramePro/FrameProProfiler.cpp
+++ b/Engine/Source/Runtime/Core/Private/FramePro/FrameProProfiler.cpp
@@ -10,6 +10,7 @@
 #include "HAL/FileManager.h"
 #include "HAL/ThreadManager.h"
 #include "Misc/App.h"
+#include "Misc/EngineVersion.h"
 #include "Misc/Paths.h"
 #include "Runtime/Launch/Resources/Version.h"
 
