Use the generic platform for getting a valid platform target name

diff --git a/Engine/Source/Runtime/Core/Private/Linux/LinuxPlatformMisc.cpp b/Engine/Source/Runtime/Core/Private/Linux/LinuxPlatformMisc.cpp
--- a/Engine/Source/Runtime/Core/Private/Linux/LinuxPlatformMisc.cpp
+++ b/Engine/Source/Runtime/Core/Private/Linux/LinuxPlatformMisc.cpp
@@ -4,11 +4,8 @@
 	LinuxPlatformMisc.cpp: Linux implementations of misc platform functions
 =============================================================================*/
 
+/* Add back if we overload any functions for GenericPlatform here
 #include "Linux/LinuxPlatformMisc.h"
 #include "Containers/Array.h"
 #include "Containers/UnrealString.h"
-
-void FLinuxPlatformMisc::GetValidTargetPlatforms(class TArray<class FString>& TargetPlatformNames)
-{
-	TargetPlatformNames.Add(TEXT("Linux"));
-}
+*/
diff --git a/Engine/Source/Runtime/Core/Private/Unix/UnixPlatformMisc.cpp b/Engine/Source/Runtime/Core/Private/Unix/UnixPlatformMisc.cpp
--- a/Engine/Source/Runtime/Core/Private/Unix/UnixPlatformMisc.cpp
+++ b/Engine/Source/Runtime/Core/Private/Unix/UnixPlatformMisc.cpp
@@ -298,11 +298,6 @@ FString FUnixPlatformMisc::GetOSVersion()
 	return FString();
 }
 
-void FUnixPlatformMisc::GetValidTargetPlatforms(class TArray<class FString>& TargetPlatformNames)
-{
-	TargetPlatformNames.Add(TEXT("Unix"));
-}
-
 const TCHAR* FUnixPlatformMisc::GetSystemErrorMessage(TCHAR* OutBuffer, int32 BufferCount, int32 Error)
 {
 	check(OutBuffer && BufferCount);
diff --git a/Engine/Source/Runtime/Core/Public/Linux/LinuxPlatformMisc.h b/Engine/Source/Runtime/Core/Public/Linux/LinuxPlatformMisc.h
--- a/Engine/Source/Runtime/Core/Public/Linux/LinuxPlatformMisc.h
+++ b/Engine/Source/Runtime/Core/Public/Linux/LinuxPlatformMisc.h
@@ -11,9 +11,10 @@
 /**
  * Linux implementation of the misc OS functions
  */
+
+// A stub for now. Will allow to override any functions if needed
 struct CORE_API FLinuxPlatformMisc : public FUnixPlatformMisc
 {
-	static void GetValidTargetPlatforms(class TArray<class FString>& TargetPlatformNames);
 };
 
 typedef FLinuxPlatformMisc FPlatformMisc;
diff --git a/Engine/Source/Runtime/Core/Public/Unix/UnixPlatformMisc.h b/Engine/Source/Runtime/Core/Public/Unix/UnixPlatformMisc.h
--- a/Engine/Source/Runtime/Core/Public/Unix/UnixPlatformMisc.h
+++ b/Engine/Source/Runtime/Core/Public/Unix/UnixPlatformMisc.h
@@ -74,7 +74,6 @@ struct CORE_API FUnixPlatformMisc : public FGenericPlatformMisc
 
 	static void LowLevelOutputDebugString(const TCHAR *Message);
 
-	static void GetValidTargetPlatforms(class TArray<class FString>& TargetPlatformNames);
 	static void RequestExit(bool Force);
 	static void RequestExitWithStatus(bool Force, uint8 ReturnCode);
 	static const TCHAR* GetSystemErrorMessage(TCHAR* OutBuffer, int32 BufferCount, int32 Error);
