Add setters to TTlsAutoCleanupValue

diff --git a/Engine/Source/Runtime/Core/Public/HAL/TlsAutoCleanup.h b/Engine/Source/Runtime/Core/Public/HAL/TlsAutoCleanup.h
--- a/Engine/Source/Runtime/Core/Public/HAL/TlsAutoCleanup.h
+++ b/Engine/Source/Runtime/Core/Public/HAL/TlsAutoCleanup.h
@@ -34,6 +34,18 @@ public:
 		return Value;
 	}
 
+	/* Sets the value. */
+	void Set(const T& InValue)
+	{
+		Value = InValue;
+	}
+
+	/* Sets the value. */
+	void Set(T&& InValue)
+	{
+		Value = MoveTemp(InValue);
+	}
+
 private:
 
 	/** The value. */
