Prevent macOS from putting the editor into App Nap mode while it's hidden and running an external process

diff --git a/Engine/Source/Runtime/Core/Private/Mac/MacPlatformProcess.cpp b/Engine/Source/Runtime/Core/Private/Mac/MacPlatformProcess.cpp
--- a/Engine/Source/Runtime/Core/Private/Mac/MacPlatformProcess.cpp
+++ b/Engine/Source/Runtime/Core/Private/Mac/MacPlatformProcess.cpp
@@ -352,7 +352,9 @@ bool FMacPlatformProcess::ExecProcess( const TCHAR* URL, const TCHAR* Params, in
 		
 		NSAutoReadPipe* StdErrPipe = [[NSAutoReadPipe new] autorelease];
 		[ProcessHandle setStandardError: (id)[StdErrPipe Pipe]];
-		
+
+		id<NSObject> Activity = [[NSProcessInfo processInfo] beginActivityWithOptions:NSActivityUserInitiated reason:@"ExecProcess"];
+
 		@try
 		{
 			[ProcessHandle launch];
@@ -373,7 +375,12 @@ bool FMacPlatformProcess::ExecProcess( const TCHAR* URL, const TCHAR* Params, in
 			{
 				[StdErrPipe copyPipeData: *OutStdErr];
 			}
-			
+
+			if (Activity)
+			{
+				[[NSProcessInfo processInfo] endActivity:Activity];
+			}
+
 			return true;
 		}
 		@catch (NSException* Exc)
@@ -386,6 +393,12 @@ bool FMacPlatformProcess::ExecProcess( const TCHAR* URL, const TCHAR* Params, in
 			{
 				*OutStdErr = FString([Exc reason]);
 			}
+
+			if (Activity)
+			{
+				[[NSProcessInfo processInfo] endActivity:Activity];
+			}
+
 			return false;
 		}
 	}
@@ -552,7 +565,9 @@ FProcHandle FMacPlatformProcess::CreateProc( const TCHAR* URL, const TCHAR* Parm
 		*OutProcessID = ProcessHandle ? [ProcessHandle processIdentifier] : 0;
 	}
 
-	return FProcHandle(ProcessHandle);
+	FProcHandle Handle(ProcessHandle);
+	Handle.Activity = [[NSProcessInfo processInfo] beginActivityWithOptions:NSActivityUserInitiated reason:@"CreateProc"];
+	return Handle;
 }
 
 FProcHandle FMacPlatformProcess::OpenProcess(uint32 ProcessID)
@@ -581,6 +596,10 @@ void FMacPlatformProcess::CloseProc( FProcHandle & ProcessHandle )
 {
 	SCOPED_AUTORELEASE_POOL;
 	[(NSTask*)ProcessHandle.Get() release];
+	if (ProcessHandle.Activity)
+	{
+		[[NSProcessInfo processInfo] endActivity:ProcessHandle.Activity];
+	}
 	ProcessHandle.Reset();
 }
 
diff --git a/Engine/Source/Runtime/Core/Public/Mac/MacPlatformProcess.h b/Engine/Source/Runtime/Core/Public/Mac/MacPlatformProcess.h
--- a/Engine/Source/Runtime/Core/Public/Mac/MacPlatformProcess.h
+++ b/Engine/Source/Runtime/Core/Public/Mac/MacPlatformProcess.h
@@ -16,12 +16,16 @@ public:
 	/** Default constructor. */
 	FORCEINLINE FProcHandle()
 		: TProcHandle()
+		, Activity(nil)
 	{}
 
 	/** Initialization constructor. */
 	FORCEINLINE explicit FProcHandle( HandleType Other )
 		: TProcHandle( Other )
+		, Activity(nil)
 	{}
+
+	id<NSObject> Activity;
 };
 
 /**
@@ -143,4 +147,4 @@ struct CORE_API FMacPlatformProcess : public FGenericPlatformProcess
 	static const TCHAR* UserLogsDir();
 };
 
-typedef FMacPlatformProcess FPlatformProcess;
\ No newline at end of file
+typedef FMacPlatformProcess FPlatformProcess;
