Fixing an issue that can cause intermittent crashes.

diff --git a/Engine/Source/Runtime/Core/Private/Mac/MacPlatformProcess.cpp b/Engine/Source/Runtime/Core/Private/Mac/MacPlatformProcess.cpp
--- a/Engine/Source/Runtime/Core/Private/Mac/MacPlatformProcess.cpp
+++ b/Engine/Source/Runtime/Core/Private/Mac/MacPlatformProcess.cpp
@@ -566,7 +566,12 @@ FProcHandle FMacPlatformProcess::CreateProc( const TCHAR* URL, const TCHAR* Parm
 	}
 
 	FProcHandle Handle(ProcessHandle);
-	Handle.Activity = [[NSProcessInfo processInfo] beginActivityWithOptions:NSActivityUserInitiated reason:@"CreateProc"];
+	if(ProcessHandle)
+	{
+		Handle.Activity = [[NSProcessInfo processInfo] beginActivityWithOptions:NSActivityUserInitiated reason:@"CreateProc"];
+		[Handle.Activity retain];
+	}
+
 	return Handle;
 }
 
@@ -595,11 +600,12 @@ void FMacPlatformProcess::WaitForProc( FProcHandle& ProcessHandle )
 void FMacPlatformProcess::CloseProc( FProcHandle & ProcessHandle )
 {
 	SCOPED_AUTORELEASE_POOL;
-	[(NSTask*)ProcessHandle.Get() release];
 	if (ProcessHandle.Activity)
 	{
 		[[NSProcessInfo processInfo] endActivity:ProcessHandle.Activity];
+		[ProcessHandle.Activity release];
 	}
+	[(NSTask*)ProcessHandle.Get() release];
 	ProcessHandle.Reset();
 }
 
