Copying JamieD's changes over from Release-4.20

diff --git a/Engine/Source/Runtime/Core/Private/Internationalization/TextHistory.cpp b/Engine/Source/Runtime/Core/Private/Internationalization/TextHistory.cpp
--- a/Engine/Source/Runtime/Core/Private/Internationalization/TextHistory.cpp
+++ b/Engine/Source/Runtime/Core/Private/Internationalization/TextHistory.cpp
@@ -1272,8 +1272,9 @@ FStringTableEntryConstPtr FTextHistory_StringTableEntry::GetStringTableEntry(con
 		{
 			if (bStringTableAssetPendingLoad && IsInGameThread())
 			{
-				// This path should never be taken when EDL is enabled
-				check(!GEventDrivenLoaderEnabled);
+                // This path should never be taken when EDL is enabled
+                // TODO: This assert is tripping in some cases (see UE-61107)
+                //check(!GEventDrivenLoaderEnabled);
 
 				// Attempt to load the string table asset now
 				FTextHistory_StringTableEntry* NonConstThis = const_cast<FTextHistory_StringTableEntry*>(this);
