Process created will have a different group pid (gpid) then the parent

diff --git a/Engine/Build/BatchFiles/RunUAT.sh b/Engine/Build/BatchFiles/RunUAT.sh
--- a/Engine/Build/BatchFiles/RunUAT.sh
+++ b/Engine/Build/BatchFiles/RunUAT.sh
@@ -11,7 +11,8 @@ echo Running AutomationTool...
 echo
 
 TermHandler() {
-	kill -TERM $UATPid 2> /dev/null
+	GPID=$(ps -o pgid= $UATPid)
+	kill -TERM -$GPID 2> /dev/null
 }
 
 # loop over the arguments, quoting spaces to pass to UAT proper
diff --git a/Engine/Source/Runtime/Core/Private/Unix/UnixPlatformProcess.cpp b/Engine/Source/Runtime/Core/Private/Unix/UnixPlatformProcess.cpp
--- a/Engine/Source/Runtime/Core/Private/Unix/UnixPlatformProcess.cpp
+++ b/Engine/Source/Runtime/Core/Private/Unix/UnixPlatformProcess.cpp
@@ -868,6 +868,9 @@ FProcHandle FUnixPlatformProcess::CreateProc(const TCHAR* URL, const TCHAR* Parm
 	posix_spawnattr_setsigdefault(&SpawnAttr, &SetToDefaultSignalSet);
 	SpawnFlags |= POSIX_SPAWN_SETSIGDEF;
 
+	// Makes spawned processes have its own unique group id so we can kill the entire group with out killing the parent
+	SpawnFlags |= POSIX_SPAWN_SETPGROUP;
+
 	int PosixSpawnErrNo = -1;
 	if (PipeWriteChild || PipeReadChild)
 	{
