[NO NON-PII COMMIT MESSAGE]

diff --git a/Engine/Source/Runtime/Core/Private/HAL/LowLevelMemTracker.cpp b/Engine/Source/Runtime/Core/Private/HAL/LowLevelMemTracker.cpp
--- a/Engine/Source/Runtime/Core/Private/HAL/LowLevelMemTracker.cpp
+++ b/Engine/Source/Runtime/Core/Private/HAL/LowLevelMemTracker.cpp
@@ -56,6 +56,7 @@ DECLARE_LLM_MEMORY_STAT(TEXT("GenericPlatformMallocCrash"), STAT_GenericPlatform
 DECLARE_LLM_MEMORY_STAT(TEXT("Engine Misc"), STAT_EngineMiscLLM, STATGROUP_LLMFULL);
 DECLARE_LLM_MEMORY_STAT(TEXT("TaskGraph Misc Tasks"), STAT_TaskGraphTasksMiscLLM, STATGROUP_LLMFULL);
 DECLARE_LLM_MEMORY_STAT(TEXT("Audio"), STAT_AudioLLM, STATGROUP_LLMFULL);
+DECLARE_LLM_MEMORY_STAT(TEXT("AudioMixer"), STAT_AudioMixerLLM, STATGROUP_LLMFULL);
 DECLARE_LLM_MEMORY_STAT(TEXT("FName"), STAT_FNameLLM, STATGROUP_LLMFULL);
 DECLARE_LLM_MEMORY_STAT(TEXT("Networking"), STAT_NetworkingLLM, STATGROUP_LLMFULL);
 DECLARE_LLM_MEMORY_STAT(TEXT("Meshes"), STAT_MeshesLLM, STATGROUP_LLMFULL);
diff --git a/Engine/Source/Runtime/Core/Public/HAL/LowLevelMemTracker.h b/Engine/Source/Runtime/Core/Public/HAL/LowLevelMemTracker.h
--- a/Engine/Source/Runtime/Core/Public/HAL/LowLevelMemTracker.h
+++ b/Engine/Source/Runtime/Core/Public/HAL/LowLevelMemTracker.h
@@ -133,6 +133,7 @@ enum class ELLMTagSet : uint8
 	macro(EngineMisc,							"EngineMisc",					GET_STATFNAME(STAT_EngineMiscLLM),							GET_STATFNAME(STAT_EngineSummaryLLM))		\
 	macro(TaskGraphTasksMisc,					"TaskGraphMiscTasks",			GET_STATFNAME(STAT_TaskGraphTasksMiscLLM),					GET_STATFNAME(STAT_EngineSummaryLLM))		\
 	macro(Audio,								"Audio",						GET_STATFNAME(STAT_AudioLLM),								GET_STATFNAME(STAT_AudioSummaryLLM))		\
+	macro(AudioMixer,							"AudioMixer",					GET_STATFNAME(STAT_AudioMixerLLM),							GET_STATFNAME(STAT_AudioSummaryLLM))	\
 	macro(FName,								"FName",						GET_STATFNAME(STAT_FNameLLM),								GET_STATFNAME(STAT_EngineSummaryLLM))		\
 	macro(Networking,							"Networking",					GET_STATFNAME(STAT_NetworkingLLM),							GET_STATFNAME(STAT_EngineSummaryLLM))		\
 	macro(Meshes,								"Meshes",						GET_STATFNAME(STAT_MeshesLLM),								GET_STATFNAME(STAT_MeshesSummaryLLM))		\
