CircularQueue

Merging Dev-Enterprise Count fix

- Going back to int32 var in the Count method to fix error when Tail is smaller than Head
- Adding a test to verify this case

diff --git a/Engine/Source/Runtime/Core/Private/Tests/Misc/CircularQueueTest.cpp b/Engine/Source/Runtime/Core/Private/Tests/Misc/CircularQueueTest.cpp
--- a/Engine/Source/Runtime/Core/Private/Tests/Misc/CircularQueueTest.cpp
+++ b/Engine/Source/Runtime/Core/Private/Tests/Misc/CircularQueueTest.cpp
@@ -60,6 +60,26 @@ bool FCircularQueueTest::RunTest(const FString& Parameters)
 		TestFalse(TEXT("A queue that had all items removed must not be full"), Queue.IsFull());
 	}
 
+	// queue with index wrapping around
+	{
+		TCircularQueue<int32> Queue(QueueSize);
+
+		//Fill queue
+		for (int32 Index = 0; Index < QueueSize - 1; ++Index)
+		{
+			TestTrue(TEXT("Adding to non-full queue must succeed"), Queue.Enqueue(Index));
+		}
+		
+		int32 Value = 0;
+		const int32 ExpectedSize = QueueSize - 1;
+		for (int32 Index = 0; Index < QueueSize; ++Index)
+		{
+			TestEqual(TEXT("Number of elements must be valid for all permutation of Tail and Head"), Queue.Count(), ExpectedSize);
+			TestTrue(TEXT("Removing from a non-empty queue must succeed"), Queue.Dequeue(Value));
+			TestTrue(TEXT("Adding to non-full queue must succeed"), Queue.Enqueue(Index));
+		}
+	}
+
 	return true;
 }
 
diff --git a/Engine/Source/Runtime/Core/Public/Containers/CircularQueue.h b/Engine/Source/Runtime/Core/Public/Containers/CircularQueue.h
--- a/Engine/Source/Runtime/Core/Public/Containers/CircularQueue.h
+++ b/Engine/Source/Runtime/Core/Public/Containers/CircularQueue.h
@@ -47,7 +47,7 @@ public:
 	 */
 	uint32 Count() const
 	{
-		int64 Count = Tail.Load() - Head.Load();
+		int32 Count = Tail.Load() - Head.Load();
 
 		if (Count < 0)
 		{
