Fire the ApplicationWillTerminate delegate when the editor console window is about to close. Should allow analytics to track cases where the editor is closed via the console window or when running as a commandlet.

diff --git a/Engine/Source/Runtime/Core/Private/Windows/WindowsPlatformMisc.cpp b/Engine/Source/Runtime/Core/Private/Windows/WindowsPlatformMisc.cpp
--- a/Engine/Source/Runtime/Core/Private/Windows/WindowsPlatformMisc.cpp
+++ b/Engine/Source/Runtime/Core/Private/Windows/WindowsPlatformMisc.cpp
@@ -549,6 +549,9 @@ void FWindowsPlatformMisc::PlatformInit()
  */
 static BOOL WINAPI ConsoleCtrlHandler( ::DWORD /*Type*/ )
 {
+	// Notify anyone listening that we're about to terminate
+	FCoreDelegates::ApplicationWillTerminateDelegate.Broadcast();
+
 	// make sure as much data is written to disk as possible
 	if (GLog)
 	{
