Fixed typo

diff --git a/Engine/Source/Runtime/Core/Private/Misc/Paths.cpp b/Engine/Source/Runtime/Core/Private/Misc/Paths.cpp
--- a/Engine/Source/Runtime/Core/Private/Misc/Paths.cpp
+++ b/Engine/Source/Runtime/Core/Private/Misc/Paths.cpp
@@ -556,7 +556,7 @@ const TArray<FString>& FPaths::GetRestrictedFolderNames()
 
 	if (!bBuiltArray)
 	{
-		RestrictedFolderNames.Add(TEXT("NotForLicencees"));
+		RestrictedFolderNames.Add(TEXT("NotForLicensees"));
 		RestrictedFolderNames.Add(TEXT("NoRedist"));
 		RestrictedFolderNames.Add(TEXT("CarefullyRedist"));
 		RestrictedFolderNames.Add(TEXT("EpicInternal"));
