Fix missing #include
#rnx

diff --git a/Engine/Source/Runtime/Core/Public/Math/IntVector.h b/Engine/Source/Runtime/Core/Public/Math/IntVector.h
--- a/Engine/Source/Runtime/Core/Public/Math/IntVector.h
+++ b/Engine/Source/Runtime/Core/Public/Math/IntVector.h
@@ -6,6 +6,7 @@
 #include "Misc/Crc.h"
 #include "Math/UnrealMathUtility.h"
 #include "Containers/UnrealString.h"
+#include "Serialization/StructuredArchive.h"
 
 /**
  * Structure for integer vectors in 3-d space.
