GitHub #4705 : Added weak lambda's for delegates and multicast delegates.

diff --git a/Engine/Source/Runtime/Core/Public/Delegates/DelegateInstancesImpl.h b/Engine/Source/Runtime/Core/Public/Delegates/DelegateInstancesImpl.h
--- a/Engine/Source/Runtime/Core/Public/Delegates/DelegateInstancesImpl.h
+++ b/Engine/Source/Runtime/Core/Public/Delegates/DelegateInstancesImpl.h
@@ -926,3 +926,148 @@ public:
 		return true;
 	}
 };
+
+/**
+ * Implements a weak object delegate binding for C++ functors, e.g. lambdas.
+ */
+template <typename FuncType, typename FunctorType, typename... VarTypes>
+class TWeakBaseFunctorDelegateInstance;
+
+template <typename WrappedRetValType, typename... ParamTypes, typename FunctorType, typename... VarTypes>
+class TWeakBaseFunctorDelegateInstance<WrappedRetValType(ParamTypes...), FunctorType, VarTypes...> : public IBaseDelegateInstance<typename TUnwrapType<WrappedRetValType>::Type(ParamTypes...)>
+{
+public:
+	typedef typename TUnwrapType<WrappedRetValType>::Type RetValType;
+
+private:
+	static_assert(TAreTypesEqual<FunctorType, typename TRemoveReference<FunctorType>::Type>::Value, "FunctorType cannot be a reference");
+
+	typedef IBaseDelegateInstance<typename TUnwrapType<WrappedRetValType>::Type(ParamTypes...)> Super;
+	typedef TWeakBaseFunctorDelegateInstance<RetValType(ParamTypes...), FunctorType, VarTypes...>   UnwrappedThisType;
+
+public:
+	TWeakBaseFunctorDelegateInstance(const UObject* InContextObject, const FunctorType& InFunctor, VarTypes... Vars)
+		: ContextObject(InContextObject)
+		, Functor      (InFunctor)
+		, Payload      (Vars...)
+		, Handle       (FDelegateHandle::GenerateNewHandle)
+	{
+	}
+
+	TWeakBaseFunctorDelegateInstance(const UObject* InContextObject, FunctorType&& InFunctor, VarTypes... Vars)
+		: ContextObject(InContextObject)
+		, Functor      (MoveTemp(InFunctor))
+		, Payload      (Vars...)
+		, Handle       (FDelegateHandle::GenerateNewHandle)
+	{
+	}
+
+	// IDelegateInstance interface
+
+#if USE_DELEGATE_TRYGETBOUNDFUNCTIONNAME
+
+	virtual FName TryGetBoundFunctionName() const override final
+	{
+		return NAME_None;
+	}
+
+#endif
+
+	virtual UObject* GetUObject() const override final
+	{
+		return ContextObject.Get();
+	}
+
+	// Deprecated
+	virtual bool HasSameObject(const void* InContextObject) const override final
+	{
+		return GetUObject() == InContextObject;
+	}
+
+	virtual bool IsSafeToExecute() const override final
+	{
+		return ContextObject.IsValid();
+	}
+
+public:
+	// IBaseDelegateInstance interface
+	virtual void CreateCopy(FDelegateBase& Base) override final
+	{
+		new (Base) UnwrappedThisType(*(UnwrappedThisType*)this);
+	}
+
+	virtual RetValType Execute(ParamTypes... Params) const override final
+	{
+		return Payload.ApplyAfter(Functor, Params...);
+	}
+
+	virtual FDelegateHandle GetHandle() const override final
+	{
+		return Handle;
+	}
+
+public:
+	/**
+	 * Creates a new static function delegate binding for the given function pointer.
+	 *
+	 * @param InFunctor C++ functor
+	 * @return The new delegate.
+	 */
+	FORCEINLINE static void Create(FDelegateBase& Base, const UObject* InContextObject, const FunctorType& InFunctor, VarTypes... Vars)
+	{
+		new (Base) UnwrappedThisType(InContextObject, InFunctor, Vars...);
+	}
+	FORCEINLINE static void Create(FDelegateBase& Base, const UObject* InContextObject, FunctorType&& InFunctor, VarTypes... Vars)
+	{
+		new (Base) UnwrappedThisType(InContextObject, MoveTemp(InFunctor), Vars...);
+	}
+
+private:
+	// Context object - the validity of this object controls the validity of the lambda
+	FWeakObjectPtr ContextObject;
+
+	// C++ functor
+	// We make this mutable to allow mutable lambdas to be bound and executed.  We don't really want to
+	// model the Functor as being a direct subobject of the delegate (which would maintain transivity of
+	// const - because the binding doesn't affect the substitutability of a copied delegate.
+	mutable typename TRemoveConst<FunctorType>::Type Functor;
+
+	// Payload member variables, if any.
+	TTuple<VarTypes...> Payload;
+
+	// The handle of this delegate
+	FDelegateHandle Handle;
+};
+
+template <typename FunctorType, typename... ParamTypes, typename... VarTypes>
+class TWeakBaseFunctorDelegateInstance<void(ParamTypes...), FunctorType, VarTypes...> : public TWeakBaseFunctorDelegateInstance<TTypeWrapper<void>(ParamTypes...), FunctorType, VarTypes...>
+{
+	typedef TWeakBaseFunctorDelegateInstance<TTypeWrapper<void>(ParamTypes...), FunctorType, VarTypes...> Super;
+
+public:
+	/**
+	 * Creates and initializes a new instance.
+	 *
+	 * @param InFunctor C++ functor
+	 */
+	TWeakBaseFunctorDelegateInstance(const UObject* InContextObject, const FunctorType& InFunctor, VarTypes... Vars)
+		: Super(InContextObject, InFunctor, Vars...)
+	{
+	}
+
+	TWeakBaseFunctorDelegateInstance(const UObject* InContextObject, FunctorType&& InFunctor, VarTypes... Vars)
+		: Super(InContextObject, MoveTemp(InFunctor), Vars...)
+	{
+	}
+
+	virtual bool ExecuteIfSafe(ParamTypes... Params) const override final
+	{
+		if (IsSafeToExecute())
+		{
+			Super::Execute(Params...);
+			return true;
+		}
+
+		return false;
+	}
+};
diff --git a/Engine/Source/Runtime/Core/Public/Delegates/DelegateSignatureImpl.inl b/Engine/Source/Runtime/Core/Public/Delegates/DelegateSignatureImpl.inl
--- a/Engine/Source/Runtime/Core/Public/Delegates/DelegateSignatureImpl.inl
+++ b/Engine/Source/Runtime/Core/Public/Delegates/DelegateSignatureImpl.inl
@@ -133,6 +133,20 @@ public:
 		return Result;
 	}
 
+	/**
+	 * Static: Creates a weak object C++ lambda delegate
+	 * technically this works for any functor types, but lambdas are the primary use case
+	 */
+	template<typename UserClass, typename FunctorType, typename... VarTypes>
+	FUNCTION_CHECK_RETURN_START
+	inline static TBaseDelegate<RetValType, ParamTypes...> CreateWeakLambda(UserClass* InUserObject, FunctorType&& InFunctor, VarTypes... Vars)
+	FUNCTION_CHECK_RETURN_END
+	{
+		TBaseDelegate<RetValType, ParamTypes...> Result;
+		TWeakBaseFunctorDelegateInstance<TFuncType, typename TRemoveReference<FunctorType>::Type, VarTypes...>::Create(Result, InUserObject, Forward<FunctorType>(InFunctor), Vars...);
+		return Result;
+	}
+
 	/**
 	 * Static: Creates a raw C++ pointer member function delegate.
 	 *
@@ -406,6 +420,16 @@ public:
 		*this = CreateLambda(Forward<FunctorType>(InFunctor), Vars...);
 	}
 
+	/**
+	 * Static: Binds a weak object C++ lambda delegate
+	 * technically this works for any functor types, but lambdas are the primary use case
+	 */
+	template<typename UserClass, typename FunctorType, typename... VarTypes>
+	inline void BindWeakLambda(UserClass* InUserObject, FunctorType&& InFunctor, VarTypes... Vars)
+	{
+		*this = CreateWeakLambda(InUserObject, Forward<FunctorType>(InFunctor), Vars...);
+	}
+
 	/**
 	 * Binds a raw C++ pointer delegate.
 	 *
@@ -715,6 +739,19 @@ public:
 		return Add(FDelegate::CreateLambda(Forward<FunctorType>(InFunctor), Vars...));
 	}
 
+	/**
+	 * Adds a weak object C++ lambda delegate
+	 * technically this works for any functor types, but lambdas are the primary use case
+	 *
+	 * @param	InUserObject	User object to bind to
+	 * @param	InFunctor		Functor (e.g. Lambda)
+	 */
+	template<typename UserClass, typename FunctorType, typename... VarTypes>
+	inline FDelegateHandle AddWeakLambda(UserClass* InUserObject, FunctorType&& InFunctor, VarTypes... Vars)
+	{
+		return Add(FDelegate::CreateWeakLambda(InUserObject, Forward<FunctorType>(InFunctor), Vars...));
+	}
+
 	/**
 	 * Adds a raw C++ pointer delegate.
 	 *
