Fix for forward declared CoreUObject weakobject types in delegates when building in Clang.

diff --git a/Engine/Source/Runtime/Core/Public/Delegates/DelegateInstancesImpl.h b/Engine/Source/Runtime/Core/Public/Delegates/DelegateInstancesImpl.h
--- a/Engine/Source/Runtime/Core/Public/Delegates/DelegateInstancesImpl.h
+++ b/Engine/Source/Runtime/Core/Public/Delegates/DelegateInstancesImpl.h
@@ -930,11 +930,11 @@ public:
 /**
  * Implements a weak object delegate binding for C++ functors, e.g. lambdas.
  */
-template <typename FuncType, typename FunctorType, typename... VarTypes>
+template <typename UserClass, typename FuncType, typename FunctorType, typename... VarTypes>
 class TWeakBaseFunctorDelegateInstance;
 
-template <typename WrappedRetValType, typename... ParamTypes, typename FunctorType, typename... VarTypes>
-class TWeakBaseFunctorDelegateInstance<WrappedRetValType(ParamTypes...), FunctorType, VarTypes...> : public IBaseDelegateInstance<typename TUnwrapType<WrappedRetValType>::Type(ParamTypes...)>
+template <typename UserClass, typename WrappedRetValType, typename... ParamTypes, typename FunctorType, typename... VarTypes>
+class TWeakBaseFunctorDelegateInstance<UserClass, WrappedRetValType(ParamTypes...), FunctorType, VarTypes...> : public IBaseDelegateInstance<typename TUnwrapType<WrappedRetValType>::Type(ParamTypes...)>
 {
 public:
 	typedef typename TUnwrapType<WrappedRetValType>::Type RetValType;
@@ -946,7 +946,7 @@ private:
 	typedef TWeakBaseFunctorDelegateInstance<RetValType(ParamTypes...), FunctorType, VarTypes...>   UnwrappedThisType;
 
 public:
-	TWeakBaseFunctorDelegateInstance(const UObject* InContextObject, const FunctorType& InFunctor, VarTypes... Vars)
+	TWeakBaseFunctorDelegateInstance(UserClass* InContextObject, const FunctorType& InFunctor, VarTypes... Vars)
 		: ContextObject(InContextObject)
 		, Functor      (InFunctor)
 		, Payload      (Vars...)
@@ -954,7 +954,7 @@ public:
 	{
 	}
 
-	TWeakBaseFunctorDelegateInstance(const UObject* InContextObject, FunctorType&& InFunctor, VarTypes... Vars)
+	TWeakBaseFunctorDelegateInstance(UserClass* InContextObject, FunctorType&& InFunctor, VarTypes... Vars)
 		: ContextObject(InContextObject)
 		, Functor      (MoveTemp(InFunctor))
 		, Payload      (Vars...)
@@ -1013,18 +1013,18 @@ public:
 	 * @param InFunctor C++ functor
 	 * @return The new delegate.
 	 */
-	FORCEINLINE static void Create(FDelegateBase& Base, const UObject* InContextObject, const FunctorType& InFunctor, VarTypes... Vars)
+	FORCEINLINE static void Create(FDelegateBase& Base, UserClass* InContextObject, const FunctorType& InFunctor, VarTypes... Vars)
 	{
 		new (Base) UnwrappedThisType(InContextObject, InFunctor, Vars...);
 	}
-	FORCEINLINE static void Create(FDelegateBase& Base, const UObject* InContextObject, FunctorType&& InFunctor, VarTypes... Vars)
+	FORCEINLINE static void Create(FDelegateBase& Base, UserClass* InContextObject, FunctorType&& InFunctor, VarTypes... Vars)
 	{
 		new (Base) UnwrappedThisType(InContextObject, MoveTemp(InFunctor), Vars...);
 	}
 
 private:
 	// Context object - the validity of this object controls the validity of the lambda
-	FWeakObjectPtr ContextObject;
+	TWeakObjectPtr<UserClass> ContextObject;
 
 	// C++ functor
 	// We make this mutable to allow mutable lambdas to be bound and executed.  We don't really want to
@@ -1039,8 +1039,8 @@ private:
 	FDelegateHandle Handle;
 };
 
-template <typename FunctorType, typename... ParamTypes, typename... VarTypes>
-class TWeakBaseFunctorDelegateInstance<void(ParamTypes...), FunctorType, VarTypes...> : public TWeakBaseFunctorDelegateInstance<TTypeWrapper<void>(ParamTypes...), FunctorType, VarTypes...>
+template <typename UserClass, typename FunctorType, typename... ParamTypes, typename... VarTypes>
+class TWeakBaseFunctorDelegateInstance<UserClass, void(ParamTypes...), FunctorType, VarTypes...> : public TWeakBaseFunctorDelegateInstance<UserClass, TTypeWrapper<void>(ParamTypes...), FunctorType, VarTypes...>
 {
 	typedef TWeakBaseFunctorDelegateInstance<TTypeWrapper<void>(ParamTypes...), FunctorType, VarTypes...> Super;
 
@@ -1050,12 +1050,12 @@ public:
 	 *
 	 * @param InFunctor C++ functor
 	 */
-	TWeakBaseFunctorDelegateInstance(const UObject* InContextObject, const FunctorType& InFunctor, VarTypes... Vars)
+	TWeakBaseFunctorDelegateInstance(UserClass* InContextObject, const FunctorType& InFunctor, VarTypes... Vars)
 		: Super(InContextObject, InFunctor, Vars...)
 	{
 	}
 
-	TWeakBaseFunctorDelegateInstance(const UObject* InContextObject, FunctorType&& InFunctor, VarTypes... Vars)
+	TWeakBaseFunctorDelegateInstance(UserClass* InContextObject, FunctorType&& InFunctor, VarTypes... Vars)
 		: Super(InContextObject, MoveTemp(InFunctor), Vars...)
 	{
 	}
diff --git a/Engine/Source/Runtime/Core/Public/Delegates/DelegateSignatureImpl.inl b/Engine/Source/Runtime/Core/Public/Delegates/DelegateSignatureImpl.inl
--- a/Engine/Source/Runtime/Core/Public/Delegates/DelegateSignatureImpl.inl
+++ b/Engine/Source/Runtime/Core/Public/Delegates/DelegateSignatureImpl.inl
@@ -143,7 +143,7 @@ public:
 	FUNCTION_CHECK_RETURN_END
 	{
 		TBaseDelegate<RetValType, ParamTypes...> Result;
-		TWeakBaseFunctorDelegateInstance<TFuncType, typename TRemoveReference<FunctorType>::Type, VarTypes...>::Create(Result, InUserObject, Forward<FunctorType>(InFunctor), Vars...);
+		TWeakBaseFunctorDelegateInstance<UserClass, FuncType, typename TRemoveReference<FunctorType>::Type, VarTypes...>::Create(Result, InUserObject, Forward<FunctorType>(InFunctor), Vars...);
 		return Result;
 	}
 
