Another Clang fix.

diff --git a/Engine/Source/Runtime/Core/Public/Delegates/DelegateInstancesImpl.h b/Engine/Source/Runtime/Core/Public/Delegates/DelegateInstancesImpl.h
--- a/Engine/Source/Runtime/Core/Public/Delegates/DelegateInstancesImpl.h
+++ b/Engine/Source/Runtime/Core/Public/Delegates/DelegateInstancesImpl.h
@@ -1062,7 +1062,7 @@ public:
 
 	virtual bool ExecuteIfSafe(ParamTypes... Params) const override final
 	{
-		if (IsSafeToExecute())
+		if (Super::IsSafeToExecute())
 		{
 			Super::Execute(Params...);
 			return true;
