Typo fix.

diff --git a/Engine/Source/Runtime/Core/Public/Delegates/DelegateSignatureImpl.inl b/Engine/Source/Runtime/Core/Public/Delegates/DelegateSignatureImpl.inl
--- a/Engine/Source/Runtime/Core/Public/Delegates/DelegateSignatureImpl.inl
+++ b/Engine/Source/Runtime/Core/Public/Delegates/DelegateSignatureImpl.inl
@@ -143,7 +143,7 @@ public:
 	FUNCTION_CHECK_RETURN_END
 	{
 		TBaseDelegate<RetValType, ParamTypes...> Result;
-		TWeakBaseFunctorDelegateInstance<UserClass, FuncType, typename TRemoveReference<FunctorType>::Type, VarTypes...>::Create(Result, InUserObject, Forward<FunctorType>(InFunctor), Vars...);
+		TWeakBaseFunctorDelegateInstance<UserClass, TFuncType, typename TRemoveReference<FunctorType>::Type, VarTypes...>::Create(Result, InUserObject, Forward<FunctorType>(InFunctor), Vars...);
 		return Result;
 	}
 
