Fix for binding an unbound TFunction to another TFunction with a different signature.  Also all null pointers now count as unbindings, not just nullptr.
TIsMemberPointer added.
TIsATFunction and TIsATFunctionRef renamed to remove the 'A's.

diff --git a/Engine/Source/Runtime/Core/Public/Templates/Function.h b/Engine/Source/Runtime/Core/Public/Templates/Function.h
--- a/Engine/Source/Runtime/Core/Public/Templates/Function.h
+++ b/Engine/Source/Runtime/Core/Public/Templates/Function.h
@@ -13,6 +13,8 @@
 #include "Templates/Invoke.h"
 #include "Templates/IsConstructible.h"
 #include "Templates/IsInvocable.h"
+#include "Templates/IsMemberPointer.h"
+#include "Templates/IsPointer.h"
 #include "Containers/ContainerAllocationPolicies.h"
 #include "Math/UnrealMathUtility.h"
 #include <new>
@@ -47,14 +49,22 @@ class TFunctionRef;
 /**
  * Traits class which checks if T is a TFunction<> type.
  */
-template <typename T> struct TIsATFunction               { enum { Value = false }; };
-template <typename T> struct TIsATFunction<TFunction<T>> { enum { Value = true  }; };
+template <typename T> struct TIsTFunction               { enum { Value = false }; };
+template <typename T> struct TIsTFunction<TFunction<T>> { enum { Value = true  }; };
+
+template <typename T> struct TIsTFunction<const          T> { enum { Value = TIsTFunction<T>::Value }; };
+template <typename T> struct TIsTFunction<      volatile T> { enum { Value = TIsTFunction<T>::Value }; };
+template <typename T> struct TIsTFunction<const volatile T> { enum { Value = TIsTFunction<T>::Value }; };
 
 /**
- * Traits class which checks if T is a TFunction<> type.
+ * Traits class which checks if T is a TFunctionRef<> type.
  */
-template <typename T> struct TIsATFunctionRef                  { enum { Value = false }; };
-template <typename T> struct TIsATFunctionRef<TFunctionRef<T>> { enum { Value = true  }; };
+template <typename T> struct TIsTFunctionRef                  { enum { Value = false }; };
+template <typename T> struct TIsTFunctionRef<TFunctionRef<T>> { enum { Value = true  }; };
+
+template <typename T> struct TIsTFunctionRef<const          T> { enum { Value = TIsTFunctionRef<T>::Value }; };
+template <typename T> struct TIsTFunctionRef<      volatile T> { enum { Value = TIsTFunctionRef<T>::Value }; };
+template <typename T> struct TIsTFunctionRef<const volatile T> { enum { Value = TIsTFunctionRef<T>::Value }; };
 
 /**
  * Private implementation details of TFunction and TFunctionRef.
@@ -355,6 +365,31 @@ namespace UE4Function_Private
 		TIsInvocable<FunctorType, ParamTypes...>
 	{
 	};
+
+	template <typename T>
+	struct TIsNullableBinding :
+		TOr<
+			TIsPointer<T>,
+			TIsMemberPointer<T>,
+			TIsTFunction<T>
+		>
+	{
+	};
+
+	template <typename T>
+	FORCEINLINE typename TEnableIf<TIsNullableBinding<T>::Value, bool>::Type IsBound(const T& Func)
+	{
+		// Function pointers, data member pointers, member function pointers and TFunctions
+		// can all be null/unbound, so test them using their boolean state.
+		return !!Func;
+	}
+
+	template <typename T>
+	FORCEINLINE typename TEnableIf<!TIsNullableBinding<T>::Value, bool>::Type IsBound(const T& Func)
+	{
+		// We can't tell if any other generic callable can be invoked, so just assume they can be.
+		return true;
+	}
 }
 
 /**
@@ -435,6 +470,8 @@ public:
 		// This constructor is disabled for function types because we want it to call the function pointer overload.
 		// It is also disabled for TFunctionRef types because VC is incorrectly treating it as a copy constructor.
 
+		checkf(UE4Function_Private::IsBound(Functor), TEXT("Cannot bind a null/unbound callable to a TFunctionRef"));
+
 		Set(&Functor);
 	}
 
@@ -457,6 +494,8 @@ public:
 		// This constructor is disabled for function types because we want it to call the function pointer overload.
 		// It is also disabled for TFunctionRef types because VC is incorrectly treating it as a copy constructor.
 
+		checkf(UE4Function_Private::IsBound(Functor), TEXT("Cannot bind a null/unbound callable to a TFunctionRef"));
+
 		Set(&Functor);
 	}
 
@@ -477,6 +516,8 @@ public:
 	{
 		// This constructor is enabled only for function types because we don't want weird errors from it being called with arbitrary pointers.
 
+		checkf(Function, TEXT("Cannot bind a null/unbound callable to a TFunctionRef"));
+
 		Set(Function);
 	}
 
@@ -623,10 +664,17 @@ public:
 		// reference):
 		//
 		// TFunction<int32(float)> MyFunction = [=](float F) -> int32 { return MyFunctionRef(F); };
-		static_assert(!TIsATFunctionRef<DecayedFunctorType>::Value, "Cannot construct a TFunction from a TFunctionRef");
+		static_assert(!TIsTFunctionRef<DecayedFunctorType>::Value, "Cannot construct a TFunction from a TFunctionRef");
 
-		OwnedType* NewObj = new (Storage) OwnedType(Forward<FunctorType>(InFunc));
-		Super::Set(&NewObj->Obj);
+		if (UE4Function_Private::IsBound(InFunc))
+		{
+			OwnedType* NewObj = new (Storage) OwnedType(Forward<FunctorType>(InFunc));
+			Super::Set(&NewObj->Obj);
+		}
+		else
+		{
+			Super::Unset();
+		}
 	}
 
 	/**
diff --git a/Engine/Source/Runtime/Core/Public/Templates/IsMemberPointer.h b/Engine/Source/Runtime/Core/Public/Templates/IsMemberPointer.h
new file mode 100644
--- /dev/null
+++ b/Engine/Source/Runtime/Core/Public/Templates/IsMemberPointer.h
@@ -0,0 +1,20 @@
+// Copyright 1998-2018 Epic Games, Inc. All Rights Reserved.
+
+#pragma once
+
+#include "CoreTypes.h"
+
+/**
+ * Traits class which tests if a type is a pointer to member (data member or member function).
+ */
+template <typename T>
+struct TIsMemberPointer
+{
+	enum { Value = false };
+};
+
+template <typename T, typename U> struct TIsMemberPointer<T U::*> { enum { Value = true }; };
+
+template <typename T> struct TIsMemberPointer<const          T> { enum { Value = TIsPointer<T>::Value }; };
+template <typename T> struct TIsMemberPointer<      volatile T> { enum { Value = TIsPointer<T>::Value }; };
+template <typename T> struct TIsMemberPointer<const volatile T> { enum { Value = TIsPointer<T>::Value }; };
diff --git a/Engine/Source/Runtime/Core/Public/Templates/IsPointer.h b/Engine/Source/Runtime/Core/Public/Templates/IsPointer.h
--- a/Engine/Source/Runtime/Core/Public/Templates/IsPointer.h
+++ b/Engine/Source/Runtime/Core/Public/Templates/IsPointer.h
@@ -13,10 +13,7 @@ struct TIsPointer
 	enum { Value = false };
 };
 
-template <typename T> struct TIsPointer<               T*> { enum { Value = true }; };
-template <typename T> struct TIsPointer<const          T*> { enum { Value = true }; };
-template <typename T> struct TIsPointer<      volatile T*> { enum { Value = true }; };
-template <typename T> struct TIsPointer<const volatile T*> { enum { Value = true }; };
+template <typename T> struct TIsPointer<T*> { enum { Value = true }; };
 
 template <typename T> struct TIsPointer<const          T> { enum { Value = TIsPointer<T>::Value }; };
 template <typename T> struct TIsPointer<      volatile T> { enum { Value = TIsPointer<T>::Value }; };
