TBA: ArchiveFromStructuredArchive adapter uses Inner to determine if it is outputting to text, and sets it's own ArIsTextFormat to false

diff --git a/Engine/Source/Runtime/Core/Private/Serialization/ArchiveFromStructuredArchive.cpp b/Engine/Source/Runtime/Core/Private/Serialization/ArchiveFromStructuredArchive.cpp
--- a/Engine/Source/Runtime/Core/Private/Serialization/ArchiveFromStructuredArchive.cpp
+++ b/Engine/Source/Runtime/Core/Private/Serialization/ArchiveFromStructuredArchive.cpp
@@ -11,8 +11,9 @@ FArchiveFromStructuredArchive::FArchiveFromStructuredArchive(FStructuredArchive:
 	// For some reason, the FArchive copy constructor will copy all the trivial members of the source archive, but then specifically set ArIsFilterEditorOnly to false, with a comment saying
 	// they don't know why it's doing this... make sure we inherit this flag here!
 	ArIsFilterEditorOnly = InnerArchive.ArIsFilterEditorOnly;
+	ArIsTextFormat = false;
 	
-	if (IsTextFormat())
+	if (InnerArchive.IsTextFormat())
 	{
 		Root = Slot.EnterRecord();
 
@@ -46,7 +47,7 @@ bool FArchiveFromStructuredArchive::Close()
 
 int64 FArchiveFromStructuredArchive::Tell()
 {
-	if (IsTextFormat())
+	if (InnerArchive.IsTextFormat())
 	{
 		return Pos;
 	}
@@ -64,7 +65,7 @@ int64 FArchiveFromStructuredArchive::TotalSize()
 
 void FArchiveFromStructuredArchive::Seek(int64 InPos)
 {
-	if (IsTextFormat())
+	if (InnerArchive.IsTextFormat())
 	{
 		check(Pos >= 0 && Pos <= Buffer.Num());
 		Pos = InPos;
@@ -77,7 +78,7 @@ void FArchiveFromStructuredArchive::Seek(int64 InPos)
 
 bool FArchiveFromStructuredArchive::AtEnd()
 {
-	if (IsTextFormat())
+	if (InnerArchive.IsTextFormat())
 	{
 		return Pos == Buffer.Num();
 	}
@@ -89,7 +90,7 @@ bool FArchiveFromStructuredArchive::AtEnd()
 
 FArchive& FArchiveFromStructuredArchive::operator<<(class FName& Value)
 {
-	if (IsTextFormat())
+	if (InnerArchive.IsTextFormat())
 	{
 		if (IsLoading())
 		{
@@ -117,7 +118,7 @@ FArchive& FArchiveFromStructuredArchive::operator<<(class FName& Value)
 
 FArchive& FArchiveFromStructuredArchive::operator<<(class UObject*& Value)
 {
-	if(IsTextFormat())
+	if(InnerArchive.IsTextFormat())
 	{
 		if (IsLoading())
 		{
@@ -186,7 +187,7 @@ FArchive& FArchiveFromStructuredArchive::operator<<(class UObject*& Value)
 
 FArchive& FArchiveFromStructuredArchive::operator<<(class FText& Value)
 {
-	if (IsTextFormat())
+	if (InnerArchive.IsTextFormat())
 	{
 		FText::SerializeText(*this, Value);
 	}
@@ -199,7 +200,7 @@ FArchive& FArchiveFromStructuredArchive::operator<<(class FText& Value)
 
 void FArchiveFromStructuredArchive::Serialize(void* V, int64 Length)
 {
-	if (IsTextFormat())
+	if (InnerArchive.IsTextFormat())
 	{
 		if (IsLoading())
 		{
@@ -228,7 +229,7 @@ void FArchiveFromStructuredArchive::Serialize(void* V, int64 Length)
 
 void FArchiveFromStructuredArchive::Commit()
 {
-	if (IsTextFormat())
+	if (InnerArchive.IsTextFormat())
 	{
 		SerializeInternal(Root.GetValue());
 	}
diff --git a/Engine/Source/Runtime/CoreUObject/Private/Serialization/ArchiveUObjectFromStructuredArchive.cpp b/Engine/Source/Runtime/CoreUObject/Private/Serialization/ArchiveUObjectFromStructuredArchive.cpp
--- a/Engine/Source/Runtime/CoreUObject/Private/Serialization/ArchiveUObjectFromStructuredArchive.cpp
+++ b/Engine/Source/Runtime/CoreUObject/Private/Serialization/ArchiveUObjectFromStructuredArchive.cpp
@@ -6,7 +6,7 @@ FArchiveUObjectFromStructuredArchive::FArchiveUObjectFromStructuredArchive(FStru
 	: FArchiveFromStructuredArchive(Slot)
 	, bPendingSerialize(true)
 {
-	if (IsTextFormat() && IsLoading())
+	if (InnerArchive.IsTextFormat() && IsLoading())
 	{
 		Commit();
 	}
