Incorrect assert removed.

diff --git a/Engine/Source/Runtime/Core/Public/Misc/MemStack.h b/Engine/Source/Runtime/Core/Public/Misc/MemStack.h
--- a/Engine/Source/Runtime/Core/Public/Misc/MemStack.h
+++ b/Engine/Source/Runtime/Core/Public/Misc/MemStack.h
@@ -293,7 +293,6 @@ public:
 		FORCEINLINE void MoveToEmpty(ForElementType& Other)
 		{
 			checkSlow(this != &Other);
-			check(!Data);
 
 			Data       = Other.Data;
 			Other.Data = nullptr;
