TBA: Converted FText over to structured archive. Implemented saving, but not loading.

diff --git a/Engine/Source/Runtime/Core/Private/Internationalization/StringTableCore.cpp b/Engine/Source/Runtime/Core/Private/Internationalization/StringTableCore.cpp
--- a/Engine/Source/Runtime/Core/Private/Internationalization/StringTableCore.cpp
+++ b/Engine/Source/Runtime/Core/Private/Internationalization/StringTableCore.cpp
@@ -587,10 +587,10 @@ void FStringTableRedirects::RedirectTableIdAndKey(FName& InOutTableId, FString&
 }
 
 
-void FStringTableReferenceCollection::CollectAssetReferences(const FName InTableId, FArchive& InAr)
+void FStringTableReferenceCollection::CollectAssetReferences(const FName InTableId, FStructuredArchive::FRecord Record)
 {
-	if (InAr.IsObjectReferenceCollector())
+	if (Record.GetUnderlyingArchive().IsObjectReferenceCollector())
 	{
-		IStringTableEngineBridge::CollectStringTableAssetReferences(InTableId, InAr);
+		IStringTableEngineBridge::CollectStringTableAssetReferences(InTableId, Record.EnterField(FIELD_NAME_TEXT("AssetReferences")));
 	}
 }
diff --git a/Engine/Source/Runtime/Core/Private/Internationalization/Text.cpp b/Engine/Source/Runtime/Core/Private/Internationalization/Text.cpp
--- a/Engine/Source/Runtime/Core/Private/Internationalization/Text.cpp
+++ b/Engine/Source/Runtime/Core/Private/Internationalization/Text.cpp
@@ -17,6 +17,8 @@
 #include "Internationalization/TextNamespaceUtil.h"
 #include "Internationalization/FastDecimalFormat.h"
 
+#include "Serialization/ArchiveFromStructuredArchive.h"
+
 #include "UObject/EditorObjectVersion.h"
 #include "HAL/PlatformProcess.h"
 
@@ -119,27 +121,27 @@ FNumberFormattingOptions::FNumberFormattingOptions()
 
 }
 
-FArchive& operator<<(FArchive& Ar, FNumberFormattingOptions& Value)
+void operator<<(FStructuredArchive::FSlot Slot, FNumberFormattingOptions& Value)
 {
-	Ar.UsingCustomVersion(FEditorObjectVersion::GUID);
+	FArchive& UnderlyingArchive = Slot.GetUnderlyingArchive();
+	FStructuredArchive::FRecord Record = Slot.EnterRecord();
+	UnderlyingArchive.UsingCustomVersion(FEditorObjectVersion::GUID);
 
-	if (Ar.CustomVer(FEditorObjectVersion::GUID) >= FEditorObjectVersion::AddedAlwaysSignNumberFormattingOption)
+	if (UnderlyingArchive.CustomVer(FEditorObjectVersion::GUID) >= FEditorObjectVersion::AddedAlwaysSignNumberFormattingOption)
 	{
-		Ar << Value.AlwaysSign;
+		Record << NAMED_ITEM("AlwaysSign", Value.AlwaysSign);
 	}
 
-	Ar << Value.UseGrouping;
+	Record << NAMED_ITEM("UseGrouping", Value.UseGrouping);
 
 	int8 RoundingModeInt8 = (int8)Value.RoundingMode;
-	Ar << RoundingModeInt8;
+	Record << NAMED_ITEM("RoundingMode", RoundingModeInt8);
 	Value.RoundingMode = (ERoundingMode)RoundingModeInt8;
 
-	Ar << Value.MinimumIntegralDigits;
-	Ar << Value.MaximumIntegralDigits;
-	Ar << Value.MinimumFractionalDigits;
-	Ar << Value.MaximumFractionalDigits;
-
-	return Ar;
+	Record << NAMED_ITEM("MinimumIntegralDigits", Value.MinimumIntegralDigits);
+	Record << NAMED_ITEM("MaximumIntegralDigits", Value.MaximumIntegralDigits);
+	Record << NAMED_ITEM("MinimumFractionalDigits", Value.MinimumFractionalDigits);
+	Record << NAMED_ITEM("MaximumFractionalDigits", Value.MaximumFractionalDigits);
 }
 
 uint32 GetTypeHash( const FNumberFormattingOptions& Key )
@@ -724,26 +726,34 @@ bool FText::FindText( const FString& Namespace, const FString& Key, FText& OutTe
 
 void FText::SerializeText(FArchive& Ar, FText& Value)
 {
+	SerializeText(FStructuredArchiveFromArchive(Ar).GetSlot(), Value);
+}
+
+void FText::SerializeText(FStructuredArchive::FSlot Slot, FText& Value)
+{
+	FArchive& UnderlyingArchive = Slot.GetUnderlyingArchive();
+	FStructuredArchive::FRecord Record = Slot.EnterRecord();
+
 	//When duplicating, the CDO is used as the template, then values for the instance are assigned.
 	//If we don't duplicate the string, the CDO and the instance are both pointing at the same thing.
 	//This would result in all subsequently duplicated objects stamping over formerly duplicated ones.
 
 	// Older FText's stored their "SourceString", that is now stored in a history class so move it there
-	if (Ar.IsLoading() && Ar.UE4Ver() < VER_UE4_FTEXT_HISTORY)
+	if (UnderlyingArchive.IsLoading() && UnderlyingArchive.UE4Ver() < VER_UE4_FTEXT_HISTORY)
 	{
 		FString SourceStringToImplantIntoHistory;
-		Ar << SourceStringToImplantIntoHistory;
+		Record << NAMED_FIELD(SourceStringToImplantIntoHistory);
 
 		FTextDisplayStringPtr DisplayString;
 
 		// Namespaces and keys are no longer stored in the FText, we need to read them in and discard
-		if (Ar.UE4Ver() >= VER_UE4_ADDED_NAMESPACE_AND_KEY_DATA_TO_FTEXT)
+		if (UnderlyingArchive.UE4Ver() >= VER_UE4_ADDED_NAMESPACE_AND_KEY_DATA_TO_FTEXT)
 		{
 			FString Namespace;
 			FString Key;
 
-			Ar << Namespace;
-			Ar << Key;
+			Record << NAMED_FIELD(Namespace);
+			Record << NAMED_FIELD(Key);
 
 			// Get the DisplayString using the namespace, key, and source string.
 			DisplayString = FTextLocalizationManager::Get().GetDisplayString(Namespace, Key, &SourceStringToImplantIntoHistory);
@@ -758,61 +768,61 @@ void FText::SerializeText(FArchive& Ar, FText& Value)
 	}
 
 #if WITH_EDITOR
-	if (Ar.IsCooking() && Ar.IsSaving() && Ar.IsPersistent() && (Ar.GetDebugSerializationFlags() & DSF_EnableCookerWarnings))
+	if (UnderlyingArchive.IsCooking() && UnderlyingArchive.IsSaving() && UnderlyingArchive.IsPersistent() && (UnderlyingArchive.GetDebugSerializationFlags() & DSF_EnableCookerWarnings))
 	{
 		if (!!(Value.Flags & ETextFlag::ConvertedProperty))
 		{
-			UE_LOG(LogText, Warning, TEXT("Saving FText \"%s\" which has been converted at load time please resave source package %s to avoid determinisitic cook and localization issues."), *Value.ToString(), *Ar.GetArchiveName());
+			UE_LOG(LogText, Warning, TEXT("Saving FText \"%s\" which has been converted at load time please resave source package %s to avoid determinisitic cook and localization issues."), *Value.ToString(), *UnderlyingArchive.GetArchiveName());
 		}
 		else if (!!(Value.Flags & ETextFlag::InitializedFromString))
 		{
-			UE_LOG(LogText, Warning, TEXT("Saving FText \"%s\" which has been initialized from FString at cook time resave of source package %s may fix issue."), *Value.ToString(), *Ar.GetArchiveName())
+			UE_LOG(LogText, Warning, TEXT("Saving FText \"%s\" which has been initialized from FString at cook time resave of source package %s may fix issue."), *Value.ToString(), *UnderlyingArchive.GetArchiveName())
 		}
 	}
 #endif
 
 	const int32 OriginalFlags = Value.Flags;
 
-	if(Ar.IsSaving())
+	if(UnderlyingArchive.IsSaving())
 	{
 		Value.TextData->PersistText(); // We always need to do this when saving so that we can save the history correctly
-		if(Ar.IsPersistent())
+		if(UnderlyingArchive.IsPersistent())
 		{
 			Value.Flags &= ~(ETextFlag::ConvertedProperty | ETextFlag::InitializedFromString); // Remove conversion flag before saving.
 		}
 	}
-	Ar << Value.Flags;
+	Record << NAMED_ITEM("Flags", Value.Flags);
 
-	if (Ar.IsLoading() && Ar.IsPersistent())
+	if (UnderlyingArchive.IsLoading() && UnderlyingArchive.IsPersistent())
 	{
 		Value.Flags &= ~(ETextFlag::ConvertedProperty | ETextFlag::InitializedFromString); // Remove conversion flag before saving.
 	}
 
-	if (Ar.IsSaving())
+	if (UnderlyingArchive.IsSaving())
 	{
 		Value.Flags = OriginalFlags;
 	}
 
-	if (Ar.UE4Ver() >= VER_UE4_FTEXT_HISTORY)
+	if (UnderlyingArchive.UE4Ver() >= VER_UE4_FTEXT_HISTORY)
 	{
 		bool bSerializeHistory = true;
 
-		if (Ar.IsSaving())
+		if (UnderlyingArchive.IsSaving())
 		{
 			// Skip the history for empty texts
 			bSerializeHistory = !Value.IsEmpty();
 
 			if (!bSerializeHistory)
 			{
-				int8 NoHistory = INDEX_NONE;
-				Ar << NoHistory;
+				int8 HistoryType = INDEX_NONE;
+				Record << NAMED_FIELD(HistoryType);
 			}
 		}
-		else if (Ar.IsLoading())
+		else if (UnderlyingArchive.IsLoading())
 		{
 			// The type is serialized during the serialization of the history, during deserialization we need to deserialize it and create the correct history
 			int8 HistoryType = INDEX_NONE;
-			Ar << HistoryType;
+			Record << NAMED_FIELD(HistoryType);
 
 			// Create the history class based on the serialized type
 			switch((ETextHistoryType)HistoryType)
@@ -889,16 +899,16 @@ void FText::SerializeText(FArchive& Ar, FText& Value)
 		if(bSerializeHistory)
 		{
 			FTextHistory& MutableTextHistory = Value.TextData->GetMutableTextHistory();
-			MutableTextHistory.Serialize(Ar);
+			MutableTextHistory.Serialize(Record);
 
 			if (Value.TextData->OwnsLocalizedString())
 			{
-				MutableTextHistory.SerializeForDisplayString(Ar, Value.TextData->GetMutableLocalizedString());
+				MutableTextHistory.SerializeForDisplayString(Record, Value.TextData->GetMutableLocalizedString());
 			}
 		}
 	}
 
-	if(Ar.IsLoading())
+	if(UnderlyingArchive.IsLoading())
 	{
 		Value.Rebuild();
 	}
@@ -911,7 +921,7 @@ void FText::SerializeText(FArchive& Ar, FText& Value)
 
 	if( Value.ShouldGatherForLocalization() )
 	{
-		Ar.ThisRequiresLocalizationGather();
+		UnderlyingArchive.ThisRequiresLocalizationGather();
 	}
 }
 
@@ -1078,46 +1088,45 @@ bool FText::IdenticalTo( const FText& Other ) const
 	return TextData == Other.TextData || TextData->GetLocalizedString() == Other.TextData->GetLocalizedString();
 }
 
-FArchive& operator<<(FArchive& Ar, FFormatArgumentValue& Value)
+void operator<<(FStructuredArchive::FSlot Slot, FFormatArgumentValue& Value)
 {
-	int8 TypeAsInt8 = Value.Type;
-	Ar << TypeAsInt8;
-	Value.Type = (EFormatArgumentType::Type)TypeAsInt8;
+	FStructuredArchive::FRecord Record = Slot.EnterRecord();
+	int8 TypeAsInt8 = Value.GetType();
+	Record << NAMED_ITEM("Type", TypeAsInt8);
+	Value  = (EFormatArgumentType::Type)TypeAsInt8;
 
 	switch(Value.Type)
 	{
 	case EFormatArgumentType::Double:
 		{
-			Ar << Value.DoubleValue;
+			Record << NAMED_ITEM("Value", Value.DoubleValue);
 			break;
 		}
 	case EFormatArgumentType::Float:
 		{
-			Ar << Value.FloatValue;
+			Record << NAMED_ITEM("Value", Value.FloatValue);
 			break;
 		}
 	case EFormatArgumentType::Int:
 		{
-			Ar << Value.IntValue;
+			Record << NAMED_ITEM("Value", Value.IntValue);
 			break;
 		}
 	case EFormatArgumentType::UInt:
 		{
-			Ar << Value.UIntValue;
+			Record << NAMED_ITEM("Value", Value.UIntValue);
 			break;
 		}
 	case EFormatArgumentType::Text:
 		{
-			if(Ar.IsLoading())
+			if(Slot.GetUnderlyingArchive().IsLoading())
 			{
 				Value.TextValue = FText();
 			}
-			Ar << Value.TextValue.GetValue();
+			Record << NAMED_ITEM("Value", Value.TextValue.GetValue());
 			break;
 		}
 	}
-	
-	return Ar;
 }
 
 FString FFormatArgumentValue::ToFormattedString(const bool bInRebuildText, const bool bInRebuildAsSource) const
@@ -1183,37 +1192,40 @@ void FFormatArgumentData::ResetValue()
 	ArgumentValueGender = ETextGender::Masculine;
 }
 
-FArchive& operator<<(FArchive& Ar, FFormatArgumentData& Value)
+void operator<<(FStructuredArchive::FSlot Slot, FFormatArgumentData& Value)
 {
-	Ar.UsingCustomVersion(FEditorObjectVersion::GUID);
+	FArchive& UnderlyingArchive = Slot.GetUnderlyingArchive();
+	FStructuredArchive::FRecord Record = Slot.EnterRecord();
+
+	UnderlyingArchive.UsingCustomVersion(FEditorObjectVersion::GUID);
 
-	if (Ar.IsLoading())
+	if (UnderlyingArchive.IsLoading())
 	{
 		// ArgumentName was changed to be FString rather than FText, so we need to convert older data to ensure serialization stays happy outside of UStruct::SerializeTaggedProperties.
-		if (Ar.UE4Ver() >= VER_UE4_K2NODE_VAR_REFERENCEGUIDS) // There was no version bump for this change, but VER_UE4_K2NODE_VAR_REFERENCEGUIDS was made at almost the same time.
+		if (UnderlyingArchive.UE4Ver() >= VER_UE4_K2NODE_VAR_REFERENCEGUIDS) // There was no version bump for this change, but VER_UE4_K2NODE_VAR_REFERENCEGUIDS was made at almost the same time.
 		{
-			Ar << Value.ArgumentName;
+			Record << NAMED_ITEM("ArgumentName", Value.ArgumentName);
 		}
 		else
 		{
 			FText TempValue;
-			Ar << TempValue;
+			Record << NAMED_ITEM("ArgumentName", TempValue);
 			Value.ArgumentName = TempValue.ToString();
 		}
 	}
-	if (Ar.IsSaving())
+	if (UnderlyingArchive.IsSaving())
 	{
-		Ar << Value.ArgumentName;
+		Record << NAMED_ITEM("ArgumentName", Value.ArgumentName);
 	}
 
 	uint8 TypeAsByte = Value.ArgumentValueType;
-	if (Ar.IsLoading())
+	if (UnderlyingArchive.IsLoading())
 	{
 		Value.ResetValue();
 
-		if (Ar.CustomVer(FEditorObjectVersion::GUID) >= FEditorObjectVersion::TextFormatArgumentDataIsVariant)
+		if (UnderlyingArchive.CustomVer(FEditorObjectVersion::GUID) >= FEditorObjectVersion::TextFormatArgumentDataIsVariant)
 		{
-			Ar << TypeAsByte;
+			Record << NAMED_ITEM("Type", TypeAsByte);
 		}
 		else
 		{
@@ -1221,31 +1233,32 @@ FArchive& operator<<(FArchive& Ar, FFormatArgumentData& Value)
 			TypeAsByte = EFormatArgumentType::Text;
 		}
 	}
-	else if (Ar.IsSaving())
+	else if (UnderlyingArchive.IsSaving())
 	{
-		Ar << TypeAsByte;
+		Record << NAMED_ITEM("Type", TypeAsByte);
 	}
 
 	Value.ArgumentValueType = (EFormatArgumentType::Type)TypeAsByte;
 	switch (Value.ArgumentValueType)
 	{
 	case EFormatArgumentType::Int:
-		Ar << Value.ArgumentValueInt;
+		Record << NAMED_ITEM("Value", Value.ArgumentValueInt);
 		break;
 	case EFormatArgumentType::Float:
-		Ar << Value.ArgumentValueFloat;
+		Record << NAMED_ITEM("Value", Value.ArgumentValueFloat);
 		break;
 	case EFormatArgumentType::Text:
-		Ar << Value.ArgumentValue;
+		Record << NAMED_ITEM("Value", Value.ArgumentValue);
 		break;
 	case EFormatArgumentType::Gender:
-		Ar << (uint8&)Value.ArgumentValueGender;
-		break;
+		{
+			uint8& Gender = (uint8&)Value.ArgumentValueGender;
+			Record << NAMED_ITEM("Value", Gender);
+			break;
+		}
 	default:
 		break;
 	}
-
-	return Ar;
 }
 
 FTextSnapshot::FTextSnapshot()
diff --git a/Engine/Source/Runtime/Core/Private/Internationalization/TextHistory.cpp b/Engine/Source/Runtime/Core/Private/Internationalization/TextHistory.cpp
--- a/Engine/Source/Runtime/Core/Private/Internationalization/TextHistory.cpp
+++ b/Engine/Source/Runtime/Core/Private/Internationalization/TextHistory.cpp
@@ -63,9 +63,9 @@ bool FTextHistory::GetHistoricNumericData(const FText& InText, FHistoricTextNume
 	return false;
 }
 
-void FTextHistory::SerializeForDisplayString(FArchive& Ar, FTextDisplayStringPtr& InOutDisplayString)
+void FTextHistory::SerializeForDisplayString(FStructuredArchive::FRecord Record, FTextDisplayStringPtr& InOutDisplayString)
 {
-	if(Ar.IsLoading())
+	if(Record.GetUnderlyingArchive().IsLoading())
 	{
 		// We will definitely need to do a rebuild later
 		Revision = 0;
@@ -135,19 +135,21 @@ const FString* FTextHistory_Base::GetSourceString() const
 	return &SourceString;
 }
 
-void FTextHistory_Base::Serialize( FArchive& Ar )
+void FTextHistory_Base::Serialize(FStructuredArchive::FRecord Record)
 {
 	// If I serialize out the Namespace and Key HERE, then we can load it up.
-	if(Ar.IsSaving())
+	if(Record.GetUnderlyingArchive().IsSaving())
 	{
 		int8 HistoryType = (int8)ETextHistoryType::Base;
-		Ar << HistoryType;
+		Record << NAMED_FIELD(HistoryType);
 	}
 }
 
-void FTextHistory_Base::SerializeForDisplayString(FArchive& Ar, FTextDisplayStringPtr& InOutDisplayString)
+void FTextHistory_Base::SerializeForDisplayString(FStructuredArchive::FRecord Record, FTextDisplayStringPtr& InOutDisplayString)
 {
-	if(Ar.IsLoading())
+	FArchive& BaseArchive = Record.GetUnderlyingArchive();
+
+	if(BaseArchive.IsLoading())
 	{
 		// We will definitely need to do a rebuild later
 		Revision = 0;
@@ -155,16 +157,16 @@ void FTextHistory_Base::SerializeForDisplayString(FArchive& Ar, FTextDisplayStri
 		FString Namespace;
 		FString Key;
 
-		Ar << Namespace;
-		Ar << Key;
-		Ar << SourceString;
+		Record << NAMED_FIELD(Namespace);
+		Record << NAMED_FIELD(Key);
+		Record << NAMED_FIELD(SourceString);
 
 #if USE_STABLE_LOCALIZATION_KEYS
 		// Make sure the package namespace for this text property is up-to-date
 		// We do this on load (as well as save) to handle cases where data is being duplicated, as it will be written by one package and loaded into another
-		if (GIsEditor && !Ar.HasAnyPortFlags(PPF_DuplicateVerbatim | PPF_DuplicateForPIE))
+		if (GIsEditor && !Record.GetUnderlyingArchive().HasAnyPortFlags(PPF_DuplicateVerbatim | PPF_DuplicateForPIE))
 		{
-			const FString PackageNamespace = TextNamespaceUtil::GetPackageNamespace(Ar);
+			const FString PackageNamespace = TextNamespaceUtil::GetPackageNamespace(BaseArchive);
 			if (!PackageNamespace.IsEmpty())
 			{
 				const FString FullNamespace = TextNamespaceUtil::BuildFullNamespace(Namespace, PackageNamespace);
@@ -189,7 +191,7 @@ void FTextHistory_Base::SerializeForDisplayString(FArchive& Ar, FTextDisplayStri
 		// Using the deserialized namespace and key, find the DisplayString.
 		InOutDisplayString = FTextLocalizationManager::Get().GetDisplayString(Namespace, Key, &SourceString);
 	}
-	else if(Ar.IsSaving())
+	else if(BaseArchive.IsSaving())
 	{
 		check(InOutDisplayString.IsValid());
 
@@ -197,7 +199,7 @@ void FTextHistory_Base::SerializeForDisplayString(FArchive& Ar, FTextDisplayStri
 		FString Key;
 		const bool bFoundNamespaceAndKey = FTextLocalizationManager::Get().FindNamespaceAndKeyFromDisplayString(InOutDisplayString.ToSharedRef(), Namespace, Key);
 
-		if (Ar.IsCooking())
+		if (BaseArchive.IsCooking())
 		{
 			// We strip the package localization off the serialized text for a cooked game, as they're not used at runtime
 			Namespace = TextNamespaceUtil::StripPackageNamespace(Namespace);
@@ -205,9 +207,9 @@ void FTextHistory_Base::SerializeForDisplayString(FArchive& Ar, FTextDisplayStri
 		else
 #if USE_STABLE_LOCALIZATION_KEYS
 		// Make sure the package namespace for this text property is up-to-date
-		if (GIsEditor && !Ar.HasAnyPortFlags(PPF_DuplicateVerbatim | PPF_DuplicateForPIE))
+		if (GIsEditor && !BaseArchive.HasAnyPortFlags(PPF_DuplicateVerbatim | PPF_DuplicateForPIE))
 		{
-			const FString PackageNamespace = TextNamespaceUtil::GetPackageNamespace(Ar);
+			const FString PackageNamespace = TextNamespaceUtil::GetPackageNamespace(BaseArchive);
 			if (!PackageNamespace.IsEmpty())
 			{
 				const FString FullNamespace = TextNamespaceUtil::BuildFullNamespace(Namespace, PackageNamespace);
@@ -223,7 +225,7 @@ void FTextHistory_Base::SerializeForDisplayString(FArchive& Ar, FTextDisplayStri
 #endif // USE_STABLE_LOCALIZATION_KEYS
 
 		// If this has no key, give it a GUID for a key
-		if (!bFoundNamespaceAndKey && GIsEditor && (Ar.IsPersistent() && !Ar.HasAnyPortFlags(PPF_Duplicate)))
+		if (!bFoundNamespaceAndKey && GIsEditor && (BaseArchive.IsPersistent() && !BaseArchive.HasAnyPortFlags(PPF_Duplicate)))
 		{
 			Key = FGuid::NewGuid().ToString();
 			if (!FTextLocalizationManager::Get().AddDisplayString(InOutDisplayString.ToSharedRef(), Namespace, Key))
@@ -235,13 +237,13 @@ void FTextHistory_Base::SerializeForDisplayString(FArchive& Ar, FTextDisplayStri
 		}
 
 		// Serialize the Namespace
-		Ar << Namespace;
+		Record << NAMED_FIELD(Namespace);
 
 		// Serialize the Key
-		Ar << Key;
+		Record << NAMED_FIELD(Key);
 
 		// Serialize the SourceString
-		Ar << SourceString;
+		Record << NAMED_FIELD(SourceString);
 	}
 }
 
@@ -282,27 +284,29 @@ FString FTextHistory_NamedFormat::BuildInvariantDisplayString() const
 	return FTextFormatter::FormatStr(SourceFmt, Arguments, true, true);
 }
 
-void FTextHistory_NamedFormat::Serialize( FArchive& Ar )
+void FTextHistory_NamedFormat::Serialize(FStructuredArchive::FRecord Record)
 {
-	if(Ar.IsSaving())
+	FArchive& BaseArchive = Record.GetUnderlyingArchive();
+
+	if(BaseArchive.IsSaving())
 	{
 		int8 HistoryType = (int8)ETextHistoryType::NamedFormat;
-		Ar << HistoryType;
+		Record << NAMED_FIELD(HistoryType);
 	}
 
-	if (Ar.IsSaving())
+	if (BaseArchive.IsSaving())
 	{
 		FText FormatText = SourceFmt.GetSourceText();
-		Ar << FormatText;
+		Record.EnterField(FIELD_NAME_TEXT("FormatText")) << FormatText;
 	}
-	else if (Ar.IsLoading())
+	else if (BaseArchive.IsLoading())
 	{
 		FText FormatText;
-		Ar << FormatText;
+		Record.EnterField(FIELD_NAME_TEXT("FormatText")) << FormatText;
 		SourceFmt = FTextFormat(FormatText);
 	}
 
-	Ar << Arguments;
+	Record << NAMED_FIELD(Arguments);
 }
 
 void FTextHistory_NamedFormat::GetHistoricFormatData(const FText& InText, TArray<FHistoricTextFormatData>& OutHistoricFormatData) const
@@ -361,27 +365,29 @@ FString FTextHistory_OrderedFormat::BuildInvariantDisplayString() const
 	return FTextFormatter::FormatStr(SourceFmt, Arguments, true, true);
 }
 
-void FTextHistory_OrderedFormat::Serialize( FArchive& Ar )
+void FTextHistory_OrderedFormat::Serialize(FStructuredArchive::FRecord Record)
 {
-	if(Ar.IsSaving())
+	FArchive& BaseArchive = Record.GetUnderlyingArchive();
+
+	if(BaseArchive.IsSaving())
 	{
 		int8 HistoryType = (int8)ETextHistoryType::OrderedFormat;
-		Ar << HistoryType;
+		Record << NAMED_FIELD(HistoryType);
 	}
 
-	if (Ar.IsSaving())
+	if (BaseArchive.IsSaving())
 	{
 		FText FormatText = SourceFmt.GetSourceText();
-		Ar << FormatText;
+		Record.EnterField(FIELD_NAME_TEXT("FormatText")) << FormatText;
 	}
-	else if (Ar.IsLoading())
+	else if (BaseArchive.IsLoading())
 	{
 		FText FormatText;
-		Ar << FormatText;
+		Record.EnterField(FIELD_NAME_TEXT("FormatText")) << FormatText;
 		SourceFmt = FTextFormat(FormatText);
 	}
 
-	Ar << Arguments;
+	Record << NAMED_FIELD(Arguments);
 }
 
 void FTextHistory_OrderedFormat::GetHistoricFormatData(const FText& InText, TArray<FHistoricTextFormatData>& OutHistoricFormatData) const
@@ -447,27 +453,29 @@ FString FTextHistory_ArgumentDataFormat::BuildInvariantDisplayString() const
 	return FTextFormatter::FormatStr(SourceFmt, Arguments, true, true);
 }
 
-void FTextHistory_ArgumentDataFormat::Serialize( FArchive& Ar )
+void FTextHistory_ArgumentDataFormat::Serialize(FStructuredArchive::FRecord Record)
 {
-	if(Ar.IsSaving())
+	FArchive& BaseArchive = Record.GetUnderlyingArchive();
+
+	if(BaseArchive.IsSaving())
 	{
 		int8 HistoryType = (int8)ETextHistoryType::ArgumentFormat;
-		Ar << HistoryType;
+		Record << NAMED_FIELD(HistoryType);
 	}
 
-	if (Ar.IsSaving())
+	if (BaseArchive.IsSaving())
 	{
 		FText FormatText = SourceFmt.GetSourceText();
-		Ar << FormatText;
+		Record.EnterField(FIELD_NAME_TEXT("FormatText")) << FormatText;
 	}
-	else if (Ar.IsLoading())
+	else if (BaseArchive.IsLoading())
 	{
 		FText FormatText;
-		Ar << FormatText;
+		Record.EnterField(FIELD_NAME_TEXT("FormatText")) << FormatText;
 		SourceFmt = FTextFormat(FormatText);
 	}
 
-	Ar << Arguments;
+	Record << NAMED_FIELD(Arguments);
 }
 
 void FTextHistory_ArgumentDataFormat::GetHistoricFormatData(const FText& InText, TArray<FHistoricTextFormatData>& OutHistoricFormatData) const
@@ -545,14 +553,16 @@ FTextHistory_FormatNumber& FTextHistory_FormatNumber::operator=(FTextHistory_For
 	return *this;
 }
 
-void FTextHistory_FormatNumber::Serialize(FArchive& Ar)
+void FTextHistory_FormatNumber::Serialize(FStructuredArchive::FRecord Record)
 {
-	Ar << SourceValue;
+	FArchive& BaseArchive = Record.GetUnderlyingArchive();
+
+	Record << NAMED_FIELD(SourceValue);
 
 	bool bHasFormatOptions = FormatOptions.IsSet();
-	Ar << bHasFormatOptions;
+	Record << NAMED_FIELD(bHasFormatOptions);
 
-	if(Ar.IsLoading())
+	if(BaseArchive.IsLoading())
 	{
 		if(bHasFormatOptions)
 		{
@@ -566,18 +576,19 @@ void FTextHistory_FormatNumber::Serialize(FArchive& Ar)
 	if(bHasFormatOptions)
 	{
 		check(FormatOptions.IsSet());
-		Ar << FormatOptions.GetValue();
+		FNumberFormattingOptions& Options = FormatOptions.GetValue();
+		Record << NAMED_FIELD(Options);
 	}
 
-	if(Ar.IsSaving())
+	if(BaseArchive.IsSaving())
 	{
 		FString CultureName = TargetCulture.IsValid()? TargetCulture->GetName() : FString();
-		Ar << CultureName;
+		Record << NAMED_FIELD(CultureName);
 	}
-	else if(Ar.IsLoading())
+	else if(BaseArchive.IsLoading())
 	{
 		FString CultureName;
-		Ar << CultureName;
+		Record << NAMED_FIELD(CultureName);
 
 		if(!CultureName.IsEmpty())
 		{
@@ -646,15 +657,17 @@ FString FTextHistory_AsNumber::BuildInvariantDisplayString() const
 	return BuildNumericDisplayString(FormattingRules);
 }
 
-void FTextHistory_AsNumber::Serialize( FArchive& Ar )
+void FTextHistory_AsNumber::Serialize(FStructuredArchive::FRecord Record)
 {
-	if(Ar.IsSaving())
+	FArchive& BaseArchive = Record.GetUnderlyingArchive();
+
+	if(BaseArchive.IsSaving())
 	{
 		int8 HistoryType = (int8)ETextHistoryType::AsNumber;
-		Ar << HistoryType;
+		Record << NAMED_FIELD(HistoryType);
 	}
 
-	FTextHistory_FormatNumber::Serialize(Ar);
+	FTextHistory_FormatNumber::Serialize(Record);
 }
 
 bool FTextHistory_AsNumber::GetHistoricNumericData(const FText& InText, FHistoricTextNumericData& OutHistoricNumericData) const
@@ -702,15 +715,17 @@ FString FTextHistory_AsPercent::BuildInvariantDisplayString() const
 	return BuildNumericDisplayString(FormattingRules, 100);
 }
 
-void FTextHistory_AsPercent::Serialize( FArchive& Ar )
+void FTextHistory_AsPercent::Serialize(FStructuredArchive::FRecord Record)
 {
-	if(Ar.IsSaving())
+	FArchive& BaseArchive = Record.GetUnderlyingArchive();
+
+	if(BaseArchive.IsSaving())
 	{
 		int8 HistoryType = (int8)ETextHistoryType::AsPercent;
-		Ar << HistoryType;
+		Record << NAMED_FIELD(HistoryType);
 	}
 
-	FTextHistory_FormatNumber::Serialize(Ar);
+	FTextHistory_FormatNumber::Serialize(Record);
 }
 
 bool FTextHistory_AsPercent::GetHistoricNumericData(const FText& InText, FHistoricTextNumericData& OutHistoricNumericData) const
@@ -766,20 +781,22 @@ FString FTextHistory_AsCurrency::BuildInvariantDisplayString() const
 	return BuildNumericDisplayString(FormattingRules);
 }
 
-void FTextHistory_AsCurrency::Serialize( FArchive& Ar )
+void FTextHistory_AsCurrency::Serialize(FStructuredArchive::FRecord Record)
 {
-	if(Ar.IsSaving())
+	FArchive& BaseArchive = Record.GetUnderlyingArchive();
+
+	if(BaseArchive.IsSaving())
 	{
 		int8 HistoryType = (int8)ETextHistoryType::AsCurrency;
-		Ar << HistoryType;
+		Record << NAMED_FIELD(HistoryType);
 	}
 
-	if (Ar.UE4Ver() >= VER_UE4_ADDED_CURRENCY_CODE_TO_FTEXT)
+	if (BaseArchive.UE4Ver() >= VER_UE4_ADDED_CURRENCY_CODE_TO_FTEXT)
 	{
-		Ar << CurrencyCode;
+		Record << NAMED_FIELD(CurrencyCode);
 	}
 
-	FTextHistory_FormatNumber::Serialize(Ar);
+	FTextHistory_FormatNumber::Serialize(Record);
 }
 
 ///////////////////////////////////////
@@ -815,34 +832,36 @@ FTextHistory_AsDate& FTextHistory_AsDate::operator=(FTextHistory_AsDate&& Other)
 	return *this;
 }
 
-void FTextHistory_AsDate::Serialize(FArchive& Ar)
+void FTextHistory_AsDate::Serialize(FStructuredArchive::FRecord Record)
 {
-	if(Ar.IsSaving())
+	FArchive& BaseArchive = Record.GetUnderlyingArchive();
+
+	if(BaseArchive.IsSaving())
 	{
 		int8 HistoryType = (int8)ETextHistoryType::AsDate;
-		Ar << HistoryType;
+		Record << NAMED_FIELD(HistoryType);
 	}
 
-	Ar << SourceDateTime;
+	Record << NAMED_FIELD(SourceDateTime);
 
 	int8 DateStyleInt8 = (int8)DateStyle;
-	Ar << DateStyleInt8;
+	Record << NAMED_FIELD(DateStyleInt8);
 	DateStyle = (EDateTimeStyle::Type)DateStyleInt8;
 
-	if( Ar.UE4Ver() >= VER_UE4_FTEXT_HISTORY_DATE_TIMEZONE )
+	if( BaseArchive.UE4Ver() >= VER_UE4_FTEXT_HISTORY_DATE_TIMEZONE )
 	{
-		Ar << TimeZone;
+		Record << NAMED_FIELD(TimeZone);
 	}
 
-	if(Ar.IsSaving())
+	if(BaseArchive.IsSaving())
 	{
 		FString CultureName = TargetCulture.IsValid()? TargetCulture->GetName() : FString();
-		Ar << CultureName;
+		Record << NAMED_FIELD(CultureName);
 	}
-	else if(Ar.IsLoading())
+	else if(BaseArchive.IsLoading())
 	{
 		FString CultureName;
-		Ar << CultureName;
+		Record << NAMED_FIELD(CultureName);
 
 		if(!CultureName.IsEmpty())
 		{
@@ -902,31 +921,33 @@ FTextHistory_AsTime& FTextHistory_AsTime::operator=(FTextHistory_AsTime&& Other)
 	return *this;
 }
 
-void FTextHistory_AsTime::Serialize(FArchive& Ar)
+void FTextHistory_AsTime::Serialize(FStructuredArchive::FRecord Record)
 {
-	if(Ar.IsSaving())
+	FArchive& BaseArchive = Record.GetUnderlyingArchive();
+
+	if(BaseArchive.IsSaving())
 	{
 		int8 HistoryType = (int8)ETextHistoryType::AsTime;
-		Ar << HistoryType;
+		Record << NAMED_FIELD(HistoryType);
 	}
 
-	Ar << SourceDateTime;
+	Record << NAMED_FIELD(SourceDateTime);
 
 	int8 TimeStyleInt8 = (int8)TimeStyle;
-	Ar << TimeStyleInt8;
+	Record << NAMED_ITEM("TimeStyle", TimeStyleInt8);
 	TimeStyle = (EDateTimeStyle::Type)TimeStyleInt8;
 
-	Ar << TimeZone;
+	Record << NAMED_FIELD(TimeZone);
 
-	if(Ar.IsSaving())
+	if(BaseArchive.IsSaving())
 	{
 		FString CultureName = TargetCulture.IsValid()? TargetCulture->GetName() : FString();
-		Ar << CultureName;
+		Record << NAMED_FIELD(CultureName);
 	}
-	else if(Ar.IsLoading())
+	else if(BaseArchive.IsLoading())
 	{
 		FString CultureName;
-		Ar << CultureName;
+		Record << NAMED_FIELD(CultureName);
 
 		if(!CultureName.IsEmpty())
 		{
@@ -989,35 +1010,37 @@ FTextHistory_AsDateTime& FTextHistory_AsDateTime::operator=(FTextHistory_AsDateT
 	return *this;
 }
 
-void FTextHistory_AsDateTime::Serialize(FArchive& Ar)
+void FTextHistory_AsDateTime::Serialize(FStructuredArchive::FRecord Record)
 {
-	if(Ar.IsSaving())
+	FArchive& BaseArchive = Record.GetUnderlyingArchive();
+
+	if(BaseArchive.IsSaving())
 	{
 		int8 HistoryType = (int8)ETextHistoryType::AsDateTime;
-		Ar << HistoryType;
+		Record << NAMED_FIELD(HistoryType);
 	}
 
-	Ar << SourceDateTime;
+	Record << NAMED_FIELD(SourceDateTime);
 
 	int8 DateStyleInt8 = (int8)DateStyle;
-	Ar << DateStyleInt8;
+	Record << NAMED_ITEM("DateStyle", DateStyleInt8);
 	DateStyle = (EDateTimeStyle::Type)DateStyleInt8;
 
 	int8 TimeStyleInt8 = (int8)TimeStyle;
-	Ar << TimeStyleInt8;
+	Record << NAMED_ITEM("TimeStyle", TimeStyleInt8);
 	TimeStyle = (EDateTimeStyle::Type)TimeStyleInt8;
 
-	Ar << TimeZone;
+	Record << NAMED_FIELD(TimeZone);
 
-	if(Ar.IsSaving())
+	if(BaseArchive.IsSaving())
 	{
 		FString CultureName = TargetCulture.IsValid()? TargetCulture->GetName() : FString();
-		Ar << CultureName;
+		Record << NAMED_FIELD(CultureName);
 	}
-	else if(Ar.IsLoading())
+	else if(BaseArchive.IsLoading())
 	{
 		FString CultureName;
-		Ar << CultureName;
+		Record << NAMED_FIELD(CultureName);
 
 		if(!CultureName.IsEmpty())
 		{
@@ -1071,16 +1094,20 @@ FTextHistory_Transform& FTextHistory_Transform::operator=(FTextHistory_Transform
 	return *this;
 }
 
-void FTextHistory_Transform::Serialize(FArchive& Ar)
+void FTextHistory_Transform::Serialize(FStructuredArchive::FRecord Record)
 {
-	if (Ar.IsSaving())
+	FArchive& BaseArchive = Record.GetUnderlyingArchive();
+
+	if (BaseArchive.IsSaving())
 	{
 		int8 HistoryType = (int8)ETextHistoryType::Transform;
-		Ar << HistoryType;
+		Record << NAMED_FIELD(HistoryType);
 	}
 
-	Ar << SourceText;
-	Ar << (uint8&)TransformType;
+	Record << NAMED_FIELD(SourceText);
+
+	uint8& TransformTypeRef = (uint8&)TransformType;
+	Record << NAMED_ITEM("TransformType", TransformTypeRef);
 }
 
 FString FTextHistory_Transform::BuildLocalizedDisplayString() const
@@ -1204,21 +1231,23 @@ void FTextHistory_StringTableEntry::GetTableIdAndKey(FName& OutTableId, FString&
 	OutKey = Key;
 }
 
-void FTextHistory_StringTableEntry::Serialize(FArchive& Ar)
+void FTextHistory_StringTableEntry::Serialize(FStructuredArchive::FRecord Record)
 {
-	if (Ar.IsSaving())
+	FArchive& BaseArchive = Record.GetUnderlyingArchive();
+
+	if (BaseArchive.IsSaving())
 	{
 		int8 HistoryType = (int8)ETextHistoryType::StringTableEntry;
-		Ar << HistoryType;
+		Record << NAMED_FIELD(HistoryType);
 	}
 
-	if (Ar.IsLoading())
+	if (BaseArchive.IsLoading())
 	{
 		// We will definitely need to do a rebuild later
 		Revision = 0;
 
-		Ar << TableId;
-		Ar << Key;
+		Record << NAMED_FIELD(TableId);
+		Record << NAMED_FIELD(Key);
 
 		// String Table assets should already have been created via dependency loading when using the EDL (although they may not be fully loaded yet)
 		FStringTableRedirects::RedirectTableIdAndKey(TableId, Key, GEventDrivenLoaderEnabled ? EStringTableLoadingPolicy::Find : EStringTableLoadingPolicy::FindOrLoad);
@@ -1227,7 +1256,7 @@ void FTextHistory_StringTableEntry::Serialize(FArchive& Ar)
 		StringTableEntry.Reset();
 		GetStringTableEntry();
 	}
-	else if (Ar.IsSaving())
+	else if (BaseArchive.IsSaving())
 	{
 		// Update the table ID and key on save to make sure they're up-to-date
 		FStringTableEntryConstPtr StringTableEntryPin = GetStringTableEntry();
@@ -1237,17 +1266,17 @@ void FTextHistory_StringTableEntry::Serialize(FArchive& Ar)
 			FStringTableRegistry::Get().FindTableIdAndKey(DisplayString.ToSharedRef(), TableId, Key);
 		}
 
-		Ar << TableId;
-		Ar << Key;
+		Record << NAMED_FIELD(TableId);
+		Record << NAMED_FIELD(Key);
 	}
 
 	// Collect string table asset references
-	FStringTableReferenceCollection::CollectAssetReferences(TableId, Ar);
+	FStringTableReferenceCollection::CollectAssetReferences(TableId, Record);
 }
 
-void FTextHistory_StringTableEntry::SerializeForDisplayString(FArchive& Ar, FTextDisplayStringPtr& InOutDisplayString)
+void FTextHistory_StringTableEntry::SerializeForDisplayString(FStructuredArchive::FRecord Record, FTextDisplayStringPtr& InOutDisplayString)
 {
-	if (Ar.IsLoading())
+	if (Record.GetUnderlyingArchive().IsLoading())
 	{
 		// We will definitely need to do a rebuild later
 		Revision = 0;
diff --git a/Engine/Source/Runtime/Core/Private/Internationalization/TextHistory.h b/Engine/Source/Runtime/Core/Private/Internationalization/TextHistory.h
--- a/Engine/Source/Runtime/Core/Private/Internationalization/TextHistory.h
+++ b/Engine/Source/Runtime/Core/Private/Internationalization/TextHistory.h
@@ -51,11 +51,11 @@ public:
 	/** Build the display string for the invariant culture */
 	virtual FString BuildInvariantDisplayString() const = 0;
 	
-	/** Serializes the history to/from an FArchive */
-	virtual void Serialize(FArchive& Ar) = 0;
+	/** Serializes the history to/from a structured archive slot */
+	virtual void Serialize(FStructuredArchive::FRecord Record) = 0;
 
 	/** Serializes data needed to get the FText's DisplayString */
-	virtual void SerializeForDisplayString(FArchive& Ar, FTextDisplayStringPtr& InOutDisplayString);
+	virtual void SerializeForDisplayString(FStructuredArchive::FRecord Record, FTextDisplayStringPtr& InOutDisplayString);
 
 	/** Returns TRUE if the Revision is out of date */
 	virtual bool IsOutOfDate() const;
@@ -103,8 +103,8 @@ public:
 	virtual ETextHistoryType GetType() const override { return ETextHistoryType::Base; }
 	virtual FString BuildLocalizedDisplayString() const override;
 	virtual FString BuildInvariantDisplayString() const override;
-	virtual void Serialize(FArchive& Ar) override;
-	virtual void SerializeForDisplayString(FArchive& Ar, FTextDisplayStringPtr& InOutDisplayString) override;
+	virtual void Serialize(FStructuredArchive::FRecord Record) override;
+	virtual void SerializeForDisplayString(FStructuredArchive::FRecord Record, FTextDisplayStringPtr& InOutDisplayString) override;
 	virtual const FString* GetSourceString() const override;
 	//~ End FTextHistory Interface
 
@@ -137,7 +137,7 @@ public:
 	virtual ETextHistoryType GetType() const override { return ETextHistoryType::NamedFormat; }
 	virtual FString BuildLocalizedDisplayString() const override;
 	virtual FString BuildInvariantDisplayString() const override;
-	virtual void Serialize(FArchive& Ar) override;
+	virtual void Serialize(FStructuredArchive::FRecord Record) override;
 	virtual void GetHistoricFormatData(const FText& InText, TArray<FHistoricTextFormatData>& OutHistoricFormatData) const override;
 	//~ End FTextHistory Interface
 
@@ -167,7 +167,7 @@ public:
 	virtual ETextHistoryType GetType() const override { return ETextHistoryType::OrderedFormat; }
 	virtual FString BuildLocalizedDisplayString() const override;
 	virtual FString BuildInvariantDisplayString() const override;
-	virtual void Serialize(FArchive& Ar) override;
+	virtual void Serialize(FStructuredArchive::FRecord Record) override;
 	virtual void GetHistoricFormatData(const FText& InText, TArray<FHistoricTextFormatData>& OutHistoricFormatData) const override;
 	//~ End FTextHistory Interface
 
@@ -197,7 +197,7 @@ public:
 	virtual ETextHistoryType GetType() const override { return ETextHistoryType::ArgumentFormat; }
 	virtual FString BuildLocalizedDisplayString() const override;
 	virtual FString BuildInvariantDisplayString() const override;
-	virtual void Serialize(FArchive& Ar) override;
+	virtual void Serialize(FStructuredArchive::FRecord Record) override;
 	virtual void GetHistoricFormatData(const FText& InText, TArray<FHistoricTextFormatData>& OutHistoricFormatData) const override;
 	//~ End FTextHistory Interface
 
@@ -224,7 +224,7 @@ public:
 	FTextHistory_FormatNumber& operator=(FTextHistory_FormatNumber&& Other);
 
 	//~ Begin FTextHistory Interface
-	virtual void Serialize(FArchive& Ar) override;
+	virtual void Serialize(FStructuredArchive::FRecord Record) override;
 	//~ End FTextHistory interface
 
 protected:
@@ -259,7 +259,7 @@ public:
 	virtual ETextHistoryType GetType() const override { return ETextHistoryType::AsNumber; }
 	virtual FString BuildLocalizedDisplayString() const override;
 	virtual FString BuildInvariantDisplayString() const override;
-	virtual void Serialize(FArchive& Ar) override;
+	virtual void Serialize(FStructuredArchive::FRecord Record) override;
 	virtual bool GetHistoricNumericData(const FText& InText, FHistoricTextNumericData& OutHistoricNumericData) const override;
 	//~ End FTextHistory interface
 
@@ -284,7 +284,7 @@ public:
 	virtual ETextHistoryType GetType() const override { return ETextHistoryType::AsPercent; }
 	virtual FString BuildLocalizedDisplayString() const override;
 	virtual FString BuildInvariantDisplayString() const override;
-	virtual void Serialize(FArchive& Ar) override;
+	virtual void Serialize(FStructuredArchive::FRecord Record) override;
 	virtual bool GetHistoricNumericData(const FText& InText, FHistoricTextNumericData& OutHistoricNumericData) const override;
 	//~ End FTextHistory interface
 
@@ -309,7 +309,7 @@ public:
 	virtual ETextHistoryType GetType() const override { return ETextHistoryType::AsCurrency; }
 	virtual FString BuildLocalizedDisplayString() const override;
 	virtual FString BuildInvariantDisplayString() const override;
-	virtual void Serialize(FArchive& Ar) override;
+	virtual void Serialize(FStructuredArchive::FRecord Record) override;
 	//~ End FTextHistory Interface
 
 private:
@@ -336,7 +336,7 @@ public:
 	virtual ETextHistoryType GetType() const override { return ETextHistoryType::AsDate; }
 	virtual FString BuildLocalizedDisplayString() const override;
 	virtual FString BuildInvariantDisplayString() const override;
-	virtual void Serialize(FArchive& Ar) override;
+	virtual void Serialize(FStructuredArchive::FRecord Record) override;
 	//~ End FTextHistory Interface
 
 private:
@@ -369,7 +369,7 @@ public:
 	virtual ETextHistoryType GetType() const override { return ETextHistoryType::AsTime; }
 	virtual FString BuildLocalizedDisplayString() const override;
 	virtual FString BuildInvariantDisplayString() const override;
-	virtual void Serialize(FArchive& Ar) override;
+	virtual void Serialize(FStructuredArchive::FRecord Record) override;
 	//~ End FTextHistory Interface
 
 private:
@@ -402,7 +402,7 @@ public:
 	virtual ETextHistoryType GetType() const override { return ETextHistoryType::AsDateTime; }
 	virtual FString BuildLocalizedDisplayString() const override;
 	virtual FString BuildInvariantDisplayString() const override;
-	virtual void Serialize(FArchive& Ar) override;
+	virtual void Serialize(FStructuredArchive::FRecord Record) override;
 	//~ End FTextHistory Interfaces
 
 private:
@@ -445,7 +445,7 @@ public:
 	virtual ETextHistoryType GetType() const override { return ETextHistoryType::Transform; }
 	virtual FString BuildLocalizedDisplayString() const override;
 	virtual FString BuildInvariantDisplayString() const override;
-	virtual void Serialize(FArchive& Ar) override;
+	virtual void Serialize(FStructuredArchive::FRecord Record) override;
 	virtual void GetHistoricFormatData(const FText& InText, TArray<FHistoricTextFormatData>& OutHistoricFormatData) const override;
 	virtual bool GetHistoricNumericData(const FText& InText, FHistoricTextNumericData& OutHistoricNumericData) const override;
 	//~ End FTextHistory Interfaces
@@ -476,8 +476,8 @@ public:
 	virtual ETextHistoryType GetType() const override { return ETextHistoryType::StringTableEntry; }
 	virtual FString BuildLocalizedDisplayString() const override;
 	virtual FString BuildInvariantDisplayString() const override;
-	virtual void Serialize(FArchive& Ar) override;
-	virtual void SerializeForDisplayString(FArchive& Ar, FTextDisplayStringPtr& InOutDisplayString) override;
+	virtual void Serialize(FStructuredArchive::FRecord Record) override;
+	virtual void SerializeForDisplayString(FStructuredArchive::FRecord Record, FTextDisplayStringPtr& InOutDisplayString) override;
 	virtual const FString* GetSourceString() const override;
 	//~ End FTextHistory Interface
 
diff --git a/Engine/Source/Runtime/Core/Private/Serialization/StructuredArchive.cpp b/Engine/Source/Runtime/Core/Private/Serialization/StructuredArchive.cpp
--- a/Engine/Source/Runtime/Core/Private/Serialization/StructuredArchive.cpp
+++ b/Engine/Source/Runtime/Core/Private/Serialization/StructuredArchive.cpp
@@ -346,6 +346,13 @@ void FStructuredArchive::FSlot::operator<< (UObject*& Value)
 	Ar.LeaveSlot();
 }
 
+void FStructuredArchive::FSlot::operator<< (FText& Value)
+{
+	Ar.EnterSlot(ElementId);
+	Ar.Formatter.Serialize(Value);
+	Ar.LeaveSlot();
+}
+
 void FStructuredArchive::FSlot::operator<< (FWeakObjectPtr& Value)
 {
 	Ar.EnterSlot(ElementId);
diff --git a/Engine/Source/Runtime/Core/Public/Internationalization/StringTableCore.h b/Engine/Source/Runtime/Core/Public/Internationalization/StringTableCore.h
--- a/Engine/Source/Runtime/Core/Public/Internationalization/StringTableCore.h
+++ b/Engine/Source/Runtime/Core/Public/Internationalization/StringTableCore.h
@@ -175,11 +175,11 @@ public:
 	}
 
 	/** Collect a string table asset reference */
-	static void CollectStringTableAssetReferences(const FName InTableId, FArchive& InAr)
+	static void CollectStringTableAssetReferences(const FName InTableId, FStructuredArchive::FSlot Slot)
 	{
 		if (InstancePtr)
 		{
-			InstancePtr->CollectStringTableAssetReferencesImpl(InTableId, InAr);
+			InstancePtr->CollectStringTableAssetReferencesImpl(InTableId, Slot);
 		}
 	}
 
@@ -187,7 +187,7 @@ protected:
 	virtual ~IStringTableEngineBridge() {}
 
 	virtual void RedirectAndLoadStringTableAssetImpl(FName& InOutTableId, const EStringTableLoadingPolicy InLoadingPolicy) = 0;
-	virtual void CollectStringTableAssetReferencesImpl(const FName InTableId, FArchive& InAr) = 0;
+	virtual void CollectStringTableAssetReferencesImpl(const FName InTableId, FStructuredArchive::FSlot Slot) = 0;
 
 	/** Singleton instance, populated by the derived type */
 	static IStringTableEngineBridge* InstancePtr;
@@ -213,5 +213,5 @@ struct CORE_API FStringTableRedirects
 struct CORE_API FStringTableReferenceCollection
 {
 	/** Collect asset references from the given table ID */
-	static void CollectAssetReferences(const FName InTableId, FArchive& InAr);
+	static void CollectAssetReferences(const FName InTableId, FStructuredArchive::FRecord Record);
 };
diff --git a/Engine/Source/Runtime/Core/Public/Internationalization/Text.h b/Engine/Source/Runtime/Core/Public/Internationalization/Text.h
--- a/Engine/Source/Runtime/Core/Public/Internationalization/Text.h
+++ b/Engine/Source/Runtime/Core/Public/Internationalization/Text.h
@@ -160,7 +160,7 @@ struct CORE_API FNumberFormattingOptions
 	int32 MaximumFractionalDigits;
 	FNumberFormattingOptions& SetMaximumFractionalDigits( int32 InValue ){ MaximumFractionalDigits = InValue; return *this; }
 
-	friend FArchive& operator<<(FArchive& Ar, FNumberFormattingOptions& Value);
+	friend void operator<<(FStructuredArchive::FSlot Slot, FNumberFormattingOptions& Value);
 
 	/** Get the hash code to use for the given formatting options */
 	friend uint32 GetTypeHash( const FNumberFormattingOptions& Key );
@@ -182,7 +182,7 @@ struct CORE_API FNumberParsingOptions
 	bool UseGrouping;
 	FNumberParsingOptions& SetUseGrouping( bool InValue ){ UseGrouping = InValue; return *this; }
 
-	friend FArchive& operator<<(FArchive& Ar, FNumberParsingOptions& Value);
+	friend void operator<<(FStructuredArchive::FSlot Slot, FNumberParsingOptions& Value);
 
 	/** Get the hash code to use for the given parsing options */
 	friend uint32 GetTypeHash( const FNumberParsingOptions& Key );
@@ -537,7 +537,8 @@ private:
 
 	FText( FString&& InSourceString, const FString& InNamespace, const FString& InKey, uint32 InFlags=0 );
 
-	static void SerializeText( FArchive& Ar, FText& Value );
+	static void SerializeText(FArchive& Ar, FText& Value);
+	static void SerializeText(FStructuredArchive::FSlot Slot, FText& Value);
 
 	/** Returns the source string of the FText */
 	const FString& GetSourceString() const;
@@ -578,6 +579,8 @@ public:
 	friend class FStringTableRegistry;
 	friend class FArchive;
 	friend class FArchiveFromStructuredArchive;
+	friend class FJsonArchiveInputFormatter;
+	friend class FJsonArchiveOutputFormatter;
 	friend class UTextProperty;
 	friend class FFormatArgumentValue;
 	friend class FTextHistory_NamedFormat;
@@ -650,7 +653,7 @@ public:
 		UIntValue = (uint64)Value;
 	}
 
-	friend FArchive& operator<<(FArchive& Ar, FFormatArgumentValue& Value);
+	friend void operator<<(FStructuredArchive::FSlot Slot, FFormatArgumentValue& Value);
 
 	FString ToFormattedString(const bool bInRebuildText, const bool bInRebuildAsSource) const;
 	void ToFormattedString(const bool bInRebuildText, const bool bInRebuildAsSource, FString& OutResult) const;
@@ -722,7 +725,7 @@ struct CORE_API FFormatArgumentData
 
 	void ResetValue();
 
-	friend FArchive& operator<<(FArchive& Ar, FFormatArgumentData& Value);
+	friend void operator<<(FStructuredArchive::FSlot Slot, FFormatArgumentData& Value);
 
 	FString ArgumentName;
 
diff --git a/Engine/Source/Runtime/Core/Public/Misc/DateTime.h b/Engine/Source/Runtime/Core/Public/Misc/DateTime.h
--- a/Engine/Source/Runtime/Core/Public/Misc/DateTime.h
+++ b/Engine/Source/Runtime/Core/Public/Misc/DateTime.h
@@ -706,6 +706,17 @@ public:
 		return Ar << DateTime.Ticks;
 	}
 
+	/**
+	 * Serializes the given date and time from or into the specified structured archive slot.
+	 *
+	 * @param Slot The structured archive slot to serialize from or into.
+	 * @param DateTime The date and time value to serialize.
+	 */
+	friend CORE_API void operator<<(FStructuredArchive::FSlot Slot, FDateTime& DateTime)
+	{
+		return Slot << DateTime.Ticks;
+	}
+
 	/**
 	 * Gets the hash for the specified date and time.
 	 *
diff --git a/Engine/Source/Runtime/Core/Public/Serialization/Formatters/BinaryArchiveFormatter.h b/Engine/Source/Runtime/Core/Public/Serialization/Formatters/BinaryArchiveFormatter.h
--- a/Engine/Source/Runtime/Core/Public/Serialization/Formatters/BinaryArchiveFormatter.h
+++ b/Engine/Source/Runtime/Core/Public/Serialization/Formatters/BinaryArchiveFormatter.h
@@ -49,6 +49,7 @@ public:
 	virtual void Serialize(FString& Value) override;
 	virtual void Serialize(FName& Value) override;
 	virtual void Serialize(UObject*& Value) override;
+	virtual void Serialize(FText& Value) override;
 	virtual void Serialize(FWeakObjectPtr& Value) override;
 	virtual void Serialize(FSoftObjectPtr& Value) override;
 	virtual void Serialize(FSoftObjectPath& Value) override;
@@ -214,6 +215,11 @@ inline void FBinaryArchiveFormatter::Serialize(UObject*& Value)
 	Inner << Value;
 }
 
+inline void FBinaryArchiveFormatter::Serialize(FText& Value)
+{
+	Inner << Value;
+}
+
 inline void FBinaryArchiveFormatter::Serialize(FWeakObjectPtr& Value)
 {
 	Inner << Value;
diff --git a/Engine/Source/Runtime/Core/Public/Serialization/StructuredArchive.h b/Engine/Source/Runtime/Core/Public/Serialization/StructuredArchive.h
--- a/Engine/Source/Runtime/Core/Public/Serialization/StructuredArchive.h
+++ b/Engine/Source/Runtime/Core/Public/Serialization/StructuredArchive.h
@@ -126,6 +126,7 @@ public:
 		void operator << (FString& Value);
 		void operator << (FName& Value);
 		void operator << (UObject*& Value);
+		void operator << (FText& Value);
 		void operator << (FWeakObjectPtr& Value);
 		void operator << (FSoftObjectPtr& Value);
 		void operator << (FSoftObjectPath& Value);
@@ -597,6 +598,11 @@ private:
 		Ar.Formatter.Serialize(Value);
 	}
 
+	FORCEINLINE void FStructuredArchive::FSlot::operator<< (FText& Value)
+	{
+		Ar.Formatter.Serialize(Value);
+	}
+
 	FORCEINLINE void FStructuredArchive::FSlot::operator<< (FWeakObjectPtr& Value)
 	{
 		Ar.Formatter.Serialize(Value);
diff --git a/Engine/Source/Runtime/Core/Public/Serialization/StructuredArchiveFormatter.h b/Engine/Source/Runtime/Core/Public/Serialization/StructuredArchiveFormatter.h
--- a/Engine/Source/Runtime/Core/Public/Serialization/StructuredArchiveFormatter.h
+++ b/Engine/Source/Runtime/Core/Public/Serialization/StructuredArchiveFormatter.h
@@ -100,6 +100,7 @@ public:
 	virtual void Serialize(FString& Value) = 0;
 	virtual void Serialize(FName& Value) = 0;
 	virtual void Serialize(UObject*& Value) = 0;
+	virtual void Serialize(FText& Value) = 0;
 	virtual void Serialize(struct FWeakObjectPtr& Value) = 0;
 	virtual void Serialize(struct FSoftObjectPtr& Value) = 0;
 	virtual void Serialize(struct FSoftObjectPath& Value) = 0;
diff --git a/Engine/Source/Runtime/CoreUObject/Private/Serialization/Formatters/JsonArchiveInputFormatter.cpp b/Engine/Source/Runtime/CoreUObject/Private/Serialization/Formatters/JsonArchiveInputFormatter.cpp
--- a/Engine/Source/Runtime/CoreUObject/Private/Serialization/Formatters/JsonArchiveInputFormatter.cpp
+++ b/Engine/Source/Runtime/CoreUObject/Private/Serialization/Formatters/JsonArchiveInputFormatter.cpp
@@ -283,6 +283,11 @@ void FJsonArchiveInputFormatter::Serialize(UObject*& Value)
 	}
 }
 
+void FJsonArchiveInputFormatter::Serialize(FText& Value)
+{
+	checkNoEntry();	// Implement me
+}
+
 void FJsonArchiveInputFormatter::Serialize(FWeakObjectPtr& Value)
 {
 	UObject* Object = nullptr;
diff --git a/Engine/Source/Runtime/CoreUObject/Private/Serialization/Formatters/JsonArchiveOutputFormatter.cpp b/Engine/Source/Runtime/CoreUObject/Private/Serialization/Formatters/JsonArchiveOutputFormatter.cpp
--- a/Engine/Source/Runtime/CoreUObject/Private/Serialization/Formatters/JsonArchiveOutputFormatter.cpp
+++ b/Engine/Source/Runtime/CoreUObject/Private/Serialization/Formatters/JsonArchiveOutputFormatter.cpp
@@ -1,6 +1,7 @@
 // Copyright 1998-2018 Epic Games, Inc. All Rights Reserved.
 
 #include "Serialization/Formatters/JsonArchiveOutputFormatter.h"
+#include "Serialization/MemoryWriter.h"
 #include "Misc/Base64.h"
 #include "Misc/SecureHash.h"
 #include "UObject/Object.h"
@@ -239,6 +240,13 @@ void FJsonArchiveOutputFormatter::Serialize(UObject*& Value)
 	}
 }
 
+void FJsonArchiveOutputFormatter::Serialize(FText& Value)
+{
+	FStructuredArchive ChildArchive(*this);
+	FText::SerializeText(ChildArchive.Open(), Value);
+	ChildArchive.Close();
+}
+
 void FJsonArchiveOutputFormatter::Serialize(FWeakObjectPtr& Value)
 {
 	if (Value.IsValid())
diff --git a/Engine/Source/Runtime/CoreUObject/Private/UObject/TextProperty.cpp b/Engine/Source/Runtime/CoreUObject/Private/UObject/TextProperty.cpp
--- a/Engine/Source/Runtime/CoreUObject/Private/UObject/TextProperty.cpp
+++ b/Engine/Source/Runtime/CoreUObject/Private/UObject/TextProperty.cpp
@@ -85,9 +85,7 @@ bool UTextProperty::Identical( const void* A, const void* B, uint32 PortFlags )
 void UTextProperty::SerializeItem(FStructuredArchive::FSlot Slot, void* Value, void const* Defaults) const
 {
 	TCppType* TextPtr = GetPropertyValuePtr(Value);
-
-	FArchiveUObjectFromStructuredArchive Ar(Slot);
-	Ar << *TextPtr;
+	Slot << *TextPtr;
 }
 
 void UTextProperty::ExportTextItem( FString& ValueStr, const void* PropertyValue, const void* DefaultValue, UObject* Parent, int32 PortFlags, UObject* ExportRootScope ) const
diff --git a/Engine/Source/Runtime/CoreUObject/Public/Serialization/Formatters/JsonArchiveInputFormatter.h b/Engine/Source/Runtime/CoreUObject/Public/Serialization/Formatters/JsonArchiveInputFormatter.h
--- a/Engine/Source/Runtime/CoreUObject/Public/Serialization/Formatters/JsonArchiveInputFormatter.h
+++ b/Engine/Source/Runtime/CoreUObject/Public/Serialization/Formatters/JsonArchiveInputFormatter.h
@@ -60,6 +60,7 @@ public:
 	virtual void Serialize(FString& Value) override;
 	virtual void Serialize(FName& Value) override;
 	virtual void Serialize(UObject*& Value) override;
+	virtual void Serialize(FText& Value) override;
 	virtual void Serialize(FWeakObjectPtr& Value) override;
 	virtual void Serialize(FSoftObjectPtr& Value) override;
 	virtual void Serialize(FSoftObjectPath& Value) override;
diff --git a/Engine/Source/Runtime/CoreUObject/Public/Serialization/Formatters/JsonArchiveOutputFormatter.h b/Engine/Source/Runtime/CoreUObject/Public/Serialization/Formatters/JsonArchiveOutputFormatter.h
--- a/Engine/Source/Runtime/CoreUObject/Public/Serialization/Formatters/JsonArchiveOutputFormatter.h
+++ b/Engine/Source/Runtime/CoreUObject/Public/Serialization/Formatters/JsonArchiveOutputFormatter.h
@@ -50,6 +50,7 @@ public:
 	virtual void Serialize(FString& Value) override;
 	virtual void Serialize(FName& Value) override;
 	virtual void Serialize(UObject*& Value) override;
+	virtual void Serialize(FText& Value) override;
 	virtual void Serialize(FWeakObjectPtr& Value) override;
 	virtual void Serialize(FSoftObjectPtr& Value) override;
 	virtual void Serialize(FSoftObjectPath& Value) override;
diff --git a/Engine/Source/Runtime/Engine/Private/Internationalization/StringTable.cpp b/Engine/Source/Runtime/Engine/Private/Internationalization/StringTable.cpp
--- a/Engine/Source/Runtime/Engine/Private/Internationalization/StringTable.cpp
+++ b/Engine/Source/Runtime/Engine/Private/Internationalization/StringTable.cpp
@@ -103,12 +103,12 @@ private:
 		}
 	}
 
-	virtual void CollectStringTableAssetReferencesImpl(const FName InTableId, FArchive& InAr) override
+	virtual void CollectStringTableAssetReferencesImpl(const FName InTableId, FStructuredArchive::FSlot Slot) override
 	{
-		check(InAr.IsObjectReferenceCollector());
+		check(Slot.GetUnderlyingArchive().IsObjectReferenceCollector());
 
-		UStringTable* StringTableAsset = FStringTableRegistry::Get().FindStringTableAsset(InTableId);
-		InAr << StringTableAsset;
+		UObject* StringTableAsset = FStringTableRegistry::Get().FindStringTableAsset(InTableId);
+		Slot << StringTableAsset;
 	}
 
 	static FSoftObjectPath GetAssetReference(const FName InTableId)
