Add missing #include
#rnx

diff --git a/Engine/Source/Runtime/Core/Private/Unix/UnixPlatformCrashContext.cpp b/Engine/Source/Runtime/Core/Private/Unix/UnixPlatformCrashContext.cpp
--- a/Engine/Source/Runtime/Core/Private/Unix/UnixPlatformCrashContext.cpp
+++ b/Engine/Source/Runtime/Core/Private/Unix/UnixPlatformCrashContext.cpp
@@ -25,7 +25,7 @@
 #include "HAL/PlatformMallocCrash.h"
 #include "Unix/UnixPlatformRunnableThread.h"
 #include "HAL/ExceptionHandling.h"
-
+#include "Stats/Stats.h"
 #include "HAL/ThreadHeartBeat.h"
 
 extern CORE_API bool GIsGPUCrashed;
