Added an 'unsafe' flavor of GetNormal2D to FVector #UE4

diff --git a/Engine/Source/Runtime/Core/Public/Math/Vector.h b/Engine/Source/Runtime/Core/Public/Math/Vector.h
--- a/Engine/Source/Runtime/Core/Public/Math/Vector.h
+++ b/Engine/Source/Runtime/Core/Public/Math/Vector.h
@@ -625,6 +625,14 @@ public:
 	 */
 	FVector GetSafeNormal2D(float Tolerance=SMALL_NUMBER) const;
 
+	/**
+	* Calculates normalized 2D version of vector without checking for zero length.
+	*
+	* @return Normalized 2D version of vector.
+	* @see GetSafeNormal2D()
+	*/
+	FVector GetUnsafeNormal2D() const;
+
 	/**
 	 * Returns the cosine of the angle between this vector and another projected onto the XY plane (no Z).
 	 *
@@ -1802,6 +1810,12 @@ FORCEINLINE FVector FVector::GetSafeNormal2D(float Tolerance) const
 	return FVector(X*Scale, Y*Scale, 0.f);
 }
 
+FORCEINLINE FVector FVector::GetUnsafeNormal2D() const
+{
+	const float Scale = FMath::InvSqrt(X * X + Y * Y);
+	return FVector(X*Scale, Y*Scale, 0.f);
+}
+
 FORCEINLINE float FVector::CosineAngle2D(FVector B) const
 {
 	FVector A(*this);
