If unable to get dbghelp module info, fallback to use available module info (fixes PC calculations for portable stack walk)

(Originally CL 4206596 by [REDACTED_PERSON])

diff --git a/Engine/Source/Runtime/Core/Private/Windows/WindowsPlatformStackWalk.cpp b/Engine/Source/Runtime/Core/Private/Windows/WindowsPlatformStackWalk.cpp
--- a/Engine/Source/Runtime/Core/Private/Windows/WindowsPlatformStackWalk.cpp
+++ b/Engine/Source/Runtime/Core/Private/Windows/WindowsPlatformStackWalk.cpp
@@ -1,4 +1,4 @@
-﻿// Copyright 1998-2018 Epic Games, Inc. All Rights Reserved.
+// Copyright 1998-2018 Epic Games, Inc. All Rights Reserved.
 
 #include "Windows/WindowsPlatformStackWalk.h"
 #include "HAL/PlatformMemory.h"
@@ -15,6 +15,7 @@
 #include "HAL/PlatformProcess.h"
 #include "CoreGlobals.h"
 #include "Misc/ConfigCacheIni.h"
+#include "Misc/OutputDeviceRedirector.h"
 
 #include "Windows/WindowsHWrapper.h"
 #include "Windows/AllowWindowsPlatformTypes.h"
@@ -742,10 +743,12 @@ int32 FWindowsPlatformStackWalk::GetProcessModuleSignatures(FStackWalkModuleInfo
 		FGetModuleFileNameEx( ProcessHandle, ModuleHandlePointer[ModuleIndex], ImageName, MAX_PATH );
 		FGetModuleBaseName( ProcessHandle, ModuleHandlePointer[ModuleIndex], ModuleName, MAX_PATH );
 
-		// Load module.
-		if(SymGetModuleInfoW64(ProcessHandle, (DWORD64)ModuleInfo.lpBaseOfDll, &Img))
+		FStackWalkModuleInfo Info = { 0 };
+
+		// Load module and get rich image help information
+		if (SymGetModuleInfoW64(ProcessHandle, (DWORD64)ModuleInfo.lpBaseOfDll, &Img))
 		{
-			FStackWalkModuleInfo Info = {0};
+
 			Info.BaseOfImage = Img.BaseOfImage;
 			FCString::Strcpy(Info.ImageName, Img.ImageName);
 			Info.ImageSize = Img.ImageSize;
@@ -755,10 +758,22 @@ int32 FWindowsPlatformStackWalk::GetProcessModuleSignatures(FStackWalkModuleInfo
 			Info.PdbSig = Img.PdbSig;
 			FMemory::Memcpy(&Info.PdbSig70, &Img.PdbSig70, sizeof(GUID));
 			Info.TimeDateStamp = Img.TimeDateStamp;
-
-			ModuleSignatures[SignatureIndex] = Info;
-			++SignatureIndex;
 		}
+		else
+		{
+			// Unable to get image help information, so fallback to the module info that is available
+			Info.BaseOfImage = (uint64)ModuleInfo.lpBaseOfDll;
+			FCString::Strcpy(Info.ImageName, ImageName);
+			Info.ImageSize = ModuleInfo.SizeOfImage;
+			FCString::Strcpy(Info.LoadedImageName, ImageName);
+			FCString::Strcpy(Info.ModuleName, ModuleName);
+
+			UE_LOG(LogWindows, Verbose, TEXT("SymGetModuleInfoW64 failed, rich module information unavailable. Error Code: %u"), GetLastError());
+		}
+
+		ModuleSignatures[SignatureIndex] = Info;
+		++SignatureIndex;
+
 	}
 
 	// Free the module handle pointer allocated in case the static array was insufficient.
