Fix backup log timestamps.

When engine starts up, it makes a backup of previous run's log with timestamp appended. Previous code used the current engine startup time; basically recording the time it was backed up, not the time it was generated. This change honors the file's original timestamp. This way files sort properly and we get accurate information about when each log was generated.

diff --git a/Engine/Source/Runtime/Core/Private/Misc/OutputDeviceFile.cpp b/Engine/Source/Runtime/Core/Private/Misc/OutputDeviceFile.cpp
--- a/Engine/Source/Runtime/Core/Private/Misc/OutputDeviceFile.cpp
+++ b/Engine/Source/Runtime/Core/Private/Misc/OutputDeviceFile.cpp
@@ -368,12 +368,18 @@ void FOutputDeviceFile::Flush()
  */
 void FOutputDeviceFile::CreateBackupCopy(const TCHAR* Filename)
 {
-	if (IFileManager::Get().FileSize(Filename) > 0)
+	IFileManager& FileManager = IFileManager::Get();
+	if (FileManager.FileSize(Filename) > 0) // file exists and is not empty
 	{
 		FString Name, Extension;
 		FString(Filename).Split(TEXT("."), &Name, &Extension, ESearchCase::CaseSensitive, ESearchDir::FromEnd);
-		FString BackupFilename = FString::Printf(TEXT("%s%s%s.%s"), *Name, BACKUP_LOG_FILENAME_POSTFIX, *GSystemStartTime, *Extension);
-		IFileManager::Get().Copy(*BackupFilename, Filename, false);
+		FDateTime OriginalTime = FileManager.GetTimeStamp(Filename);
+		FString BackupFilename = FString::Printf(TEXT("%s%s%s.%s"), *Name, BACKUP_LOG_FILENAME_POSTFIX, *OriginalTime.ToString(), *Extension);
+		if (FileManager.Copy(*BackupFilename, Filename, false) == COPY_OK)
+		{
+			FileManager.SetTimeStamp(*BackupFilename, OriginalTime);
+		}
+		// We use Copy + SetTimeStamp instead of Move because caller might want to append to original log file.
 	}
 }
 
