Fixed FrameNumber numeric input fields not understanding "*" which prevented pressing Ctrl + G and then Enter (to escape focus) when evaluating a sub-frame. Fixed Timecode not displaying Subframe markers.

diff --git a/Engine/Source/Runtime/Core/Public/Math/BasicMathExpressionEvaluator.h b/Engine/Source/Runtime/Core/Public/Math/BasicMathExpressionEvaluator.h
--- a/Engine/Source/Runtime/Core/Public/Math/BasicMathExpressionEvaluator.h
+++ b/Engine/Source/Runtime/Core/Public/Math/BasicMathExpressionEvaluator.h
@@ -8,29 +8,29 @@
 #include "Misc/ExpressionParserTypes.h"
 #include "Internationalization/FastDecimalFormat.h"
 
-#define DEFINE_EXPRESSION_OPERATOR_NODE(TYPE, ...) \
+#define DEFINE_EXPRESSION_OPERATOR_NODE(EXPORTAPI, TYPE, ...) \
 namespace ExpressionParser {\
-	struct CORE_API TYPE { static const TCHAR* const Moniker; }; \
+	struct EXPORTAPI TYPE { static const TCHAR* const Moniker; }; \
 }\
 DEFINE_EXPRESSION_NODE_TYPE(ExpressionParser::TYPE, __VA_ARGS__)
 
 /** Define some expression types for basic arithmetic */
 DEFINE_EXPRESSION_NODE_TYPE(double, 0x8444A8A3, 0x19AE4E13, 0xBCFA75EE, 0x39982BD6)
 
-DEFINE_EXPRESSION_OPERATOR_NODE(FSubExpressionStart, 0xCC40A083, 0xADBF46E2, 0xA93D12BB, 0x525D7417)
-DEFINE_EXPRESSION_OPERATOR_NODE(FSubExpressionEnd, 0x125E4C67, 0x96EB48C4, 0x8894E09C, 0xB3CD56BF)
-
-DEFINE_EXPRESSION_OPERATOR_NODE(FPlus, 0x6F88756B, 0xF9234263, 0x9B13614F, 0x2706074B)
-DEFINE_EXPRESSION_OPERATOR_NODE(FPlusEquals, 0x05A878C9, 0x0E6C44A4, 0x9A73920D, 0x3175AB48)
-DEFINE_EXPRESSION_OPERATOR_NODE(FMinus, 0xE6240779, 0xEE2849CF, 0x95A0E648, 0x22D58713)
-DEFINE_EXPRESSION_OPERATOR_NODE(FMinusEquals, 0x8D1E08E3, 0x5F534245, 0xA987AD7E, 0xDB78A4B7)
-DEFINE_EXPRESSION_OPERATOR_NODE(FStar, 0xF287B3BF, 0x35DF4141, 0xA8B4F57B, 0x7E06DF47)
-DEFINE_EXPRESSION_OPERATOR_NODE(FStarEquals, 0xBA359BB0, 0xCEB54682, 0x9160EB4F, 0xD1687C7F)
-DEFINE_EXPRESSION_OPERATOR_NODE(FForwardSlash, 0xF99670F8, 0x74794256, 0xBB0CAE6D, 0xC67CD5B6)
-DEFINE_EXPRESSION_OPERATOR_NODE(FForwardSlashEquals, 0x4AFE0CF8, 0xF9054360, 0xBE5DCE80, 0xDC2E22F6)
-DEFINE_EXPRESSION_OPERATOR_NODE(FPercent, 0x936E4434, 0x0A014F2D, 0xBEEC90D3, 0x4D3ECEA2)
-DEFINE_EXPRESSION_OPERATOR_NODE(FSquareRoot, 0xE7C03E11, 0x9DE84B4B, 0xBA4C2B76, 0x69BF028E)
-DEFINE_EXPRESSION_OPERATOR_NODE(FPower, 0x93388F8D, 0x1D9B4DFE, 0xBD4D6CC4, 0x12D1DE99)
+DEFINE_EXPRESSION_OPERATOR_NODE(CORE_API, FSubExpressionStart, 0xCC40A083, 0xADBF46E2, 0xA93D12BB, 0x525D7417)
+DEFINE_EXPRESSION_OPERATOR_NODE(CORE_API, FSubExpressionEnd, 0x125E4C67, 0x96EB48C4, 0x8894E09C, 0xB3CD56BF)
+
+DEFINE_EXPRESSION_OPERATOR_NODE(CORE_API, FPlus, 0x6F88756B, 0xF9234263, 0x9B13614F, 0x2706074B)
+DEFINE_EXPRESSION_OPERATOR_NODE(CORE_API, FPlusEquals, 0x05A878C9, 0x0E6C44A4, 0x9A73920D, 0x3175AB48)
+DEFINE_EXPRESSION_OPERATOR_NODE(CORE_API, FMinus, 0xE6240779, 0xEE2849CF, 0x95A0E648, 0x22D58713)
+DEFINE_EXPRESSION_OPERATOR_NODE(CORE_API, FMinusEquals, 0x8D1E08E3, 0x5F534245, 0xA987AD7E, 0xDB78A4B7)
+DEFINE_EXPRESSION_OPERATOR_NODE(CORE_API, FStar, 0xF287B3BF, 0x35DF4141, 0xA8B4F57B, 0x7E06DF47)
+DEFINE_EXPRESSION_OPERATOR_NODE(CORE_API, FStarEquals, 0xBA359BB0, 0xCEB54682, 0x9160EB4F, 0xD1687C7F)
+DEFINE_EXPRESSION_OPERATOR_NODE(CORE_API, FForwardSlash, 0xF99670F8, 0x74794256, 0xBB0CAE6D, 0xC67CD5B6)
+DEFINE_EXPRESSION_OPERATOR_NODE(CORE_API, FForwardSlashEquals, 0x4AFE0CF8, 0xF9054360, 0xBE5DCE80, 0xDC2E22F6)
+DEFINE_EXPRESSION_OPERATOR_NODE(CORE_API, FPercent, 0x936E4434, 0x0A014F2D, 0xBEEC90D3, 0x4D3ECEA2)
+DEFINE_EXPRESSION_OPERATOR_NODE(CORE_API, FSquareRoot, 0xE7C03E11, 0x9DE84B4B, 0xBA4C2B76, 0x69BF028E)
+DEFINE_EXPRESSION_OPERATOR_NODE(CORE_API, FPower, 0x93388F8D, 0x1D9B4DFE, 0xBD4D6CC4, 0x12D1DE99)
 	
 namespace ExpressionParser
 {
diff --git a/Engine/Source/Runtime/TimeManagement/Private/FrameNumberTimeEvaluator.cpp b/Engine/Source/Runtime/TimeManagement/Private/FrameNumberTimeEvaluator.cpp
--- a/Engine/Source/Runtime/TimeManagement/Private/FrameNumberTimeEvaluator.cpp
+++ b/Engine/Source/Runtime/TimeManagement/Private/FrameNumberTimeEvaluator.cpp
@@ -38,6 +38,7 @@ FFrameNumberTimeEvaluator::FFrameNumberTimeEvaluator()
 	FrameTokenDefinitions.DefineToken(&ConsumeSymbol<FFrames>);
 	FrameTokenDefinitions.DefineToken(&ConsumeSymbol<FPlus>);
 	FrameTokenDefinitions.DefineToken(&ConsumeSymbol<FMinus>);
+	FrameTokenDefinitions.DefineToken(&ConsumeSymbol<FStar>);
 	FrameTokenDefinitions.DefineToken(&ConsumeLocalizedNumber);
 
 	TimeTokenDefinitions.IgnoreWhitespace();
@@ -92,7 +93,7 @@ TValueOrError<FFrameTime, FExpressionError> FFrameNumberTimeEvaluator::EvaluateT
 		{
 			OutDirectlyParsed = true;
 		}
-		else if (Node.Cast<FBracketStart>() || Node.Cast<FBracketEnd>() || Node.Cast<FPlus>() || Node.Cast<FMinus>())
+		else if (Node.Cast<FBracketStart>() || Node.Cast<FBracketEnd>() || Node.Cast<FPlus>() || Node.Cast<FMinus>() || Node.Cast<FStar>())
 		{
 			if (Node.Cast<FMinus>())
 			{
@@ -217,11 +218,15 @@ TValueOrError<FFrameTime, FExpressionError> FFrameNumberTimeEvaluator::EvaluateF
 			Tokens.RemoveAt(i);
 			i--;
 		}
-		else if (Node.Cast<FMinus>())
+		else if (Node.Cast<FMinus>() || Node.Cast<FStar>())
 		{
+			if (Node.Cast<FMinus>())
+			{
+				bIsNegative = true;
+			}
+
 			Tokens.RemoveAt(i);
 			i--;
-			bIsNegative = true;
 		}
 	}
 
diff --git a/Engine/Source/Runtime/TimeManagement/Public/FrameNumberNumericInterface.h b/Engine/Source/Runtime/TimeManagement/Public/FrameNumberNumericInterface.h
--- a/Engine/Source/Runtime/TimeManagement/Public/FrameNumberNumericInterface.h
+++ b/Engine/Source/Runtime/TimeManagement/Public/FrameNumberNumericInterface.h
@@ -92,20 +92,16 @@ private:
 			return FString::Printf(TEXT("%.2f s"), TimeInSeconds);
 		}
 		case EFrameNumberDisplayFormats::NonDropFrameTimecode:
-		{
-			FFrameTime InternalFrameNumber = FFrameTime::FromDecimal(Value);
-			FFrameNumber PlayRateFrameNumber = FFrameRate::TransformTime(InternalFrameNumber, SourceFrameRate, DestinationFrameRate).FloorToFrame();
-
-			FTimecode AsNonDropTimecode = FTimecode::FromFrameNumber(PlayRateFrameNumber, DestinationFrameRate, false);
-			return FString::Printf(TEXT("[%s]"), *AsNonDropTimecode.ToString(false));
-		}
 		case EFrameNumberDisplayFormats::DropFrameTimecode:
 		{
 			FFrameTime InternalFrameNumber = FFrameTime::FromDecimal(Value);
-			FFrameNumber PlayRateFrameNumber = FFrameRate::TransformTime(InternalFrameNumber, SourceFrameRate, DestinationFrameRate).FloorToFrame();
+			FFrameTime DisplayTime = FFrameRate::TransformTime(InternalFrameNumber, SourceFrameRate, DestinationFrameRate);
+			FString SubframeIndicator = FMath::IsNearlyZero(DisplayTime.GetSubFrame()) ? TEXT("") : TEXT("*");
+
+			bool bIsDropTimecode = Format == EFrameNumberDisplayFormats::DropFrameTimecode;
 
-			FTimecode AsDropTimecode = FTimecode::FromFrameNumber(PlayRateFrameNumber, DestinationFrameRate, true);
-			return FString::Printf(TEXT("[%s]"), *AsDropTimecode.ToString(false));
+			FTimecode AsNonDropTimecode = FTimecode::FromFrameNumber(DisplayTime.FloorToFrame(), DestinationFrameRate, bIsDropTimecode);
+			return FString::Printf(TEXT("[%s%s]"), *AsNonDropTimecode.ToString(false), *SubframeIndicator);
 		}
 		default:
 			return FString(TEXT("Unsupported Format"));
@@ -123,7 +119,7 @@ private:
 		// also check to see if they explicitly specified that format, or if the evaluator just happens to be able to parse that.
 
 		// All of these will convert into the frame resolution from the user's input before returning.
-		FFrameNumberTimeEvaluator Eval;
+ 		FFrameNumberTimeEvaluator Eval;
 		bool bWasTimecodeText;
 		TValueOrError<FFrameTime, FExpressionError> TimecodeResult = Eval.EvaluateTimecode(*InString, SourceFrameRate, DestinationFrameRate, /*Out*/bWasTimecodeText);
 		bool bWasFrameText;
diff --git a/Engine/Source/Runtime/TimeManagement/Public/FrameNumberTimeEvaluator.h b/Engine/Source/Runtime/TimeManagement/Public/FrameNumberTimeEvaluator.h
--- a/Engine/Source/Runtime/TimeManagement/Public/FrameNumberTimeEvaluator.h
+++ b/Engine/Source/Runtime/TimeManagement/Public/FrameNumberTimeEvaluator.h
@@ -6,27 +6,20 @@
 #include "Misc/Optional.h"
 #include "Templates/ValueOrError.h"
 #include "Misc/ExpressionParserTypes.h"
+#include "Math/BasicMathExpressionEvaluator.h"
 
 struct FFrameTime;
 struct FFrameRate;
 
-#define DEFINE_EXPRESSION_OPERATOR_NODE2(TYPE, ...) \
-namespace ExpressionParser {\
-	struct TIMEMANAGEMENT_API TYPE { static const TCHAR* const Moniker; }; \
-}\
-DEFINE_EXPRESSION_NODE_TYPE(ExpressionParser::TYPE, __VA_ARGS__)
-
-// @todo sequencer-timecode This is a duplicate from the BasicMathEvaluator, but that macro exports with CORE_API so we can't use it here.
-// We should probably just explicitly define these nodes in the namespace instead of duplicating the macro.
-DEFINE_EXPRESSION_OPERATOR_NODE2(FHour, 0x702443E8, 0xEF9A48A0, 0x8DC56394, 0x48F2632A);
-DEFINE_EXPRESSION_OPERATOR_NODE2(FMinute, 0x4FED0D41, 0x298D481C, 0xAC022899, 0xF196E817);
-DEFINE_EXPRESSION_OPERATOR_NODE2(FSecond, 0x43365BB1, 0x9A6B457E, 0xA6C71BB3, 0xCE6D5ACD);
-DEFINE_EXPRESSION_OPERATOR_NODE2(FMillisecond, 0xBFA9E4FB, 0x45799275, 0x0EFA76E3, 0x76E3533F);
-DEFINE_EXPRESSION_OPERATOR_NODE2(FFrames, 0xA0D341C4, 0xA51D40D3, 0xA9DBDD45, 0x83E1A21A);
-DEFINE_EXPRESSION_OPERATOR_NODE2(FTimecodeDelimiter, 0x6E947008, 0x523D4A17, 0xB5F7BCDA, 0xA43B62FC);
-DEFINE_EXPRESSION_OPERATOR_NODE2(FDropcodeDelimiter, 0x80AF2C2C, 0x544A451F, 0x822E5B24, 0xAFC78EEF);
-DEFINE_EXPRESSION_OPERATOR_NODE2(FBracketStart, 0xA7358BD1, 0xB4EF466D, 0xA336CD84, 0x14C93D2E);
-DEFINE_EXPRESSION_OPERATOR_NODE2(FBracketEnd, 0xE49500E9, 0x03E64440, 0x87802630, 0xF1C1DDDF);
+DEFINE_EXPRESSION_OPERATOR_NODE(TIMEMANAGEMENT_API, FHour, 0x702443E8, 0xEF9A48A0, 0x8DC56394, 0x48F2632A);
+DEFINE_EXPRESSION_OPERATOR_NODE(TIMEMANAGEMENT_API, FMinute, 0x4FED0D41, 0x298D481C, 0xAC022899, 0xF196E817);
+DEFINE_EXPRESSION_OPERATOR_NODE(TIMEMANAGEMENT_API, FSecond, 0x43365BB1, 0x9A6B457E, 0xA6C71BB3, 0xCE6D5ACD);
+DEFINE_EXPRESSION_OPERATOR_NODE(TIMEMANAGEMENT_API, FMillisecond, 0xBFA9E4FB, 0x45799275, 0x0EFA76E3, 0x76E3533F);
+DEFINE_EXPRESSION_OPERATOR_NODE(TIMEMANAGEMENT_API, FFrames, 0xA0D341C4, 0xA51D40D3, 0xA9DBDD45, 0x83E1A21A);
+DEFINE_EXPRESSION_OPERATOR_NODE(TIMEMANAGEMENT_API, FTimecodeDelimiter, 0x6E947008, 0x523D4A17, 0xB5F7BCDA, 0xA43B62FC);
+DEFINE_EXPRESSION_OPERATOR_NODE(TIMEMANAGEMENT_API, FDropcodeDelimiter, 0x80AF2C2C, 0x544A451F, 0x822E5B24, 0xAFC78EEF);
+DEFINE_EXPRESSION_OPERATOR_NODE(TIMEMANAGEMENT_API, FBracketStart, 0xA7358BD1, 0xB4EF466D, 0xA336CD84, 0x14C93D2E);
+DEFINE_EXPRESSION_OPERATOR_NODE(TIMEMANAGEMENT_API, FBracketEnd, 0xE49500E9, 0x03E64440, 0x87802630, 0xF1C1DDDF);
 
 struct TIMEMANAGEMENT_API FFrameNumberTimeEvaluator
 {
