Use platform layer correctly for the binary directory

diff --git a/Engine/Source/Runtime/Core/Private/Unix/UnixPlatformProcess.cpp b/Engine/Source/Runtime/Core/Private/Unix/UnixPlatformProcess.cpp
--- a/Engine/Source/Runtime/Core/Private/Unix/UnixPlatformProcess.cpp
+++ b/Engine/Source/Runtime/Core/Private/Unix/UnixPlatformProcess.cpp
@@ -365,7 +365,7 @@ const TCHAR* FUnixPlatformProcess::ExecutableName(bool bRemoveExtension)
 
 FString FUnixPlatformProcess::GenerateApplicationPath( const FString& AppName, EBuildConfigurations::Type BuildConfiguration)
 {
-	FString PlatformName = GetBinariesSubdirectory();
+	FString PlatformName = FPlatformProcess::GetBinariesSubdirectory();
 	FString ExecutablePath = FString::Printf(TEXT("../../../Engine/Binaries/%s/%s"), *PlatformName, *AppName);
 	
 	if (BuildConfiguration != EBuildConfigurations::Development && BuildConfiguration != EBuildConfigurations::DebugGame)
