Remove extra mallocs when crash handling

Still need to look into gmalloc calls, such as using FStrings during Ensure/Crash handling

diff --git a/Engine/Source/Runtime/ApplicationCore/Private/Unix/UnixErrorOutputDevice.cpp b/Engine/Source/Runtime/ApplicationCore/Private/Unix/UnixErrorOutputDevice.cpp
--- a/Engine/Source/Runtime/ApplicationCore/Private/Unix/UnixErrorOutputDevice.cpp
+++ b/Engine/Source/Runtime/ApplicationCore/Private/Unix/UnixErrorOutputDevice.cpp
@@ -94,12 +94,6 @@ void FUnixErrorOutputDevice::HandleError()
 
 		GLog->Flush();
 
-		// do not copy if graphics have not been initialized or if we're on the wrong thread
-		if (FApp::CanEverRender() && IsInGameThread())
-		{
-			FPlatformApplicationMisc::ClipboardCopy(GErrorHist);
-		}
-
 		FPlatformMisc::SubmitErrorReport(GErrorHist, EErrorReportMode::Interactive);
 		FCoreDelegates::OnShutdownAfterError.Broadcast();
 #if !PLATFORM_EXCEPTIONS_DISABLED
diff --git a/Engine/Source/Runtime/Core/Private/Unix/UnixPlatformCrashContext.cpp b/Engine/Source/Runtime/Core/Private/Unix/UnixPlatformCrashContext.cpp
--- a/Engine/Source/Runtime/Core/Private/Unix/UnixPlatformCrashContext.cpp
+++ b/Engine/Source/Runtime/Core/Private/Unix/UnixPlatformCrashContext.cpp
@@ -430,10 +430,6 @@ void FUnixCrashContext::GenerateCrashInfoAndLaunchReporter(bool bReportingNonCra
 
 		// try launching the tool and wait for its exit, if at all
 		const TCHAR * RelativePathToCrashReporter = TEXT("../../../Engine/Binaries/Linux/CrashReportClient");	// FIXME: painfully hard-coded
-		if (!FPaths::FileExists(RelativePathToCrashReporter))
-		{
-			RelativePathToCrashReporter = TEXT("../../../engine/binaries/linux/crashreportclient");	// FIXME: even more painfully hard-coded
-		}
 
 		FString CrashReportLogFilename = LogBaseFilename + TEXT("-CRC") + LogExtension;
 		FString CrashReportLogFilepath = FPaths::Combine(*LogFolder, *CrashReportLogFilename);
diff --git a/Engine/Source/Runtime/Core/Private/Unix/UnixPlatformStackWalk.cpp b/Engine/Source/Runtime/Core/Private/Unix/UnixPlatformStackWalk.cpp
--- a/Engine/Source/Runtime/Core/Private/Unix/UnixPlatformStackWalk.cpp
+++ b/Engine/Source/Runtime/Core/Private/Unix/UnixPlatformStackWalk.cpp
@@ -904,27 +904,27 @@ namespace
 	{
 	public:
 		explicit RecordReader(const char* Path)
-			: SymbolFile(fopen(Path, "r")),
+			: SymbolFileFD(open(Path, O_RDONLY)),
 			  StartOffset(sizeof(RecordsHeader))
 		{
-			if (SymbolFile)
+			if (SymbolFileFD)
 			{
-				fread(static_cast<void*>(&RecordCount), sizeof(RecordsHeader), 1, SymbolFile);
-				StartOffset = sizeof(RecordsHeader);
+				// TODO check for EINTR
+				read(SymbolFileFD, static_cast<void*>(&RecordCount), sizeof(RecordsHeader));
 			}
 		}
 
 		~RecordReader()
 		{
-			if (SymbolFile)
+			if (SymbolFileFD != -1)
 			{
-				fclose(SymbolFile);
+				close(SymbolFileFD);
 			}
 		}
 
 		bool IsValid() const
 		{
-			return SymbolFile != nullptr;
+			return SymbolFileFD != -1 && RecordCount > 0;
 		}
 
 		uint32_t GetRecordCount() const
@@ -941,8 +941,8 @@ namespace
 
 			Record Out;
 
-			fseek(SymbolFile, StartOffset + Index * sizeof(Record), SEEK_SET);
-			fread(static_cast<void*>(&Out), sizeof(Record), 1, SymbolFile);
+			lseek(SymbolFileFD, StartOffset + Index * sizeof(Record), SEEK_SET);
+			read(SymbolFileFD, static_cast<void*>(&Out), sizeof(Record));
 
 			return Out;
 		}
@@ -955,8 +955,8 @@ namespace
 				return;
 			}
 
-			fseek(SymbolFile, StartOffset + RecordCount * sizeof(Record) + Offset, SEEK_SET);
-			fread(Buffer, MaxSize, 1, SymbolFile);
+			lseek(SymbolFileFD, StartOffset + RecordCount * sizeof(Record) + Offset, SEEK_SET);
+			read(SymbolFileFD, Buffer, MaxSize);
 
 			// Read the max chunk we can read, then find the next '\n' and replace that with '\0'
 			for (int i = 0; i < MaxSize; i++)
@@ -973,7 +973,7 @@ namespace
 		}
 
 	private:
-		FILE* SymbolFile = nullptr;
+		int SymbolFileFD;
 
 		// For now can only be up to 4GB
 		uint32_t StartOffset = 0;
