PR #4272: Fix typo. (Contributed by [REDACTED_PERSON])

diff --git a/Engine/Source/Runtime/Core/Public/Logging/LogVerbosity.h b/Engine/Source/Runtime/Core/Public/Logging/LogVerbosity.h
--- a/Engine/Source/Runtime/Core/Public/Logging/LogVerbosity.h
+++ b/Engine/Source/Runtime/Core/Public/Logging/LogVerbosity.h
@@ -18,7 +18,7 @@ namespace ELogVerbosity
 		/** Not used */
 		NoLogging		= 0,
 
-		/** Always prints s fatal error to console (and log file) and crashes (even if logging is disabled) */
+		/** Always prints a fatal error to console (and log file) and crashes (even if logging is disabled) */
 		Fatal,
 
 		/** 
