Fix duplicated struct/class from slightly different submit into main coming back into dev-editor

diff --git a/Engine/Source/Runtime/Core/Private/Unix/UnixPlatformFile.cpp b/Engine/Source/Runtime/Core/Private/Unix/UnixPlatformFile.cpp
--- a/Engine/Source/Runtime/Core/Private/Unix/UnixPlatformFile.cpp
+++ b/Engine/Source/Runtime/Core/Private/Unix/UnixPlatformFile.cpp
@@ -270,72 +270,6 @@ private:
 };
 
 extern int32 CORE_API GMaxNumberFileMappingCache;
-namespace
-{
-	double MaxInvalidCacheTime = 0.5; // 500ms
-	struct FileEntry
-	{
-		FString File;
-		bool bInvalid = false;
-		double CacheTime = 0.0f;
-		bool IsInvalid() const
-		{
-			double Current = FPlatformTime::Seconds();
-			return bInvalid && Current - CacheTime <= MaxInvalidCacheTime;
-		}
-	};
-	struct FileMapCache
-	{
-		virtual const FileEntry* Find(const FString& Key) = 0;
-		virtual void AddEntry(const FString& Key, const FString& Elem) = 0;
-		virtual void Invalidate(const FString& Key) = 0;
-	};
-	class FileMapCacheDummy : public FileMapCache
-	{
-	public:
-		const FileEntry* Find(const FString& Key) override
-		{
-			return nullptr;
-		}
-		void AddEntry(const FString& Key, const FString& Elem) override
-		{ }
-		void Invalidate(const FString& Key) override
-		{ }
-	};
-	class FileMapCacheDefault : public FileMapCache
-	{
-	public:
-		FileMapCacheDefault()
-			: Cache(GMaxNumberFileMappingCache)
-		{
-		}
-		const FileEntry* Find(const FString& Key) override
-		{
-			FScopeLock ScopeLock(&Mutex);
-			return Cache.FindAndTouch(Key);
-		}
-		void AddEntry(const FString& Key, const FString& Elem) override
-		{
-			FScopeLock ScopeLock(&Mutex);
-			Cache.Add(Key, {Elem, Elem.IsEmpty(), FPlatformTime::Seconds()});
-		}
-		void Invalidate(const FString& Key) override
-		{
-			FScopeLock ScopeLock(&Mutex);
-			Cache.Remove(Key);
-		}
-	private:
-		FCriticalSection Mutex;
-		TLruCache<FString, FileEntry> Cache;
-	};
-	FileMapCache& GetFileMapCache()
-	{
-		static FileMapCacheDefault DefaultCache;
-		static FileMapCacheDummy DummyCache;
-		return GMaxNumberFileMappingCache > 0 ? *static_cast<FileMapCache*>(&DefaultCache) : *static_cast<FileMapCache*>(&DummyCache);
-	}
-}
-
 
 namespace
 {
