Linux: Pass a command line argument to crash reporter to show or skip a user agreement popup

diff --git a/Engine/Source/Programs/CrashReportClient/Private/CrashReportClientApp.cpp b/Engine/Source/Programs/CrashReportClient/Private/CrashReportClientApp.cpp
--- a/Engine/Source/Programs/CrashReportClient/Private/CrashReportClientApp.cpp
+++ b/Engine/Source/Programs/CrashReportClient/Private/CrashReportClientApp.cpp
@@ -54,6 +54,9 @@ static FString GameNameFromCmd;
 /** GUID of the crash passed via the command line. */
 static FString CrashGUIDFromCmd;
 
+/** Disable the showing of the yes/no crash upload popup if passed via command line */
+static bool bSkipCrashUploadPopup = false;
+
 /**
  * Look for the report to upload, either in the command line or in the platform's report queue
  */
@@ -107,6 +110,11 @@ void ParseCommandLine(const TCHAR* CommandLine)
 		{
 			CrashGUIDFromCmd = Params.FindRef(TEXT("CrashGUID"));
 		}
+
+		if (Switches.Contains(TEXT("SkipPopup")))
+		{
+			bSkipCrashUploadPopup = true;
+		}
 	}
 
 	if (FoundReportDirectoryAbsolutePaths.Num() == 0)
@@ -341,10 +349,10 @@ void RunCrashReportClient(const TCHAR* CommandLine)
 	// For now Linux needs to ask if we want to submit a crash report unless specifically configured not to
 	if (PLATFORM_LINUX)
 	{
-		bool bAgreeToServerCrashUpload = true;
-		GConfig->GetBool(TEXT("CrashReportClient"), TEXT("bAgreeToServerCrashUpload"), bAgreeToServerCrashUpload, GEngineIni);
-		if(bAgreeToServerCrashUpload && FPlatformMisc::MessageBoxExt(EAppMsgType::YesNo, TEXT("A crash has occurred. Would you like to submit this crash report?"), TEXT("Crash Client Reporter")) == EAppReturnType::No)
+		// We are ignoring the Unattened bool here as for Linux we are always set to true
+		if(!bSkipCrashUploadPopup && FPlatformMisc::MessageBoxExt(EAppMsgType::YesNo, TEXT("A crash has occurred. Would you like to submit this crash report?"), TEXT("Crash Client Reporter")) == EAppReturnType::No)
 		{
+			UE_LOG(CrashReportClientLog, Warning, TEXT("User has selected no to crash upload. Crash report will not be uploaded."));
 			GIsRequestingExit = true;
 		}
 	}
diff --git a/Engine/Source/Runtime/Core/Private/Unix/UnixPlatformCrashContext.cpp b/Engine/Source/Runtime/Core/Private/Unix/UnixPlatformCrashContext.cpp
--- a/Engine/Source/Runtime/Core/Private/Unix/UnixPlatformCrashContext.cpp
+++ b/Engine/Source/Runtime/Core/Private/Unix/UnixPlatformCrashContext.cpp
@@ -14,6 +14,7 @@
 #include "Misc/Paths.h"
 #include "Misc/FileHelper.h"
 #include "Delegates/IDelegateInstance.h"
+#include "Misc/ConfigCacheIni.h"
 #include "Misc/Guid.h"
 #include "Misc/OutputDeviceRedirector.h"
 #include "Misc/OutputDeviceError.h"
@@ -382,170 +383,196 @@ void FUnixCrashContext::GenerateCrashInfoAndLaunchReporter(bool bReportingNonCra
 		CrashGuid = FGuid::NewGuid().ToString();
 	}
 
-	FString CrashInfoFolder = FPaths::Combine(*FPaths::ProjectSavedDir(), TEXT("Crashes"), *FString::Printf(TEXT("%sinfo-%s-pid-%d-%s"), bReportingNonCrash ? TEXT("ensure") : TEXT("crash"), FApp::GetProjectName(), getpid(), *CrashGuid));
-	FString CrashInfoAbsolute = FPaths::ConvertRelativePathToFull(CrashInfoFolder);
-	if (IFileManager::Get().MakeDirectory(*CrashInfoAbsolute, true))
+	/* Table showing the desired behavior when wanting to show a pop up or not. As well as avoiding starting CRC
+	 *  based on an *.ini setting bAgreedToCrashUpload and whether or not we are unattended
+	 *
+	 *  Unattended | AgreeToUpload || Show Popup | Start CRC
+	 *  ----------------------------------------------------
+	 *      1      |       1       ||     0      |    1
+	 *      1      |       0       ||     0      |    0
+	 *      0      |       1       ||     0      |    1
+	 *      0      |       0       ||     1      |    1
+	 */
+
+	// Suppress the user input dialog if we're running in unattended mode
+	bool bUnattended = FApp::IsUnattended() || (!IsInteractiveEnsureMode() && bReportingNonCrash) || IsRunningDedicatedServer();
+
+	// By default we wont upload unless the *.ini has set this to true
+	bool bAgreedToCrashUpload = false;
+	GConfig->GetBool(TEXT("CrashReportClient"), TEXT("bAgreedToCrashUpload"), bAgreedToCrashUpload, GEngineIni);
+
+	bool bSkipCRC = bUnattended && !bAgreedToCrashUpload;
+
+	if (!bSkipCRC)
 	{
-		// generate "minidump"
-		GenerateReport(FPaths::Combine(*CrashInfoAbsolute, TEXT("Diagnostics.txt")));
-		
-		// generate "minidump" (just >1 byte)
-		GenerateMinidump(FPaths::Combine(*CrashInfoAbsolute, TEXT("minidump.dmp")));
-
-		// Introduces a new runtime crash context. Will replace all Windows related crash reporting.
-		FString FilePath(CrashInfoFolder);
-		FilePath += TEXT("/");
-		FilePath += FGenericCrashContext::CrashContextRuntimeXMLNameW;
-		SerializeAsXML(*FilePath);
-
-		// copy log
-		FString LogSrcAbsolute = FPlatformOutputDevices::GetAbsoluteLogFilename();
-		FString LogFolder = FPaths::GetPath(LogSrcAbsolute);
-		FString LogFilename = FPaths::GetCleanFilename(LogSrcAbsolute);
-		FString LogBaseFilename = FPaths::GetBaseFilename(LogSrcAbsolute);
-		FString LogExtension = FPaths::GetExtension(LogSrcAbsolute, true);
-		FString LogDstAbsolute = FPaths::Combine(*CrashInfoAbsolute, *LogFilename);
-		FPaths::NormalizeDirectoryName(LogDstAbsolute);
-
-		// Flush out the log
-		GLog->Flush();
+		FString CrashInfoFolder = FPaths::Combine(*FPaths::ProjectSavedDir(), TEXT("Crashes"), *FString::Printf(TEXT("%sinfo-%s-pid-%d-%s"), bReportingNonCrash ? TEXT("ensure") : TEXT("crash"), FApp::GetProjectName(), getpid(), *CrashGuid));
+		FString CrashInfoAbsolute = FPaths::ConvertRelativePathToFull(CrashInfoFolder);
+		if (IFileManager::Get().MakeDirectory(*CrashInfoAbsolute, true))
+		{
+			// generate "minidump"
+			GenerateReport(FPaths::Combine(*CrashInfoAbsolute, TEXT("Diagnostics.txt")));
+
+			// generate "minidump" (just >1 byte)
+			GenerateMinidump(FPaths::Combine(*CrashInfoAbsolute, TEXT("minidump.dmp")));
+
+			// Introduces a new runtime crash context. Will replace all Windows related crash reporting.
+			FString FilePath(CrashInfoFolder);
+			FilePath += TEXT("/");
+			FilePath += FGenericCrashContext::CrashContextRuntimeXMLNameW;
+			SerializeAsXML(*FilePath);
+
+			// copy log
+			FString LogSrcAbsolute = FPlatformOutputDevices::GetAbsoluteLogFilename();
+			FString LogFolder = FPaths::GetPath(LogSrcAbsolute);
+			FString LogFilename = FPaths::GetCleanFilename(LogSrcAbsolute);
+			FString LogBaseFilename = FPaths::GetBaseFilename(LogSrcAbsolute);
+			FString LogExtension = FPaths::GetExtension(LogSrcAbsolute, true);
+			FString LogDstAbsolute = FPaths::Combine(*CrashInfoAbsolute, *LogFilename);
+			FPaths::NormalizeDirectoryName(LogDstAbsolute);
+
+			// Flush out the log
+			GLog->Flush();
 
 #if !NO_LOGGING
-		bool bMemoryOnly = FPlatformOutputDevices::GetLog()->IsMemoryOnly();
-		bool bBacklogEnabled = FOutputDeviceRedirector::Get()->IsBacklogEnabled();
-
-		// The minimum free space on drive for saving a crash log
-		const uint64 MinDriveSpaceForCrashLog = 250 * 1024 * 1024;
-		// Max log file size to copy to report folder as will be filtered before submission to crash reporting backend
-		const uint64 MaxFileSizeForCrashLog = 100 * 1024 * 1024;
-
-		// Check whether server has low disk space available
-		uint64 TotalDiskSpace = 0;
-		uint64 TotalDiskFreeSpace = 0;
-		bool bLowDriveSpace = false;
-		if (FPlatformMisc::GetDiskTotalAndFreeSpace(*LogDstAbsolute, TotalDiskSpace, TotalDiskFreeSpace))
-		{
-			if (TotalDiskFreeSpace < MinDriveSpaceForCrashLog)
+			bool bMemoryOnly = FPlatformOutputDevices::GetLog()->IsMemoryOnly();
+			bool bBacklogEnabled = FOutputDeviceRedirector::Get()->IsBacklogEnabled();
+
+			// The minimum free space on drive for saving a crash log
+			const uint64 MinDriveSpaceForCrashLog = 250 * 1024 * 1024;
+			// Max log file size to copy to report folder as will be filtered before submission to crash reporting backend
+			const uint64 MaxFileSizeForCrashLog = 100 * 1024 * 1024;
+
+			// Check whether server has low disk space available
+			uint64 TotalDiskSpace = 0;
+			uint64 TotalDiskFreeSpace = 0;
+			bool bLowDriveSpace = false;
+			if (FPlatformMisc::GetDiskTotalAndFreeSpace(*LogDstAbsolute, TotalDiskSpace, TotalDiskFreeSpace))
 			{
-				bLowDriveSpace = true;
+				if (TotalDiskFreeSpace < MinDriveSpaceForCrashLog)
+				{
+					bLowDriveSpace = true;
+				}
 			}
-		}
 
-		if (!bLowDriveSpace)
-		{
-			if (bMemoryOnly || bBacklogEnabled)
+			if (!bLowDriveSpace)
 			{
-				FArchive* LogFile = IFileManager::Get().CreateFileWriter(*LogDstAbsolute, FILEWRITE_AllowRead);
-				if (LogFile)
+				if (bMemoryOnly || bBacklogEnabled)
 				{
-					if (bMemoryOnly)
+					FArchive* LogFile = IFileManager::Get().CreateFileWriter(*LogDstAbsolute, FILEWRITE_AllowRead);
+					if (LogFile)
 					{
-						FPlatformOutputDevices::GetLog()->Dump(*LogFile);
+						if (bMemoryOnly)
+						{
+							FPlatformOutputDevices::GetLog()->Dump(*LogFile);
+						}
+						else
+						{
+							FOutputDeviceArchiveWrapper Wrapper(LogFile);
+							GLog->SerializeBacklog(&Wrapper);
+						}
+
+						LogFile->Flush();
+						delete LogFile;
+					}
+				}
+				else
+				{
+					const bool bReplace = true;
+					const bool bEvenIfReadOnly = false;
+					const bool bAttributes = false;
+					FCopyProgress* const CopyProgress = nullptr;
+
+					// Only copy the log file if it is MaxLogFileSizeForCrashReport or less (note it will be zlib compressed for submission to data router)
+					if (IFileManager::Get().FileExists(*LogSrcAbsolute) && IFileManager::Get().FileSize(*LogSrcAbsolute) <= MaxFileSizeForCrashLog)
+					{
+						static_cast<void>(IFileManager::Get().Copy(*LogDstAbsolute, *LogSrcAbsolute, bReplace, bEvenIfReadOnly, bAttributes, CopyProgress, FILEREAD_AllowWrite, FILEWRITE_AllowRead));	// best effort, so don't care about result: couldn't copy -> tough, no log
 					}
 					else
 					{
-						FOutputDeviceArchiveWrapper Wrapper(LogFile);
-						GLog->SerializeBacklog(&Wrapper);
+						FFileHelper::SaveStringToFile(FString(TEXT("Log not available, too large for submission to crash reporting backend")), *LogDstAbsolute);
 					}
-
-					LogFile->Flush();
-					delete LogFile;
 				}
 			}
-			else
+			else if (TotalDiskFreeSpace >= MaxFileSizeForCrashLog)
 			{
-				const bool bReplace = true;
-				const bool bEvenIfReadOnly = false;
-				const bool bAttributes = false;
-				FCopyProgress* const CopyProgress = nullptr;
-
-				// Only copy the log file if it is MaxLogFileSizeForCrashReport or less (note it will be zlib compressed for submission to data router)				
-				if (IFileManager::Get().FileExists(*LogSrcAbsolute) && IFileManager::Get().FileSize(*LogSrcAbsolute) <= MaxFileSizeForCrashLog)
-				{
-					static_cast<void>(IFileManager::Get().Copy(*LogDstAbsolute, *LogSrcAbsolute, bReplace, bEvenIfReadOnly, bAttributes, CopyProgress, FILEREAD_AllowWrite, FILEWRITE_AllowRead));	// best effort, so don't care about result: couldn't copy -> tough, no log
-				}
-				else
-				{
-					FFileHelper::SaveStringToFile(FString(TEXT("Log not available, too large for submission to crash reporting backend")), *LogDstAbsolute);
-				}
+				FFileHelper::SaveStringToFile(FString(TEXT("Log not available, server has low available disk space")), *LogDstAbsolute);
 			}
-		}
-		else if (TotalDiskFreeSpace >= MaxFileSizeForCrashLog)
-		{
-			FFileHelper::SaveStringToFile(FString(TEXT("Log not available, server has low available disk space")), *LogDstAbsolute);
-		}
 #endif // !NO_LOGGING
 
-		// If present, include the crash report config file to pass config values to the CRC
-		const TCHAR* CrashConfigFilePath = GetCrashConfigFilePath();
-		if (IFileManager::Get().FileExists(CrashConfigFilePath))
-		{
-			FString CrashConfigFilename = FPaths::GetCleanFilename(CrashConfigFilePath);
-			FString CrashConfigDstAbsolute = FPaths::Combine(*CrashInfoAbsolute, *CrashConfigFilename);
-			static_cast<void>(IFileManager::Get().Copy(*CrashConfigDstAbsolute, CrashConfigFilePath));	// best effort, so don't care about result
-		}
+			// If present, include the crash report config file to pass config values to the CRC
+			const TCHAR* CrashConfigFilePath = GetCrashConfigFilePath();
+			if (IFileManager::Get().FileExists(CrashConfigFilePath))
+			{
+				FString CrashConfigFilename = FPaths::GetCleanFilename(CrashConfigFilePath);
+				FString CrashConfigDstAbsolute = FPaths::Combine(*CrashInfoAbsolute, *CrashConfigFilename);
+				static_cast<void>(IFileManager::Get().Copy(*CrashConfigDstAbsolute, CrashConfigFilePath));	// best effort, so don't care about result
+			}
 
-		// try launching the tool and wait for its exit, if at all
-		const TCHAR * RelativePathToCrashReporter = TEXT("../../../Engine/Binaries/Linux/CrashReportClient");	// FIXME: painfully hard-coded
+			// try launching the tool and wait for its exit, if at all
+			const TCHAR * RelativePathToCrashReporter = TEXT("../../../Engine/Binaries/Linux/CrashReportClient");	// FIXME: painfully hard-coded
 
-		FString CrashReportLogFilename = LogBaseFilename + TEXT("-CRC") + LogExtension;
-		FString CrashReportLogFilepath = FPaths::Combine(*LogFolder, *CrashReportLogFilename);
-		FString CrashReportClientArguments = TEXT(" -Abslog=");
-		CrashReportClientArguments += *CrashReportLogFilepath;
-		CrashReportClientArguments += TEXT(" ");
+			FString CrashReportLogFilename = LogBaseFilename + TEXT("-CRC") + LogExtension;
+			FString CrashReportLogFilepath = FPaths::Combine(*LogFolder, *CrashReportLogFilename);
+			FString CrashReportClientArguments = TEXT(" -Abslog=");
+			CrashReportClientArguments += *CrashReportLogFilepath;
+			CrashReportClientArguments += TEXT(" ");
 
-		// Suppress the user input dialog if we're running in unattended mode
-		bool bNoDialog = FApp::IsUnattended() || (!IsInteractiveEnsureMode() && bReportingNonCrash) || IsRunningDedicatedServer();
-		if (bNoDialog)
-		{
-			CrashReportClientArguments += TEXT(" -Unattended ");
-		}
+			if (bUnattended)
+			{
+				CrashReportClientArguments += TEXT(" -Unattended ");
+			}
 
-		// Whether to clean up crash reports after send
-		if (IsRunningDedicatedServer() && FParse::Param(FCommandLine::Get(), TEXT("CleanCrashReports")))
-		{
-			CrashReportClientArguments += TEXT(" -CleanCrashReports ");
-		}
+			if (bAgreedToCrashUpload)
+			{
+				CrashReportClientArguments += TEXT(" -SkipPopup ");
+			}
 
-		CrashReportClientArguments += CrashInfoAbsolute + TEXT("/");
+			// Whether to clean up crash reports after send
+			if (IsRunningDedicatedServer() && FParse::Param(FCommandLine::Get(), TEXT("CleanCrashReports")))
+			{
+				CrashReportClientArguments += TEXT(" -CleanCrashReports ");
+			}
 
-		if (bReportingNonCrash)
-		{
-			// If we're reporting non-crash, try to avoid spinning here and instead do that in the tick.
-			// However, if there was already a crash reporter running (i.e. we hit ensure() too quickly), take a hitch here
-			if (UnixCrashReporterTracker::CurrentTicker.IsValid())
+			CrashReportClientArguments += CrashInfoAbsolute + TEXT("/");
+
+			if (bReportingNonCrash)
 			{
-				// do not wait indefinitely, allow 45 second hitch (anticipating callstack parsing)
-				const double kEnsureTimeOut = 45.0;
-				const double kEnsureSleepInterval = 0.1;
-				if (!UnixCrashReporterTracker::WaitForProcWithTimeout(UnixCrashReporterTracker::CurrentlyRunningCrashReporter, kEnsureTimeOut, kEnsureSleepInterval))
+				// If we're reporting non-crash, try to avoid spinning here and instead do that in the tick.
+				// However, if there was already a crash reporter running (i.e. we hit ensure() too quickly), take a hitch here
+				if (UnixCrashReporterTracker::CurrentTicker.IsValid())
 				{
-					FPlatformProcess::TerminateProc(UnixCrashReporterTracker::CurrentlyRunningCrashReporter);
+					// do not wait indefinitely, allow 45 second hitch (anticipating callstack parsing)
+					const double kEnsureTimeOut = 45.0;
+					const double kEnsureSleepInterval = 0.1;
+					if (!UnixCrashReporterTracker::WaitForProcWithTimeout(UnixCrashReporterTracker::CurrentlyRunningCrashReporter, kEnsureTimeOut, kEnsureSleepInterval))
+					{
+						FPlatformProcess::TerminateProc(UnixCrashReporterTracker::CurrentlyRunningCrashReporter);
+					}
+
+					UnixCrashReporterTracker::Tick(0.001f);	// tick one more time to make it clean up after itself
 				}
 
-				UnixCrashReporterTracker::Tick(0.001f);	// tick one more time to make it clean up after itself
+				UnixCrashReporterTracker::CurrentlyRunningCrashReporter = FPlatformProcess::CreateProc(RelativePathToCrashReporter, *CrashReportClientArguments, true, false, false, NULL, 0, NULL, NULL);
+				UnixCrashReporterTracker::CurrentTicker = FTicker::GetCoreTicker().AddTicker(FTickerDelegate::CreateStatic(&UnixCrashReporterTracker::Tick), 1.f);
 			}
+			else
+			{
+				// spin here until CrashReporter exits
+				FProcHandle RunningProc = FPlatformProcess::CreateProc(RelativePathToCrashReporter, *CrashReportClientArguments, true, false, false, NULL, 0, NULL, NULL);
 
-			UnixCrashReporterTracker::CurrentlyRunningCrashReporter = FPlatformProcess::CreateProc(RelativePathToCrashReporter, *CrashReportClientArguments, true, false, false, NULL, 0, NULL, NULL);
-			UnixCrashReporterTracker::CurrentTicker = FTicker::GetCoreTicker().AddTicker(FTickerDelegate::CreateStatic(&UnixCrashReporterTracker::Tick), 1.f);
-		}
-		else
-		{
-			// spin here until CrashReporter exits
-			FProcHandle RunningProc = FPlatformProcess::CreateProc(RelativePathToCrashReporter, *CrashReportClientArguments, true, false, false, NULL, 0, NULL, NULL);
+				// do not wait indefinitely - can be more generous about the hitch than in ensure() case
+				// NOTE: Chris.Wood - increased from 3 to 8 mins because server crashes were timing out and getting lost
+				// NOTE: Do not increase above 8.5 mins without altering watchdog scripts to match
+				const double kCrashTimeOut = 8 * 60.0;
 
-			// do not wait indefinitely - can be more generous about the hitch than in ensure() case
-			// NOTE: Chris.Wood - increased from 3 to 8 mins because server crashes were timing out and getting lost
-			// NOTE: Do not increase above 8.5 mins without altering watchdog scripts to match
-			const double kCrashTimeOut = 8 * 60.0;
+				const double kCrashSleepInterval = 1.0;
+				if (!UnixCrashReporterTracker::WaitForProcWithTimeout(RunningProc, kCrashTimeOut, kCrashSleepInterval))
+				{
+					FPlatformProcess::TerminateProc(RunningProc);
+				}
 
-			const double kCrashSleepInterval = 1.0;
-			if (!UnixCrashReporterTracker::WaitForProcWithTimeout(RunningProc, kCrashTimeOut, kCrashSleepInterval))
-			{
-				FPlatformProcess::TerminateProc(RunningProc);
+				FPlatformProcess::CloseProc(RunningProc);
 			}
-
-			FPlatformProcess::CloseProc(RunningProc);
 		}
 	}
 
