Linux: Write to stderr when we fail to find expected to find sym file

diff --git a/Engine/Source/Runtime/Core/Private/Unix/UnixPlatformStackWalk.cpp b/Engine/Source/Runtime/Core/Private/Unix/UnixPlatformStackWalk.cpp
--- a/Engine/Source/Runtime/Core/Private/Unix/UnixPlatformStackWalk.cpp
+++ b/Engine/Source/Runtime/Core/Private/Unix/UnixPlatformStackWalk.cpp
@@ -18,6 +18,7 @@
 #include <dlfcn.h>
 #include <cxxabi.h>
 #include <stdio.h>
+#include <unistd.h>
 
 // FIXME Remove this define once we remove the old way symbolicate
 #define _ELFDEFINITIONS_H_
@@ -1027,7 +1028,7 @@ namespace
 				FCStringAnsi::Strcpy(out_SymbolInfo.FunctionName, info.dli_sname);
 			}
 
-			char ModuleSymbolPath[MAX_PATH + 1];
+			ANSICHAR ModuleSymbolPath[MAX_PATH + 1];
 
 			// We cant assume if we are relative we have not chdir to a different working dir.
 			if (FPaths::IsRelative(info.dli_fname))
@@ -1103,6 +1104,22 @@ namespace
 					}
 				}
 			}
+			// We only care if we fail to find our own *.sym file
+			else if (!FCStringAnsi::Strcmp(SOName, TCHAR_TO_ANSI(FPlatformProcess::ExecutableName())))
+			{
+				static bool bReported = false;
+				if (!bReported)
+				{
+					// Will not be part of UE_LOG as it would potentially allocate memory
+					const ANSICHAR* Message = "Failed to find symbol file, expected location:\n\"";
+					write(STDOUT_FILENO, Message, FCStringAnsi::Strlen(Message));
+					write(STDOUT_FILENO, ModuleSymbolPath, FCStringAnsi::Strlen(ModuleSymbolPath));
+					write(STDOUT_FILENO, "\"\n", 2);
+
+					// This will likely happen multiple times, so only write out once
+					bReported = true;
+				}
+			}
 		}
 
 		return false;
