Linux: Use MallocCrash when hitting out of memory issues in BinnedAllocFromOS

diff --git a/Engine/Source/Runtime/Core/Private/GenericPlatform/GenericPlatformMemory.cpp b/Engine/Source/Runtime/Core/Private/GenericPlatform/GenericPlatformMemory.cpp
--- a/Engine/Source/Runtime/Core/Private/GenericPlatform/GenericPlatformMemory.cpp
+++ b/Engine/Source/Runtime/Core/Private/GenericPlatform/GenericPlatformMemory.cpp
@@ -289,9 +289,7 @@ void* FGenericPlatformMemory::BinnedAllocFromOS( SIZE_T Size )
 	Pointer = PointerWeGotFromMMap;
 	if (Pointer == MAP_FAILED)
 	{
-		const int ErrNo = errno;
-		UE_LOG(LogHAL, Fatal, TEXT("mmap(len=%llu, size as passed %llu) failed with errno = %d (%s)"), (uint64)(ActualSizeMapped + DescriptorSize), (uint64)Size,
-			ErrNo, StringCast< TCHAR >(strerror(ErrNo)).Get());
+		FPlatformMemory::OnOutOfMemory(ActualSizeMapped, ExpectedAlignment);
 		// unreachable
 		return nullptr;
 	}
@@ -311,9 +309,7 @@ void* FGenericPlatformMemory::BinnedAllocFromOS( SIZE_T Size )
 			// unmap the part before
 			if (munmap(Pointer, SizeToNextAlignedPointer) != 0)
 			{
-				const int ErrNo = errno;
-				UE_LOG(LogHAL, Fatal, TEXT("munmap(addr=%p, len=%llu) failed with errno = %d (%s)"), Pointer, (uint64)(SizeToNextAlignedPointer),
-					ErrNo, StringCast< TCHAR >(strerror(ErrNo)).Get());
+				FPlatformMemory::OnOutOfMemory(SizeToNextAlignedPointer, ExpectedAlignment);
 				// unreachable
 				return nullptr;
 			}
@@ -341,9 +337,7 @@ void* FGenericPlatformMemory::BinnedAllocFromOS( SIZE_T Size )
 		{
 			if (munmap(TailPtr, TailSize) != 0)
 			{
-				const int ErrNo = errno;
-				UE_LOG(LogHAL, Fatal, TEXT("munmap(addr=%p, len=%llu) failed with errno = %d (%s)"), TailPtr, (uint64)TailSize,
-					ErrNo, StringCast< TCHAR >(strerror(ErrNo)).Get());
+				FPlatformMemory::OnOutOfMemory(TailSize, ExpectedAlignment);
 				// unreachable
 				return nullptr;
 			}
@@ -423,18 +417,16 @@ void FGenericPlatformMemory::BinnedFreeToOS( void* Ptr, SIZE_T Size )
 
 		if (UNLIKELY(munmap(PointerToUnmap, SizeToUnmap) != 0))
 		{
-			const int ErrNo = errno;
-			UE_LOG(LogHAL, Fatal, TEXT("munmap(addr=%p, len=%llu, size as passed %llu) failed with errno = %d (%s)"), PointerToUnmap, (uint64)SizeToUnmap, (uint64)Size,
-				ErrNo, StringCast< TCHAR >(strerror(ErrNo)).Get());
+			FPlatformMemory::OnOutOfMemory(SizeToUnmap, 0);
+			// unreachable
 		}
 	}
 	else
 	{
 		if (UNLIKELY(munmap(Ptr, SizeInWholePages) != 0))
 		{
-			const int ErrNo = errno;
-			UE_LOG(LogHAL, Fatal, TEXT("munmap(addr=%p, len=%llu, size as passed %llu) failed with errno = %d (%s)"), Ptr, (uint64)SizeInWholePages, (uint64)Size,
-				ErrNo, StringCast< TCHAR >(strerror(ErrNo)).Get());
+			FPlatformMemory::OnOutOfMemory(SizeInWholePages, 0);
+			// unreachable
 		}
 	}
 #endif // UE4_PLATFORM_USES_MMAP_FOR_BINNED_OS_ALLOCS
diff --git a/Engine/Source/Runtime/Core/Private/Unix/UnixPlatformMemory.cpp b/Engine/Source/Runtime/Core/Private/Unix/UnixPlatformMemory.cpp
--- a/Engine/Source/Runtime/Core/Private/Unix/UnixPlatformMemory.cpp
+++ b/Engine/Source/Runtime/Core/Private/Unix/UnixPlatformMemory.cpp
@@ -6,6 +6,8 @@
 
 #include "Unix/UnixPlatformMemory.h"
 #include "Misc/AssertionMacros.h"
+#include "Misc/CoreDelegates.h"
+#include "Misc/FeedbackContext.h"
 #include "Math/NumericLimits.h"
 #include "Math/UnrealMathUtility.h"
 #include "Templates/UnrealTemplate.h"
@@ -17,6 +19,8 @@
 #include "HAL/MallocBinned2.h"
 #include "HAL/MallocReplayProxy.h"
 #include "HAL/MallocStomp.h"
+#include "HAL/PlatformMallocCrash.h"
+
 #if PLATFORM_FREEBSD
 	#include <kvm.h>
 #else
@@ -881,6 +885,48 @@ bool FUnixPlatformMemory::UnmapNamedSharedMemoryRegion(FSharedMemoryRegion * Mem
 	return bAllSucceeded;
 }
 
+void FUnixPlatformMemory::OnOutOfMemory(uint64 Size, uint32 Alignment)
+{
+	// Update memory stats before we enter the crash handler.
+	OOMAllocationSize = Size;
+	OOMAllocationAlignment = Alignment;
+
+	// only call this code one time - if already OOM, abort
+	if (bIsOOM)
+	{
+		return;
+	}
+	bIsOOM = true;
+
+	FMalloc* Prev = GMalloc;
+	FPlatformMallocCrash::Get().SetAsGMalloc();
+
+	FPlatformMemoryStats PlatformMemoryStats = FPlatformMemory::GetStats();
+
+	UE_LOG(LogMemory, Warning, TEXT("MemoryStats:")\
+		TEXT("\n\tAvailablePhysical %llu")\
+		TEXT("\n\t AvailableVirtual %llu")\
+		TEXT("\n\t     UsedPhysical %llu")\
+		TEXT("\n\t PeakUsedPhysical %llu")\
+		TEXT("\n\t      UsedVirtual %llu")\
+		TEXT("\n\t  PeakUsedVirtual %llu"),
+		(uint64)PlatformMemoryStats.AvailablePhysical,
+		(uint64)PlatformMemoryStats.AvailableVirtual,
+		(uint64)PlatformMemoryStats.UsedPhysical,
+		(uint64)PlatformMemoryStats.PeakUsedPhysical,
+		(uint64)PlatformMemoryStats.UsedVirtual,
+		(uint64)PlatformMemoryStats.PeakUsedVirtual);
+	if (GWarn)
+	{
+		Prev->DumpAllocatorStats(*GWarn);
+	}
+
+	// let any registered handlers go
+	FCoreDelegates::GetMemoryTrimDelegate().Broadcast();
+
+	UE_LOG(LogMemory, Fatal, TEXT("Ran out of memory allocating %llu bytes with alignment %u"), Size, Alignment);
+	// unreachable
+}
 
 /**
 * LLM uses these low level functions (LLMAlloc and LLMFree) to allocate memory. It grabs
diff --git a/Engine/Source/Runtime/Core/Public/Unix/UnixPlatformMemory.h b/Engine/Source/Runtime/Core/Public/Unix/UnixPlatformMemory.h
--- a/Engine/Source/Runtime/Core/Public/Unix/UnixPlatformMemory.h
+++ b/Engine/Source/Runtime/Core/Public/Unix/UnixPlatformMemory.h
@@ -77,6 +77,7 @@ struct CORE_API FUnixPlatformMemory : public FGenericPlatformMemory
 	static FSharedMemoryRegion * MapNamedSharedMemoryRegion(const FString& InName, bool bCreate, uint32 AccessMode, SIZE_T Size);
 	static bool UnmapNamedSharedMemoryRegion(FSharedMemoryRegion * MemoryRegion);
 	static bool GetLLMAllocFunctions(void*(*&OutAllocFunction)(size_t), void(*&OutFreeFunction)(void*, size_t), int32& OutAlignment);
+	static CA_NO_RETURN void OnOutOfMemory(uint64 Size, uint32 Alignment);
 	//~ End FGenericPlatformMemory Interface
 };
 
