Linux: Fix deadlock in a file cache which could be locked in a crash handler

diff --git a/Engine/Source/Runtime/Core/Private/Unix/UnixPlatformCrashContext.cpp b/Engine/Source/Runtime/Core/Private/Unix/UnixPlatformCrashContext.cpp
--- a/Engine/Source/Runtime/Core/Private/Unix/UnixPlatformCrashContext.cpp
+++ b/Engine/Source/Runtime/Core/Private/Unix/UnixPlatformCrashContext.cpp
@@ -628,6 +628,8 @@ void DefaultCrashHandler(const FUnixCrashContext & Context)
 /** Global pointer to crash handler */
 void (* GCrashHandlerPointer)(const FGenericCrashContext & Context) = NULL;
 
+extern int32 CORE_API GMaxNumberFileMappingCache;
+
 /** True system-specific crash handler that gets called first */
 void PlatformCrashHandler(int32 Signal, siginfo_t* Info, void* Context)
 {
@@ -639,6 +641,9 @@ void PlatformCrashHandler(int32 Signal, siginfo_t* Info, void* Context)
 	// Switch to malloc crash.
 	FPlatformMallocCrash::Get().SetAsGMalloc();
 
+	// Once we crash we can no longer try to find cache files. This can cause a deadlock when crashing while having locked in that file cache
+	GMaxNumberFileMappingCache = 0;
+
 	FUnixCrashContext CrashContext;
 	CrashContext.InitFromSignal(Signal, Info, Context);
 	CrashContext.FirstCrashHandlerFrame = static_cast<uint64*>(__builtin_return_address(0));
