Linux: If we have a 0 LineNumber lets try to use to the previous Record.

Still an issue with non-virtual thunks reporting line number zero but it seems even windows skips these frames. GDB reports a different callsite that doesnt seem super related (possibly?) Nothing to do with thunking though.

diff --git a/Engine/Source/Runtime/Core/Private/Unix/UnixPlatformStackWalk.cpp b/Engine/Source/Runtime/Core/Private/Unix/UnixPlatformStackWalk.cpp
--- a/Engine/Source/Runtime/Core/Private/Unix/UnixPlatformStackWalk.cpp
+++ b/Engine/Source/Runtime/Core/Private/Unix/UnixPlatformStackWalk.cpp
@@ -935,6 +935,7 @@ namespace
 
 		Record GetRecord(int Index) const
 		{
+			// When we remove this check, make sure we handle possible out of bounds cases
 			if (Index > RecordCount || Index < 0)
 			{
 				return {};
@@ -1075,6 +1076,17 @@ namespace
 
 					if (AddressToFind >= Current.Address && AddressToFind < Current.Address + Size)
 					{
+						// Hack when we have a zero line number to attempt to use the previous record for a better guess.
+						// non-virtual thunks seem to cause a bunch of these but this will not fix those.
+						if (Current.LineNumber == 0)
+						{
+							Record Previous = Reader.GetRecord(Middle - 1);
+							if (Previous.LineNumber > 0 && Previous.LineNumber != static_cast<uint32_t>(-1))
+							{
+								Current.LineNumber = Previous.LineNumber;
+							}
+						}
+
 						Reader.ReadOffsetIntoMemory(out_SymbolInfo.Filename, sizeof(out_SymbolInfo.Filename), Current.FileRelativeOffset);
 						Reader.ReadOffsetIntoMemory(out_SymbolInfo.FunctionName, sizeof(out_SymbolInfo.FunctionName), Current.SymbolRelativeOffset);
 						out_SymbolInfo.LineNumber = Current.LineNumber;
