Fix for missing include.

diff --git a/Engine/Source/Runtime/Core/Private/Mac/MacPlatformProcess.cpp b/Engine/Source/Runtime/Core/Private/Mac/MacPlatformProcess.cpp
--- a/Engine/Source/Runtime/Core/Private/Mac/MacPlatformProcess.cpp
+++ b/Engine/Source/Runtime/Core/Private/Mac/MacPlatformProcess.cpp
@@ -5,6 +5,7 @@
 =============================================================================*/
 
 #include "Mac/MacPlatformProcess.h"
+#include "Mac/MacPlatform.h"
 #include "Apple/ApplePlatformRunnableThread.h"
 #include "Misc/App.h"
 #include "Misc/Paths.h"
