Fix mac CIS

diff --git a/Engine/Source/Runtime/Core/Private/Mac/MacPlatformProcess.cpp b/Engine/Source/Runtime/Core/Private/Mac/MacPlatformProcess.cpp
--- a/Engine/Source/Runtime/Core/Private/Mac/MacPlatformProcess.cpp
+++ b/Engine/Source/Runtime/Core/Private/Mac/MacPlatformProcess.cpp
@@ -850,7 +850,7 @@ const TCHAR* FMacPlatformProcess::UserLogsDir()
 
 const TCHAR* FMacPlatformProcess::UserHomeDir()
 {
-	static TCHAR Result[MAX_PATH] = TEXT("");
+	static TCHAR Result[MAC_MAX_PATH] = TEXT("");
 	if (!Result[0])
 	{
 		SCOPED_AUTORELEASE_POOL;
