Fix FQuat::NetSerialize modifying the original value

diff --git a/Engine/Source/Runtime/Core/Private/Math/UnrealMath.cpp b/Engine/Source/Runtime/Core/Private/Math/UnrealMath.cpp
--- a/Engine/Source/Runtime/Core/Private/Math/UnrealMath.cpp
+++ b/Engine/Source/Runtime/Core/Private/Math/UnrealMath.cpp
@@ -797,10 +797,12 @@ FVector FQuat::Euler() const
 
 bool FQuat::NetSerialize(FArchive& Ar, class UPackageMap*, bool& bOutSuccess)
 {
-	FQuat &Q = *this;
+	FQuat Q;
 
 	if (Ar.IsSaving())
 	{
+		Q = *this;
+
 		// Make sure we have a non null SquareSum. It shouldn't happen with a quaternion, but better be safe.
 		if(Q.SizeSquared() <= SMALL_NUMBER)
 		{
@@ -809,7 +811,10 @@ bool FQuat::NetSerialize(FArchive& Ar, class UPackageMap*, bool& bOutSuccess)
 		else
 		{
 			// All transmitted quaternions *MUST BE* unit quaternions, in which case we can deduce the value of W.
-			Q.Normalize();
+			if (!ensure(Q.IsNormalized()))
+			{
+				Q.Normalize();
+			}
 			// force W component to be non-negative
 			if (Q.W < 0.f)
 			{
@@ -822,6 +827,7 @@ bool FQuat::NetSerialize(FArchive& Ar, class UPackageMap*, bool& bOutSuccess)
 	}
 
 	Ar << Q.X << Q.Y << Q.Z;
+
 	if ( Ar.IsLoading() )
 	{
 		const float XYZMagSquared = (Q.X*Q.X + Q.Y*Q.Y + Q.Z*Q.Z);
@@ -841,6 +847,8 @@ bool FQuat::NetSerialize(FArchive& Ar, class UPackageMap*, bool& bOutSuccess)
 			Q.Y *= XYZInvMag;
 			Q.Z *= XYZInvMag;
 		}
+
+		*this = Q;
 	}
 
 	bOutSuccess = true;
