Fix static analysis CIS
This function should be moved to SceneUtils.h and the dependency in AsyncTaskDownloadImage resolved another way.

diff --git a/Engine/Source/Runtime/Core/Public/HAL/Platform.h b/Engine/Source/Runtime/Core/Public/HAL/Platform.h
--- a/Engine/Source/Runtime/Core/Public/HAL/Platform.h
+++ b/Engine/Source/Runtime/Core/Public/HAL/Platform.h
@@ -890,9 +890,9 @@ namespace TypeTests
 
 // this function is used to suppress static analysis warnings
 #if PLATFORM_HTML5
-#define IsHTML5Platform() true
+FORCEINLINE bool IsHTML5Platform() { return true; }
 #else
-#define IsHTML5Platform() false
+FORCEINLINE bool IsHTML5Platform() { return false; }
 #endif
 
 
