Core: Fixed closing quotes with preceeding escaped backslashes not being parsed correctly

diff --git a/Engine/Source/Runtime/Core/Private/Misc/Parse.cpp b/Engine/Source/Runtime/Core/Private/Misc/Parse.cpp
--- a/Engine/Source/Runtime/Core/Private/Misc/Parse.cpp
+++ b/Engine/Source/Runtime/Core/Private/Misc/Parse.cpp
@@ -1069,7 +1069,6 @@ bool FParse::LineExtended(const TCHAR** Stream, FString& Result, int32& LinesCon
 	bool bGotStream = false;
 	bool bIsQuoted = false;
 	bool bIgnore = false;
-	bool bIsEscapedDoubleQuote = false;
 	int32 BracketDepth = 0;
 
 	Result = TEXT("");
@@ -1126,18 +1125,15 @@ bool FParse::LineExtended(const TCHAR** Stream, FString& Result, int32& LinesCon
 			BracketDepth--;
 			(*Stream)++;
 		}
+		// specifically consume escaped backslashes and quotes within quoted strings
+		else if (bIsQuoted && !bIgnore && (*Stream)[0] == TEXT('\\') && ( (*Stream)[1] == TEXT('\"') || (*Stream)[1] == TEXT('\\') ))
+		{
+			Result.AppendChars(*Stream, 2);
+			(*Stream) += 2;
+		}
 		else
 		{
-			// Check quoting.
-			bIsEscapedDoubleQuote |= ((*Stream)[0] == TEXT('\\') && (*Stream)[1] == TEXT('\"'));
-			if (!bIsEscapedDoubleQuote)
-			{
-				bIsQuoted = bIsQuoted ^ (**Stream == TEXT('\"'));
-			}
-			else if (**Stream == TEXT('\"'))
-			{
-				bIsEscapedDoubleQuote = false;
-			}
+			bIsQuoted = bIsQuoted ^ (**Stream == TEXT('\"'));
 
 			// Got stuff.
 			if (!bIgnore)
@@ -1234,3 +1230,50 @@ bool FParse::SchemeNameFromURI(const TCHAR* URI, FString& OutSchemeName)
 		}
 	}
 }
+
+
+#if !(UE_BUILD_SHIPPING || UE_BUILD_TEST)
+
+#include "Misc/AutomationTest.h"
+IMPLEMENT_SIMPLE_AUTOMATION_TEST(FParseLineExtendedTest, "System.Core.Misc.ParseLineExtended", EAutomationTestFlags::ApplicationContextMask | EAutomationTestFlags::SmokeFilter)
+bool FParseLineExtendedTest::RunTest(const FString& Parameters)
+{
+	const TCHAR* Tests[] = {
+		TEXT("Test string"),                            // Normal string
+		TEXT("{Test string}"),                          // Braced string
+		TEXT("\"Test string\""),                        // Quoted string
+		TEXT("\"Test \\\"string\\\"\""),                // Quoted string w/ escaped quotes
+		TEXT("a=\"Test\", b=\"Test\""),                 // Quoted value list
+		TEXT("a=\"Test\\\\\", b=\"{Test}\""),           // Quoted value list w/ escaped backslash preceeding closing quote
+		TEXT("a=\"Test\\\\\\\" String\", b=\"{Test}\""),// Quoted value list w/ escaped backslash preceeding escaped quote
+		TEXT("Test=(Inner=\"{content}\")"),             // Nested value list
+	};
+
+	const TCHAR* Expected[] = {
+		TEXT("Test string"),
+		TEXT("Test string"),
+		TEXT("\"Test string\""),
+		TEXT("\"Test \\\"string\\\"\""),
+		TEXT("a=\"Test\", b=\"Test\""),
+		TEXT("a=\"Test\\\\\", b=\"{Test}\""),
+		TEXT("a=\"Test\\\\\\\" String\", b=\"{Test}\""),
+		TEXT("Test=(Inner=\"{content}\")"),
+	};
+
+	int32 LinesConsumed = 0;
+	FString Result;
+
+	for (int32 Index = 0; Index < ARRAY_COUNT(Tests); ++Index)
+	{
+		LinesConsumed = 0;
+		Result.Reset();
+
+		const TCHAR* Stream = Tests[Index];
+		bool bSuccess = FParse::LineExtended(&Stream, Result, LinesConsumed, false);
+		TestTrue(*FString::Printf(TEXT("Expecting parsed line [%s] to be [%s]. Result was [%s]."), Tests[Index], Expected[Index], *Result), bSuccess && Result == Expected[Index]);
+	}
+
+	return true;
+}
+
+#endif //!(UE_BUILD_SHIPPING || UE_BUILD_TEST)
