PR #5116: Enable the equivilent of -Wreorder on MSVC (Contributed by [REDACTED_PERSON])

diff --git a/Engine/Plugins/Editor/SpeedTreeImporter/Source/SpeedTreeImporter/Private/SpeedTreeImportFactory.cpp b/Engine/Plugins/Editor/SpeedTreeImporter/Source/SpeedTreeImporter/Private/SpeedTreeImportFactory.cpp
--- a/Engine/Plugins/Editor/SpeedTreeImporter/Source/SpeedTreeImporter/Private/SpeedTreeImportFactory.cpp
+++ b/Engine/Plugins/Editor/SpeedTreeImporter/Source/SpeedTreeImporter/Private/SpeedTreeImportFactory.cpp
@@ -64,17 +64,10 @@
 
 #if WITH_SPEEDTREE
 
-#ifdef __clang__ // @todo: ThirdParty/SpeedTree/SpeedTreeSDK-v7.0/Include/Core/RenderState_inl.h:52:2: warning: field 'm_eHueVariation' will be initialized after field 'm_eShaderGenerationMode' [-Wreorder]
-	#pragma clang diagnostic push
-	#pragma clang diagnostic ignored "-Wreorder"
-#endif
-
+THIRD_PARTY_INCLUDES_START
 #include "Core/Core.h"
 #include "TreeReader.h"
-
-#ifdef __clang__
-	#pragma clang diagnostic pop
-#endif
+THIRD_PARTY_INCLUDES_END
 
 #endif // WITH_SPEEDTREE
 #include "PhysicsEngine/SphereElem.h"
diff --git a/Engine/Plugins/Experimental/PixelStreaming/Source/PixelStreaming/Private/ProxyConnection.cpp b/Engine/Plugins/Experimental/PixelStreaming/Source/PixelStreaming/Private/ProxyConnection.cpp
--- a/Engine/Plugins/Experimental/PixelStreaming/Source/PixelStreaming/Private/ProxyConnection.cpp
+++ b/Engine/Plugins/Experimental/PixelStreaming/Source/PixelStreaming/Private/ProxyConnection.cpp
@@ -16,10 +16,10 @@
 FProxyConnection::FProxyConnection(const FString& IP, uint16 Port, FStreamer& Streamer) :
 	Streamer(Streamer),
 	InputDevice(FModuleManager::Get().GetModuleChecked<IPixelStreamingPlugin>("PixelStreaming").GetInputDevice()),
-	Thread(TEXT("WebRTC Proxy Connection"), [this, IP, Port]() { Run(IP, Port); }),
 	Socket(nullptr),
 	Listener(nullptr),
-	ExitRequested(false)
+	ExitRequested(false),
+	Thread(TEXT("WebRTC Proxy Connection"), [this, IP, Port]() { Run(IP, Port); })
 {}
 
 FProxyConnection::~FProxyConnection()
diff --git a/Engine/Plugins/Experimental/ProxyLODPlugin/Source/ProxyLOD/Private/ProxyLODMeshTypes.h b/Engine/Plugins/Experimental/ProxyLODPlugin/Source/ProxyLOD/Private/ProxyLODMeshTypes.h
--- a/Engine/Plugins/Experimental/ProxyLODPlugin/Source/ProxyLOD/Private/ProxyLODMeshTypes.h
+++ b/Engine/Plugins/Experimental/ProxyLODPlugin/Source/ProxyLOD/Private/ProxyLODMeshTypes.h
@@ -471,18 +471,18 @@ private:
 
 template <typename SimplifierVertexType>
 TAOSMesh<SimplifierVertexType>::TAOSMesh() :
-	NumVertexes(0),
-	NumIndexes(0),
 	Vertexes(NULL),
-	Indexes(NULL)
+	Indexes(NULL),
+	NumVertexes(0),
+	NumIndexes(0)
 {}
 
 template <typename SimplifierVertexType>
 TAOSMesh<SimplifierVertexType>::TAOSMesh(int32 VertCount, int32 FaceCount) :
-	NumVertexes(0),
-	NumIndexes(0),
 	Vertexes(NULL),
-	Indexes(NULL)
+	Indexes(NULL),
+	NumVertexes(0),
+	NumIndexes(0)
 {
 	this->Resize(VertCount, FaceCount);
 }
diff --git a/Engine/Plugins/Runtime/WindowsMixedReality/Source/WindowsMixedRealityHMD/Private/WindowsMixedRealityHMD.cpp b/Engine/Plugins/Runtime/WindowsMixedReality/Source/WindowsMixedRealityHMD/Private/WindowsMixedRealityHMD.cpp
--- a/Engine/Plugins/Runtime/WindowsMixedReality/Source/WindowsMixedRealityHMD/Private/WindowsMixedRealityHMD.cpp
+++ b/Engine/Plugins/Runtime/WindowsMixedReality/Source/WindowsMixedRealityHMD/Private/WindowsMixedRealityHMD.cpp
@@ -859,11 +859,11 @@ namespace WindowsMixedReality
 
 	FWindowsMixedRealityHMD::FWindowsMixedRealityHMD(const FAutoRegister& AutoRegister)
 		: FSceneViewExtensionBase(AutoRegister)
+		, ScreenScalePercentage(1.0f)
 		, mCustomPresent(nullptr)
+		, HMDTrackingOrigin(EHMDTrackingOrigin::Floor)
 		, CurrOrientation(FQuat::Identity)
 		, CurrPosition(FVector::ZeroVector)
-		, HMDTrackingOrigin(EHMDTrackingOrigin::Floor)
-		, ScreenScalePercentage(1.0f)
 	{
 		static const FName RendererModuleName("Renderer");
 		RendererModule = FModuleManager::GetModulePtr<IRendererModule>(RendererModuleName);
diff --git a/Engine/Source/Runtime/AVIWriter/Private/AVIWriter.cpp b/Engine/Source/Runtime/AVIWriter/Private/AVIWriter.cpp
--- a/Engine/Source/Runtime/AVIWriter/Private/AVIWriter.cpp
+++ b/Engine/Source/Runtime/AVIWriter/Private/AVIWriter.cpp
@@ -22,27 +22,27 @@ IMPLEMENT_MODULE(FAVIWriterModule, AVIWriter);
 
 #include "Windows/WindowsHWrapper.h"
 #include "Windows/AllowWindowsPlatformTypes.h"
+
 typedef TCHAR* PTCHAR;
-#pragma warning(push)
-#pragma warning(disable : 4263) // 'function' : member function does not override any base class virtual member function
-#pragma warning(disable : 4264) // 'virtual_function' : no override available for virtual member function from base 'cla
-#pragma warning(disable : 4265) // 'class' : class has virtual functions, but destructor is not virtual
-#if USING_CODE_ANALYSIS
-	#pragma warning(disable:6509)  // Invalid annotation: 'return' cannot be referenced in some contexts
-	#pragma warning(disable:6101)  // Returning uninitialized memory '*lpdwExitCode'.  A successful path through the function does not set the named _Out_ parameter.
-	#pragma warning(disable:28204) // 'Func' has an override at 'file' and only the override is annotated for _Param_(N): when an override is annotated, the base (this function) should be similarly annotated.
-#endif
-#ifdef __clang__
-#pragma clang diagnostic push
-#pragma clang diagnostic ignored "-Wreorder"
-#endif
+
+THIRD_PARTY_INCLUDES_START
+#if !defined(__clang__)
+	#pragma warning(push)
+	#pragma warning(disable: 4263) // 'function' : member function does not override any base class virtual member function
+	#pragma warning(disable: 4264) // 'virtual_function' : no override available for virtual member function from base 'cla
+	#if USING_CODE_ANALYSIS
+		#pragma warning(disable: 6509)  // Invalid annotation: 'return' cannot be referenced in some contexts
+		#pragma warning(disable: 6101)  // Returning uninitialized memory '*lpdwExitCode'.  A successful path through the function does not set the named _Out_ parameter.
+		#pragma warning(disable: 28204) // 'Func' has an override at 'file' and only the override is annotated for _Param_(N): when an override is annotated, the base (this function) should be similarly annotated.
+	#endif // USING_CODE_ANALYSIS
+#endif // !defined(__clang__)
 #include <streams.h>
 #include <dshow.h>
 #include <initguid.h>
-#ifdef __clang__
-#pragma clang diagnostic pop
-#endif
-#pragma warning(pop)
+#if !defined(__clang__)
+	#pragma warning(pop)
+#endif // !defined(__clang__)
+THIRD_PARTY_INCLUDES_END
 
 #include "Windows/HideWindowsPlatformTypes.h"
 
diff --git a/Engine/Source/Runtime/AVIWriter/Private/CapturePin.h b/Engine/Source/Runtime/AVIWriter/Private/CapturePin.h
--- a/Engine/Source/Runtime/AVIWriter/Private/CapturePin.h
+++ b/Engine/Source/Runtime/AVIWriter/Private/CapturePin.h
@@ -13,32 +13,34 @@ struct FCapturedFrame;
 
 #if PLATFORM_WINDOWS && !UE_BUILD_MINIMAL
 
+#include "Windows/WindowsHWrapper.h"
+#include "Windows/AllowWindowsPlatformTypes.h"
+
 typedef TCHAR* PTCHAR;
 
-#ifdef __clang__
+THIRD_PARTY_INCLUDES_START
+#if defined(__clang__)
 	#pragma clang diagnostic push
-	#pragma clang diagnostic ignored "-Wreorder"	// warning : field 'x' will be initialized after field 'y' [-Wreorder]
 	#pragma clang diagnostic ignored "-Woverloaded-virtual"	 // note: hidden overloaded virtual function 'x' declared here: different number of parameters
 #else
 	#pragma warning(push)
-	#pragma warning(disable : 4263) // 'function' : member function does not override any base class virtual member function
-	#pragma warning(disable : 4264) // 'virtual_function' : no override available for virtual member function from base 
-	#pragma warning(disable : 4265) // 'class' : class has virtual functions, but destructor is not virtual
+	#pragma warning(disable: 4263) // 'function' : member function does not override any base class virtual member function
+	#pragma warning(disable: 4264) // 'virtual_function' : no override available for virtual member function from base 
 	#if USING_CODE_ANALYSIS
-		#pragma warning(disable:6509)  // Invalid annotation: 'return' cannot be referenced in some contexts
-		#pragma warning(disable:6101)  // Returning uninitialized memory '*lpdwExitCode'.  A successful path through the function does not set the named _Out_ parameter.
-		#pragma warning(disable:28204) // 'Func' has an override at 'file' and only the override is annotated for _Param_(N): when an override is annotated, the base (this function) should be similarly annotated.
-	#endif
-#endif
-#include "Windows/WindowsHWrapper.h"
-#include "Windows/AllowWindowsPlatformTypes.h"
+		#pragma warning(disable: 6101)  // Returning uninitialized memory '*lpdwExitCode'.  A successful path through the function does not set the named _Out_ parameter.
+		#pragma warning(disable: 6509)  // Invalid annotation: 'return' cannot be referenced in some contexts
+		#pragma warning(disable: 28204) // 'Func' has an override at 'file' and only the override is annotated for _Param_(N): when an override is annotated, the base (this function) should be similarly annotated.
+	#endif // USING_CODE_ANALYSIS
+#endif // defined(__clang__)
 #include <streams.h>
-#include "Windows/HideWindowsPlatformTypes.h"
-#ifdef __clang__
+#if defined(__clang__)
 	#pragma clang diagnostic pop
 #else
 	#pragma warning(pop)
-#endif
+#endif // defined(__clang__)
+THIRD_PARTY_INCLUDES_END
+
+#include "Windows/HideWindowsPlatformTypes.h"
 
 class FAVIWriter;
 struct FCapturedFrame;
diff --git a/Engine/Source/Runtime/Core/Public/Android/AndroidPlatformCompilerPreSetup.h b/Engine/Source/Runtime/Core/Public/Android/AndroidPlatformCompilerPreSetup.h
--- a/Engine/Source/Runtime/Core/Public/Android/AndroidPlatformCompilerPreSetup.h
+++ b/Engine/Source/Runtime/Core/Public/Android/AndroidPlatformCompilerPreSetup.h
@@ -7,6 +7,7 @@
 // Disable common CA warnings around SDK includes
 #ifndef THIRD_PARTY_INCLUDES_START
 	#define THIRD_PARTY_INCLUDES_START \
+		PRAGMA_DISABLE_REORDER_WARNINGS \
 		PRAGMA_DISABLE_SHADOW_VARIABLE_WARNINGS \
 		PRAGMA_DISABLE_UNDEFINED_IDENTIFIER_WARNINGS \
 		PRAGMA_DISABLE_MISSING_VIRTUAL_DESTRUCTOR_WARNINGS \
@@ -18,7 +19,8 @@
 		PRAGMA_ENABLE_DEPRECATION_WARNINGS \
 		PRAGMA_ENABLE_MISSING_VIRTUAL_DESTRUCTOR_WARNINGS \
 		PRAGMA_ENABLE_UNDEFINED_IDENTIFIER_WARNINGS \
-		PRAGMA_ENABLE_SHADOW_VARIABLE_WARNINGS
+		PRAGMA_ENABLE_SHADOW_VARIABLE_WARNINGS \
+		PRAGMA_ENABLE_REORDER_WARNINGS
 #endif
 
 // Make certain warnings always be warnings, even despite -Werror.
diff --git a/Engine/Source/Runtime/Core/Public/Apple/ApplePlatformCompilerPreSetup.h b/Engine/Source/Runtime/Core/Public/Apple/ApplePlatformCompilerPreSetup.h
--- a/Engine/Source/Runtime/Core/Public/Apple/ApplePlatformCompilerPreSetup.h
+++ b/Engine/Source/Runtime/Core/Public/Apple/ApplePlatformCompilerPreSetup.h
@@ -7,6 +7,7 @@
 // Disable common CA warnings around SDK includes
 #ifndef THIRD_PARTY_INCLUDES_START
 	#define THIRD_PARTY_INCLUDES_START \
+		PRAGMA_DISABLE_REORDER_WARNINGS \
 		PRAGMA_DISABLE_SHADOW_VARIABLE_WARNINGS \
 		PRAGMA_DISABLE_UNDEFINED_IDENTIFIER_WARNINGS \
 		PRAGMA_DISABLE_MISSING_VIRTUAL_DESTRUCTOR_WARNINGS \
@@ -18,13 +19,13 @@
 		PRAGMA_ENABLE_DEPRECATION_WARNINGS \
 		PRAGMA_ENABLE_MISSING_VIRTUAL_DESTRUCTOR_WARNINGS \
 		PRAGMA_ENABLE_UNDEFINED_IDENTIFIER_WARNINGS \
-		PRAGMA_ENABLE_SHADOW_VARIABLE_WARNINGS
+		PRAGMA_ENABLE_SHADOW_VARIABLE_WARNINGS \
+		PRAGMA_ENABLE_REORDER_WARNINGS
 #endif
 
 // Make certain warnings always be warnings, even despite -Werror.
 // Rationale: we don't want to suppress those as there are plans to address them (e.g. UE-12341), but breaking builds due to these warnings is very expensive
 // since they cannot be caught by all compilers that we support. They are deemed to be relatively safe to be ignored, at least until all SDKs/toolchains start supporting them.
-#pragma clang diagnostic warning "-Wreorder"
 #pragma clang diagnostic warning "-Wparentheses-equality"
 #pragma clang diagnostic ignored "-Winconsistent-missing-override"
 #pragma clang diagnostic ignored "-Wundefined-bool-conversion"
diff --git a/Engine/Source/Runtime/Core/Public/Clang/ClangPlatformCompilerPreSetup.h b/Engine/Source/Runtime/Core/Public/Clang/ClangPlatformCompilerPreSetup.h
--- a/Engine/Source/Runtime/Core/Public/Clang/ClangPlatformCompilerPreSetup.h
+++ b/Engine/Source/Runtime/Core/Public/Clang/ClangPlatformCompilerPreSetup.h
@@ -73,6 +73,17 @@
 		_Pragma("clang diagnostic pop")
 #endif // PRAGMA_ENABLE_MISSING_VIRTUAL_DESTRUCTOR_WARNINGS
 
+#ifndef PRAGMA_DISABLE_REORDER_WARNINGS
+	#define PRAGMA_DISABLE_REORDER_WARNINGS \
+		_Pragma("clang diagnostic push") \
+		_Pragma("clang diagnostic ignored \"-Wreorder\"")
+#endif // PRAGMA_DISABLE_REORDER_WARNINGS
+
+#ifndef PRAGMA_ENABLE_REORDER_WARNINGS
+	#define PRAGMA_ENABLE_REORDER_WARNINGS \
+		_Pragma("clang diagnostic pop")
+#endif // PRAGMA_ENABLE_REORDER_WARNINGS
+
 #ifndef PRAGMA_POP
 	#define PRAGMA_POP \
 		_Pragma("clang diagnostic pop")
diff --git a/Engine/Source/Runtime/Core/Public/HTML5/HTML5PlatformCompilerPreSetup.h b/Engine/Source/Runtime/Core/Public/HTML5/HTML5PlatformCompilerPreSetup.h
--- a/Engine/Source/Runtime/Core/Public/HTML5/HTML5PlatformCompilerPreSetup.h
+++ b/Engine/Source/Runtime/Core/Public/HTML5/HTML5PlatformCompilerPreSetup.h
@@ -7,6 +7,7 @@
 // Disable common CA warnings around SDK includes
 #ifndef THIRD_PARTY_INCLUDES_START
 	#define THIRD_PARTY_INCLUDES_START \
+		PRAGMA_DISABLE_REORDER_WARNINGS \
 		PRAGMA_DISABLE_SHADOW_VARIABLE_WARNINGS \
 		PRAGMA_DISABLE_UNDEFINED_IDENTIFIER_WARNINGS \
 		PRAGMA_DISABLE_MISSING_VIRTUAL_DESTRUCTOR_WARNINGS \
@@ -18,5 +19,6 @@
 		PRAGMA_ENABLE_DEPRECATION_WARNINGS \
 		PRAGMA_ENABLE_MISSING_VIRTUAL_DESTRUCTOR_WARNINGS \
 		PRAGMA_ENABLE_UNDEFINED_IDENTIFIER_WARNINGS \
-		PRAGMA_ENABLE_SHADOW_VARIABLE_WARNINGS
+		PRAGMA_ENABLE_SHADOW_VARIABLE_WARNINGS \
+		PRAGMA_ENABLE_REORDER_WARNINGS
 #endif
diff --git a/Engine/Source/Runtime/Core/Public/Unix/UnixPlatformCompilerPreSetup.h b/Engine/Source/Runtime/Core/Public/Unix/UnixPlatformCompilerPreSetup.h
--- a/Engine/Source/Runtime/Core/Public/Unix/UnixPlatformCompilerPreSetup.h
+++ b/Engine/Source/Runtime/Core/Public/Unix/UnixPlatformCompilerPreSetup.h
@@ -34,6 +34,7 @@
 // Disable common CA warnings around SDK includes
 #ifndef THIRD_PARTY_INCLUDES_START
 	#define THIRD_PARTY_INCLUDES_START \
+		PRAGMA_DISABLE_REORDER_WARNINGS \
 		PRAGMA_DISABLE_SHADOW_VARIABLE_WARNINGS \
 		PRAGMA_DISABLE_NULL_DEREFERENCE_WARNINGS \
 		PRAGMA_DISABLE_UNDEFINED_IDENTIFIER_WARNINGS \
@@ -45,16 +46,15 @@
 #ifndef THIRD_PARTY_INCLUDES_END
 	#define THIRD_PARTY_INCLUDES_END \
 		PRAGMA_ENABLE_DEPRECATION_WARNINGS \
+		PRAGMA_ENABLE_DEPRECATED_REGISTER_WARNINGS \
 		PRAGMA_ENABLE_MISSING_VIRTUAL_DESTRUCTOR_WARNINGS \
 		PRAGMA_ENABLE_UNDEFINED_IDENTIFIER_WARNINGS \
 		PRAGMA_ENABLE_NULL_DEREFERENCE_WARNINGS \
 		PRAGMA_ENABLE_SHADOW_VARIABLE_WARNINGS \
-		PRAGMA_ENABLE_DEPRECATED_REGISTER_WARNINGS
+		PRAGMA_ENABLE_REORDER_WARNINGS
 #endif
 
 // Make certain warnings always be warnings, even despite -Werror.
 // Rationale: we don't want to suppress those as there are plans to address them (e.g. UE-12341), but breaking builds due to these warnings is very expensive
 // since they cannot be caught by all compilers that we support. They are deemed to be relatively safe to be ignored, at least until all SDKs/toolchains start supporting them.
-//#pragma clang diagnostic warning "-Winconsistent-missing-override"
-#pragma clang diagnostic warning "-Wreorder"
 #pragma clang diagnostic warning "-Wparentheses-equality"
diff --git a/Engine/Source/Runtime/Core/Public/Windows/WindowsPlatformCompilerPreSetup.h b/Engine/Source/Runtime/Core/Public/Windows/WindowsPlatformCompilerPreSetup.h
--- a/Engine/Source/Runtime/Core/Public/Windows/WindowsPlatformCompilerPreSetup.h
+++ b/Engine/Source/Runtime/Core/Public/Windows/WindowsPlatformCompilerPreSetup.h
@@ -4,52 +4,44 @@
 
 #if defined(__clang__)
 	#include "Clang/ClangPlatformCompilerPreSetup.h"
-#endif
 
-#ifndef DISABLE_DEPRECATION
-	#if !defined(__clang__)
-		#define DEPRECATED(VERSION, MESSAGE) __declspec(deprecated(MESSAGE " Please update your code to the new API before upgrading to the next release, otherwise your project will no longer compile."))
-
-		#define PRAGMA_DISABLE_DEPRECATION_WARNINGS \
-			__pragma (warning(push)) \
-			__pragma (warning(disable:4995)) \
-			__pragma (warning(disable:4996))
-
-		#define PRAGMA_ENABLE_DEPRECATION_WARNINGS \
-			__pragma (warning(pop))
-	#endif
-#endif // DISABLE_DEPRECATION
-
-#if defined(__clang__)
 	// Disable common CA warnings around SDK includes
 	#ifndef THIRD_PARTY_INCLUDES_START
 		#define THIRD_PARTY_INCLUDES_START \
+			PRAGMA_DISABLE_REORDER_WARNINGS \
 			PRAGMA_DISABLE_SHADOW_VARIABLE_WARNINGS \
 			PRAGMA_DISABLE_UNDEFINED_IDENTIFIER_WARNINGS \
 			PRAGMA_DISABLE_DEPRECATION_WARNINGS
-	#endif
+	#endif // THIRD_PARTY_INCLUDES_START
 
 	#ifndef THIRD_PARTY_INCLUDES_END
 		#define THIRD_PARTY_INCLUDES_END \
 			PRAGMA_ENABLE_DEPRECATION_WARNINGS \
 			PRAGMA_ENABLE_UNDEFINED_IDENTIFIER_WARNINGS \
-			PRAGMA_ENABLE_SHADOW_VARIABLE_WARNINGS
-	#endif
+			PRAGMA_ENABLE_SHADOW_VARIABLE_WARNINGS \
+			PRAGMA_ENABLE_REORDER_WARNINGS
+	#endif // THIRD_PARTY_INCLUDES_END
 #else
-	// VC++
-	// 4456 - declaration of 'LocalVariable' hides previous local declaration
-	// 4457 - declaration of 'LocalVariable' hides function parameter
-	// 4458 - declaration of 'LocalVariable' hides class member
-	// 4459 - declaration of 'LocalVariable' hides global declaration
-	// 6244 - local declaration of <variable> hides previous declaration at <line> of <file>
+	#ifndef DISABLE_DEPRECATION
+		#define DEPRECATED(VERSION, MESSAGE) __declspec(deprecated(MESSAGE " Please update your code to the new API before upgrading to the next release, otherwise your project will no longer compile."))
+
+		#define PRAGMA_DISABLE_DEPRECATION_WARNINGS \
+			__pragma (warning(push)) \
+			__pragma (warning(disable: 4995)) /* 'function': name was marked as #pragma deprecated */ \
+			__pragma (warning(disable: 4996)) /* The compiler encountered a deprecated declaration. */
+
+		#define PRAGMA_ENABLE_DEPRECATION_WARNINGS \
+			__pragma (warning(pop))
+	#endif // DISABLE_DEPRECATION
+
 	#ifndef PRAGMA_DISABLE_SHADOW_VARIABLE_WARNINGS
 		#define PRAGMA_DISABLE_SHADOW_VARIABLE_WARNINGS \
 			__pragma (warning(push)) \
-			__pragma (warning(disable:4456)) \
-			__pragma (warning(disable:4457)) \
-			__pragma (warning(disable:4458)) \
-			__pragma (warning(disable:4459)) \
-			__pragma (warning(disable:6244))
+			__pragma (warning(disable: 4456)) /* declaration of 'LocalVariable' hides previous local declaration */ \
+			__pragma (warning(disable: 4457)) /* declaration of 'LocalVariable' hides function parameter */ \
+			__pragma (warning(disable: 4458)) /* declaration of 'LocalVariable' hides class member */ \
+			__pragma (warning(disable: 4459)) /* declaration of 'LocalVariable' hides global declaration */ \
+			__pragma (warning(disable: 6244)) /* local declaration of <variable> hides previous declaration at <line> of <file> */
 	#endif // PRAGMA_DISABLE_SHADOW_VARIABLE_WARNINGS
 
 	#ifndef PRAGMA_ENABLE_SHADOW_VARIABLE_WARNINGS
@@ -60,7 +52,7 @@
 	#ifndef PRAGMA_DISABLE_UNDEFINED_IDENTIFIER_WARNINGS
 		#define PRAGMA_DISABLE_UNDEFINED_IDENTIFIER_WARNINGS \
 			__pragma(warning(push)) \
-			__pragma(warning(disable: 4668))  /* 'symbol' is not defined as a preprocessor macro, replacing with '0' for 'directives' */
+			__pragma(warning(disable: 4668)) /* 'symbol' is not defined as a preprocessor macro, replacing with '0' for 'directives' */
 	#endif // PRAGMA_DISABLE_UNDEFINED_IDENTIFIER_WARNINGS
 
 	#ifndef PRAGMA_ENABLE_UNDEFINED_IDENTIFIER_WARNINGS
@@ -71,7 +63,7 @@
 	#ifndef PRAGMA_DISABLE_MISSING_VIRTUAL_DESTRUCTOR_WARNINGS
 		#define PRAGMA_DISABLE_MISSING_VIRTUAL_DESTRUCTOR_WARNINGS \
 			__pragma(warning(push)) \
-			__pragma(warning(disable:4265))
+			__pragma(warning(disable: 4265)) /* class' : class has virtual functions, but destructor is not virtual */
 	#endif // PRAGMA_DISABLE_MISSING_VIRTUAL_DESTRUCTOR_WARNINGS
 
 	#ifndef PRAGMA_ENABLE_MISSING_VIRTUAL_DESTRUCTOR_WARNINGS
@@ -79,6 +71,17 @@
 			__pragma(warning(pop))
 	#endif // PRAGMA_ENABLE_MISSING_VIRTUAL_DESTRUCTOR_WARNINGS
 
+	#ifndef PRAGMA_DISABLE_REORDER_WARNINGS
+		#define PRAGMA_DISABLE_REORDER_WARNINGS \
+			__pragma(warning(push)) \
+			__pragma(warning(disable: 5038)) /* data member 'member1' will be initialized after data member 'member2' data member 'member' will be initialized after base class 'base_class' */
+	#endif // PRAGMA_DISABLE_REORDER_WARNINGS
+
+	#ifndef PRAGMA_ENABLE_REORDER_WARNINGS
+		#define PRAGMA_ENABLE_REORDER_WARNINGS \
+			__pragma(warning(pop))
+	#endif // PRAGMA_ENABLE_REORDER_WARNINGS
+
 	#ifndef PRAGMA_POP
 		#define PRAGMA_POP \
 			__pragma(warning(pop))
@@ -105,11 +108,12 @@
 			__pragma(warning(disable: 28252)) /* Inconsistent annotation for '<func>': return/function has '<annotation>' on the prior instance. */ \
 			__pragma(warning(disable: 28253)) /* Inconsistent annotation for '<func>': _Param_(<num>) has '<annotation>' on the prior instance. */ \
 			__pragma(warning(disable: 28301)) /* No annotations for first declaration of '<func>'. */ \
+			PRAGMA_DISABLE_REORDER_WARNINGS \
 			PRAGMA_DISABLE_UNDEFINED_IDENTIFIER_WARNINGS \
 			PRAGMA_DISABLE_SHADOW_VARIABLE_WARNINGS \
 			PRAGMA_DISABLE_MISSING_VIRTUAL_DESTRUCTOR_WARNINGS \
 			PRAGMA_DISABLE_DEPRECATION_WARNINGS
-	#endif
+	#endif // THIRD_PARTY_INCLUDES_START
 
 	#ifndef THIRD_PARTY_INCLUDES_END
 		#define THIRD_PARTY_INCLUDES_END \
@@ -117,19 +121,10 @@
 			PRAGMA_ENABLE_MISSING_VIRTUAL_DESTRUCTOR_WARNINGS \
 			PRAGMA_ENABLE_SHADOW_VARIABLE_WARNINGS \
 			PRAGMA_ENABLE_UNDEFINED_IDENTIFIER_WARNINGS \
+			PRAGMA_ENABLE_REORDER_WARNINGS \
 			__pragma(warning(pop))
-	#endif
-#endif
+	#endif // THIRD_PARTY_INCLUDES_END
 
-#if !defined(__clang__)
 	#define EMIT_CUSTOM_WARNING_AT_LINE(Line, Warning) \
 		__pragma(message(WARNING_LOCATION(Line) ": warning C4996: " Warning))
-#endif
-
-#if defined(__clang__)
-	// Make certain warnings always be warnings, even despite -Werror.
-	// Rationale: we don't want to suppress those as there are plans to address them (e.g. UE-12341), but breaking builds due to these warnings is very expensive
-	// since they cannot be caught by all compilers that we support. They are deemed to be relatively safe to be ignored, at least until all SDKs/toolchains start supporting them.
-	#pragma clang diagnostic warning "-Wreorder"
-#endif
-
+#endif // defined(__clang__)
diff --git a/Engine/Source/Runtime/Core/Public/Windows/WindowsPlatformCompilerSetup.h b/Engine/Source/Runtime/Core/Public/Windows/WindowsPlatformCompilerSetup.h
--- a/Engine/Source/Runtime/Core/Public/Windows/WindowsPlatformCompilerSetup.h
+++ b/Engine/Source/Runtime/Core/Public/Windows/WindowsPlatformCompilerSetup.h
@@ -30,6 +30,7 @@ static_assert(_MSC_VER < 1910 || _MSC_VER >= 1913, "Visual Studio 2017 version 1
 #pragma warning (error: 4700      4702 4703 4704 4705 4706 4707 4708 4709 4710 4711 4712 4713 4714 4715 4716 4717 4718 4719 4720 4721 4722 4723 4724 4725 4726 4727 4728 4729      4731 4732 4733 4734 4735 4736 4737      4739 4740 4741 4742 4743 4744 4745 4746 4747 4748 4749 4750 4751 4752 4753 4754 4755 4756 4757 4758 4759 4760 4761 4762 4763 4764 4765 4766           4769      4771 4772 4773      4775 4776 4777 4778 4779 4780 4781 4782 4783 4784 4785 4786 4787 4788 4789 4790 4791 4792 4793 4794 4795 4796 4797 4798 4799)
 #pragma warning (error: 4800 4801 4802 4803 4804 4805 4806 4807 4808 4809 4810 4811 4812 4813 4814 4815 4816 4817 4818           4821 4822 4823 4824 4825      4827      4829 4830 4831 4832 4833 4834 4835 4836           4839 4840 4841 4842 4843 4844 4845 4846 4847 4848 4849 4850 4851 4852 4853 4854 4855 4856 4857 4858 4859 4860 4861 4862 4863 4864 4865      4867      4869 4870 4871 4872 4873 4874 4875 4876 4877 4878 4879 4880 4881 4882 4883 4884 4885 4886 4887 4888 4889 4890 4891 4892 4893 4894 4895 4896 4897 4898 4899)
 #pragma warning (error: 4900 4901 4902 4903 4904 4905 4906 4907 4908 4909 4910 4911 4912 4913 4914 4915 4916 4917 4918 4919 4920 4921 4922 4923 4924 4925 4926 4927 4928 4929 4930 4931 4932 4933 4934 4935 4936 4937 4938 4939 4940 4941 4942 4943 4944 4945 4946 4947 4948 4949 4950 4951 4952 4953 4954 4955 4956 4957 4958 4959 4960 4961      4963 4964 4965 4966 4967 4968 4969 4970 4971 4972 4973 4974 4975 4976 4977 4978 4979 4980 4981 4982 4983 4984 4985                4989 4990 4991 4992 4993 4994 4995      4997 4998 4999)
+#pragma warning (error:                                                                                                                                                                                               5038                                                                                                                                                                                                                                                                                                                 )
 
 //
 // Skipped warnings, which are not disabled below. Most are disabled by default. Might be useful to look through, re-enable some and fix the code.
diff --git a/Engine/Source/Runtime/D3D12RHI/Private/D3D12Residency.h b/Engine/Source/Runtime/D3D12RHI/Private/D3D12Residency.h
--- a/Engine/Source/Runtime/D3D12RHI/Private/D3D12Residency.h
+++ b/Engine/Source/Runtime/D3D12RHI/Private/D3D12Residency.h
@@ -19,11 +19,15 @@ namespace D3DX12Residency
 #else
 #include "D3D12Util.h"
 #include "Windows/AllowWindowsPlatformTypes.h"
+
+THIRD_PARTY_INCLUDES_START
 #include "dxgi1_6.h"
 #pragma warning(push)
-#pragma warning(disable : 6031)
+#pragma warning(disable: 6031)
 	#include <D3DX12Residency.h>
 #pragma warning(pop)
+THIRD_PARTY_INCLUDES_END
+
 #include "Windows/HideWindowsPlatformTypes.h"
 #endif
 
