UE4 object version changed. Required to fix serialized case-insensitve hashes of wide character strings containing accent characters (128-255).

diff --git a/Engine/Source/Runtime/Core/Private/UObject/UnrealNames.cpp b/Engine/Source/Runtime/Core/Private/UObject/UnrealNames.cpp
--- a/Engine/Source/Runtime/Core/Private/UObject/UnrealNames.cpp
+++ b/Engine/Source/Runtime/Core/Private/UObject/UnrealNames.cpp
@@ -1286,6 +1286,12 @@ FArchive& operator<<(FArchive& Ar, FNameEntrySerialized& E)
 			Ar << E.NonCasePreservingHash;
 			Ar << E.CasePreservingHash;
 			E.bWereHashesLoaded = true;
+
+			if (Ar.UE4Ver() < VER_UE4_FIX_WIDE_STRING_CRC && E.IsWide())
+			{
+				// Regenerate hashes for packages saved before a fix to how hashes are generated
+				E.NonCasePreservingHash = GetRawNonCasePreservingHash(E.GetWideName());
+			}
 		}
 	}
 	else
diff --git a/Engine/Source/Runtime/Core/Public/Misc/Char.h b/Engine/Source/Runtime/Core/Public/Misc/Char.h
--- a/Engine/Source/Runtime/Core/Public/Misc/Char.h
+++ b/Engine/Source/Runtime/Core/Public/Misc/Char.h
@@ -135,7 +135,18 @@ typedef TChar<ANSICHAR> FCharAnsi;
 	TCHAR specialized functions
 -----------------------------------------------------------------------------*/
 
-template <> inline TChar<WIDECHAR>::CharType TChar<WIDECHAR>::ToUpper(CharType Char)	{ return ::towupper(Char); }
+template <> inline TChar<WIDECHAR>::CharType TChar<WIDECHAR>::ToUpper(CharType Char)
+{
+#if PLATFORM_SWITCH && !PLATFORM_TCHAR_IS_4_BYTES
+	// for 2 byte wide chars use ansi toupper() if it is not an actual wide char (value less than 128)
+	// towupper() is slow on some platforms.
+	if (Char < 128)
+	{
+		return ::toupper((ANSICHAR)Char);
+	}
+#endif
+	return ::towupper(Char);
+}
 template <> inline TChar<WIDECHAR>::CharType TChar<WIDECHAR>::ToLower(CharType Char)	{ return ::towlower(Char); }
 template <> inline bool TChar<WIDECHAR>::IsUpper(CharType Char)							{ return ::iswupper(Char) != 0; }
 template <> inline bool TChar<WIDECHAR>::IsLower(CharType Char)							{ return ::iswlower(Char) != 0; }
diff --git a/Engine/Source/Runtime/Core/Public/Misc/Crc.h b/Engine/Source/Runtime/Core/Public/Misc/Crc.h
--- a/Engine/Source/Runtime/Core/Public/Misc/Crc.h
+++ b/Engine/Source/Runtime/Core/Public/Misc/Crc.h
@@ -124,13 +124,7 @@ inline uint32 FCrc::Strihash_DEPRECATED(const WIDECHAR* Data)
 	uint32 Hash=0;
 	while( *Data )
 	{
-#if !PLATFORM_TCHAR_IS_4_BYTES
-		// for 2 byte wide chars use ansi toupper() if it's not an actual wide char (upper byte == 0)
-		// towupper() is painfully slow on Switch ATM.
-		WIDECHAR Ch = (*Data & 0xFF00) ? TChar<WIDECHAR>::ToUpper(*Data++) : TChar<ANSICHAR>::ToUpper(*Data++);
-#else
 		WIDECHAR Ch = TChar<WIDECHAR>::ToUpper(*Data++);
-#endif
 		uint16  B  = Ch;
 		Hash     = ((Hash >> 8) & 0x00FFFFFF) ^ CRCTable_DEPRECATED[(Hash ^ B) & 0x000000FF];
 		B        = Ch>>8;
diff --git a/Engine/Source/Runtime/Core/Public/UObject/ObjectVersion.h b/Engine/Source/Runtime/Core/Public/UObject/ObjectVersion.h
--- a/Engine/Source/Runtime/Core/Public/UObject/ObjectVersion.h
+++ b/Engine/Source/Runtime/Core/Public/UObject/ObjectVersion.h
@@ -629,6 +629,8 @@ enum EUnrealEngineObjectUE4Version
 	VER_UE4_POINTLIGHT_SOURCE_ORIENTATION,
 	// LocalizationId has been added to the package summary (editor-only)
 	VER_UE4_ADDED_PACKAGE_SUMMARY_LOCALIZATION_ID,
+	// Fixed case insensitive hashes of wide strings containing character values from 128-255
+	VER_UE4_FIX_WIDE_STRING_CRC,
 
 	// -----<new versions can be added before this line>-------------------------------------------------
 	// - this needs to be the last line (see note below)
