Add a custom implementation of the callstack -> portable callstack translation function for Windows, to avoid calling symbolication functions that we don't need to use.

diff --git a/Engine/Source/Runtime/Core/Private/Windows/WindowsPlatformCrashContext.cpp b/Engine/Source/Runtime/Core/Private/Windows/WindowsPlatformCrashContext.cpp
--- a/Engine/Source/Runtime/Core/Private/Windows/WindowsPlatformCrashContext.cpp
+++ b/Engine/Source/Runtime/Core/Private/Windows/WindowsPlatformCrashContext.cpp
@@ -28,6 +28,7 @@
 #include <strsafe.h>
 #include <dbghelp.h>
 #include <Shlwapi.h>
+#include <psapi.h>
 
 #ifndef UE_LOG_CRASH_CALLSTACK
 	#define UE_LOG_CRASH_CALLSTACK 1
@@ -36,6 +37,68 @@
 #pragma comment( lib, "version.lib" )
 #pragma comment( lib, "Shlwapi.lib" )
 
+void FWindowsPlatformCrashContext::SetPortableCallStack(const TArray<FProgramCounterSymbolInfo>& Stack)
+{
+	// Get all the module handles for the current process. Each module handle is its base address.
+	TArray<HMODULE, TInlineAllocator<128>> ProcessModuleHandles;
+	for (;;)
+	{
+		DWORD BufferSize = ProcessModuleHandles.Num() * sizeof(HMODULE);
+		DWORD RequiredBufferSize = 0;
+		if (!EnumProcessModules(GetCurrentProcess(), ProcessModuleHandles.GetData(), BufferSize, &RequiredBufferSize))
+		{
+			return;
+		}
+		if(RequiredBufferSize <= BufferSize)
+		{
+			break;
+		}
+		ProcessModuleHandles.SetNum(RequiredBufferSize / sizeof(HMODULE));
+	}
+
+	// Sort the handles by address. This allows us to do a binary search for the module containing an address.
+	Algo::Sort(ProcessModuleHandles);
+
+	// Prepare the callstack buffer
+	CallStack.Reset(Stack.Num());
+
+	// Create the crash context
+	for (int Idx = 0; Idx < Stack.Num(); Idx++)
+	{
+		int ModuleIdx = Algo::UpperBound(ProcessModuleHandles, (HMODULE)Stack[Idx].ProgramCounter) - 1;
+		if (ModuleIdx < 0 || ModuleIdx >= ProcessModuleHandles.Num())
+		{
+			CallStack.Add(FCrashStackFrame(TEXT("Unknown"), 0, Stack[Idx].ProgramCounter));
+		}
+		else
+		{
+			TCHAR ModuleName[MAX_PATH];
+			if (GetModuleFileNameW(ProcessModuleHandles[ModuleIdx], ModuleName, ARRAY_COUNT(ModuleName)) != 0)
+			{
+				TCHAR* ModuleNameEnd = FCString::Strrchr(ModuleName, '\\');
+				if(ModuleNameEnd != nullptr)
+				{
+					FMemory::Memmove(ModuleName, ModuleNameEnd + 1, (FCString::Strlen(ModuleNameEnd + 1) + 1) * sizeof(TCHAR));
+				}
+
+				TCHAR* ModuleNameExt = FCString::Strrchr(ModuleName, '.');
+				if(ModuleNameExt != nullptr)
+				{
+					*ModuleNameExt = 0;
+				}
+			}
+			else
+			{
+				FCString::Strcpy(ModuleName, TEXT("Unknown"));
+			}
+
+			uint64 BaseAddress = (uint64)ProcessModuleHandles[ModuleIdx];
+			uint64 Offset = Stack[Idx].ProgramCounter - BaseAddress;
+			CallStack.Add(FCrashStackFrame(ModuleName, BaseAddress, Offset));
+		}
+	}
+}
+
 void FWindowsPlatformCrashContext::AddPlatformSpecificProperties() const
 {
 	AddCrashProperty(TEXT("PlatformIsRunningWindows"), 1);
diff --git a/Engine/Source/Runtime/Core/Public/GenericPlatform/GenericPlatformCrashContext.h b/Engine/Source/Runtime/Core/Public/GenericPlatform/GenericPlatformCrashContext.h
--- a/Engine/Source/Runtime/Core/Public/GenericPlatform/GenericPlatformCrashContext.h
+++ b/Engine/Source/Runtime/Core/Public/GenericPlatform/GenericPlatformCrashContext.h
@@ -74,11 +74,11 @@ struct FCrashStackFrame
 	uint64 BaseAddress;
 	uint64 Offset;
 
-	FCrashStackFrame(const FString& ModuleNameIn, uint64 BaseAddressIn, uint64 OffsetIn)
+	FCrashStackFrame(FString ModuleNameIn, uint64 BaseAddressIn, uint64 OffsetIn)
+		: ModuleName(MoveTemp(ModuleNameIn))
+		, BaseAddress(BaseAddressIn)
+		, Offset(OffsetIn)
 	{
-		ModuleName = ModuleNameIn;
-		BaseAddress = BaseAddressIn;
-		Offset = OffsetIn;
 	}
 };
 
@@ -212,7 +212,7 @@ public:
 	void CapturePortableCallStack(int32 NumStackFramesToIgnore, void* Context);
 	
 	/** Sets the portable callstack to a specified stack */
-	void SetPortableCallStack(const TArray<FProgramCounterSymbolInfo>& Stack);
+	virtual void SetPortableCallStack(const TArray<FProgramCounterSymbolInfo>& Stack);
 
 	/**
 	 * @return whether this crash is a non-crash event
diff --git a/Engine/Source/Runtime/Core/Public/Windows/WindowsPlatformCrashContext.h b/Engine/Source/Runtime/Core/Public/Windows/WindowsPlatformCrashContext.h
--- a/Engine/Source/Runtime/Core/Public/Windows/WindowsPlatformCrashContext.h
+++ b/Engine/Source/Runtime/Core/Public/Windows/WindowsPlatformCrashContext.h
@@ -14,6 +14,8 @@ struct CORE_API FWindowsPlatformCrashContext : public FGenericCrashContext
 		bIsEnsure = bInIsEnsure;
 	}
 
+	virtual void SetPortableCallStack(const TArray<FProgramCounterSymbolInfo>& Stack);
+
 	virtual void AddPlatformSpecificProperties() const override;
 };
 
