Pass raw callstacks into FGenericCrashContext::SetPortableCallstack(). We don't actually need symbol information for each frame, and it's expensive to compute.

diff --git a/Engine/Source/Runtime/Core/Private/GenericPlatform/GenericPlatformCrashContext.cpp b/Engine/Source/Runtime/Core/Private/GenericPlatform/GenericPlatformCrashContext.cpp
--- a/Engine/Source/Runtime/Core/Private/GenericPlatform/GenericPlatformCrashContext.cpp
+++ b/Engine/Source/Runtime/Core/Private/GenericPlatform/GenericPlatformCrashContext.cpp
@@ -652,53 +652,57 @@ FORCENOINLINE void FGenericCrashContext::CapturePortableCallStack(int32 NumStack
 		NumStackFramesToIgnore++;
 	}
 
-	const int32 MaxDepth = 100;
-	TArray<FProgramCounterSymbolInfo> Stack = FPlatformStackWalk::GetStack(NumStackFramesToIgnore, MaxDepth, Context);
-	return SetPortableCallStack(Stack);
-}
+	// Capture the stack trace
+	static const int StackTraceMaxDepth = 100;
+	uint64 StackTrace[StackTraceMaxDepth];
+	FMemory::Memzero(StackTrace);
+	int32 StackTraceDepth = FPlatformStackWalk::CaptureStackBackTrace(StackTrace, StackTraceMaxDepth, Context);
 
-void FGenericCrashContext::SetPortableCallStack(const TArray<FProgramCounterSymbolInfo>& Stack)
-{
-	uint32 ModuleEntries = (uint32)FPlatformStackWalk::GetProcessModuleCount();
+	// Make sure we don't exceed the current stack depth
+	NumStackFramesToIgnore = FMath::Min(NumStackFramesToIgnore, StackTraceDepth);
 
-	if (ModuleEntries)
-	{
-		TArray<FStackWalkModuleInfo> ProcessModules;
-		ProcessModules.AddUninitialized(ModuleEntries);
-		FPlatformStackWalk::GetProcessModuleSignatures(ProcessModules.GetData(), ProcessModules.Max());
-
-		TMap<FString, uint64> ImageBases;
-		int32 ModuleIndex = 0;
-		for (TArray<FProgramCounterSymbolInfo>::TConstIterator Itr(Stack); Itr; ++Itr)
-		{
-			FString ModuleName = FPaths::GetBaseFilename(Itr->ModuleName);
+	// Generate the portable callstack from it
+	SetPortableCallStack(StackTrace + NumStackFramesToIgnore, StackTraceDepth - NumStackFramesToIgnore);
+}
 
-			if (!ImageBases.Contains(ModuleName))
-			{
-				for (TArray<FStackWalkModuleInfo>::TConstIterator ProcessModuleItr(ProcessModules); ProcessModuleItr; ++ProcessModuleItr)
-				{
-					FString ProcessModuleName = FPaths::GetBaseFilename(ProcessModuleItr->ImageName);
+void FGenericCrashContext::SetPortableCallStack(const uint64* StackFrames, int32 NumStackFrames)
+{
+	// Get all the modules in the current process
+	uint32 NumModules = (uint32)FPlatformStackWalk::GetProcessModuleCount();
 
-					if (!ModuleName.Compare(ProcessModuleName, ESearchCase::IgnoreCase))
-					{
-						ImageBases.Add(ModuleName, ProcessModuleItr->BaseOfImage);
-						break;
-					}
-				}
+	TArray<FStackWalkModuleInfo> Modules;
+	Modules.AddUninitialized(NumModules);
 
-			}
+	NumModules = FPlatformStackWalk::GetProcessModuleSignatures(Modules.GetData(), NumModules);
+	Modules.SetNum(NumModules);
 
-			uint64 BaseOfImage = MAX_uint64;
+	// Update the callstack with offsets from each module
+	CallStack.Reset(NumStackFrames);
+	for(int32 Idx = 0; Idx < NumStackFrames; Idx++)
+	{
+		const uint64 StackFrame = StackFrames[Idx];
 
-			if (ImageBases.Contains(ModuleName))
+		// Try to find the module containing this stack frame
+		const FStackWalkModuleInfo* FoundModule = nullptr;
+		for(const FStackWalkModuleInfo& Module : Modules)
+		{
+			if(StackFrame >= Module.BaseOfImage && StackFrame < Module.BaseOfImage + Module.ImageSize)
 			{
-				BaseOfImage = ImageBases[ModuleName];
+				FoundModule = &Module;
+				break;
 			}
+		}
 
-			CallStack.Add(FCrashStackFrame(ModuleName, BaseOfImage, Itr->ProgramCounter > BaseOfImage ? Itr->ProgramCounter - BaseOfImage : MAX_uint64));
+		// Add the callstack item
+		if(FoundModule == nullptr)
+		{
+			CallStack.Add(FCrashStackFrame(TEXT("Unknown"), 0, StackFrame));
+		}
+		else
+		{
+			CallStack.Add(FCrashStackFrame(FPaths::GetBaseFilename(FoundModule->ImageName), FoundModule->BaseOfImage, StackFrame - FoundModule->BaseOfImage));
 		}
 	}
-
 }
 
 FProgramCounterSymbolInfoEx::FProgramCounterSymbolInfoEx( FString InModuleName, FString InFunctionName, FString InFilename, uint32 InLineNumber, uint64 InSymbolDisplacement, uint64 InOffsetInModule, uint64 InProgramCounter ) :
diff --git a/Engine/Source/Runtime/Core/Private/HAL/ThreadHeartBeat.cpp b/Engine/Source/Runtime/Core/Private/HAL/ThreadHeartBeat.cpp
--- a/Engine/Source/Runtime/Core/Private/HAL/ThreadHeartBeat.cpp
+++ b/Engine/Source/Runtime/Core/Private/HAL/ThreadHeartBeat.cpp
@@ -237,12 +237,7 @@ void FORCENOINLINE FThreadHeartBeat::OnHang(double HangDuration, uint32 ThreadTh
 		FCString::Strcpy(GErrorMessage, ARRAY_COUNT(GErrorMessage), *ErrorMessage);
 
 		// Skip macros and FDebug, we always want this to fire
-		TArray<FProgramCounterSymbolInfo> Stack;
-		for (int32 Idx = 0; Idx < NumStackFrames; Idx++)
-		{
-			FPlatformStackWalk::ProgramCounterToSymbolInfo(StackFrames[Idx], Stack.AddDefaulted_GetRef());
-		}
-		ReportHang(*ErrorMessage, Stack);
+		ReportHang(*ErrorMessage, StackFrames, NumStackFrames);
 
 		GErrorMessage[0] = '\0';
 #endif // PLATFORM_DESKTOP
diff --git a/Engine/Source/Runtime/Core/Private/Mac/MacPlatformMisc.cpp b/Engine/Source/Runtime/Core/Private/Mac/MacPlatformMisc.cpp
--- a/Engine/Source/Runtime/Core/Private/Mac/MacPlatformMisc.cpp
+++ b/Engine/Source/Runtime/Core/Private/Mac/MacPlatformMisc.cpp
@@ -1877,7 +1877,7 @@ void NewReportEnsure( const TCHAR* ErrorMessage, int NumStackFramesToIgnore )
 	EnsureLock.Unlock();
 }
 
-void ReportHang(const TCHAR* ErrorMessage, const TArray<FProgramCounterSymbolInfo>& Stack)
+void ReportHang(const TCHAR* ErrorMessage, const uint64* StackFrames, int32 NumStackFrames)
 {
 	EnsureLock.Lock();
 	if (!bReentranceGuard && FMacApplicationInfo::CrashReporter != nil)
@@ -1886,7 +1886,7 @@ void ReportHang(const TCHAR* ErrorMessage, const TArray<FProgramCounterSymbolInf
 
 		const bool bIsEnsure = true;
 		FMacCrashContext EnsureContext(bIsEnsure);
-		EnsureContext.SetPortableCallStack(Stack);
+		EnsureContext.SetPortableCallStack(StackFrames, NumStackFrames);
 		EnsureContext.GenerateEnsureInfoAndLaunchReporter();
 
 		bReentranceGuard = false;
diff --git a/Engine/Source/Runtime/Core/Private/Unix/UnixPlatformStackWalk.cpp b/Engine/Source/Runtime/Core/Private/Unix/UnixPlatformStackWalk.cpp
--- a/Engine/Source/Runtime/Core/Private/Unix/UnixPlatformStackWalk.cpp
+++ b/Engine/Source/Runtime/Core/Private/Unix/UnixPlatformStackWalk.cpp
@@ -807,7 +807,7 @@ void NewReportEnsure(const TCHAR* ErrorMessage, int NumStackFramesToIgnore)
 	EnsureLock.Unlock();
 }
 
-void ReportHang(const TCHAR* ErrorMessage, const TArray<FProgramCounterSymbolInfo>& Stack)
+void ReportHang(const TCHAR* ErrorMessage, const uint64* StackFrames, int32 NumStackFrames)
 {
 	EnsureLock.Lock();
 	if (!bReentranceGuard)
@@ -816,7 +816,7 @@ void ReportHang(const TCHAR* ErrorMessage, const TArray<FProgramCounterSymbolInf
 
 		const bool bIsEnsure = true;
 		FUnixCrashContext EnsureContext(bIsEnsure);
-		EnsureContext.SetPortableCallStack(Stack);
+		EnsureContext.SetPortableCallStack(StackFrames, NumStackFrames);
 		EnsureContext.GenerateCrashInfoAndLaunchReporter(true);
 
 		bReentranceGuard = false;
diff --git a/Engine/Source/Runtime/Core/Private/Windows/WindowsPlatformCrashContext.cpp b/Engine/Source/Runtime/Core/Private/Windows/WindowsPlatformCrashContext.cpp
--- a/Engine/Source/Runtime/Core/Private/Windows/WindowsPlatformCrashContext.cpp
+++ b/Engine/Source/Runtime/Core/Private/Windows/WindowsPlatformCrashContext.cpp
@@ -37,7 +37,7 @@
 #pragma comment( lib, "version.lib" )
 #pragma comment( lib, "Shlwapi.lib" )
 
-void FWindowsPlatformCrashContext::SetPortableCallStack(const TArray<FProgramCounterSymbolInfo>& Stack)
+void FWindowsPlatformCrashContext::SetPortableCallStack(const uint64* StackTrace, int32 StackTraceDepth)
 {
 	// Get all the module handles for the current process. Each module handle is its base address.
 	TArray<HMODULE, TInlineAllocator<128>> ProcessModuleHandles;
@@ -60,15 +60,15 @@ void FWindowsPlatformCrashContext::SetPortableCallStack(const TArray<FProgramCou
 	Algo::Sort(ProcessModuleHandles);
 
 	// Prepare the callstack buffer
-	CallStack.Reset(Stack.Num());
+	CallStack.Reset(StackTraceDepth);
 
 	// Create the crash context
-	for (int Idx = 0; Idx < Stack.Num(); Idx++)
+	for (int Idx = 0; Idx < StackTraceDepth; Idx++)
 	{
-		int ModuleIdx = Algo::UpperBound(ProcessModuleHandles, (HMODULE)Stack[Idx].ProgramCounter) - 1;
+		int ModuleIdx = Algo::UpperBound(ProcessModuleHandles, (HMODULE)StackTrace[Idx]) - 1;
 		if (ModuleIdx < 0 || ModuleIdx >= ProcessModuleHandles.Num())
 		{
-			CallStack.Add(FCrashStackFrame(TEXT("Unknown"), 0, Stack[Idx].ProgramCounter));
+			CallStack.Add(FCrashStackFrame(TEXT("Unknown"), 0, StackTrace[Idx]));
 		}
 		else
 		{
@@ -93,7 +93,7 @@ void FWindowsPlatformCrashContext::SetPortableCallStack(const TArray<FProgramCou
 			}
 
 			uint64 BaseAddress = (uint64)ProcessModuleHandles[ModuleIdx];
-			uint64 Offset = Stack[Idx].ProgramCounter - BaseAddress;
+			uint64 Offset = StackTrace[Idx] - BaseAddress;
 			CallStack.Add(FCrashStackFrame(ModuleName, BaseAddress, Offset));
 		}
 	}
@@ -397,12 +397,12 @@ void NewReportEnsure_Inner( const TCHAR* ErrorMessage )
 #endif	// WINVER
 }
 
-void ReportHang(const TCHAR* ErrorMessage, const TArray<FProgramCounterSymbolInfo>& Stack)
+void ReportHang(const TCHAR* ErrorMessage, const uint64* StackFrames, int32 NumStackFrames)
 {
 	const bool bIsEnsure = true;
 
 	FWindowsPlatformCrashContext CrashContext(bIsEnsure);
-	CrashContext.SetPortableCallStack(Stack);
+	CrashContext.SetPortableCallStack(StackFrames, NumStackFrames);
 
 	EErrorReportUI ReportUI = IsInteractiveEnsureMode() ? EErrorReportUI::ShowDialog : EErrorReportUI::ReportInUnattendedMode;
 	ReportCrashUsingCrashReportClient(CrashContext, nullptr, ErrorMessage, ReportUI, bIsEnsure);
diff --git a/Engine/Source/Runtime/Core/Public/GenericPlatform/GenericPlatformCrashContext.h b/Engine/Source/Runtime/Core/Public/GenericPlatform/GenericPlatformCrashContext.h
--- a/Engine/Source/Runtime/Core/Public/GenericPlatform/GenericPlatformCrashContext.h
+++ b/Engine/Source/Runtime/Core/Public/GenericPlatform/GenericPlatformCrashContext.h
@@ -212,7 +212,7 @@ public:
 	void CapturePortableCallStack(int32 NumStackFramesToIgnore, void* Context);
 	
 	/** Sets the portable callstack to a specified stack */
-	virtual void SetPortableCallStack(const TArray<FProgramCounterSymbolInfo>& Stack);
+	virtual void SetPortableCallStack(const uint64* StackFrames, int32 NumStackFrames);
 
 	/**
 	 * @return whether this crash is a non-crash event
diff --git a/Engine/Source/Runtime/Core/Public/HAL/ExceptionHandling.h b/Engine/Source/Runtime/Core/Public/HAL/ExceptionHandling.h
--- a/Engine/Source/Runtime/Core/Public/HAL/ExceptionHandling.h
+++ b/Engine/Source/Runtime/Core/Public/HAL/ExceptionHandling.h
@@ -25,7 +25,7 @@ extern CORE_API TCHAR MiniDumpFilenameW[1024];
 // #CrashReport: 2014-10-09 These methods are specific to windows, remove from here.
 extern CORE_API int32 ReportCrash( Windows::LPEXCEPTION_POINTERS ExceptionInfo );
 extern CORE_API void NewReportEnsure( const TCHAR* ErrorMessage, int NumStackFramesToIgnore );
-extern CORE_API void ReportHang(const TCHAR*, const TArray<FProgramCounterSymbolInfo>& Stack);
+extern CORE_API void ReportHang(const TCHAR*, const uint64* StackFrames, int32 NumStackFrames);
 #elif PLATFORM_XBOXONE
 #include "XboxOne/XboxOneSystemIncludes.h"
 // #CrashReport: 2014-10-09 Should be move to another file
@@ -36,10 +36,10 @@ extern CORE_API void NewReportEnsure( const TCHAR* ErrorMessage, int NumStackFra
 #include <signal.h>
 extern CORE_API int32 ReportCrash( ucontext_t *Context, int32 Signal, struct __siginfo* Info );
 extern CORE_API void NewReportEnsure( const TCHAR* ErrorMessage, int NumStackFramesToIgnore );
-extern CORE_API void ReportHang(const TCHAR*, const TArray<FProgramCounterSymbolInfo>& Stack);
+extern CORE_API void ReportHang(const TCHAR*, const uint64* StackFrames, int32 NumStackFrames);
 #elif PLATFORM_UNIX
 extern CORE_API void NewReportEnsure( const TCHAR* ErrorMessage, int NumStackFramesToIgnore );
-extern CORE_API void ReportHang(const TCHAR*, const TArray<FProgramCounterSymbolInfo>& Stack);
+extern CORE_API void ReportHang(const TCHAR*, const uint64* StackFrames, int32 NumStackFrames);
 #endif
 
 enum class ECrashType
diff --git a/Engine/Source/Runtime/Core/Public/Windows/WindowsPlatformCrashContext.h b/Engine/Source/Runtime/Core/Public/Windows/WindowsPlatformCrashContext.h
--- a/Engine/Source/Runtime/Core/Public/Windows/WindowsPlatformCrashContext.h
+++ b/Engine/Source/Runtime/Core/Public/Windows/WindowsPlatformCrashContext.h
@@ -14,7 +14,7 @@ struct CORE_API FWindowsPlatformCrashContext : public FGenericCrashContext
 		bIsEnsure = bInIsEnsure;
 	}
 
-	virtual void SetPortableCallStack(const TArray<FProgramCounterSymbolInfo>& Stack);
+	virtual void SetPortableCallStack(const uint64* StackFrames, int32 NumStackFrames);
 
 	virtual void AddPlatformSpecificProperties() const override;
 };
