Merging //UE4/Dev-Main to Dev-Build (//UE4/Dev-Build)

#rnx

diff --git a/Engine/Source/Runtime/Core/Public/HAL/Platform.h b/Engine/Source/Runtime/Core/Public/HAL/Platform.h
--- a/Engine/Source/Runtime/Core/Public/HAL/Platform.h
+++ b/Engine/Source/Runtime/Core/Public/HAL/Platform.h
@@ -135,7 +135,7 @@
 
 //@port Identify the platform here and include the platform header to setup the platform types, etc
 #if PLATFORM_WINDOWS
-	#include "Windows/WindowsPlatform.h"
+	#include "Windows/WIndowsPlatform.h"	// this is the actual filename on disk, alas cannot be easily renamed in source control
 #elif PLATFORM_PS4
 	#include "PS4/PS4Platform.h"
 #elif PLATFORM_XBOXONE
