Make sure modules actually exist before assuming they're up to date.

#rnx

diff --git a/Engine/Source/Runtime/Core/Private/Modules/ModuleManager.cpp b/Engine/Source/Runtime/Core/Private/Modules/ModuleManager.cpp
--- a/Engine/Source/Runtime/Core/Private/Modules/ModuleManager.cpp
+++ b/Engine/Source/Runtime/Core/Private/Modules/ModuleManager.cpp
@@ -176,6 +176,14 @@ bool FModuleManager::IsModuleUpToDate(const FName InModuleName) const
 	TMap<FName, FString> ModulePathMap;
 	FindModulePaths(*InModuleName.ToString(), ModulePathMap);
 
+	for (const TPair<FName, FString>& Pair : ModulePathMap)
+	{
+		if (!FPaths::FileExists(*Pair.Value))
+		{
+			return false;
+		}
+	}
+
 	return ModulePathMap.Num() == 1;
 }
 #endif
