PR #5117: Cleaned up MSVC warning declarations added documentation links to all warnings (Contributed by [REDACTED_PERSON])

#rnx

diff --git a/Engine/Source/Runtime/Core/Public/Windows/WindowsPlatformCompilerSetup.h b/Engine/Source/Runtime/Core/Public/Windows/WindowsPlatformCompilerSetup.h
--- a/Engine/Source/Runtime/Core/Public/Windows/WindowsPlatformCompilerSetup.h
+++ b/Engine/Source/Runtime/Core/Public/Windows/WindowsPlatformCompilerSetup.h
@@ -6,21 +6,26 @@
 #error PLATFORM_WINDOWS not defined
 #endif
 
-/**
-* Future-proofing the min version check so we keep bumping it whenever we upgrade.
-*/
+//
+// Future-proofing the min version check so we keep bumping it whenever we upgrade.
+//
 #if defined(_MSC_VER) && _MSC_VER > 1929 
 	#pragma message("Detected compiler newer than Visual Studio 2019, please update min version checking in WindowsPlatformCompilerSetup.h")
 #endif
 
-/**
- * We require at least Visual Studio 2015 Update 3 to compile
- */
+//
+// We require at least Visual Studio 2015 Update 3 to compile
+//
 static_assert(_MSC_VER != 1900 || _MSC_FULL_VER >= 190024210, "Visual Studio 2015 Update 3 is required to compile on Windows (http://go.microsoft.com/fwlink/?LinkId=691129)");
 static_assert(_MSC_VER >= 1900, "Visual Studio 2015 or later is required to compile on Windows");
 static_assert(_MSC_VER < 1910 || _MSC_VER >= 1913, "Visual Studio 2017 version 15.6 is required to compile on Windows. Please install updates through the Visual Studio installer.");
 static_assert(_MSC_VER != 1914 && _MSC_VER != 1915, "Visual Studio 2017 versions 15.7 and 15.8 are known to have code generation bugs that affect UE4. Please update to version 15.9.");
 
+//
+// Manually enable all warnings as errors, except ones that are explicitly skipped.
+// Warnings that we explicitly disable later are still included in there.
+//
+
 #pragma warning (error:      4001 4002 4003 4004      4006 4007 4008 4009 4010 4011 4012 4013 4014 4015 4016 4017 4018 4019 4020 4021 4022 4023 4024 4025 4026 4027 4028 4029 4030 4031 4032 4033 4034 4035 4036 4037 4038 4039 4040 4041 4042 4043 4044 4045 4046 4047 4048 4049 4050 4051 4052 4053 4054 4055 4056 4057 4058 4059 4060           4063 4064 4065 4066 4067 4068 4069 4070 4071 4072 4073 4074 4075 4076 4077 4078 4079 4080 4081 4082 4083 4084 4085 4086 4087 4088 4089 4090 4091 4092 4093 4094 4095 4096 4097 4098 4099)
 #pragma warning (error: 4100 4101 4102 4103 4104 4105 4106 4107 4108 4109 4110 4111 4112 4113 4114 4115 4116 4117 4118 4119 4120 4121 4122 4123 4124 4125 4126 4127 4128 4129 4130 4131 4132 4133 4134 4135 4136 4137 4138 4139 4140 4141 4142 4143 4144 4145 4146 4147 4148 4149 4150 4151 4152 4153 4154 4155 4156 4157 4158 4159 4160 4161 4162 4163 4164      4166 4167 4168 4169 4170 4171 4172 4173 4174 4175 4176 4177 4178 4179 4180 4181 4182 4183 4184 4185 4186 4187 4188 4189 4190 4191 4192 4193 4194 4195 4196 4197 4198 4199)
 #pragma warning (error: 4200 4201 4202 4203 4204 4205 4206 4207 4208 4209 4210 4211 4212 4213 4214 4215 4216 4217 4218 4219 4220 4221 4222 4223 4224 4225 4226 4227 4228 4229 4230 4231 4232 4233 4234 4235 4236 4237 4238 4239 4240 4241      4243 4244 4245 4246 4247 4248 4249 4250 4251 4252 4253      4255 4256 4257 4258 4259 4260 4261 4262 4263 4264           4267 4268 4269 4270 4271 4272 4273 4274 4275 4276 4277 4278 4279 4280 4281 4282 4283 4284 4285 4286 4287 4288 4289 4290 4291 4292 4293 4294 4295      4297 4298 4299)
@@ -35,184 +40,214 @@ static_assert(_MSC_VER != 1914 && _MSC_VER != 1915, "Visual Studio 2017 versions
 
 //
 // Skipped warnings, which are not disabled below. Most are disabled by default. Might be useful to look through, re-enable some and fix the code.
+// NOTE: We don't use #pragma warning (default: ####) on these warnings because default enables warnings that default to disabled.
 //
 
-// 4005 - 'identifier' : macro redefinition																							http://msdn.microsoft.com/en-us/library/8d10sc3w.aspx
-// 4061 - enumerator 'identifier' in switch of enum 'enumeration' is not explicitly handled by a case label:						http://msdn.microsoft.com/en-us/library/96f5t7fy.aspx
-// 4062 - enumerator 'identifier' in switch of enum 'enumeration' is not handled													http://msdn.microsoft.com/en-us/library/fdt9w8tf.aspx
-// 4165 - 'HRESULT' is being converted to 'bool'; are you sure this is what you want?
-// 4242 - 'identifier' : conversion from 'type1' to 'type2', possible loss of data													http://msdn.microsoft.com/en-us/library/3hca13eh.aspx
-// 4254 - 'operator' : conversion from 'type1' to 'type2', possible loss of data													http://msdn.microsoft.com/en-us/library/3fbf7w04.aspx
-// 4265 - 'class' : class has virtual functions, but destructor is not virtual														http://msdn.microsoft.com/en-us/library/wzxffy8c.aspx
-// 4266 - 'function' : no override available for virtual member function from base 'type'; function is hidden						http://msdn.microsoft.com/en-us/library/4b76ty10.aspx
-// 4305 - 'identifier' : truncation from 'type1' to 'type2'																			http://msdn.microsoft.com/en-us/library/0as1ke3f.aspx
-// 4296 - 'operator' : expression is always false																					http://msdn.microsoft.com/en-us/library/wz2y40yt.aspx
-// 4307 - '' : integral constant overflow																							http://msdn.microsoft.com/en-us/library/4kze989h.aspx 
-// 4311 - 'variable' : pointer truncation from 'type' to 'type'																		http://msdn.microsoft.com/en-us/library/4t91x2k5.aspx
-// 4312 - 'operation' : conversion from 'type1' to 'type2' of greater size															http://msdn.microsoft.com/en-us/library/h97f4b9y.aspx
-// 4339 - 'type' : use of undefined type detected in CLR meta-data - use of this type may lead to a runtime exception				http://msdn.microsoft.com/en-us/library/3fxw8y6x.aspx
-// 4342 - behavior change: 'function' called, but a member operator was called in previous versions									http://msdn.microsoft.com/en-us/library/z8910865.aspx
+// 4005 - 'identifier' : macro redefinition																							https://docs.microsoft.com/en-us/cpp/error-messages/compiler-warnings/compiler-warning-level-1-c4005
+// 4061 - enumerator 'identifier' in switch of enum 'enumeration' is not explicitly handled by a case label							https://docs.microsoft.com/en-us/cpp/error-messages/compiler-warnings/compiler-warning-level-4-c4061
+// 4062 - enumerator 'identifier' in switch of enum 'enumeration' is not handled													https://docs.microsoft.com/en-us/cpp/error-messages/compiler-warnings/compiler-warning-level-4-c4062
+
+// 4165 - 'HRESULT' is being converted to 'bool'; are you sure this is what you want?												https://docs.microsoft.com/en-us/cpp/error-messages/compiler-warnings/compiler-warning-level-1-c4165
+
+// 4242 - 'identifier' : conversion from 'type1' to 'type2', possible loss of data													https://docs.microsoft.com/en-us/cpp/error-messages/compiler-warnings/compiler-warning-level-4-c4242
+// 4254 - 'operator' : conversion from 'type1' to 'type2', possible loss of data													https://docs.microsoft.com/en-us/cpp/error-messages/compiler-warnings/compiler-warning-level-4-c4254
+// 4265 - 'class' : class has virtual functions, but destructor is not virtual														https://docs.microsoft.com/en-us/cpp/error-messages/compiler-warnings/compiler-warning-level-3-c4265
+// 4266 - 'function' : no override available for virtual member function from base 'type'; function is hidden						https://docs.microsoft.com/en-us/cpp/error-messages/compiler-warnings/compiler-warning-level-4-c4266
+// 4296 - 'operator' : expression is always false																					https://docs.microsoft.com/en-us/cpp/error-messages/compiler-warnings/compiler-warning-level-4-c4296
+
+// 4305 - 'identifier' : truncation from 'type1' to 'type2'																			https://docs.microsoft.com/en-us/cpp/error-messages/compiler-warnings/compiler-warning-level-1-c4305
+// 4307 - '' : integral constant overflow																							https://docs.microsoft.com/en-us/cpp/error-messages/compiler-warnings/compiler-warning-level-2-c4307
+// 4311 - 'variable' : pointer truncation from 'type' to 'type'																		https://docs.microsoft.com/en-us/cpp/error-messages/compiler-warnings/compiler-warning-level-1-c4311
+// 4312 - 'operation' : conversion from 'type1' to 'type2' of greater size															https://docs.microsoft.com/en-us/cpp/error-messages/compiler-warnings/compiler-warning-level-1-c4312
+// 4339 - 'type' : use of undefined type detected in CLR meta-data - use of this type may lead to a runtime exception				https://docs.microsoft.com/en-us/cpp/error-messages/compiler-warnings/compiler-warning-level-4-c4339
+// 4342 - behavior change: 'function' called, but a member operator was called in previous versions									https://docs.microsoft.com/en-us/cpp/error-messages/compiler-warnings/compiler-warning-level-1-c4342
 // 4345 - behavior change: an object of POD type constructed with an initializer of the form () will be default-initialized			http://msdn.microsoft.com/en-us/library/wewb47ee.aspx
-// 4350 - behavior change: 'member1' called instead of 'member2'																	http://msdn.microsoft.com/en-us/library/0eestyah.aspx
-// 4365 - 'action' : conversion from 'type_1' to 'type_2', signed/unsigned mismatch													http://msdn.microsoft.com/en-us/library/ms173683.aspx
-// 4370 - layout of class has changed from a previous version of the compiler due to better packing
-// 4371 - layout of class may have changed from a previous version of the compiler due to better packing of member
-// 4373 - '%$S': virtual function overrides '%$pS', previous versions of the compiler did not override when parameters only differed by const/volatile qualifiers	http://msdn.microsoft.com/en-us/library/bb384874.aspx
-// 4388 - '==' : signed/unsigned mismatch																							http://msdn.microsoft.com/en-us/library/jj155806.aspx
-// 4412 - 'function': function signature contains type 'type'; C++ objects are unsafe to pass between pure code and mixed or native	http://msdn.microsoft.com/en-us/library/ms235599.aspx
-// 4435 - 'class1' : Object layout under /vd2 will change due to virtual base 'class2'
-// 4437 - dynamic_cast from virtual base 'class1' to 'class2' could fail in some contexts											http://msdn.microsoft.com/en-us/library/23k5d385.aspx
-// 4444 - top level '__unaligned' is not implemented in this context
-
-// Shadow variable declaration warnings. These should eventually be fixed up and reenabled. Only effect VS2015.
-// 4456 - declaration of 'LocalVariable' hides previous local declaration
-// 4457 - declaration of 'LocalVariable' hides function parameter
-// 4458 - declaration of 'parameter' hides class member
-// 4459 - declaration of 'LocalVariable' hides global declaration
-
-// 4463 - overflow; assigning 1 to bit-field that can only hold values from -1 to 0
-// 4471 - a forward declaration of an unscoped enumeration must have an underlying type (int assumed)
-// 4472 - 'enum' is a native enum: add an access specifier (private/public) to declare a managed enum
-// 4481 - nonstandard extension used: override specifier 'keyword'																	http://msdn.microsoft.com/en-us/library/ms173703.aspx
-// 4510 - 'class' : default constructor could not be generated																		http://msdn.microsoft.com/en-us/library/2cf74y2b.aspx
-// 4547 - 'operator' : operator before comma has no effect; expected operator with side-effect										http://msdn.microsoft.com/en-us/library/y1724hsf.aspx
-// 4548 - expression before comma has no effect; expected expression with side-effect												http://msdn.microsoft.com/en-us/library/yxyxx8fx.aspx
-// 4549 - 'operator' : operator before comma has no effect; did you intend 'operator'?												http://msdn.microsoft.com/en-us/library/60yhzzeh.aspx
-// 4555 - expression has no effect; expected expression with side-effect															http://msdn.microsoft.com/en-us/library/k64a6he5.aspx
-// 4571 - Informational: catch(...) semantics changed since Visual C++ 7.1; structured exceptions (SEH) are no longer caught		http://msdn.microsoft.com/en-us/library/55s8esw4.aspx
-// 4574 - 'FALSE' is defined to be '0': did you mean to use '#if FALSE'?
-// 4577 - 'noexcept' used with no exception handling mode specified; termination on exception is not guaranteed. Specify /EHsc		
-// 4608 - 'symbol1' has already been initialized by another union member in the initializer list, 'symbol2'
-// 4651 - 'definition' specified for precompiled header but not for current compile													http://msdn.microsoft.com/en-us/library/h6dykdte.aspx
-// 4692 - 'function': signature of non - private member contains assembly private native type 'native_type'							http://msdn.microsoft.com/en-us/library/ms173713.aspx
-// 4701 - Potentially uninitialized local variable 'name' used																		http://msdn.microsoft.com/en-us/library/1wea5zwe.aspx
-// 4730 - 'main' : mixing _m64 and floating point expressions may result in incorrect code											http://msdn.microsoft.com/en-us/library/3z3ww2w3.aspx
-// 4738 - storing 32-bit float result in memory, possible loss of performance														http://msdn.microsoft.com/en-us/library/c24hdbz6.aspx
-// 4767 - section name 'symbol' is longer than 8 characters and will be truncated by the linker
-// 4774 - 'sprintf_s' : format string expected in argument 3 is not a string literal
-// 4819 - The file contains a character that cannot be represented in the current code page (xxx). Save the file in Unicode format to prevent data loss	https://msdn.microsoft.com/en-us/library/ms173715.aspx
-// 4820 - 'bytes' bytes padding added after construct 'member_name'																	http://msdn.microsoft.com/en-us/library/t7khkyth.aspx
-// 4826 - Conversion from 'type1 ' to 'type_2' is sign-extended. This may cause unexpected runtime behavior.						http://msdn.microsoft.com/en-us/library/ms235307.aspx
-// 4828 - The file contains a character starting at offset ... that is illegal in the current source character set (codepage ...).
-// 4837 - trigraph detected: '??%c' replaced by '%c'																				http://msdn.microsoft.com/en-us/library/cc664919.aspx
+// 4350 - behavior change: 'member1' called instead of 'member2'																	https://docs.microsoft.com/en-us/cpp/error-messages/compiler-warnings/compiler-warning-level-1-c4350
+// 4365 - 'action' : conversion from 'type_1' to 'type_2', signed/unsigned mismatch													https://docs.microsoft.com/en-us/cpp/error-messages/compiler-warnings/compiler-warning-level-4-c4365
+// 4370 - layout of class has changed from a previous version of the compiler due to better packing									https://docs.microsoft.com/en-us/cpp/error-messages/compiler-warnings/compiler-warnings-c4200-through-c4399
+// 4371 - layout of class may have changed from a previous version of the compiler due to better packing of member					https://docs.microsoft.com/en-us/cpp/error-messages/compiler-warnings/c4371
+// 4373 - 'function': virtual function overrides 'base_function', previous versions of the compiler did not override when parameters only differed by const/volatile qualifiers	https://docs.microsoft.com/en-us/cpp/error-messages/compiler-warnings/compiler-warning-level-3-c4373
+// 4388 - '==' : signed/unsigned mismatch																							https://docs.microsoft.com/en-us/cpp/error-messages/compiler-warnings/compiler-warnings-c4200-through-c4399
+// 4412 - 'function': function signature contains type 'type'; C++ objects are unsafe to pass between pure code and mixed or native	https://docs.microsoft.com/en-us/cpp/error-messages/compiler-warnings/compiler-warning-level-2-c4412
+// 4435 - 'class1' : Object layout under /vd2 will change due to virtual base 'class2'												https://docs.microsoft.com/en-us/cpp/error-messages/compiler-warnings/compiler-warning-level-4-c4435
+// 4437 - dynamic_cast from virtual base 'class1' to 'class2' could fail in some contexts											https://docs.microsoft.com/en-us/cpp/error-messages/compiler-warnings/compiler-warning-level-4-c4437
+// 4444 - 'identifier': top level '__unaligned' is not implemented in this context													https://docs.microsoft.com/en-us/cpp/error-messages/compiler-warnings/compiler-warnings-c4400-through-c4599
+
+// NOTE: Shadow variable declaration warnings. These should eventually be fixed up and reenabled.
+// 4456 - declaration of 'identifier' hides previous local declaration																https://docs.microsoft.com/en-us/cpp/error-messages/compiler-warnings/compiler-warning-level-4-c4456
+// 4457 - declaration of 'identifier' hides function parameter																		https://docs.microsoft.com/en-us/cpp/error-messages/compiler-warnings/compiler-warning-level-4-c4457
+// 4458 - declaration of 'identifier' hides class member																			https://docs.microsoft.com/en-us/cpp/error-messages/compiler-warnings/compiler-warning-level-4-c4458
+// 4459 - declaration of 'identifier' hides global declaration																		https://docs.microsoft.com/en-us/cpp/error-messages/compiler-warnings/compiler-warning-level-4-c4459
+
+// 4463 - overflow; assigning value to bit-field that can only hold values from low_value to high_value								https://docs.microsoft.com/en-us/cpp/error-messages/compiler-warnings/compiler-warning-level-4-c4463
+// 4471 - a forward declaration of an unscoped enumeration must have an underlying type (int assumed)								https://docs.microsoft.com/en-us/cpp/error-messages/compiler-warnings/compiler-warning-level-4-c4471
+// 4472 - 'identifier' is a native enum: add an access specifier (private/public) to declare a 'WinRT|managed' enum					https://docs.microsoft.com/en-us/cpp/error-messages/compiler-warnings/compiler-warnings-c4400-through-c4599
+// 4481 - nonstandard extension used: override specifier 'keyword'																	https://docs.microsoft.com/en-us/cpp/error-messages/compiler-warnings/compiler-warning-level-4-c4481
+
+// 4510 - 'class' : default constructor could not be generated																		https://docs.microsoft.com/en-us/cpp/error-messages/compiler-warnings/compiler-warning-level-4-c4510
+// 4547 - 'operator' : operator before comma has no effect; expected operator with side-effect										https://docs.microsoft.com/en-us/cpp/error-messages/compiler-warnings/compiler-warning-level-1-c4547
+// 4548 - expression before comma has no effect; expected expression with side-effect												https://docs.microsoft.com/en-us/cpp/error-messages/compiler-warnings/compiler-warning-level-1-c4548
+// 4549 - 'operator' : operator before comma has no effect; did you intend 'operator'?												https://docs.microsoft.com/en-us/cpp/error-messages/compiler-warnings/compiler-warning-level-1-c4549
+// 4555 - expression has no effect; expected expression with side-effect															https://docs.microsoft.com/en-us/cpp/error-messages/compiler-warnings/compiler-warning-level-1-c4555
+// 4571 - Informational: catch(...) semantics changed since Visual C++ 7.1; structured exceptions (SEH) are no longer caught		https://docs.microsoft.com/en-us/cpp/error-messages/compiler-warnings/compiler-warning-level-4-c4571
+// 4574 - 'Identifier' is defined to be '0': did you mean to use '#if identifier'?													https://docs.microsoft.com/en-us/cpp/error-messages/compiler-warnings/compiler-warnings-c4400-through-c4599
+// 4577 - 'noexcept' used with no exception handling mode specified; termination on exception is not guaranteed. Specify /EHsc		https://docs.microsoft.com/en-us/cpp/error-messages/compiler-warnings/compiler-warnings-c4400-through-c4599
+
+// 4608 - 'union_member' has already been initialized by another union member in the initializer list, 'union_member'				https://docs.microsoft.com/en-us/cpp/error-messages/compiler-warnings/compiler-warning-level-3-c4608
+// 4619 - #pragma warning : there is no warning number 'number'																		https://docs.microsoft.com/en-us/cpp/error-messages/compiler-warnings/compiler-warning-level-3-c4619
+// 4643 - Forward declaring 'identifier' in namespace std is not permitted by the C++ Standard.										https://docs.microsoft.com/en-us/cpp/cpp-conformance-improvements-2017
+// 4651 - 'definition' specified for precompiled header but not for current compile													https://docs.microsoft.com/en-us/cpp/error-messages/compiler-warnings/compiler-warning-level-1-c4651
+// 4668 - 'symbol' is not defined as a preprocessor macro, replacing with '0' for 'directives'										https://docs.microsoft.com/en-us/cpp/error-messages/compiler-warnings/compiler-warning-level-4-c4668
+// 4692 - 'function': signature of non-private member contains assembly private native type 'native_type'							https://docs.microsoft.com/en-us/cpp/error-messages/compiler-warnings/compiler-warning-level-1-c4692
+
+// 4701 - Potentially uninitialized local variable 'name' used																		https://docs.microsoft.com/en-us/cpp/error-messages/compiler-warnings/compiler-warning-level-4-c4701
+// 4730 - 'main' : mixing _m64 and floating point expressions may result in incorrect code											https://docs.microsoft.com/en-us/cpp/error-messages/compiler-warnings/compiler-warning-level-1-c4730
+// 4738 - storing 32-bit float result in memory, possible loss of performance														https://docs.microsoft.com/en-us/cpp/error-messages/compiler-warnings/compiler-warning-level-3-c4738
+// 4767 - section name '%s' is longer than 8 characters and will be truncated by the linker											https://docs.microsoft.com/en-us/cpp/error-messages/compiler-warnings/compiler-warnings-c4600-through-c4799
+// 4770 - partially validated enum 'name' used as index																				https://docs.microsoft.com/en-us/cpp/error-messages/compiler-warnings/compiler-warnings-c4600-through-c4799
+// 4774 - 'string' : format string expected in argument number is not a string literal												https://docs.microsoft.com/en-us/cpp/error-messages/compiler-warnings/compiler-warnings-c4600-through-c4799
+
+// 4819 - The file contains a character that cannot be represented in the current code page (number). Save the file in Unicode format to prevent data loss.	https://docs.microsoft.com/en-us/cpp/error-messages/compiler-warnings/compiler-warning-level-1-c4819
+// 4820 - 'bytes' bytes padding added after construct 'member_name'																	https://docs.microsoft.com/en-us/cpp/error-messages/compiler-warnings/compiler-warning-level-4-c4820
+// 4826 - Conversion from 'type1' to 'type2' is sign-extended. This may cause unexpected runtime behavior.							https://docs.microsoft.com/en-us/cpp/error-messages/compiler-warnings/compiler-warnings-c4800-through-c4999
+// 4828 - The file contains a character starting at offset ... that is illegal in the current source character set (codepage ...).	https://docs.microsoft.com/en-us/cpp/error-messages/compiler-warnings/compiler-warnings-c4800-through-c4999
+// 4837 - trigraph detected: '??character' replaced by 'character'																	https://docs.microsoft.com/en-us/cpp/error-messages/compiler-warnings/compiler-warnings-c4800-through-c4999
+// 4838 - conversion from 'type_1' to 'type_2' requires a narrowing conversion														https://docs.microsoft.com/en-us/cpp/error-messages/compiler-warnings/compiler-warning-level-1-c4838
 // 4868 - file(line_number)' compiler may not enforce left-to-right evaluation order in braced initializer list                     https://msdn.microsoft.com/en-us/library/mt656772.aspx
-// 4838 - 'type1' to 'type2' requires a narrowing conversion
-// 4962 - 'function': profile-guided optimizations disabled because optimizations caused profile data to become inconsistent
-// 4986 - 'function': exception specification does not match previous declaration													http://msdn.microsoft.com/en-us/library/jj620898.aspx
-// 4987 - nonstandard extension used: 'throw (...)'
-// 4988 - 'symbol': variable declared outside class/function scope
-#pragma warning (default : 4996)
+
+// 4962 - 'function' : Profile-guided optimizations disabled because optimizations caused profile data to become inconsistent		https://docs.microsoft.com/en-us/cpp/error-messages/compiler-warnings/compiler-warning-c4962
+// 4986 - 'function': exception specification does not match previous declaration													https://docs.microsoft.com/en-us/cpp/error-messages/compiler-warnings/compiler-warning-c4986
+// 4987 - nonstandard extension used: 'throw (...)'																					https://docs.microsoft.com/en-us/cpp/error-messages/compiler-warnings/compiler-warnings-c4800-through-c4999
+// 4988 - 'variable': variable declared outside class/function scope																https://docs.microsoft.com/en-us/cpp/error-messages/compiler-warnings/compiler-warnings-c4800-through-c4999
+// 4996 - The compiler encountered a deprecated declaration.																		https://docs.microsoft.com/en-us/cpp/error-messages/compiler-warnings/compiler-warning-level-3-c4996
+
+// 5045 - Compiler will insert Spectre mitigation for memory load if /Qspectre switch specified										https://docs.microsoft.com/en-us/cpp/error-messages/compiler-warnings/c5045
+
+//
+// Disabled warnings
+//
 
 // @todo UWP: Disabled because DbgHelp.h has some anonymous typedefs in it (not allowed in Visual Studio 2015).  We should probably just wrap that header.
-#pragma warning(disable : 4091) // 'typedef ': ignored on left of 'type declaration' when no variable is declared
+#pragma warning(disable: 4091) // 'keyword' : ignored on left of 'type' when no variable is declared								https://docs.microsoft.com/en-us/cpp/error-messages/compiler-warnings/compiler-warning-level-1-c4091
 
 // Unwanted VC++ level 4 warnings to disable.
-#pragma warning(disable : 4100) // unreferenced formal parameter										
-#pragma warning(disable : 4127) // Conditional expression is constant									
-#pragma warning(disable : 4200) // Zero-length array item at end of structure, a VC-specific extension	
-#pragma warning(disable : 4201) // nonstandard extension used : nameless struct/union	
-#pragma warning(disable : 4244) // 'expression' : conversion from 'type' to 'type', possible loss of data						
-#pragma warning(disable : 4245) // 'initializing': conversion from 'type' to 'type', signed/unsigned mismatch 
-#pragma warning(disable : 4251) // 'type' needs to have dll-interface to be used by clients of 'type'
-#pragma warning(disable : 4291) // typedef-name '' used as synonym for class-name ''    
-#pragma warning(disable : 4307) // '': integral constant overflow   
-//#pragma warning(disable : 4305) // 'argument' : truncation from 'double' to 'float' --- fp:precise builds only!!!!!!!
-#pragma warning(disable : 4324) // structure was padded due to __declspec(align())						
-#pragma warning(disable : 4355) // this used in base initializer list                                   
-#pragma warning(disable : 4373) // '%$S': virtual function overrides '%$pS', previous versions of the compiler did not override when parameters only differed by const/volatile qualifiers
-#pragma warning(disable : 4389) // signed/unsigned mismatch                                             
-#pragma warning(disable : 4511) // copy constructor could not be generated                              
-#pragma warning(disable : 4512) // assignment operator could not be generated                           
-#pragma warning(disable : 4514) // unreferenced inline function has been removed						
-#pragma warning(disable : 4599) // VS2015 update 3 : When using PCH files, mismatched -I include directories to the compiler between -Yc and -Yu compilations will now produce a new warning.
-#pragma warning(disable : 4605) // VS2015 update 3 : Seems related to 4599.
-#pragma warning(disable : 4699) // creating precompiled header											
-#pragma warning(disable : 4702) // unreachable code in inline expanded function							
-#pragma warning(disable : 4710) // inline function not expanded											
-#pragma warning(disable : 4711) // function selected for automatic inlining								
-#pragma warning(disable : 4714) // __forceinline function not expanded
-#pragma warning(disable : 4768) // __declspec ignored when followed by a link specifier (new warning)
-#pragma warning(disable : 4482) // nonstandard extension used: enum 'enum' used in qualified name (having hte enum name helps code readability and should be part of TR1 or TR2)
-#pragma warning(disable : 4748)	// /GS can not protect parameters and local variables from local buffer overrun because optimizations are disabled in function
-// NOTE: _mm_cvtpu8_ps will generate this falsely if it doesn't get inlined
-#pragma warning(disable : 4799)	// Warning: function 'ident' has no EMMS instruction
-#pragma warning(disable : 4275) // non - DLL-interface classkey 'identifier' used as base for DLL-interface classkey 'identifier'
-
-// all of the /Wall warnings that we are able to enable
-// @todo:  http://msdn2.microsoft.com/library/23k5d385(en-us,vs.80).aspx
-#pragma warning(default : 4191) // 'operator/operation' : unsafe conversion from 'type of expression' to 'type required'
-#pragma warning(disable : 4217) // 'operator' : member template functions cannot be used for copy-assignment or copy-construction
-//#pragma warning(disable : 4242) // 'variable' : conversion from 'type' to 'type', possible loss of data
-//#pragma warning(default : 4254) // 'operator' : conversion from 'type1' to 'type2', possible loss of data
-#pragma warning(default : 4255) // 'function' : no function prototype given: converting '()' to '(void)'
-#pragma warning(default : 4263) // 'function' : member function does not override any base class virtual member function
-#pragma warning(default : 4264) // 'virtual_function' : no override available for virtual member function from base 'class'; function is hidden
-#pragma warning(disable : 4267) // 'argument' : conversion from 'type1' to 'type2', possible loss of data
-#pragma warning(default : 4287) // 'operator' : unsigned/negative constant mismatch
-#pragma warning(default : 4289) // nonstandard extension used : 'var' : loop control variable declared in the for-loop is used outside the for-loop scope
-#pragma warning(disable : 4315) // 'this' pointer for member may not be aligned 8 as expected by the constructor
-#pragma warning(disable : 4316) // 'type' : object allocated on the heap may not be aligned 16
-//#pragma warning(disable : 4339) // 'type' : use of undefined type detected in CLR meta-data - use of this type may lead to a runtime exception
-#pragma warning(disable : 4347) // behavior change: 'function template' is called instead of 'function
-#pragma warning(disable : 4351) // new behavior: elements of array 'array' will be default initialized
-#pragma warning(disable : 4366) // The result of the unary '&' operator may be unaligned
-#pragma warning(disable : 4514) // unreferenced inline/local function has been removed
-#pragma warning(default : 4529) // 'member_name' : forming a pointer-to-member requires explicit use of the address-of operator ('&') and a qualified name
-#pragma warning(default : 4536) // 'type name' : type-name exceeds meta-data limit of 'limit' characters
-#pragma warning(default : 4545) // expression before comma evaluates to a function which is missing an argument list
-#pragma warning(default : 4546) // function call before comma missing argument list
-//#pragma warning(default : 4547) // 'operator' : operator before comma has no effect; expected operator with side-effect
-//#pragma warning(default : 4548) // expression before comma has no effect; expected expression with side-effect  (needed as xlocale does not compile cleanly)
-//#pragma warning(default : 4549) // 'operator' : operator before comma has no effect; did you intend 'operator'?
-//#pragma warning(disable : 4555) // expression has no effect; expected expression with side-effect
-#pragma warning(default : 4557) // '__assume' contains effect 'effect'
-#pragma warning(disable : 4577) // 'noexcept' used with no exception handling mode specified; termination on exception is not guaranteed. Specify /EHsc
-#pragma warning(disable : 4623) // 'derived class' : default constructor could not be generated because a base class default constructor is inaccessible
-#pragma warning(disable : 4625) // 'derived class' : copy constructor could not be generated because a base class copy constructor is inaccessible
-#pragma warning(disable : 4626) // 'derived class' : assignment operator could not be generated because a base class assignment operator is inaccessible
-#pragma warning(default : 4628) // digraphs not supported with -Ze. Character sequence 'digraph' not interpreted as alternate token for 'char'
-#pragma warning(disable : 4640) // 'instance' : construction of local static object is not thread-safe
-#pragma warning(default : 4682) // 'parameter' : no directional parameter attribute specified, defaulting to [in]
-#pragma warning(default : 4686) // 'user-defined type' : possible change in behavior, change in UDT return calling convention
-#pragma warning(disable : 4710) // 'function' : function not inlined / The given function was selected for inline expansion, but the compiler did not perform the inlining.
-#pragma warning(default : 4786) // 'identifier' : identifier was truncated to 'number' characters in the debug information
-#pragma warning(default : 4793) // native code generated for function 'function': 'reason'
-#pragma warning(default : 4905) // wide string literal cast to 'LPSTR'
-#pragma warning(default : 4906) // string literal cast to 'LPWSTR'
-#pragma warning(disable : 4917) // 'declarator' : a GUID cannot only be associated with a class, interface or namespace ( ocid.h breaks this)
-#pragma warning(default : 4931) // we are assuming the type library was built for number-bit pointers
-#pragma warning(default : 4946) // reinterpret_cast used between related classes: 'class1' and 'class2'
-#pragma warning(default : 4928) // illegal copy-initialization; more than one user-defined conversion has been implicitly applied
-#pragma warning(disable : 4180) // qualifier applied to function type has no meaning; ignored
-#pragma warning(disable : 4121) // 'symbol' : alignment of a member was sensitive to packing
-#pragma warning(disable : 4345) // behavior change: an object of POD type constructed with an initializer of the form () will be default-initialized
-#pragma warning(disable : 4464) // relative include path contains '..'
-#pragma warning(disable : 4592) // symbol will be dynamically initialized (implementation limitation) // warning broken in VS 2015 Update 1 & 2 - see http://stackoverflow.com/a/34027257
-#pragma warning(disable : 4828) // The file contains a character starting at offset ... that is illegal in the current source character set(codepage ...).
-#pragma warning(disable : 4868) // compiler may not enforce left-to-right evaluation order in braced initializer list
+#pragma warning(disable: 4100) // 'identifier' : unreferenced formal parameter														https://docs.microsoft.com/en-us/cpp/error-messages/compiler-warnings/compiler-warning-level-4-c4100
+#pragma warning(disable: 4121) // 'symbol' : alignment of a member was sensitive to packing											https://docs.microsoft.com/en-us/cpp/error-messages/compiler-warnings/compiler-warning-level-4-c4121
+#pragma warning(disable: 4127) // Conditional expression is constant																https://docs.microsoft.com/en-us/cpp/error-messages/compiler-warnings/compiler-warning-level-4-c4127
+#pragma warning(disable: 4180) // qualifier applied to function type has no meaning; ignored										https://docs.microsoft.com/en-us/cpp/error-messages/compiler-warnings/compiler-warning-level-1-c4180
+#pragma warning(disable: 4189) // 'identifier': local variable is initialized but not referenced									https://docs.microsoft.com/en-us/cpp/error-messages/compiler-warnings/compiler-warnings-c4000-through-c4199
+
+#pragma warning(disable: 4200) // Zero-length array item at end of structure, a VC-specific extension								https://docs.microsoft.com/en-us/cpp/error-messages/compiler-warnings/compiler-warning-levels-2-and-4-c4200
+#pragma warning(disable: 4201) // nonstandard extension used : nameless struct/union												https://docs.microsoft.com/en-us/cpp/error-messages/compiler-warnings/compiler-warning-level-4-c4201
+#pragma warning(disable: 4217) // 'operator' : member template functions cannot be used for copy-assignment or copy-construction	// No docs
+#pragma warning(disable: 4244) // 'argument' : conversion from 'type1' to 'type2', possible loss of data							https://docs.microsoft.com/en-us/cpp/error-messages/compiler-warnings/compiler-warning-level-4-c4201
+#pragma warning(disable: 4245) // 'initializing': conversion from 'type' to 'type', signed/unsigned mismatch						https://docs.microsoft.com/en-us/cpp/error-messages/compiler-warnings/compiler-warning-level-4-c4245
+#pragma warning(disable: 4251) // 'type' needs to have dll-interface to be used by clients of 'type'								https://docs.microsoft.com/en-us/cpp/error-messages/compiler-warnings/compiler-warning-level-1-c4251
+#pragma warning(disable: 4267) // 'var' : conversion from 'size_t' to 'type', possible loss of data									https://docs.microsoft.com/en-us/cpp/error-messages/compiler-warnings/compiler-warning-level-3-c4267
+#pragma warning(disable: 4275) // non - DLL-interface classkey 'identifier' used as base for DLL-interface classkey 'identifier'	https://docs.microsoft.com/en-us/cpp/error-messages/compiler-warnings/compiler-warning-level-2-c4275
+#pragma warning(disable: 4291) // typedef-name '' used as synonym for class-name ''													https://docs.microsoft.com/en-us/cpp/error-messages/compiler-warnings/compiler-warning-level-1-c4291
+
+#pragma warning(disable: 4307) // '': integral constant overflow																	https://docs.microsoft.com/en-us/cpp/error-messages/compiler-warnings/compiler-warning-level-2-c4307
+#pragma warning(disable: 4315) // 'classname': 'this' pointer for member 'member' may not be aligned 'alignment' as expected by the constructor	https://docs.microsoft.com/en-us/cpp/error-messages/compiler-warnings/compiler-warnings-c4200-through-c4399
+#pragma warning(disable: 4316) // 'identifier': object allocated on the heap may not be aligned 'alignment'							https://docs.microsoft.com/en-us/cpp/error-messages/compiler-warnings/compiler-warnings-c4200-through-c4399
+#pragma warning(disable: 4324) // structure was padded due to __declspec(align())													https://docs.microsoft.com/en-us/cpp/error-messages/compiler-warnings/compiler-warning-level-4-c4324
+#pragma warning(disable: 4347) // behavior change: 'function template' is called instead of 'function								https://docs.microsoft.com/en-us/cpp/error-messages/compiler-warnings/compiler-warnings-c4200-through-c4399
+#pragma warning(disable: 4351) // new behavior: elements of array 'array' will be default initialized								https://docs.microsoft.com/en-us/cpp/error-messages/compiler-warnings/compiler-warnings-c4200-through-c4399
+#pragma warning(disable: 4355) // this used in base initializer list																https://docs.microsoft.com/en-us/cpp/error-messages/compiler-warnings/compiler-warning-c4355
+#pragma warning(disable: 4366) // The result of the unary 'operator' operator may be unaligned										https://docs.microsoft.com/en-us/cpp/error-messages/compiler-warnings/compiler-warning-level-4-c4366
+#pragma warning(disable: 4373) // '%$S': virtual function overrides '%$pS', previous versions of the compiler did not override when parameters only differed by const/volatile qualifiers	https://docs.microsoft.com/en-us/cpp/error-messages/compiler-warnings/compiler-warning-level-3-c4373
+#pragma warning(disable: 4389) // signed/unsigned mismatch																			https://docs.microsoft.com/en-us/cpp/error-messages/compiler-warnings/compiler-warning-c4355
 
+#pragma warning(disable: 4464) // relative include path contains '..'																https://docs.microsoft.com/en-us/cpp/error-messages/compiler-warnings/c4464
+#pragma warning(disable: 4482) // nonstandard extension used: enum 'enumeration' used in qualified name								https://docs.microsoft.com/en-us/cpp/error-messages/compiler-warnings/compiler-warnings-c4400-through-c4599
+
+#pragma warning(disable: 4505) // 'function' : unreferenced local function has been removed											https://docs.microsoft.com/en-us/cpp/error-messages/compiler-warnings/compiler-warning-level-4-c4505
+#pragma warning(disable: 4511) // 'class' : copy constructor could not be generated													https://docs.microsoft.com/en-us/cpp/error-messages/compiler-warnings/compiler-warning-level-3-c4511
+#pragma warning(disable: 4512) // 'class' : assignment operator could not be generated												https://docs.microsoft.com/en-us/cpp/error-messages/compiler-warnings/compiler-warning-level-4-c4512
+#pragma warning(disable: 4514) // 'function' : unreferenced inline function has been removed										https://docs.microsoft.com/en-us/cpp/error-messages/compiler-warnings/compiler-warning-level-4-c4514
 #if UE_BUILD_DEBUG
 // xstring.h causes this warning in debug builds
-#pragma warning(disable : 4548) // expression before comma has no effect; expected expression with side-effect
+#pragma warning(disable: 4548) // expression before comma has no effect; expected expression with side-effect						https://docs.microsoft.com/en-us/cpp/error-messages/compiler-warnings/compiler-warning-level-1-c4548
+#endif
+#pragma warning(disable: 4592) // 'function': 'constexpr' call evaluation failed; function will be called at run-time				https://docs.microsoft.com/en-us/cpp/error-messages/compiler-warnings/compiler-warnings-c4400-through-c4599
+#pragma warning(disable: 4599) // 'flag path': command line argument number number does not match precompiled header				https://docs.microsoft.com/en-us/cpp/error-messages/compiler-warnings/compiler-warnings-c4400-through-c4599
+
+#pragma warning(disable: 4605) // '/Dmacro' specified on current command line, but was not specified when precompiled header was built	https://docs.microsoft.com/en-us/cpp/error-messages/compiler-warnings/compiler-warnings-c4600-through-c4799
+#pragma warning(disable: 4623) // 'derived class' : default constructor was implicitly defined as deleted because a base class default constructor is inaccessible or deleted	https://docs.microsoft.com/en-us/cpp/error-messages/compiler-warnings/compiler-warning-level-4-c4623
+#pragma warning(disable: 4625) // 'derived class' : copy constructor was implicitly defined as deleted because a base class copy constructor is inaccessible or deleted	https://docs.microsoft.com/en-us/cpp/error-messages/compiler-warnings/compiler-warning-level-4-c4625
+#pragma warning(disable: 4626) // 'derived class' : assignment operator was implicitly defined as deleted because a base class assignment operator is inaccessible or deleted	https://docs.microsoft.com/en-us/cpp/error-messages/compiler-warnings/compiler-warning-level-4-c4626
+#pragma warning(disable: 4640) // 'instance' : construction of local static object is not thread-safe								https://docs.microsoft.com/en-us/cpp/error-messages/compiler-warnings/compiler-warning-level-3-c4640
+#pragma warning(disable: 4699) // creating precompiled header																		// No docs
+
+#pragma warning(disable: 4702) // unreachable code																					https://docs.microsoft.com/en-us/cpp/error-messages/compiler-warnings/compiler-warning-level-4-c4702
+#pragma warning(disable: 4710) // 'function' : function not inlined																	https://docs.microsoft.com/en-us/cpp/error-messages/compiler-warnings/compiler-warning-level-4-c4710
+#pragma warning(disable: 4711) // function selected for automatic inlining															https://docs.microsoft.com/en-us/cpp/error-messages/compiler-warnings/compiler-warning-level-1-c4711
+#pragma warning(disable: 4714) // function 'function' marked as __forceinline not inlined											https://docs.microsoft.com/en-us/cpp/error-messages/compiler-warnings/compiler-warning-level-4-c4714
+#pragma warning(disable: 4748) // /GS can not protect parameters and local variables from local buffer overrun because optimizations are disabled in function	// No docs
+#pragma warning(disable: 4768) // __declspec attributes before linkage specification are ignored									https://docs.microsoft.com/en-us/cpp/error-messages/compiler-warnings/compiler-warnings-c4600-through-c4799
+// NOTE: _mm_cvtpu8_ps will generate this falsely if it doesn't get inlined
+#pragma warning(disable: 4799) // Warning: function 'ident' has no EMMS instruction													https://docs.microsoft.com/en-us/cpp/error-messages/compiler-warnings/compiler-warning-level-1-c4799
+
+// NOTE: https://answers.unrealengine.com/questions/701635/warning-c4828.html
+#pragma warning(disable: 4828) // The file contains a character starting at offset ... that is illegal in the current source character set(codepage ...).	// No docs
+#pragma warning(disable: 4868) // 'file(line_number)' compiler may not enforce left-to-right evaluation order in braced initializer list	https://docs.microsoft.com/en-us/cpp/error-messages/compiler-warnings/compiler-warning-c4868
+
+// NOTE: ocid.h breaks this
+#pragma warning(disable: 4917) // 'declarator' : a GUID can only be associated with a class, interface or namespace 				https://docs.microsoft.com/en-us/cpp/error-messages/compiler-warnings/compiler-warning-level-1-c4917
+#if WINVER == 0x0502
+// NOTE: WinXP hits deprecated versions of stdio across the board
+#pragma warning(disable: 4995) // 'function': name was marked as #pragma deprecated													https://docs.microsoft.com/en-us/cpp/error-messages/compiler-warnings/compiler-warning-level-3-c4995
 #endif
 
+//
+// All of the /Wall warnings that we are able to enable
+// @todo: https://docs.microsoft.com/en-us/cpp/preprocessor/compiler-warnings-that-are-off-by-default
+// NOTE: This is currently just overriding the error versions above, removing these will cause them to be errors!
+//
+
+#pragma warning(default: 4191) // 'operator/operation': unsafe conversion from 'type_of_expression' to 'type_required'				https://docs.microsoft.com/en-us/cpp/error-messages/compiler-warnings/compiler-warnings-c4000-through-c4199
+
+#pragma warning(default: 4255) // 'function' : no function prototype given: converting '()' to '(void)'								https://docs.microsoft.com/en-us/cpp/error-messages/compiler-warnings/compiler-warning-level-4-c4255
+#pragma warning(default: 4263) // 'function' : member function does not override any base class virtual member function				https://docs.microsoft.com/en-us/cpp/error-messages/compiler-warnings/compiler-warning-level-4-c4263
+#pragma warning(default: 4264) // 'virtual_function' : no override available for virtual member function from base 'class'; function is hidden	https://docs.microsoft.com/en-us/cpp/error-messages/compiler-warnings/compiler-warning-level-1-c4264
+#pragma warning(3:       4265) // 'class' : class has virtual functions, but destructor is not virtual								https://docs.microsoft.com/en-us/cpp/error-messages/compiler-warnings/compiler-warning-level-3-c4265
+#pragma warning(default: 4287) // 'operator' : unsigned/negative constant mismatch													https://docs.microsoft.com/en-us/cpp/error-messages/compiler-warnings/compiler-warning-level-1-c4264
+#pragma warning(default: 4289) // nonstandard extension used : 'var' : loop control variable declared in the for-loop is used outside the for-loop scope	https://docs.microsoft.com/en-us/cpp/error-messages/compiler-warnings/compiler-warning-level-4-c4289
+
+//#pragma warning(disable : 4339) // 'type' : use of undefined type detected in CLR meta-data - use of this type may lead to a runtime exception
+#pragma warning(disable: 4345) // behavior change: an object of POD type constructed with an initializer of the form () will be default-initialized
+
+#pragma warning(disable: 4514) // unreferenced inline/local function has been removed
+#pragma warning(default: 4529) // 'member_name' : forming a pointer-to-member requires explicit use of the address-of operator ('&') and a qualified name	// No docs
+#pragma warning(default: 4536) // 'type name' : type-name exceeds meta-data limit of 'limit' characters								https://docs.microsoft.com/en-us/cpp/error-messages/compiler-warnings/compiler-warning-level-4-c4536
+#pragma warning(default: 4545) // expression before comma evaluates to a function which is missing an argument list					https://docs.microsoft.com/en-us/cpp/error-messages/compiler-warnings/compiler-warning-level-1-c4545
+#pragma warning(default: 4546) // function call before comma missing argument list													https://docs.microsoft.com/en-us/cpp/error-messages/compiler-warnings/compiler-warning-level-1-c4546
+#pragma warning(default: 4557) // '__assume' contains effect 'effect'																https://docs.microsoft.com/en-us/cpp/error-messages/compiler-warnings/compiler-warning-level-3-c4557
+#pragma warning(disable: 4577) // 'noexcept' used with no exception handling mode specified; termination on exception is not guaranteed. Specify /EHsc
+
+#pragma warning(default: 4628) // digraphs not supported with -Ze. Character sequence 'digraph' not interpreted as alternate token for 'char'	https://docs.microsoft.com/en-us/cpp/error-messages/compiler-warnings/compiler-warning-level-4-c4289
+#pragma warning(default: 4682) // 'parameter' : no directional parameter attribute specified, defaulting to [in]					https://docs.microsoft.com/en-us/cpp/error-messages/compiler-warnings/compiler-warning-level-4-c4682
+#pragma warning(default: 4686) // 'user-defined type' : possible change in behavior, change in UDT return calling convention		https://docs.microsoft.com/en-us/cpp/error-messages/compiler-warnings/compiler-warning-level-3-c4686
+
+#pragma warning(disable: 4710) // 'function' : function not inlined / The given function was selected for inline expansion, but the compiler did not perform the inlining.
+#pragma warning(default: 4786) // 'identifier' : identifier was truncated to 'number' characters in the debug information			// No docs
+#pragma warning(default: 4793) // 'function' : function is compiled as native code: 'reason'										https://docs.microsoft.com/en-us/cpp/error-messages/compiler-warnings/compiler-warning-level-1-and-3-c4793
+
+#pragma warning(default: 4905) // wide string literal cast to 'LPSTR'																https://docs.microsoft.com/en-us/cpp/error-messages/compiler-warnings/compiler-warning-level-1-c4905
+#pragma warning(default: 4906) // string literal cast to 'LPWSTR'																	https://docs.microsoft.com/en-us/cpp/error-messages/compiler-warnings/compiler-warning-level-1-c4906
+#pragma warning(default: 4928) // illegal copy-initialization; more than one user-defined conversion has been implicitly applied	https://docs.microsoft.com/en-us/cpp/error-messages/compiler-warnings/compiler-warning-level-1-c4928
+#pragma warning(default: 4931) // we are assuming the type library was built for number-bit pointers								https://docs.microsoft.com/en-us/cpp/error-messages/compiler-warnings/compiler-warning-level-4-c4931
+#pragma warning(default: 4946) // reinterpret_cast used between related classes: 'class1' and 'class2'								https://docs.microsoft.com/en-us/cpp/error-messages/compiler-warnings/compiler-warning-level-1-c4946
+
+#pragma warning(default: 4996)
+
 // interesting ones to turn on and off at times
-#pragma warning(3 : 4265) // 'class' : class has virtual functions, but destructor is not virtual
 //#pragma warning(disable : 4266) // '' : no override available for virtual member function from base ''; function is hidden
 //#pragma warning(disable : 4296) // 'operator' : expression is always true / false
 //#pragma warning(disable : 4820) // 'bytes' bytes padding added after member 'member'
 // Mixing MMX/SSE intrinsics will cause this warning, even when it's done correctly.
 //#pragma warning(disable : 4730) //mixing _m64 and floating point expressions may result in incorrect code
 
-// It'd be nice to turn these on, but at the moment they can't be used in DEBUG due to the varargs stuff.	
-#pragma warning(disable : 4189) // local variable is initialized but not referenced 
-#pragma warning(disable : 4505) // unreferenced local function has been removed		
-
-#if WINVER == 0x0502
-// WinXP hits deprecated versions of stdio across the board
-#pragma warning(disable : 4995) // 'function': name was marked as #pragma deprecated
-#endif
-
 // If C++ exception handling is disabled, force guarding to be off.
 #if !defined(_CPPUNWIND) && !defined(__INTELLISENSE__) && !defined(HACK_HEADER_GENERATOR)
 #error "Bad VCC option: C++ exception handling must be enabled" //lint !e309 suppress as lint doesn't have this defined
