Build fixes from merge.

#rnx

diff --git a/Engine/Source/Runtime/Core/Private/Windows/WindowsPlatformCrashContext.cpp b/Engine/Source/Runtime/Core/Private/Windows/WindowsPlatformCrashContext.cpp
--- a/Engine/Source/Runtime/Core/Private/Windows/WindowsPlatformCrashContext.cpp
+++ b/Engine/Source/Runtime/Core/Private/Windows/WindowsPlatformCrashContext.cpp
@@ -782,7 +782,7 @@ private:
 		CrashContext.CaptureAllThreadContexts();
 
 		// Also mark the same number of frames to be ignored if we symbolicate from the minidump
-		CrashContext.SetNumMinidumpFramesToIgnore(IgnoreCount);
+		CrashContext.SetNumMinidumpFramesToIgnore(NumStackFramesToIgnore);
 
 		// First launch the crash reporter client.
 #if WINVER > 0x502	// Windows Error Reporting is not supported on Windows XP
