[LLM] audio tagging changes:
 - Current Audio tag is renamed to AudioMisc.
 - Adding parent Audio category that sums all existing audio categories.
 - Adding new AudioSoundWaves category to track audio content.
 - Adding new AudioVoiceChat category.
 - Removed a couple of erroneous double taggings.

diff --git a/Engine/Source/Runtime/Core/Private/HAL/LowLevelMemTracker.cpp b/Engine/Source/Runtime/Core/Private/HAL/LowLevelMemTracker.cpp
--- a/Engine/Source/Runtime/Core/Private/HAL/LowLevelMemTracker.cpp
+++ b/Engine/Source/Runtime/Core/Private/HAL/LowLevelMemTracker.cpp
@@ -63,11 +63,14 @@ DECLARE_LLM_MEMORY_STAT(TEXT("GenericPlatformMallocCrash"), STAT_GenericPlatform
 DECLARE_LLM_MEMORY_STAT(TEXT("Engine Misc"), STAT_EngineMiscLLM, STATGROUP_LLMFULL);
 DECLARE_LLM_MEMORY_STAT(TEXT("TaskGraph Misc Tasks"), STAT_TaskGraphTasksMiscLLM, STATGROUP_LLMFULL);
 DECLARE_LLM_MEMORY_STAT(TEXT("Audio"), STAT_AudioLLM, STATGROUP_LLMFULL);
+DECLARE_LLM_MEMORY_STAT(TEXT("AudioMisc"), STAT_AudioMiscLLM, STATGROUP_LLMFULL);
+DECLARE_LLM_MEMORY_STAT(TEXT("AudioSoundWaves"), STAT_AudioSoundWavesLLM, STATGROUP_LLMFULL);
 DECLARE_LLM_MEMORY_STAT(TEXT("AudioMixer"), STAT_AudioMixerLLM, STATGROUP_LLMFULL);
 DECLARE_LLM_MEMORY_STAT(TEXT("AudioPrecache"), STAT_AudioPrecacheLLM, STATGROUP_LLMFULL);
 DECLARE_LLM_MEMORY_STAT(TEXT("AudioDecompress"), STAT_AudioDecompressLLM, STATGROUP_LLMFULL);
 DECLARE_LLM_MEMORY_STAT(TEXT("AudioRealtimePrecache"), STAT_AudioRealtimePrecacheLLM, STATGROUP_LLMFULL);
 DECLARE_LLM_MEMORY_STAT(TEXT("AudioFullDecompress"), STAT_AudioFullDecompressLLM, STATGROUP_LLMFULL);
+DECLARE_LLM_MEMORY_STAT(TEXT("AudioVoiceChat"), STAT_AudioVoiceChatLLM, STATGROUP_LLMFULL);
 DECLARE_LLM_MEMORY_STAT(TEXT("FName"), STAT_FNameLLM, STATGROUP_LLMFULL);
 DECLARE_LLM_MEMORY_STAT(TEXT("Networking"), STAT_NetworkingLLM, STATGROUP_LLMFULL);
 DECLARE_LLM_MEMORY_STAT(TEXT("Meshes"), STAT_MeshesLLM, STATGROUP_LLMFULL);
diff --git a/Engine/Source/Runtime/Core/Public/HAL/LowLevelMemTracker.h b/Engine/Source/Runtime/Core/Public/HAL/LowLevelMemTracker.h
--- a/Engine/Source/Runtime/Core/Public/HAL/LowLevelMemTracker.h
+++ b/Engine/Source/Runtime/Core/Public/HAL/LowLevelMemTracker.h
@@ -20,7 +20,7 @@
 	#define LLM_SUPPORTED_PLATFORM (PLATFORM_XBOXONE || PLATFORM_PS4 || PLATFORM_WINDOWS || PLATFORM_IOS || PLATFORM_MAC || PLATFORM_ANDROID || PLATFORM_SWITCH || PLATFORM_UNIX)
 
 #ifndef ALLOW_LOW_LEVEL_MEM_TRACKER_IN_TEST
-	#define ALLOW_LOW_LEVEL_MEM_TRACKER_IN_TEST 0
+#define ALLOW_LOW_LEVEL_MEM_TRACKER_IN_TEST 1 //#TOOD: do not check in
 #endif
 
 	// *** enable/disable LLM here ***
@@ -191,11 +191,14 @@ enum class ELLMTagSet : uint8
 	macro(EngineMisc,							"EngineMisc",					GET_STATFNAME(STAT_EngineMiscLLM),							GET_STATFNAME(STAT_EngineSummaryLLM),			-1)\
 	macro(TaskGraphTasksMisc,					"TaskGraphMiscTasks",			GET_STATFNAME(STAT_TaskGraphTasksMiscLLM),					GET_STATFNAME(STAT_EngineSummaryLLM),			-1)\
 	macro(Audio,								"Audio",						GET_STATFNAME(STAT_AudioLLM),								GET_STATFNAME(STAT_AudioSummaryLLM),			-1)\
-	macro(AudioMixer,							"AudioMixer",					GET_STATFNAME(STAT_AudioMixerLLM),							GET_STATFNAME(STAT_AudioSummaryLLM),			-1)\
-	macro(AudioPrecache,						"AudioPrecache",				GET_STATFNAME(STAT_AudioPrecacheLLM),						GET_STATFNAME(STAT_AudioSummaryLLM),			-1)\
-	macro(AudioDecompress,						"AudioDecompress",				GET_STATFNAME(STAT_AudioDecompressLLM),						GET_STATFNAME(STAT_AudioSummaryLLM),			-1)\
-	macro(AudioRealtimePrecache,				"AudioRealtimePrecache",		GET_STATFNAME(STAT_AudioRealtimePrecacheLLM),				GET_STATFNAME(STAT_AudioSummaryLLM),			-1)\
-	macro(AudioFullDecompress,					"AudioFullDecompress",			GET_STATFNAME(STAT_AudioFullDecompressLLM),					GET_STATFNAME(STAT_AudioSummaryLLM),			-1)\
+	macro(AudioMisc,							"AudioMisc",					GET_STATFNAME(STAT_AudioMiscLLM),							GET_STATFNAME(STAT_AudioSummaryLLM),			ELLMTag::Audio)\
+	macro(AudioSoundWaves,						"AudioSoundWaves",				GET_STATFNAME(STAT_AudioSoundWavesLLM),						GET_STATFNAME(STAT_AudioSummaryLLM),			ELLMTag::Audio)\
+	macro(AudioMixer,							"AudioMixer",					GET_STATFNAME(STAT_AudioMixerLLM),							GET_STATFNAME(STAT_AudioSummaryLLM),			ELLMTag::Audio)\
+	macro(AudioPrecache,						"AudioPrecache",				GET_STATFNAME(STAT_AudioPrecacheLLM),						GET_STATFNAME(STAT_AudioSummaryLLM),			ELLMTag::Audio)\
+	macro(AudioDecompress,						"AudioDecompress",				GET_STATFNAME(STAT_AudioDecompressLLM),						GET_STATFNAME(STAT_AudioSummaryLLM),			ELLMTag::Audio)\
+	macro(AudioRealtimePrecache,				"AudioRealtimePrecache",		GET_STATFNAME(STAT_AudioRealtimePrecacheLLM),				GET_STATFNAME(STAT_AudioSummaryLLM),			ELLMTag::Audio)\
+	macro(AudioFullDecompress,					"AudioFullDecompress",			GET_STATFNAME(STAT_AudioFullDecompressLLM),					GET_STATFNAME(STAT_AudioSummaryLLM),			ELLMTag::Audio)\
+	macro(AudioVoiceChat,						"AudioVoiceChat",				GET_STATFNAME(STAT_AudioVoiceChatLLM),						GET_STATFNAME(STAT_AudioSummaryLLM),			ELLMTag::Audio)\
 	macro(FName,								"FName",						GET_STATFNAME(STAT_FNameLLM),								GET_STATFNAME(STAT_EngineSummaryLLM),			-1)\
 	macro(Networking,							"Networking",					GET_STATFNAME(STAT_NetworkingLLM),							GET_STATFNAME(STAT_EngineSummaryLLM),			-1)\
 	macro(Meshes,								"Meshes",						GET_STATFNAME(STAT_MeshesLLM),								GET_STATFNAME(STAT_MeshesSummaryLLM),			-1)\
diff --git a/Engine/Source/Runtime/Engine/Private/AudioDecompress.cpp b/Engine/Source/Runtime/Engine/Private/AudioDecompress.cpp
--- a/Engine/Source/Runtime/Engine/Private/AudioDecompress.cpp
+++ b/Engine/Source/Runtime/Engine/Private/AudioDecompress.cpp
@@ -613,7 +613,6 @@ FAsyncAudioDecompressWorker::FAsyncAudioDecompressWorker(USoundWave* InWave, int
 
 void FAsyncAudioDecompressWorker::DoWork()
 {
-	LLM_SCOPE(ELLMTag::Audio);
 	LLM_SCOPE(ELLMTag::AudioDecompress);
 
 	if (AudioInfo)
diff --git a/Engine/Source/Runtime/Engine/Private/AudioDevice.cpp b/Engine/Source/Runtime/Engine/Private/AudioDevice.cpp
--- a/Engine/Source/Runtime/Engine/Private/AudioDevice.cpp
+++ b/Engine/Source/Runtime/Engine/Private/AudioDevice.cpp
@@ -258,7 +258,7 @@ FAudioQualitySettings FAudioDevice::GetQualityLevelSettings()
 
 bool FAudioDevice::Init(int32 InMaxChannels)
 {
-	LLM_SCOPE(ELLMTag::Audio);
+	LLM_SCOPE(ELLMTag::AudioMisc);
 
 	if (bIsInitialized)
 	{
@@ -3663,7 +3663,7 @@ void FAudioDevice::StartSources(TArray<FWaveInstance*>& WaveInstances, int32 Fir
 
 void FAudioDevice::Update(bool bGameTicking)
 {
-	LLM_SCOPE(ELLMTag::Audio);
+	LLM_SCOPE(ELLMTag::AudioMisc);
 
 	if (!IsInAudioThread())
 	{
@@ -4059,7 +4059,7 @@ void FAudioDevice::InitializePluginListeners(UWorld* World)
 
 void FAudioDevice::AddNewActiveSound(const FActiveSound& NewActiveSound)
 {
-	LLM_SCOPE(ELLMTag::Audio);
+	LLM_SCOPE(ELLMTag::AudioMisc);
 
 	if (NewActiveSound.Sound == nullptr)
 	{
@@ -5084,7 +5084,6 @@ void FAudioDevice::Flush(UWorld* WorldToFlush, bool bClearActivatedReverb)
 
 void FAudioDevice::Precache(USoundWave* SoundWave, bool bSynchronous, bool bTrackMemory, bool bForceFullDecompression)
 {
-	LLM_SCOPE(ELLMTag::Audio);
 	LLM_SCOPE(ELLMTag::AudioPrecache);
 
 	if (SoundWave == nullptr)
diff --git a/Engine/Source/Runtime/Engine/Private/AudioStreaming.cpp b/Engine/Source/Runtime/Engine/Private/AudioStreaming.cpp
--- a/Engine/Source/Runtime/Engine/Private/AudioStreaming.cpp
+++ b/Engine/Source/Runtime/Engine/Private/AudioStreaming.cpp
@@ -530,7 +530,7 @@ void FAudioStreamingManager::ProcessPendingAsyncFileResults()
 
 void FAudioStreamingManager::UpdateResourceStreaming(float DeltaTime, bool bProcessEverything /*= false*/)
 {
-	LLM_SCOPE(ELLMTag::Audio);
+	LLM_SCOPE(ELLMTag::AudioMisc);
 
 	if (CriticalSection.TryLock())
 	{
diff --git a/Engine/Source/Runtime/Engine/Private/AudioThread.cpp b/Engine/Source/Runtime/Engine/Private/AudioThread.cpp
--- a/Engine/Source/Runtime/Engine/Private/AudioThread.cpp
+++ b/Engine/Source/Runtime/Engine/Private/AudioThread.cpp
@@ -180,7 +180,7 @@ void FAudioThread::Exit()
 
 uint32 FAudioThread::Run()
 {
-	LLM_SCOPE(ELLMTag::Audio);
+	LLM_SCOPE(ELLMTag::AudioMisc);
 
 	FMemory::SetupTLSCachesOnCurrentThread();
 	FPlatformProcess::SetupAudioThread();
diff --git a/Engine/Source/Runtime/Engine/Private/SoundWave.cpp b/Engine/Source/Runtime/Engine/Private/SoundWave.cpp
--- a/Engine/Source/Runtime/Engine/Private/SoundWave.cpp
+++ b/Engine/Source/Runtime/Engine/Private/SoundWave.cpp
@@ -274,7 +274,7 @@ void USoundWave::GetAssetRegistryTags(TArray<FAssetRegistryTag>& OutTags) const
 
 void USoundWave::Serialize( FArchive& Ar )
 {
-	LLM_SCOPE(ELLMTag::Audio);
+	LLM_SCOPE(ELLMTag::AudioSoundWaves);
 
 	DECLARE_SCOPE_CYCLE_COUNTER( TEXT("USoundWave::Serialize"), STAT_SoundWave_Serialize, STATGROUP_LoadTime );
 
@@ -684,7 +684,7 @@ void USoundWave::InvalidateCompressedData()
 
 void USoundWave::PostLoad()
 {
-	LLM_SCOPE(ELLMTag::Audio);
+	LLM_SCOPE(ELLMTag::AudioSoundWaves);
 
 	Super::PostLoad();
 
diff --git a/Engine/Source/Runtime/Engine/Public/AudioDecompress.h b/Engine/Source/Runtime/Engine/Public/AudioDecompress.h
--- a/Engine/Source/Runtime/Engine/Public/AudioDecompress.h
+++ b/Engine/Source/Runtime/Engine/Public/AudioDecompress.h
@@ -367,7 +367,7 @@ public:
 
 	void DoWork()
 	{
-		LLM_SCOPE(ELLMTag::Audio);
+		LLM_SCOPE(ELLMTag::AudioMisc);
 
 		switch(TaskType)
 		{
diff --git a/Engine/Source/Runtime/UnrealAudio/Private/UnrealAudio.cpp b/Engine/Source/Runtime/UnrealAudio/Private/UnrealAudio.cpp
--- a/Engine/Source/Runtime/UnrealAudio/Private/UnrealAudio.cpp
+++ b/Engine/Source/Runtime/UnrealAudio/Private/UnrealAudio.cpp
@@ -148,7 +148,7 @@ namespace UAudio
 
 	bool FUnrealAudioModule::InitializeAudioDevice()
 	{
-		LLM_SCOPE(ELLMTag::Audio);
+		LLM_SCOPE(ELLMTag::AudioMisc);
 
 		UnrealAudioDevice = FModuleManager::LoadModulePtr<IUnrealAudioDeviceModule>(ModuleName);
 		if (UnrealAudioDevice)
@@ -293,7 +293,7 @@ namespace UAudio
 
 	uint32 FUnrealAudioModule::Run()
 	{
-		LLM_SCOPE(ELLMTag::Audio);
+		LLM_SCOPE(ELLMTag::AudioMisc);
 
 		AudioThreadChecker.CheckThread();
 
diff --git a/Engine/Source/Runtime/Windows/XAudio2/Private/XAudio2Source.cpp b/Engine/Source/Runtime/Windows/XAudio2/Private/XAudio2Source.cpp
--- a/Engine/Source/Runtime/Windows/XAudio2/Private/XAudio2Source.cpp
+++ b/Engine/Source/Runtime/Windows/XAudio2/Private/XAudio2Source.cpp
@@ -1230,7 +1230,7 @@ int32 FXAudio2SoundSource::GetDestinationVoiceIndexForEffect( SourceDestinations
 
 void FXAudio2SoundSourceCallback::OnBufferEnd(void* BufferContext)
 {
-	LLM_SCOPE(ELLMTag::Audio);
+	LLM_SCOPE(ELLMTag::AudioMisc);
 
 	if (BufferContext)
 	{
