fix for memory leak in the IOS file system

diff --git a/Engine/Source/Runtime/Core/Private/IOS/IOSPlatformFile.cpp b/Engine/Source/Runtime/Core/Private/IOS/IOSPlatformFile.cpp
--- a/Engine/Source/Runtime/Core/Private/IOS/IOSPlatformFile.cpp
+++ b/Engine/Source/Runtime/Core/Private/IOS/IOSPlatformFile.cpp
@@ -1086,9 +1086,8 @@ FString FIOSPlatformFile::ConvertToIOSPath(const FString& Filename, bool bForWri
 		}
 		else
 		{
-			FString ReadPathBase = FString([[NSBundle mainBundle] bundlePath]) + TEXT("/cookeddata/");
-            FString OutString = ReadPathBase + Result.ToLower();
-            return OutString;
+			static FString ReadPathBase = FString([[NSBundle mainBundle] bundlePath]) + TEXT("/cookeddata/");
+            return ReadPathBase + Result.ToLower();
 		}
 	}
 
