[LLM] cleanup of a few large items out of EngineMisc:
 - transient textures moved to Textures tag
 - font cache allocations moved to UI tag
 - ICU allocations moved to Localization tag

diff --git a/Engine/Source/Runtime/Core/Private/Internationalization/ICUInternationalization.cpp b/Engine/Source/Runtime/Core/Private/Internationalization/ICUInternationalization.cpp
--- a/Engine/Source/Runtime/Core/Private/Internationalization/ICUInternationalization.cpp
+++ b/Engine/Source/Runtime/Core/Private/Internationalization/ICUInternationalization.cpp
@@ -1,6 +1,7 @@
 // Copyright 1998-2019 Epic Games, Inc. All Rights Reserved.
 #include "Internationalization/ICUInternationalization.h"
 #include "HAL/FileManager.h"
+#include "HAL/LowLevelMemTracker.h"
 #include "Misc/ScopeLock.h"
 #include "Misc/Paths.h"
 #include "Internationalization/Culture.h"
@@ -37,6 +38,7 @@ namespace
 
 		static void* U_CALLCONV Malloc(const void* context, size_t size)
 		{
+			LLM_SCOPE(ELLMTag::Localization);
 			void* Result = FMemory::Malloc(size);
 #if STATS
 			BytesInUseCount += FMemory::GetAllocSize(Result);
@@ -57,6 +59,7 @@ namespace
 
 		static void* U_CALLCONV Realloc(const void* context, void* mem, size_t size)
 		{
+			LLM_SCOPE(ELLMTag::Localization);
 			return FMemory::Realloc(mem, size);
 		}
 
@@ -747,6 +750,8 @@ UDate FICUInternationalization::UEDateTimeToICUDate(const FDateTime& DateTime)
 
 UBool FICUInternationalization::OpenDataFile(const void* InContext, void** OutFileContext, void** OutContents, const char* InPath)
 {
+	LLM_SCOPE(ELLMTag::Localization);
+
 	FICUInternationalization* This = (FICUInternationalization*)InContext;
 	check(This);
 
diff --git a/Engine/Source/Runtime/Engine/Private/Texture2D.cpp b/Engine/Source/Runtime/Engine/Private/Texture2D.cpp
--- a/Engine/Source/Runtime/Engine/Private/Texture2D.cpp
+++ b/Engine/Source/Runtime/Engine/Private/Texture2D.cpp
@@ -10,6 +10,7 @@
 #include "Misc/App.h"
 #include "HAL/PlatformFilemanager.h"
 #include "HAL/FileManager.h"
+#include "HAL/LowLevelMemTracker.h"
 #include "Misc/Paths.h"
 #include "Containers/ResourceArray.h"
 #include "UObject/UObjectIterator.h"
@@ -1001,6 +1002,8 @@ bool UTexture2D::IsFullyStreamedIn()
 
 UTexture2D* UTexture2D::CreateTransient(int32 InSizeX, int32 InSizeY, EPixelFormat InFormat)
 {
+	LLM_SCOPE(ELLMTag::Textures);
+
 	UTexture2D* NewTexture = NULL;
 	if (InSizeX > 0 && InSizeY > 0 &&
 		(InSizeX % GPixelFormats[InFormat].BlockSizeX) == 0 &&
@@ -1030,7 +1033,7 @@ UTexture2D* UTexture2D::CreateTransient(int32 InSizeX, int32 InSizeY, EPixelForm
 	}
 	else
 	{
-		UE_LOG(LogTexture, Warning, TEXT("Invalid parameters specified for UTexture2D::Create()"));
+		UE_LOG(LogTexture, Warning, TEXT("Invalid parameters specified for UTexture2D::CreateTransient()"));
 	}
 	return NewTexture;
 }
diff --git a/Engine/Source/Runtime/SlateCore/Private/Fonts/FontCacheCompositeFont.cpp b/Engine/Source/Runtime/SlateCore/Private/Fonts/FontCacheCompositeFont.cpp
--- a/Engine/Source/Runtime/SlateCore/Private/Fonts/FontCacheCompositeFont.cpp
+++ b/Engine/Source/Runtime/SlateCore/Private/Fonts/FontCacheCompositeFont.cpp
@@ -6,6 +6,7 @@
 #include "Fonts/FontBulkData.h"
 #include "Misc/FileHelper.h"
 #include "Algo/BinarySearch.h"
+#include "HAL/LowLevelMemTracker.h"
 #include "Internationalization/Culture.h"
 #include "Internationalization/Internationalization.h"
 
@@ -340,6 +341,8 @@ const FFontData& FCompositeFontCache::GetFontDataForCharacter(const FSlateFontIn
 
 TSharedPtr<FFreeTypeFace> FCompositeFontCache::GetFontFace(const FFontData& InFontData)
 {
+	LLM(ELLMTag::UI);
+
 	TSharedPtr<FFreeTypeFace> FaceAndMemory = FontFaceMap.FindRef(InFontData);
 	if (!FaceAndMemory.IsValid() && InFontData.HasFont())
 	{
diff --git a/Engine/Source/Runtime/SlateCore/Private/Fonts/FontCacheFreeType.cpp b/Engine/Source/Runtime/SlateCore/Private/Fonts/FontCacheFreeType.cpp
--- a/Engine/Source/Runtime/SlateCore/Private/Fonts/FontCacheFreeType.cpp
+++ b/Engine/Source/Runtime/SlateCore/Private/Fonts/FontCacheFreeType.cpp
@@ -4,6 +4,7 @@
 #include "SlateGlobals.h"
 #include "HAL/PlatformFile.h"
 #include "HAL/PlatformFilemanager.h"
+#include "HAL/LowLevelMemTracker.h"
 
 #if WITH_FREETYPE
 
@@ -19,6 +20,7 @@ namespace FreeTypeMemory
 
 static void* Alloc(FT_Memory Memory, long Size)
 {
+	LLM(ELLMTag::UI);
 	void* Result = FMemory::Malloc(Size);
 
 #if STATS
@@ -31,6 +33,7 @@ static void* Alloc(FT_Memory Memory, long Size)
 
 static void* Realloc(FT_Memory Memory, long CurSize, long NewSize, void* Block)
 {
+	LLM(ELLMTag::UI);
 
 #if STATS
 	long DeltaNewSize = NewSize - CurSize;
diff --git a/Engine/Source/Runtime/SlateCore/Private/Fonts/FontCacheHarfBuzz.cpp b/Engine/Source/Runtime/SlateCore/Private/Fonts/FontCacheHarfBuzz.cpp
--- a/Engine/Source/Runtime/SlateCore/Private/Fonts/FontCacheHarfBuzz.cpp
+++ b/Engine/Source/Runtime/SlateCore/Private/Fonts/FontCacheHarfBuzz.cpp
@@ -4,6 +4,7 @@
 #include "Fonts/FontCache.h"
 #include "Fonts/FontCacheFreeType.h"
 #include "Fonts/SlateFontRenderer.h"
+#include "HAL/LowLevelMemTracker.h"
 
 #if WITH_HARFBUZZ
 
@@ -12,11 +13,13 @@ extern "C"
 
 void* HarfBuzzMalloc(size_t InSizeBytes)
 {
+	LLM(ELLMTag::UI);
 	return FMemory::Malloc(InSizeBytes);
 }
 
 void* HarfBuzzCalloc(size_t InNumItems, size_t InItemSizeBytes)
 {
+	LLM(ELLMTag::UI);
 	const size_t AllocSizeBytes = InNumItems * InItemSizeBytes;
 	if (AllocSizeBytes > 0)
 	{
@@ -29,6 +32,7 @@ void* HarfBuzzCalloc(size_t InNumItems, size_t InItemSizeBytes)
 
 void* HarfBuzzRealloc(void* InPtr, size_t InSizeBytes)
 {
+	LLM(ELLMTag::UI);
 	return FMemory::Realloc(InPtr, InSizeBytes);
 }
 
