fixed copyright notice

diff --git a/Engine/Source/Runtime/Core/Private/Misc/Fnv.cpp b/Engine/Source/Runtime/Core/Private/Misc/Fnv.cpp
--- a/Engine/Source/Runtime/Core/Private/Misc/Fnv.cpp
+++ b/Engine/Source/Runtime/Core/Private/Misc/Fnv.cpp
@@ -1,4 +1,4 @@
-// Copyright 1998-2018 Epic Games, Inc. All Rights Reserved.
+// Copyright 1998-2019 Epic Games, Inc. All Rights Reserved.
 
 #include "Misc/Fnv.h"
 /** Algorithm implementation from http://tools.ietf.org/html/draft-eastlake-fnv-16 */
diff --git a/Engine/Source/Runtime/Core/Public/Misc/Fnv.h b/Engine/Source/Runtime/Core/Public/Misc/Fnv.h
--- a/Engine/Source/Runtime/Core/Public/Misc/Fnv.h
+++ b/Engine/Source/Runtime/Core/Public/Misc/Fnv.h
@@ -1,4 +1,4 @@
-// Copyright 1998-2018 Epic Games, Inc. All Rights Reserved.
+// Copyright 1998-201 Epic Games, Inc. All Rights Reserved.
 
 #pragma once
 
