Fix TSAN compilation errors on Mac.

diff --git a/Engine/Source/Runtime/Core/Public/Clang/ClangPlatformCodeAnalysis.h b/Engine/Source/Runtime/Core/Public/Clang/ClangPlatformCodeAnalysis.h
--- a/Engine/Source/Runtime/Core/Public/Clang/ClangPlatformCodeAnalysis.h
+++ b/Engine/Source/Runtime/Core/Public/Clang/ClangPlatformCodeAnalysis.h
@@ -108,8 +108,14 @@
 	#define TSAN_SAFE __attribute__((no_sanitize("thread")))
 
 	// Thread-sanitiser annotation functions.
-	extern "C" void AnnotateHappensBefore(const char *f, int l, void *addr);
-	extern "C" void AnnotateHappensAfter(const char *f, int l, void *addr);
+	#ifdef __cplusplus
+	extern "C" {
+	#endif
+		void AnnotateHappensBefore(const char *f, int l, void *addr);
+		void AnnotateHappensAfter(const char *f, int l, void *addr);
+	#ifdef __cplusplus
+	}
+	#endif
 
 	// Annotate that previous load/stores occur before addr 
 	#define TSAN_BEFORE(addr) AnnotateHappensBefore(__FILE__, __LINE__, (void*)(addr))
@@ -118,7 +124,9 @@
 	#define TSAN_AFTER(addr) AnnotateHappensAfter(__FILE__, __LINE__, (void*)(addr))
 
 	// Because annotating the global bools is tiresome...
-	#include <atomic>
-	#define TSAN_ATOMIC(Type) std::atomic<Type>
+	#ifdef __cplusplus
+	    #include <atomic>
+	    #define TSAN_ATOMIC(Type) std::atomic<Type>
+	#endif
 
 #endif
