Linux: Improved FPlatformMisc::GetLoginId()

diff --git a/Engine/Source/Runtime/Core/Private/Unix/UnixPlatformMisc.cpp b/Engine/Source/Runtime/Core/Private/Unix/UnixPlatformMisc.cpp
--- a/Engine/Source/Runtime/Core/Private/Unix/UnixPlatformMisc.cpp
+++ b/Engine/Source/Runtime/Core/Private/Unix/UnixPlatformMisc.cpp
@@ -948,3 +948,8 @@ void FUnixPlatformMisc::UngrabAllInput()
 		UngrabAllInputCallback();
 	}
 }
+
+FString FUnixPlatformMisc::GetLoginId()
+{
+	return FString::Printf(TEXT("%s + %s(%d)"), *GetOperatingSystemId(), FPlatformProcess::UserName(), geteuid());
+}
diff --git a/Engine/Source/Runtime/Core/Public/Unix/UnixPlatformMisc.h b/Engine/Source/Runtime/Core/Public/Unix/UnixPlatformMisc.h
--- a/Engine/Source/Runtime/Core/Public/Unix/UnixPlatformMisc.h
+++ b/Engine/Source/Runtime/Core/Public/Unix/UnixPlatformMisc.h
@@ -132,6 +132,7 @@ struct CORE_API FUnixPlatformMisc : public FGenericPlatformMisc
 	 */
 	static bool HasOverriddenReturnCode(uint8 * OverriddenReturnCodeToUsePtr);
 	static FString GetOSVersion();
+	static FString GetLoginId();
 
 #if STATS || ENABLE_STATNAMEDEVENTS
 	static void BeginNamedEventFrame();
