Linux: make LoginId (aka old MachineId) easier to parse (UE-70622).

- Amendment to UE-70622.

diff --git a/Engine/Source/Runtime/Core/Private/Unix/UnixPlatformMisc.cpp b/Engine/Source/Runtime/Core/Private/Unix/UnixPlatformMisc.cpp
--- a/Engine/Source/Runtime/Core/Private/Unix/UnixPlatformMisc.cpp
+++ b/Engine/Source/Runtime/Core/Private/Unix/UnixPlatformMisc.cpp
@@ -951,5 +951,5 @@ void FUnixPlatformMisc::UngrabAllInput()
 
 FString FUnixPlatformMisc::GetLoginId()
 {
-	return FString::Printf(TEXT("%s + %s(%d)"), *GetOperatingSystemId(), FPlatformProcess::UserName(), geteuid());
+	return FString::Printf(TEXT("%s-%08x"), *GetOperatingSystemId(), static_cast<uint32>(geteuid()));
 }
