- Fixes for Clang issues from recent merge

diff --git a/Engine/Source/Runtime/Core/Public/Clang/ClangPlatformCompilerPreSetup.h b/Engine/Source/Runtime/Core/Public/Clang/ClangPlatformCompilerPreSetup.h
--- a/Engine/Source/Runtime/Core/Public/Clang/ClangPlatformCompilerPreSetup.h
+++ b/Engine/Source/Runtime/Core/Public/Clang/ClangPlatformCompilerPreSetup.h
@@ -40,6 +40,28 @@
 		_Pragma("clang diagnostic pop")
 #endif // DISABLE_DEPRECATION
 
+#ifndef PRAGMA_DISABLE_OVERLOADED_VIRTUAL_WARNINGS
+	#define PRAGMA_DISABLE_OVERLOADED_VIRTUAL_WARNINGS \
+		_Pragma("clang diagnostic push") \
+		_Pragma("clang diagnostic ignored \"-Woverloaded-virtual\"")
+#endif // PRAGMA_DISABLE_OVERLOADED_VIRTUAL_WARNINGS
+
+#ifndef PRAGMA_ENABLE_OVERLOADED_VIRTUAL_WARNINGS
+	#define PRAGMA_ENABLE_OVERLOADED_VIRTUAL_WARNINGS \
+		_Pragma("clang diagnostic pop")
+#endif // PRAGMA_ENABLE_OVERLOADED_VIRTUAL_WARNINGS
+
+#ifndef PRAGMA_DISABLE_MISSING_BRACES_WARNINGS
+	#define PRAGMA_DISABLE_MISSING_BRACES_WARNINGS \
+		_Pragma("clang diagnostic push") \
+		_Pragma("clang diagnostic ignored \"-Wmissing-braces\"")
+#endif // PRAGMA_DISABLE_MISSING_BRACES_WARNINGS
+
+#ifndef PRAGMA_ENABLE_MISSING_BRACES_WARNINGS
+	#define PRAGMA_ENABLE_MISSING_BRACES_WARNINGS \
+		_Pragma("clang diagnostic pop")
+#endif // PRAGMA_ENABLE_MISSING_BRACES_WARNINGS
+
 #ifndef PRAGMA_DISABLE_SHADOW_VARIABLE_WARNINGS
 	#define PRAGMA_DISABLE_SHADOW_VARIABLE_WARNINGS \
 		_Pragma("clang diagnostic push") \
diff --git a/Engine/Source/Runtime/Core/Public/Windows/WindowsPlatformCompilerPreSetup.h b/Engine/Source/Runtime/Core/Public/Windows/WindowsPlatformCompilerPreSetup.h
--- a/Engine/Source/Runtime/Core/Public/Windows/WindowsPlatformCompilerPreSetup.h
+++ b/Engine/Source/Runtime/Core/Public/Windows/WindowsPlatformCompilerPreSetup.h
@@ -11,11 +11,15 @@
 			PRAGMA_DISABLE_REORDER_WARNINGS \
 			PRAGMA_DISABLE_SHADOW_VARIABLE_WARNINGS \
 			PRAGMA_DISABLE_UNDEFINED_IDENTIFIER_WARNINGS \
-			PRAGMA_DISABLE_DEPRECATION_WARNINGS
+			PRAGMA_DISABLE_DEPRECATION_WARNINGS \
+			PRAGMA_DISABLE_OVERLOADED_VIRTUAL_WARNINGS \
+			PRAGMA_DISABLE_MISSING_BRACES_WARNINGS
 	#endif // THIRD_PARTY_INCLUDES_START
 
 	#ifndef THIRD_PARTY_INCLUDES_END
 		#define THIRD_PARTY_INCLUDES_END \
+			PRAGMA_DISABLE_MISSING_BRACES_WARNINGS \
+			PRAGMA_DISABLE_OVERLOADED_VIRTUAL_WARNINGS \
 			PRAGMA_ENABLE_DEPRECATION_WARNINGS \
 			PRAGMA_ENABLE_UNDEFINED_IDENTIFIER_WARNINGS \
 			PRAGMA_ENABLE_SHADOW_VARIABLE_WARNINGS \
diff --git a/Engine/Source/Runtime/GameplayMediaEncoder/Public/GameplayMediaEncoder.h b/Engine/Source/Runtime/GameplayMediaEncoder/Public/GameplayMediaEncoder.h
--- a/Engine/Source/Runtime/GameplayMediaEncoder/Public/GameplayMediaEncoder.h
+++ b/Engine/Source/Runtime/GameplayMediaEncoder/Public/GameplayMediaEncoder.h
@@ -113,10 +113,10 @@ private:
 	bool bProcessMediaSamples = true;
 
 	// live streaming: quality adaptation to available uplink b/w
-	TAtomic<uint32> NewVideoBitrate = 0;
+	TAtomic<uint32> NewVideoBitrate{ 0 };
 	FThreadSafeBool bChangeBitrate = false;
 	FTimespan FramerateMonitoringStart = -1;
-	TAtomic<uint32> NewVideoFramerate = 0;
+	TAtomic<uint32> NewVideoFramerate{ 0 };
 	FThreadSafeBool bChangeFramerate = false;
 };
 
diff --git a/Engine/Source/Runtime/Windows/WindowsPlatformFeatures/Private/HighlightRecorder.h b/Engine/Source/Runtime/Windows/WindowsPlatformFeatures/Private/HighlightRecorder.h
--- a/Engine/Source/Runtime/Windows/WindowsPlatformFeatures/Private/HighlightRecorder.h
+++ b/Engine/Source/Runtime/Windows/WindowsPlatformFeatures/Private/HighlightRecorder.h
@@ -87,7 +87,7 @@ private:
 	void OnMediaSample(const FGameplayMediaEncoderSample& Sample) override;
 
 private:
-	TAtomic<EState> State = EState::Stopped;
+	TAtomic<EState> State{ EState::Stopped };
 
 	TUniquePtr<FWmfMp4Writer> Mp4Writer;
 
diff --git a/Engine/Source/Runtime/Windows/WindowsPlatformFeatures/Private/WindowsVideoRecordingSystem.h b/Engine/Source/Runtime/Windows/WindowsPlatformFeatures/Private/WindowsVideoRecordingSystem.h
--- a/Engine/Source/Runtime/Windows/WindowsPlatformFeatures/Private/WindowsVideoRecordingSystem.h
+++ b/Engine/Source/Runtime/Windows/WindowsPlatformFeatures/Private/WindowsVideoRecordingSystem.h
@@ -41,7 +41,7 @@ private:
 	void NextRecording();
 	void FinalizeCallbackOnGameThread(bool bSaved, bool bAutoContinue, FString Path, bool bBroadcast);
 
-	TAtomic<EVideoRecordingState> RecordState = EVideoRecordingState::None;
+	TAtomic<EVideoRecordingState> RecordState{ EVideoRecordingState::None };
 
 	FVideoRecordingParameters Parameters;
 
