#UE4 Fix static analysis warning and made sure editor builds running as servers behave the same as monolithic servers.

diff --git a/Engine/Source/Runtime/Core/Private/Windows/WindowsPlatformMisc.cpp b/Engine/Source/Runtime/Core/Private/Windows/WindowsPlatformMisc.cpp
--- a/Engine/Source/Runtime/Core/Private/Windows/WindowsPlatformMisc.cpp
+++ b/Engine/Source/Runtime/Core/Private/Windows/WindowsPlatformMisc.cpp
@@ -1685,7 +1685,7 @@ const TCHAR* FWindowsPlatformMisc::GetPlatformFeaturesModuleName()
 {
 	bool bModuleExists = FModuleManager::Get().ModuleExists(TEXT("WindowsPlatformFeatures"));
 	// If running a dedicated server then we use the default PlatformFeatures
-	if (bModuleExists && !UE_SERVER)
+	if (bModuleExists && !IsRunningDedicatedServer())
 	{
 		UE_LOG(LogWindows, Log, TEXT("WindowsPlatformFeatures enabled"));
 		return TEXT("WindowsPlatformFeatures");
