Push/pop TRUE and FALSE macros before defining them as constants, to accommodate code that includes MinimalWindowsAPI.h after a Windows header.

diff --git a/Engine/Source/Runtime/Core/Public/Windows/MinimalWindowsApi.h b/Engine/Source/Runtime/Core/Public/Windows/MinimalWindowsApi.h
--- a/Engine/Source/Runtime/Core/Public/Windows/MinimalWindowsApi.h
+++ b/Engine/Source/Runtime/Core/Public/Windows/MinimalWindowsApi.h
@@ -18,6 +18,13 @@
 	#define MINIMAL_WINDOWS_API extern "C" __declspec(dllimport)
 #endif
 
+// Undefine Windows types that we're going to redefine. This should be done by HideWindowsPlatformTypes.h after including any system headers, 
+// but it's less friction to just undefine them here.
+#pragma push_macro("TRUE")
+#pragma push_macro("FALSE")
+#undef TRUE
+#undef FALSE
+
 // Use strongly typed handles
 #ifndef STRICT
 #define STRICT
@@ -192,3 +199,7 @@ namespace Windows
 		return QueryPerformanceCounter((LPLARGE_INTEGER)Cycles);
 	}
 }
+
+// Restore the definitions for TRUE and FALSE
+#pragma pop_macro("FALSE")
+#pragma pop_macro("TRUE")
