Remove unnecessary code
#rnx

diff --git a/Engine/Source/Runtime/Core/Public/Delegates/DelegateCombinations.h b/Engine/Source/Runtime/Core/Public/Delegates/DelegateCombinations.h
--- a/Engine/Source/Runtime/Core/Public/Delegates/DelegateCombinations.h
+++ b/Engine/Source/Runtime/Core/Public/Delegates/DelegateCombinations.h
@@ -71,7 +71,4 @@
 #define DECLARE_DYNAMIC_DELEGATE_NineParams( DelegateName, Param1Type, Param1Name, Param2Type, Param2Name, Param3Type, Param3Name, Param4Type, Param4Name, Param5Type, Param5Name, Param6Type, Param6Name, Param7Type, Param7Name, Param8Type, Param8Name, Param9Type, Param9Name ) BODY_MACRO_COMBINE(CURRENT_FILE_ID,_,__LINE__,_DELEGATE) FUNC_DECLARE_DYNAMIC_DELEGATE( FWeakObjectPtr, DelegateName, DelegateName##_DelegateWrapper, FUNC_CONCAT( Param1Type InParam1, Param2Type InParam2, Param3Type InParam3, Param4Type InParam4, Param5Type InParam5, Param6Type InParam6, Param7Type InParam7, Param8Type InParam8, Param9Type InParam9 ), FUNC_CONCAT( *this, InParam1, InParam2, InParam3, InParam4, InParam5, InParam6, InParam7, InParam8, InParam9 ), void, Param1Type, Param2Type, Param3Type, Param4Type, Param5Type, Param6Type, Param7Type, Param8Type, Param9Type )
 #define DECLARE_DYNAMIC_MULTICAST_DELEGATE_NineParams( DelegateName, Param1Type, Param1Name, Param2Type, Param2Name, Param3Type, Param3Name, Param4Type, Param4Name, Param5Type, Param5Name, Param6Type, Param6Name, Param7Type, Param7Name, Param8Type, Param8Name, Param9Type, Param9Name ) BODY_MACRO_COMBINE(CURRENT_FILE_ID,_,__LINE__,_DELEGATE) FUNC_DECLARE_DYNAMIC_MULTICAST_DELEGATE( FWeakObjectPtr, DelegateName, DelegateName##_DelegateWrapper, FUNC_CONCAT( Param1Type InParam1, Param2Type InParam2, Param3Type InParam3, Param4Type InParam4, Param5Type InParam5, Param6Type InParam6, Param7Type InParam7, Param8Type InParam8, Param9Type InParam9 ), FUNC_CONCAT( *this, InParam1, InParam2, InParam3, InParam4, InParam5, InParam6, InParam7, InParam8, InParam9 ), void, Param1Type, Param2Type, Param3Type, Param4Type, Param5Type, Param6Type, Param7Type, Param8Type, Param9Type )
 #define DECLARE_DELEGATE_RetVal_NineParams( ReturnValueType, DelegateName, Param1Type, Param2Type, Param3Type, Param4Type, Param5Type, Param6Type, Param7Type, Param8Type, Param9Type ) FUNC_DECLARE_DELEGATE( DelegateName, ReturnValueType, Param1Type, Param2Type, Param3Type, Param4Type, Param5Type, Param6Type, Param7Type, Param8Type, Param9Type )
-
-struct FWeakObjectPtr;
-
 #define DECLARE_DYNAMIC_DELEGATE_RetVal_NineParams( ReturnValueType, DelegateName, Param1Type, Param1Name, Param2Type, Param2Name, Param3Type, Param3Name, Param4Type, Param4Name, Param5Type, Param5Name, Param6Type, Param6Name, Param7Type, Param7Name, Param8Type, Param8Name, Param9Type, Param9Name ) BODY_MACRO_COMBINE(CURRENT_FILE_ID,_,__LINE__,_DELEGATE) FUNC_DECLARE_DYNAMIC_DELEGATE_RETVAL( FWeakObjectPtr, DelegateName, DelegateName##_DelegateWrapper, ReturnValueType, FUNC_CONCAT( Param1Type InParam1, Param2Type InParam2, Param3Type InParam3, Param4Type InParam4, Param5Type InParam5, Param6Type InParam6, Param7Type InParam7, Param8Type InParam8, Param9Type InParam9 ), FUNC_CONCAT( *this, InParam1, InParam2, InParam3, InParam4, InParam5, InParam6, InParam7, InParam8, InParam9 ), ReturnValueType, Param1Type, Param2Type, Param3Type, Param4Type, Param5Type, Param6Type, Param7Type, Param8Type, Param9Type )
