Push/pop TRUE and FALSE macros between PreWindowsApi.h and PostWindowsApi.h, since they are sometimes also defined by third party libraries.

#rnx

diff --git a/Engine/Source/Runtime/Core/Public/Windows/PostWindowsApi.h b/Engine/Source/Runtime/Core/Public/Windows/PostWindowsApi.h
--- a/Engine/Source/Runtime/Core/Public/Windows/PostWindowsApi.h
+++ b/Engine/Source/Runtime/Core/Public/Windows/PostWindowsApi.h
@@ -8,8 +8,6 @@ THIRD_PARTY_INCLUDES_END
 #undef UINT
 #undef DWORD
 #undef FLOAT
-#undef TRUE
-#undef FALSE
 
 // Undo any Windows defines.
 #undef uint8
@@ -69,6 +67,8 @@ THIRD_PARTY_INCLUDES_END
 #pragma pop_macro("MAX_uint32")
 #pragma pop_macro("MAX_int32")
 #pragma pop_macro("TEXT")
+#pragma pop_macro("TRUE")
+#pragma pop_macro("FALSE")
 
 // Redefine CDECL to our version of the #define.  <AJS> Is this really necessary?
 #define CDECL	    __cdecl					/* Standard C function */
diff --git a/Engine/Source/Runtime/Core/Public/Windows/PreWindowsApi.h b/Engine/Source/Runtime/Core/Public/Windows/PreWindowsApi.h
--- a/Engine/Source/Runtime/Core/Public/Windows/PreWindowsApi.h
+++ b/Engine/Source/Runtime/Core/Public/Windows/PreWindowsApi.h
@@ -9,6 +9,8 @@
 #pragma push_macro("MAX_uint32")
 #pragma push_macro("MAX_int32")
 #pragma push_macro("TEXT")
+#pragma push_macro("TRUE")
+#pragma push_macro("FALSE")
 
 // Undefine the TEXT macro for winnt.h to redefine it, unless it's already been included
 #ifndef _WINNT_
