Remove unused macros
#rnx

diff --git a/Engine/Source/Runtime/Core/Public/Delegates/Delegate.h b/Engine/Source/Runtime/Core/Public/Delegates/Delegate.h
--- a/Engine/Source/Runtime/Core/Public/Delegates/Delegate.h
+++ b/Engine/Source/Runtime/Core/Public/Delegates/Delegate.h
@@ -265,12 +265,6 @@ class DynamicMulticastDelegateName : public TBaseDynamicMulticastDelegate<TWeakP
 	};
 
 
-// Simple macro chain to concatenate code text
-#define FUNC_COMBINE_ACTUAL( A, B ) A##B
-#define FUNC_COMBINE( A, B ) FUNC_COMBINE_ACTUAL( A, B )
-
-
-
 #define ENABLE_STATIC_FUNCTION_FNAMES (PLATFORM_COMPILER_CLANG && (__clang_major__ > 3 || (__clang_major__ == 3 && __clang_minor__ >= 5)))
 
 #if ENABLE_STATIC_FUNCTION_FNAMES
@@ -432,10 +426,6 @@ namespace UE4Delegates_Private
 	class EventName : public BaseTypeEvent { friend class OwningType; };
 
 
-// Undefine temporary macros
-#undef FUNC_COMBINE_ACTUAL
-#undef FUNC_COMBINE
-
 // Simple delegate used by various utilities such as timers
 DECLARE_DELEGATE( FSimpleDelegate );
 DECLARE_MULTICAST_DELEGATE( FSimpleMulticastDelegate );
