Linux: always select best clock, avoid using CLOC_REALTIME which might be unreliable

diff --git a/Engine/Source/Runtime/Core/Private/Unix/UnixPlatformTime.cpp b/Engine/Source/Runtime/Core/Private/Unix/UnixPlatformTime.cpp
--- a/Engine/Source/Runtime/Core/Private/Unix/UnixPlatformTime.cpp
+++ b/Engine/Source/Runtime/Core/Private/Unix/UnixPlatformTime.cpp
@@ -163,9 +163,19 @@ int FUnixTime::CalibrateAndSelectClock()
 	// do not calibrate in case of programs, so e.g. ShaderCompileWorker speed is not impacted
 	if (IS_PROGRAM)
 	{
-		FCStringAnsi::Snprintf(CalibrationLog, sizeof(CalibrationLog),
-			"Skipped benchmarking clocks because the engine is running in a standalone program mode - CLOCK_REALTIME will be used.\n");
-		return CLOCK_REALTIME;
+		struct timespec ts;
+		if (clock_gettime(CLOCK_MONOTONIC, &ts) == -1)
+		{
+			FCStringAnsi::Snprintf(CalibrationLog, sizeof(CalibrationLog),
+				"Skipped benchmarking clocks because the engine is running in a standalone program mode: CLOCK_MONOTONIC is unavailable, CLOCK_REALTIME will be used.\n");
+			return CLOCK_REALTIME;
+		}
+		else
+		{
+			FCStringAnsi::Snprintf(CalibrationLog, sizeof(CalibrationLog),
+				"Skipped benchmarking clocks because the engine is running in a standalone program mode - CLOCK_MONOTONIC will be used.\n");
+			return CLOCK_MONOTONIC;
+		}
 	}
 	else
 	{
@@ -188,8 +198,10 @@ int FUnixTime::CalibrateAndSelectClock()
 			{ CLOCK_MONOTONIC_COARSE, "CLOCK_MONOTONIC_COARSE", 0 }
 		};
 
-		int ChosenClock = 0;	// REALTIME should be always supported
-		for (int Idx = 0; Idx < ARRAY_COUNT(Clocks); ++Idx)
+		int ChosenClock = 0;	
+		// Skip measuring CLOCK_REALTIME, so it's never picked up if any other is available.
+		// CLOCK_REALTIME should be always supported - as a last resort.
+		for (int Idx = 1; Idx < ARRAY_COUNT(Clocks); ++Idx)
 		{
 			Clocks[Idx].Rate = CallsPerSecondBenchmark(Clocks[Idx].Id, Clocks[Idx].Desc);
 			if (Clocks[Idx].Rate > Clocks[ChosenClock].Rate)
@@ -203,7 +215,7 @@ int FUnixTime::CalibrateAndSelectClock()
 		FCStringAnsi::Strncat(CalibrationLog, Buffer, sizeof(CalibrationLog));
 
 		// Warn if our current clock source cannot be called at least 1M times a second (<30k a frame) as this may affect tight loops
-		if (Clocks[ChosenClock].Rate < 1000000)
+		if (ChosenClock != 0 && Clocks[ChosenClock].Rate < 1000000)
 		{
 			FCStringAnsi::Snprintf(Buffer, sizeof(Buffer), "The clock source is too slow on this machine, performance may be affected.\n");
 			FCStringAnsi::Strncat(CalibrationLog, Buffer, sizeof(CalibrationLog));
