Fix *PlatformProcess::GenerateApplicationPath() always assuming the engine is in ../../Engine. This is not necessarily true for builds with a unique build environment.

diff --git a/Engine/Source/Runtime/Core/Private/Mac/MacPlatformProcess.cpp b/Engine/Source/Runtime/Core/Private/Mac/MacPlatformProcess.cpp
--- a/Engine/Source/Runtime/Core/Private/Mac/MacPlatformProcess.cpp
+++ b/Engine/Source/Runtime/Core/Private/Mac/MacPlatformProcess.cpp
@@ -90,7 +90,7 @@ FString FMacPlatformProcess::GenerateApplicationPath( const FString& AppName, EB
 	else
 	{
 		// Try expected path of an executable inside an app package in Engine Binaries
-		FString ExecutablePath = FString::Printf(TEXT("../../../Engine/Binaries/%s/%s.app/Contents/MacOS/%s"), *PlatformName, *ExecutableName, *ExecutableName);
+		FString ExecutablePath = FPaths::EngineDir() / FString::Printf(TEXT("Binaries/%s/%s.app/Contents/MacOS/%s"), *PlatformName, *ExecutableName, *ExecutableName);
 			
 		NSString* LaunchPath = ExecutablePath.GetNSString();
 		
@@ -101,7 +101,7 @@ FString FMacPlatformProcess::GenerateApplicationPath( const FString& AppName, EB
 		else
 		{
 			// Next try expected path of a simple executable file in Engine Binaries
-			ExecutablePath = FString::Printf(TEXT("../../../Engine/Binaries/%s/%s"), *PlatformName, *ExecutableName);
+			ExecutablePath = FPaths::EngineDir() / FString::Printf(TEXT("Binaries/%s/%s"), *PlatformName, *ExecutableName);
 
 			LaunchPath = ExecutablePath.GetNSString();
 
diff --git a/Engine/Source/Runtime/Core/Private/Unix/UnixPlatformProcess.cpp b/Engine/Source/Runtime/Core/Private/Unix/UnixPlatformProcess.cpp
--- a/Engine/Source/Runtime/Core/Private/Unix/UnixPlatformProcess.cpp
+++ b/Engine/Source/Runtime/Core/Private/Unix/UnixPlatformProcess.cpp
@@ -346,7 +346,7 @@ const TCHAR* FUnixPlatformProcess::ExecutableName(bool bRemoveExtension)
 FString FUnixPlatformProcess::GenerateApplicationPath( const FString& AppName, EBuildConfigurations::Type BuildConfiguration)
 {
 	FString PlatformName = FPlatformProcess::GetBinariesSubdirectory();
-	FString ExecutablePath = FString::Printf(TEXT("../../../Engine/Binaries/%s/%s"), *PlatformName, *AppName);
+	FString ExecutablePath = FPaths::EngineDir() / FString::Printf(TEXT("Binaries/%s/%s"), *PlatformName, *AppName);
 	
 	if (BuildConfiguration != EBuildConfigurations::Development)
 	{
diff --git a/Engine/Source/Runtime/Core/Private/Windows/WindowsPlatformProcess.cpp b/Engine/Source/Runtime/Core/Private/Windows/WindowsPlatformProcess.cpp
--- a/Engine/Source/Runtime/Core/Private/Windows/WindowsPlatformProcess.cpp
+++ b/Engine/Source/Runtime/Core/Private/Windows/WindowsPlatformProcess.cpp
@@ -108,7 +108,8 @@ void FWindowsPlatformProcess::FreeDllHandle( void* DllHandle )
 FString FWindowsPlatformProcess::GenerateApplicationPath( const FString& AppName, EBuildConfigurations::Type BuildConfiguration)
 {
 	FString PlatformName = GetBinariesSubdirectory();
-	FString ExecutablePath = FString::Printf(TEXT("..\\..\\..\\Engine\\Binaries\\%s\\%s"), *PlatformName, *AppName);
+	FString ExecutablePath = FPaths::EngineDir() / FString::Printf(TEXT("Binaries/%s/%s"), *PlatformName, *AppName);
+	FPaths::MakePlatformFilename(ExecutablePath);
 
 	if (BuildConfiguration != EBuildConfigurations::Development)
 	{
