Added crash report path for ReportGPUCrash.

diff --git a/Engine/Source/Runtime/Core/Private/Mac/MacErrorOutputDevice.cpp b/Engine/Source/Runtime/Core/Private/Mac/MacErrorOutputDevice.cpp
--- a/Engine/Source/Runtime/Core/Private/Mac/MacErrorOutputDevice.cpp
+++ b/Engine/Source/Runtime/Core/Private/Mac/MacErrorOutputDevice.cpp
@@ -8,6 +8,8 @@
 #include "Mac/CocoaThread.h"
 #include "HAL/ExceptionHandling.h"
 
+extern CORE_API bool GIsGPUCrashed;
+
 FMacErrorOutputDevice::FMacErrorOutputDevice()
 :	ErrorPos(0)
 {}
@@ -47,7 +49,16 @@ void FMacErrorOutputDevice::Serialize( const TCHAR* Msg, ELogVerbosity::Type Ver
 #if PLATFORM_EXCEPTIONS_DISABLED
 		UE_DEBUG_BREAK();
 #endif
-		ReportAssert(Msg, 0);
+		const int32 NumStackFramesToIgnore = 0;
+
+		if (GIsGPUCrashed)
+		{
+			ReportGPUCrash(Msg, NumStackFramesToIgnore);
+		}
+		else
+		{
+			ReportAssert(Msg, NumStackFramesToIgnore);
+		}
 	}
 	else
 	{
diff --git a/Engine/Source/Runtime/Core/Private/Mac/MacPlatformMisc.cpp b/Engine/Source/Runtime/Core/Private/Mac/MacPlatformMisc.cpp
--- a/Engine/Source/Runtime/Core/Private/Mac/MacPlatformMisc.cpp
+++ b/Engine/Source/Runtime/Core/Private/Mac/MacPlatformMisc.cpp
@@ -1479,7 +1479,8 @@ static void DefaultCrashHandler(FMacCrashContext const& Context)
 static uint32 GMacStackIgnoreDepth = 6;
 
 /** Message for the assert triggered on this thread */
-thread_local const TCHAR* GAssertErrorMessage = nullptr;
+thread_local const TCHAR* GCrashErrorMessage = nullptr;
+thread_local ECrashContextType GCrashErrorType = ECrashContextType::Crash;
 
 /** True system-specific crash handler that gets called first */
 static void PlatformCrashHandler(int32 Signal, siginfo_t* Info, void* Context)
@@ -1490,15 +1491,15 @@ static void PlatformCrashHandler(int32 Signal, siginfo_t* Info, void* Context)
 	ECrashContextType Type;
 	const TCHAR* ErrorMessage;
 
-	if (GAssertErrorMessage == nullptr)
+	if (GCrashErrorMessage == nullptr)
 	{
 		Type = ECrashContextType::Crash;
 		ErrorMessage = TEXT("Caught signal");
 	}
 	else
 	{
-		Type = ECrashContextType::Assert;
-		ErrorMessage = GAssertErrorMessage;
+		Type = GCrashErrorType;
+		ErrorMessage = GCrashErrorMessage;
 	}
 	
 	FMacCrashContext CrashContext(Type, ErrorMessage);
@@ -1908,7 +1909,15 @@ void FMacCrashContext::GenerateEnsureInfoAndLaunchReporter() const
 
 void ReportAssert(const TCHAR* ErrorMessage, int NumStackFramesToIgnore)
 {
-	GAssertErrorMessage = ErrorMessage;
+	GCrashErrorMessage = ErrorMessage;
+	GCrashErrorType = ECrashContextType::Assert;
+	FPlatformMisc::RaiseException(1);
+}
+
+void ReportGPUCrash(const TCHAR* ErrorMessage, int NumStackFramesToIgnore)
+{
+	GCrashErrorMessage = ErrorMessage;
+	GCrashErrorType = ECrashContextType::GPUCrash;
 	FPlatformMisc::RaiseException(1);
 }
 
diff --git a/Engine/Source/Runtime/Core/Private/Unix/UnixErrorOutputDevice.cpp b/Engine/Source/Runtime/Core/Private/Unix/UnixErrorOutputDevice.cpp
--- a/Engine/Source/Runtime/Core/Private/Unix/UnixErrorOutputDevice.cpp
+++ b/Engine/Source/Runtime/Core/Private/Unix/UnixErrorOutputDevice.cpp
@@ -13,6 +13,8 @@
 #include "Misc/App.h"
 #include "HAL/ExceptionHandling.h"
 
+extern CORE_API bool GIsGPUCrashed;
+
 FUnixErrorOutputDevice::FUnixErrorOutputDevice()
 :	ErrorPos(0)
 {
@@ -60,7 +62,15 @@ void FUnixErrorOutputDevice::Serialize(const TCHAR* Msg, ELogVerbosity::Type Ver
 		//     FOutputDevice::LogfImpl()
 		//     FWindowsErrorOutputDevice::Serialize()
 		const int32 NumStackFramesToIgnore = 3;
-		ReportAssert(Msg, NumStackFramesToIgnore);
+
+		if (GIsGPUCrashed)
+		{
+			ReportGPUCrash(Msg, NumStackFramesToIgnore);
+		}
+		else
+		{
+			ReportAssert(Msg, NumStackFramesToIgnore);
+		}
 	}
 	else
 	{
diff --git a/Engine/Source/Runtime/Core/Private/Unix/UnixPlatformStackWalk.cpp b/Engine/Source/Runtime/Core/Private/Unix/UnixPlatformStackWalk.cpp
--- a/Engine/Source/Runtime/Core/Private/Unix/UnixPlatformStackWalk.cpp
+++ b/Engine/Source/Runtime/Core/Private/Unix/UnixPlatformStackWalk.cpp
@@ -782,11 +782,30 @@ int32 FUnixPlatformStackWalk::GetProcessModuleSignatures(FStackWalkModuleInfo *M
 	return Signatures.Index;
 }
 
-thread_local const TCHAR* GAssertErrorMessage = nullptr;
+thread_local const TCHAR* GCrashErrorMessage = nullptr;
+thread_local ECrashContextType GCrashErrorType = ECrashContextType::Crash;
 
 void ReportAssert(const TCHAR* ErrorMessage, int NumStackFramesToIgnore)
 {
-	GAssertErrorMessage = ErrorMessage;
+	GCrashErrorMessage = ErrorMessage;
+	GCrashErrorType = ECrashContextType::Assert;
+
+	// Store NumStackFramesToIgnore in signal data	
+	sigval UserData;
+	UserData.sival_int = NumStackFramesToIgnore + 2; // +2 for this function and sigqueue()
+	sigqueue(getpid(),  SIGSEGV, UserData);
+	
+	// Make sure we never return
+	for (;;)
+	{
+		FPlatformProcess::Sleep(60.0f);
+	}
+}
+
+void ReportGPUCrash(const TCHAR* ErrorMessage, int NumStackFramesToIgnore)
+{
+	GCrashErrorMessage = ErrorMessage;
+	GCrashErrorType = ECrashContextType::GPUCrash;
 
 	// Store NumStackFramesToIgnore in signal data	
 	sigval UserData;
diff --git a/Engine/Source/Runtime/Core/Private/Windows/WindowsErrorOutputDevice.cpp b/Engine/Source/Runtime/Core/Private/Windows/WindowsErrorOutputDevice.cpp
--- a/Engine/Source/Runtime/Core/Private/Windows/WindowsErrorOutputDevice.cpp
+++ b/Engine/Source/Runtime/Core/Private/Windows/WindowsErrorOutputDevice.cpp
@@ -21,6 +21,8 @@
 #include "HAL/ExceptionHandling.h"
 #include "Windows/WindowsHWrapper.h"
 
+extern CORE_API bool GIsGPUCrashed;
+
 FWindowsErrorOutputDevice::FWindowsErrorOutputDevice()
 {
 }
@@ -65,7 +67,15 @@ void FWindowsErrorOutputDevice::Serialize( const TCHAR* Msg, ELogVerbosity::Type
 		//     FOutputDevice::LogfImpl()
 		//     FWindowsErrorOutputDevice::Serialize()
 		const int32 NumStackFramesToIgnore = 3;
-		ReportAssert(Msg, NumStackFramesToIgnore);
+
+		if (GIsGPUCrashed)
+		{
+			ReportGPUCrash(Msg, NumStackFramesToIgnore);
+		}
+		else
+		{
+			ReportAssert(Msg, NumStackFramesToIgnore);
+		}
 	}
 	else
 	{
diff --git a/Engine/Source/Runtime/Core/Private/Windows/WindowsPlatformCrashContext.cpp b/Engine/Source/Runtime/Core/Private/Windows/WindowsPlatformCrashContext.cpp
--- a/Engine/Source/Runtime/Core/Private/Windows/WindowsPlatformCrashContext.cpp
+++ b/Engine/Source/Runtime/Core/Private/Windows/WindowsPlatformCrashContext.cpp
@@ -43,6 +43,7 @@
  * Code for an assert exception
  */
 const uint32 AssertExceptionCode = 0x4000;
+const uint32 GPUCrashExceptionCode = 0x8000;
 
 /**
  * Stores information about an assert that can be unpacked in the exception handler.
@@ -549,6 +550,17 @@ FORCENOINLINE void ReportAssert(const TCHAR* ErrorMessage, int NumStackFramesToI
 	::RaiseException( AssertExceptionCode, 0, ARRAY_COUNT(Arguments), Arguments );
 }
 
+FORCENOINLINE void ReportGPUCrash(const TCHAR* ErrorMessage, int NumStackFramesToIgnore)
+{
+	/** This is the last place to gather memory stats before exception. */
+	FGenericCrashContext::CrashMemoryStats = FPlatformMemory::GetStats();
+
+	FAssertInfo Info(ErrorMessage, NumStackFramesToIgnore + 2); // +2 for this function and RaiseException()
+
+	ULONG_PTR Arguments[] = { (ULONG_PTR)&Info };
+	::RaiseException( GPUCrashExceptionCode, 0, ARRAY_COUNT(Arguments), Arguments );
+}
+
 void ReportHang(const TCHAR* ErrorMessage, const uint64* StackFrames, int32 NumStackFrames, uint32 HungThreadId)
 {
 	if (ReportCrashCallCount > 0 || FDebug::HasAsserted())
@@ -781,7 +793,7 @@ private:
 		const TCHAR* ErrorMessage = TEXT("Unhandled exception");
 		int NumStackFramesToIgnore = 0;
 
-		// If it was an assert, allow overriding the info from the exception parameters
+		// If it was an assert or GPU crash, allow overriding the info from the exception parameters
 		if (ExceptionInfo->ExceptionRecord->ExceptionCode == AssertExceptionCode && ExceptionInfo->ExceptionRecord->NumberParameters == 1)
 		{
 			const FAssertInfo& Info = *(const FAssertInfo*)ExceptionInfo->ExceptionRecord->ExceptionInformation[0];
@@ -789,6 +801,13 @@ private:
 			ErrorMessage = Info.ErrorMessage;
 			NumStackFramesToIgnore = Info.NumStackFramesToIgnore;
 		}
+		else if (ExceptionInfo->ExceptionRecord->ExceptionCode == GPUCrashExceptionCode && ExceptionInfo->ExceptionRecord->NumberParameters == 1)
+		{
+			const FAssertInfo& Info = *(const FAssertInfo*)ExceptionInfo->ExceptionRecord->ExceptionInformation[0];
+			Type = ECrashContextType::GPUCrash;
+			ErrorMessage = Info.ErrorMessage;
+			NumStackFramesToIgnore = Info.NumStackFramesToIgnore;
+		}
 
 		// Not super safe due to dynamic memory allocations, but at least enables new functionality.
 		// Introduces a new runtime crash context. Will replace all Windows related crash reporting.
diff --git a/Engine/Source/Runtime/Core/Public/HAL/ExceptionHandling.h b/Engine/Source/Runtime/Core/Public/HAL/ExceptionHandling.h
--- a/Engine/Source/Runtime/Core/Public/HAL/ExceptionHandling.h
+++ b/Engine/Source/Runtime/Core/Public/HAL/ExceptionHandling.h
@@ -25,6 +25,7 @@ extern CORE_API TCHAR MiniDumpFilenameW[1024];
 // #CrashReport: 2014-10-09 These methods are specific to windows, remove from here.
 extern CORE_API int32 ReportCrash( Windows::LPEXCEPTION_POINTERS ExceptionInfo );
 extern CORE_API void ReportAssert(const TCHAR* ErrorMessage, int NumStackFramesToIgnore);
+extern CORE_API void ReportGPUCrash(const TCHAR* ErrorMessage, int NumStackFramesToIgnore);
 extern CORE_API void ReportEnsure(const TCHAR* ErrorMessage, int NumStackFramesToIgnore);
 extern CORE_API void ReportHang(const TCHAR*, const uint64* StackFrames, int32 NumStackFrames, uint32 HungThreadId);
 #elif PLATFORM_XBOXONE
@@ -37,10 +38,12 @@ extern CORE_API void ReportEnsure(const TCHAR* ErrorMessage, int NumStackFramesT
 #include <signal.h>
 extern CORE_API int32 ReportCrash(ucontext_t *Context, int32 Signal, struct __siginfo* Info);
 extern CORE_API void ReportAssert(const TCHAR* ErrorMessage, int NumStackFramesToIgnore);
+extern CORE_API void ReportGPUCrash(const TCHAR* ErrorMessage, int NumStackFramesToIgnore);
 extern CORE_API void ReportEnsure(const TCHAR* ErrorMessage, int NumStackFramesToIgnore);
 extern CORE_API void ReportHang(const TCHAR*, const uint64* StackFrames, int32 NumStackFrames, uint32 HungThreadId);
 #elif PLATFORM_UNIX
 extern CORE_API void ReportAssert(const TCHAR* ErrorMessage, int NumStackFramesToIgnore);
+extern CORE_API void ReportGPUCrash(const TCHAR* ErrorMessage, int NumStackFramesToIgnore);
 extern CORE_API void ReportEnsure(const TCHAR* ErrorMessage, int NumStackFramesToIgnore);
 extern CORE_API void ReportHang(const TCHAR*, const uint64* StackFrames, int32 NumStackFrames, uint32 HungThreadId);
 #endif
