Add Unit Test for Timecode.
Remove 23.976 from the drop frame list.

diff --git a/Engine/Plugins/Runtime/AR/Apple/AppleARKit/Source/AppleARKit/Private/AppleARKitTimecodeProvider.cpp b/Engine/Plugins/Runtime/AR/Apple/AppleARKit/Source/AppleARKit/Private/AppleARKitTimecodeProvider.cpp
--- a/Engine/Plugins/Runtime/AR/Apple/AppleARKit/Source/AppleARKit/Private/AppleARKitTimecodeProvider.cpp
+++ b/Engine/Plugins/Runtime/AR/Apple/AppleARKit/Source/AppleARKit/Private/AppleARKitTimecodeProvider.cpp
@@ -20,7 +20,7 @@ void UAppleARKitTimecodeProvider::Shutdown(class UEngine*)
 
 FTimecode UAppleARKitTimecodeProvider::GetTimecode() const
 {
-	return FTimecode(FPlatformTime::Seconds(), FrameRate);
+	return FTimecode(FPlatformTime::Seconds(), FrameRate, FTimecode::IsDropFormatTimecodeSupported(FrameRate));
 }
 
 FFrameRate UAppleARKitTimecodeProvider::GetFrameRate() const
diff --git a/Engine/Source/Runtime/Core/Private/Tests/Misc/TimecodeTest.cpp b/Engine/Source/Runtime/Core/Private/Tests/Misc/TimecodeTest.cpp
new file mode 100644
--- /dev/null
+++ b/Engine/Source/Runtime/Core/Private/Tests/Misc/TimecodeTest.cpp
@@ -0,0 +1,204 @@
+// Copyright 1998-2019 Epic Games, Inc. All Rights Reserved.
+
+#include "CoreTypes.h"
+#include "Async/Async.h"
+#include "Async/Future.h"
+#include "Containers/UnrealString.h"
+#include "Logging/LogMacros.h"
+#include "Math/NumericLimits.h"
+#include "Misc/FrameRate.h"
+#include "Misc/Timecode.h"
+#include "Misc/Timespan.h"
+#include "Misc/AutomationTest.h"
+
+
+#if WITH_DEV_AUTOMATION_TESTS
+
+IMPLEMENT_SIMPLE_AUTOMATION_TEST(FTimecodeTest, "System.Core.Misc.Timecode", EAutomationTestFlags::ApplicationContextMask | EAutomationTestFlags::EngineFilter)
+
+/**
+ * Run a suite of timecode conversion operations to validate conversion from timecode to timespan/FrameNumber are working
+ *
+ * Drop Frame drop a frame every minute except every 10th minute
+ * 29.97fps
+ * 00:58:01:28 ; 00:58:01:29 ; 00:58:02:00 ; 00:58:02:01 (no skip)
+ * 01:00:59:28 ; 01:00:59:29 ; 01:01:00:02 ; 01:01:00:03 (every minute, we skip frame 0 and 1)
+ * 01:09:59:28 ; 01:09:59:29 ; 01:10:00:00 ; 01:10:00:01 (except every 10th minute, we include frame 0 and 1)
+ */
+bool FTimecodeTest::RunTest(const FString& Parameters)
+{
+	FFrameRate CommonFrameRates[]{
+		FFrameRate(12, 1),
+		FFrameRate(15, 1),
+		FFrameRate(24, 1),
+		FFrameRate(25, 1),
+		FFrameRate(30, 1),
+		FFrameRate(48, 1),
+		FFrameRate(48, 2), // Should give the same result as 24/1
+		FFrameRate(50, 1),
+		FFrameRate(60, 1),
+		FFrameRate(100, 1),
+		FFrameRate(120, 1),
+		FFrameRate(240, 1),
+		FFrameRate(24000, 1001),
+		FFrameRate(30000, 1001),
+		FFrameRate(48000, 1001),
+		FFrameRate(60000, 1001),
+	};
+
+	auto ConversionWithFrameRateTest = [this](const FFrameRate FrameRate)
+	{
+		const bool bIsDropFrame = FTimecode::IsDropFormatTimecodeSupported(FrameRate);
+		int32 NumberOfErrors = 0;
+		FTimecode PreviousTimecodeValue;
+
+		const int64 StartIndex = 0;
+		for (int64 FrameIndex = StartIndex; FrameIndex <= MAX_int32; ++FrameIndex)
+		{
+			const FFrameNumber FrameNumber = static_cast<int32>(FrameIndex);
+			const FTimecode TimecodeValue = FTimecode::FromFrameNumber(FrameNumber, FrameRate, bIsDropFrame);
+			bool bDoTest = true;
+
+			// Conversion from FrameNumber to Timecode
+			if (bDoTest)
+			{
+				const FFrameNumber ExpectedFrameNumber = TimecodeValue.ToFrameNumber(FrameRate);
+				if (FrameNumber != ExpectedFrameNumber)
+				{
+					AddError(FString::Printf(TEXT("Timecode '%s' didn't convert properly from FrameNumber '%d'.")
+						, *TimecodeValue.ToString()
+						, FrameNumber.Value
+						, *FrameRate.ToPrettyText().ToString()
+					));
+					bDoTest = false;
+					++NumberOfErrors;
+				}
+			}
+
+			// Conversion from Timespan to Timecode
+			if (bDoTest)
+			{
+				const FTimespan TimespanFromTimecode = TimecodeValue.ToTimespan(FrameRate);
+				const FTimecode TimecodeFromTimespan = FTimecode::FromTimespan(TimespanFromTimecode, FrameRate, bIsDropFrame);
+
+				if (TimecodeFromTimespan != TimecodeValue)
+				{
+					AddError(FString::Printf(TEXT("Timecode '%s' didn't convert properly from Timespan '%f' for frame rate '%s'.")
+						, *TimecodeValue.ToString()
+						, TimespanFromTimecode.GetTotalSeconds()
+						, *FrameRate.ToPrettyText().ToString()
+					));
+					bDoTest = false;
+					++NumberOfErrors;
+				}
+				else if (!bIsDropFrame)
+				{
+					// Do they have the same hours, minutes, seconds
+					bool bHoursAreValid = TimespanFromTimecode.GetHours() == TimecodeValue.Hours || ((TimecodeValue.Hours%24) == TimespanFromTimecode.GetHours() && (TimecodeValue.Hours/24) == TimespanFromTimecode.GetDays());
+					bool bMinutesAreValid = TimespanFromTimecode.GetMinutes() == TimecodeValue.Minutes;
+					bool bSecondsAreValid = TimespanFromTimecode.GetSeconds() == TimecodeValue.Seconds;
+					if (!bHoursAreValid || !bMinutesAreValid || !bSecondsAreValid)
+					{
+						AddError(FString::Printf(TEXT("Timecode '%s' hours/minutes/seconds doesn't matches with Timespan '%s' from frame rate '%s'.")
+							, *TimecodeValue.ToString()
+							, *TimespanFromTimecode.ToString()
+							, *FrameRate.ToPrettyText().ToString()
+						));
+						bDoTest = false;
+						++NumberOfErrors;
+					}
+				}
+			}
+
+			// Test if the frame number is incrementing
+			bool bIsPreviousTimecodeValid = FrameIndex != StartIndex;
+			if (bDoTest && bIsPreviousTimecodeValid)
+			{
+				bool bWrongFrame = PreviousTimecodeValue.Frames + 1 != TimecodeValue.Frames
+					&& TimecodeValue.Frames != 0;
+				const bool bWrongSeconds = PreviousTimecodeValue.Seconds != TimecodeValue.Seconds
+					&& PreviousTimecodeValue.Seconds + 1 != TimecodeValue.Seconds
+					&& TimecodeValue.Seconds != 0;
+				const bool bWrongMinutes = PreviousTimecodeValue.Minutes != TimecodeValue.Minutes
+					&& PreviousTimecodeValue.Minutes + 1 != TimecodeValue.Minutes
+					&& TimecodeValue.Minutes != 0;
+
+				if (bWrongFrame && bIsDropFrame)
+				{
+					// if new minute but not multiple of 10 mins, 2|4 is expected
+					const int32 NumberOfFramesInSecond = FMath::CeilToInt(FrameRate.AsDecimal());
+					const int32 NumberOfTimecodesToDrop = NumberOfFramesInSecond <= 30 ? 2 : 4;
+					bWrongFrame = !(TimecodeValue.Frames == NumberOfTimecodesToDrop && PreviousTimecodeValue.Minutes + 1 == TimecodeValue.Minutes && TimecodeValue.Minutes % 10 != 0);
+				}
+
+				if (bWrongFrame || bWrongSeconds || bWrongMinutes)
+				{
+					AddError(FString::Printf(TEXT("Timecode '%s' is not a continuity of the previous timecode '%s' from frame rate '%s'.")
+						, *TimecodeValue.ToString()
+						, *PreviousTimecodeValue.ToString()
+						, *FrameRate.ToPrettyText().ToString()
+					));
+					bDoTest = false;
+					++NumberOfErrors;
+				}
+			}
+
+			// Test frame rate that should be equivalent
+			if (bDoTest)
+			{
+				const FFrameRate EquivalentFrameRate = FFrameRate(FrameRate.Numerator * 3, FrameRate.Denominator * 3);
+				const FTimecode EquivalentTimecodeValue = FTimecode::FromFrameNumber(FrameNumber, EquivalentFrameRate, bIsDropFrame);
+				if (TimecodeValue != EquivalentTimecodeValue)
+				{
+					AddError(FString::Printf(TEXT("Timecode '%s' didn't convert properly from FrameNumber '%d' when the frame rate is tripled.")
+						, *TimecodeValue.ToString()
+						, FrameNumber.Value
+						, *FrameRate.ToPrettyText().ToString()
+					));
+					bDoTest = false;
+					++NumberOfErrors;
+				}
+			}
+
+			// If we have a lot of errors with this frame rate, there is no need to log them all.
+			if (NumberOfErrors > 10)
+			{
+				AddWarning(FString::Printf(TEXT("Skip test for frame rate '%s'. Other errors may exists.")
+					, *FrameRate.ToPrettyText().ToString()
+				));
+				break;
+			}
+
+			PreviousTimecodeValue = TimecodeValue;
+
+			// LTC timecode support up to 40 hours
+			if (TimecodeValue.Hours >= 40)
+			{
+				break;
+			}
+		}
+
+		AddInfo(FString::Printf(TEXT("Timecode test was completed with frame rate '%s'"), *FrameRate.ToPrettyText().ToString()));
+
+		return NumberOfErrors == 0;
+	};
+
+	// Test the conversion for all common frame rate
+	TArray<TFuture<bool>> Futures;
+	for (const FFrameRate FrameRate : CommonFrameRates)
+	{
+		Futures.Add(Async<bool>(EAsyncExecution::Thread, [FrameRate, &ConversionWithFrameRateTest](){ return ConversionWithFrameRateTest(FrameRate); }));
+	}
+
+	bool bSuccessfully = true;
+	for (const TFuture<bool>& Future : Futures)
+	{
+		Future.Wait();
+		bSuccessfully == bSuccessfully && Future.Get();
+	}
+
+	return bSuccessfully;
+}
+
+
+#endif //WITH_DEV_AUTOMATION_TESTS
diff --git a/Engine/Source/Runtime/Core/Public/Misc/Timecode.h b/Engine/Source/Runtime/Core/Public/Misc/Timecode.h
--- a/Engine/Source/Runtime/Core/Public/Misc/Timecode.h
+++ b/Engine/Source/Runtime/Core/Public/Misc/Timecode.h
@@ -6,6 +6,7 @@
 #include "GenericPlatform/GenericPlatformMath.h"
 #include "Misc/FrameNumber.h"
 #include "Misc/FrameRate.h"
+#include "Misc/Timespan.h"
 
 /**
  * A timecode that stores time in HH:MM:SS format with the remainder
@@ -14,8 +15,8 @@
 struct FTimecode
 {
 	/**
-	* Default construction for UObject purposes
-	*/
+	 * Default construction for UObject purposes
+	 */
 	FTimecode()
 		: Hours(0)
 		, Minutes(0)
@@ -25,12 +26,12 @@ struct FTimecode
 	{}
 
 	/**
-	* User construction from a number of hours minutes seconds and frames.
-	* @param InbDropFrame - If true, this Timecode represents a "Drop Frame Timecode" format which
+	 * User construction from a number of hours minutes seconds and frames.
+	 * @param InbDropFrame - If true, this Timecode represents a "Drop Frame Timecode" format which
 							skips the first frames of every minute (except those ending in multiples of 10)
 							to account for drift when using a fractional NTSC framerate.
-	*/
-	FTimecode(int32 InHours, int32 InMinutes, int32 InSeconds, int32 InFrames, bool InbDropFrame)
+ 	 */
+	explicit FTimecode(int32 InHours, int32 InMinutes, int32 InSeconds, int32 InFrames, bool InbDropFrame)
 		: Hours(InHours)
 		, Minutes(InMinutes)
 		, Seconds(InSeconds)
@@ -38,21 +39,21 @@ struct FTimecode
 		, bDropFrameFormat(InbDropFrame)
 	{}
 
-	FTimecode(double InSeconds, const FFrameRate& FrameRate)
-		: bDropFrameFormat(false)
+	/**
+	 * User construction from a time in seconds
+	 * @param InbDropFrame - If true, this Timecode represents a "Drop Frame Timecode" format which
+							skips the first frames of every minute (except those ending in multiples of 10)
+							to account for drift when using a fractional NTSC framerate.
+	 */
+	explicit FTimecode(double InSeconds, const FFrameRate& InFrameRate, bool InbDropFrame)
 	{
-		int32 WholeSeconds = InSeconds;
-		const float FractionalSeconds = InSeconds - WholeSeconds;
-
-		Frames = FMath::FloorToInt(FrameRate.Numerator * FractionalSeconds);
-
-		Seconds = WholeSeconds % 60;
-		WholeSeconds /= 60;
-		Minutes = WholeSeconds % 60;
-		WholeSeconds /= 60;
-		Hours = WholeSeconds % 24;
+		int32 NumberOfFrames = InbDropFrame ? (int32)FMath::RoundToInt(InSeconds * InFrameRate.AsDecimal()) : (int32)FMath::RoundToZero(InSeconds * FMath::RoundToInt(InFrameRate.AsDecimal()));
+		*this = FromFrameNumber(FFrameNumber(NumberOfFrames), InFrameRate, InbDropFrame);
 	}
 
+	friend bool operator==(const FFrameRate& A, const FFrameRate& B);
+	friend bool operator!=(const FFrameRate& A, const FFrameRate& B);
+
 public:
 
 	/**
@@ -179,21 +180,48 @@ public:
 		}
 	}
 
+	/**
+	 * Converts this Timecode back into a timespan at the given framerate, taking into account if this is a drop-frame format timecode.
+	 */
+	FTimespan ToTimespan(const FFrameRate& InFrameRate) const
+	{
+		const FFrameNumber ConvertedFrameNumber = ToFrameNumber(InFrameRate);
+		const double NumberOfSeconds = bDropFrameFormat
+				? ConvertedFrameNumber.Value * InFrameRate.AsInterval()
+				: (double)ConvertedFrameNumber.Value / FMath::RoundToInt(InFrameRate.AsDecimal());
+		return FTimespan::FromSeconds(NumberOfSeconds);
+	}
+
+	/**
+	 * Create a FTimecode from a timespan at the given frame rate. Optionally supports creating a drop frame timecode,
+	 * which drops certain timecode display numbers to help account for NTSC frame rates which are fractional.
+	 *
+	 * @param InFrameNumber - The timespan to convert into a timecode.
+	 * @param InFrameRate	- The framerate that this timecode is based in. This should be the playback framerate as it is used to determine
+	 *						  when the Frame value wraps over.
+	 * @param InbDropFrame	- If true, the returned timecode will drop the first two frames on every minute (except when Minute % 10 == 0)
+	 *						  This is only valid for NTSC framerates (29.97, 59.94) and will assert if you try to create a drop-frame format
+	 *						  from a non-valid framerate. All framerates can be represented when in non-drop frame format.
+	 */
+	static FTimecode FromTimespan(const FTimespan& InTimespan, const FFrameRate& InFrameRate, bool InbDropFrame)
+	{
+		return FTimecode(InTimespan.GetTotalSeconds(), InFrameRate, InbDropFrame);
+	}
+
 	static bool IsDropFormatTimecodeSupported(const FFrameRate& InFrameRate)
 	{
-		// Drop Format Timecode is only valid for NTSC 23.976, 29.97, and 59.94.
-		const FFrameRate TwentyThreeNineSevenSix = FFrameRate(24000, 1001);
+		// Drop Format Timecode is only valid for 29.97 and 59.94.
 		const FFrameRate TwentyNineNineSeven = FFrameRate(30000, 1001);
 		const FFrameRate FiftyNineNineFour = FFrameRate(60000, 1001);
 
-		return InFrameRate == TwentyThreeNineSevenSix || InFrameRate == TwentyNineNineSeven || InFrameRate == FiftyNineNineFour;
+		return InFrameRate == TwentyNineNineSeven || InFrameRate == FiftyNineNineFour;
 	}
 
 	/**
-	* Get the Qualified Timecode formatted in HH:MM:SS:FF or HH;MM;SS;FF depending on if this represents drop-frame timecode or not.
-	* @param bForceSignDisplay - Forces the timecode to be prepended with a positive or negative sign.
-	Standard behavior is to only show the sign when the value is negative.
-	*/
+	 * Get the Qualified Timecode formatted in HH:MM:SS:FF or HH;MM;SS;FF depending on if this represents drop-frame timecode or not.
+	 * @param bForceSignDisplay - Forces the timecode to be prepended with a positive or negative sign.
+								  Standard behavior is to only show the sign when the value is negative.
+	 */
 	FString ToString(bool bForceSignDisplay = false) const
 	{
 		bool bHasNegativeComponent = Hours < 0 || Minutes < 0 || Seconds < 0 || Frames < 0;
@@ -235,3 +263,13 @@ public:
 	/** If true, this Timecode represents a Drop Frame timecode used to account for fractional frame rates in NTSC play rates. */
 	bool bDropFrameFormat;
 };
+
+inline bool operator==(const FTimecode& A, const FTimecode& B)
+{
+	return A.Hours == B.Hours && A.Minutes == B.Minutes && A.Seconds == B.Seconds && A.Frames == A.Frames;
+}
+
+inline bool operator!=(const FTimecode& A, const FTimecode& B)
+{
+	return !(A == B);
+}
diff --git a/Engine/Source/Runtime/Engine/Private/Engine/SystemTimeTimecodeProvider.cpp b/Engine/Source/Runtime/Engine/Private/Engine/SystemTimeTimecodeProvider.cpp
--- a/Engine/Source/Runtime/Engine/Private/Engine/SystemTimeTimecodeProvider.cpp
+++ b/Engine/Source/Runtime/Engine/Private/Engine/SystemTimeTimecodeProvider.cpp
@@ -10,8 +10,6 @@ FTimecode USystemTimeTimecodeProvider::GetTimecode() const
 {
 	const FDateTime DateTime = FDateTime::Now();
 	const FTimespan Timespan = DateTime.GetTimeOfDay();
-	const double TotalSeconds = Timespan.GetTotalSeconds();
-	FFrameNumber FrameNumber = FrameRate.AsFrameNumber(TotalSeconds);
 
-	return FTimecode::FromFrameNumber(FrameNumber, FrameRate, FTimecode::IsDropFormatTimecodeSupported(FrameRate));
+	return FTimecode::FromTimespan(Timespan, FrameRate, FTimecode::IsDropFormatTimecodeSupported(FrameRate));
 }
