Missing file from CL 5296751.

diff --git a/Engine/Source/Runtime/Core/Private/Unix/UnixPlatformCrashContext.cpp b/Engine/Source/Runtime/Core/Private/Unix/UnixPlatformCrashContext.cpp
--- a/Engine/Source/Runtime/Core/Private/Unix/UnixPlatformCrashContext.cpp
+++ b/Engine/Source/Runtime/Core/Private/Unix/UnixPlatformCrashContext.cpp
@@ -671,7 +671,8 @@ void (* GCrashHandlerPointer)(const FGenericCrashContext & Context) = NULL;
 
 extern int32 CORE_API GMaxNumberFileMappingCache;
 
-extern thread_local const TCHAR* GAssertErrorMessage;
+extern thread_local const TCHAR* GCrashErrorMessage;
+extern thread_local ECrashContextType GCrashErrorType;
 
 /** True system-specific crash handler that gets called first */
 void PlatformCrashHandler(int32 Signal, siginfo_t* Info, void* Context)
@@ -690,15 +691,15 @@ void PlatformCrashHandler(int32 Signal, siginfo_t* Info, void* Context)
 	ECrashContextType Type;
 	const TCHAR* ErrorMessage;
 
-	if (GAssertErrorMessage == nullptr)
+	if (GCrashErrorMessage == nullptr)
 	{
 		Type = ECrashContextType::Crash;
 		ErrorMessage = TEXT("Caught signal");
 	}
 	else
 	{
-		Type = ECrashContextType::Assert;
-		ErrorMessage = GAssertErrorMessage;
+		Type = GCrashErrorType;
+		ErrorMessage = GCrashErrorMessage;
 	}
 
 	FUnixCrashContext CrashContext(Type, ErrorMessage);
