[LLM] current Meshes tag replaced with VertexData tag. New Meshes parent tag added for Skeletal, Static, Instanced and Landscape tags.

diff --git a/Engine/Source/Runtime/Core/Private/HAL/LowLevelMemTracker.cpp b/Engine/Source/Runtime/Core/Private/HAL/LowLevelMemTracker.cpp
--- a/Engine/Source/Runtime/Core/Private/HAL/LowLevelMemTracker.cpp
+++ b/Engine/Source/Runtime/Core/Private/HAL/LowLevelMemTracker.cpp
@@ -109,6 +109,7 @@ DECLARE_LLM_MEMORY_STAT(TEXT("InitUObject"), STAT_InitUObjectLLM, STATGROUP_LLMF
 DECLARE_LLM_MEMORY_STAT(TEXT("VideoRecording"), STAT_VideoRecordingLLM, STATGROUP_LLMFULL);
 DECLARE_LLM_MEMORY_STAT(TEXT("CsvProfiler"), STAT_CsvProfilerLLM, STATGROUP_LLMFULL);
 DECLARE_LLM_MEMORY_STAT(TEXT("MaterialInstance"), STAT_MaterialInstanceLLM, STATGROUP_LLMFULL);
+DECLARE_LLM_MEMORY_STAT(TEXT("VertexData"), STAT_VertexDataLLM, STATGROUP_LLMFULL);
 DECLARE_LLM_MEMORY_STAT(TEXT("SkeletalMesh"), STAT_SkeletalMeshLLM, STATGROUP_LLMFULL);
 DECLARE_LLM_MEMORY_STAT(TEXT("InstancedMesh"), STAT_InstancedMeshLLM, STATGROUP_LLMFULL);
 DECLARE_LLM_MEMORY_STAT(TEXT("Landscape"), STAT_LandscapeLLM, STATGROUP_LLMFULL);
diff --git a/Engine/Source/Runtime/Core/Public/HAL/LowLevelMemTracker.h b/Engine/Source/Runtime/Core/Public/HAL/LowLevelMemTracker.h
--- a/Engine/Source/Runtime/Core/Public/HAL/LowLevelMemTracker.h
+++ b/Engine/Source/Runtime/Core/Public/HAL/LowLevelMemTracker.h
@@ -212,7 +212,7 @@ enum class ELLMTagSet : uint8
 	macro(AsyncLoading,							"AsyncLoading",					GET_STATFNAME(STAT_AsyncLoadingLLM),						GET_STATFNAME(STAT_EngineSummaryLLM),			-1)\
 	macro(UObject,								"UObject",						GET_STATFNAME(STAT_UObjectLLM),								GET_STATFNAME(STAT_UObjectSummaryLLM),			-1)\
 	macro(Animation,							"Animation",					GET_STATFNAME(STAT_AnimationLLM),							GET_STATFNAME(STAT_AnimationSummaryLLM),		-1)\
-	macro(StaticMesh,							"StaticMesh",					GET_STATFNAME(STAT_StaticMeshLLM),							GET_STATFNAME(STAT_StaticMeshSummaryLLM),		-1)\
+	macro(StaticMesh,							"StaticMesh",					GET_STATFNAME(STAT_StaticMeshLLM),							GET_STATFNAME(STAT_StaticMeshSummaryLLM),		ELLMTag::Meshes)\
 	macro(Materials,							"Materials",					GET_STATFNAME(STAT_MaterialsLLM),							GET_STATFNAME(STAT_MaterialsSummaryLLM),		-1)\
 	macro(Particles,							"Particles",					GET_STATFNAME(STAT_ParticlesLLM),							GET_STATFNAME(STAT_ParticlesSummaryLLM),		-1)\
 	macro(GC,									"GC",							GET_STATFNAME(STAT_GCLLM),									GET_STATFNAME(STAT_EngineSummaryLLM),			-1)\
@@ -234,9 +234,10 @@ enum class ELLMTagSet : uint8
 	macro(InitUObject,							"InitUObject",					GET_STATFNAME(STAT_InitUObjectLLM),							GET_STATFNAME(STAT_EngineSummaryLLM),			-1)\
 	macro(VideoRecording,						"VideoRecording",				GET_STATFNAME(STAT_VideoRecordingLLM),						GET_STATFNAME(STAT_EngineSummaryLLM),			-1)\
 	macro(MaterialInstance,						"MaterialInstance",				GET_STATFNAME(STAT_MaterialInstanceLLM),					GET_STATFNAME(STAT_EngineSummaryLLM),			-1)\
-	macro(SkeletalMesh,							"SkeletalMesh",					GET_STATFNAME(STAT_SkeletalMeshLLM),						GET_STATFNAME(STAT_EngineSummaryLLM),			-1)\
-	macro(InstancedMesh,						"InstancedMesh",				GET_STATFNAME(STAT_InstancedMeshLLM),						GET_STATFNAME(STAT_EngineSummaryLLM),			-1)\
-	macro(Landscape,							"Landscape",					GET_STATFNAME(STAT_LandscapeLLM),							GET_STATFNAME(STAT_EngineSummaryLLM),			-1)\
+	macro(VertexData,							"VertexData",					GET_STATFNAME(STAT_VertexDataLLM),							GET_STATFNAME(STAT_EngineSummaryLLM),			-1)\
+	macro(SkeletalMesh,							"SkeletalMesh",					GET_STATFNAME(STAT_SkeletalMeshLLM),						GET_STATFNAME(STAT_EngineSummaryLLM),			ELLMTag::Meshes)\
+	macro(InstancedMesh,						"InstancedMesh",				GET_STATFNAME(STAT_InstancedMeshLLM),						GET_STATFNAME(STAT_EngineSummaryLLM),			ELLMTag::Meshes)\
+	macro(Landscape,							"Landscape",					GET_STATFNAME(STAT_LandscapeLLM),							GET_STATFNAME(STAT_EngineSummaryLLM),			ELLMTag::Meshes)\
 	macro(CsvProfiler,							"CsvProfiler",					GET_STATFNAME(STAT_CsvProfilerLLM),							GET_STATFNAME(STAT_EngineSummaryLLM),			-1)
 
 /*
diff --git a/Engine/Source/Runtime/D3D12RHI/Private/D3D12Stats.cpp b/Engine/Source/Runtime/D3D12RHI/Private/D3D12Stats.cpp
--- a/Engine/Source/Runtime/D3D12RHI/Private/D3D12Stats.cpp
+++ b/Engine/Source/Runtime/D3D12RHI/Private/D3D12Stats.cpp
@@ -340,7 +340,7 @@ void UpdateBufferStats(FD3D12ResourceLocation* ResourceLocation, bool bAllocatin
 		// released and the tracking could be re-enabled in the buddy allocator.
 		// The problem is that the releasing of resources happens in a generic way
 		// (see FD3D12ResourceLocation) 
-		LLM_SCOPED_PAUSE_TRACKING_WITH_ENUM_AND_AMOUNT(ELLMTag::Meshes, RequestedSize, ELLMTracker::Default, ELLMAllocType::None);
+		LLM_SCOPED_PAUSE_TRACKING_WITH_ENUM_AND_AMOUNT(ELLMTag::VertexData, RequestedSize, ELLMTracker::Default, ELLMAllocType::None);
 		LLM_SCOPED_PAUSE_TRACKING_WITH_ENUM_AND_AMOUNT(ELLMTag::GraphicsPlatform, RequestedSize, ELLMTracker::Platform, ELLMAllocType::None);
 #endif
 	}
@@ -365,7 +365,7 @@ void UpdateBufferStats(FD3D12ResourceLocation* ResourceLocation, bool bAllocatin
 
 #if PLATFORM_WINDOWS
 		// this is a work-around on Windows. See comment above.
-		LLM_SCOPED_PAUSE_TRACKING_WITH_ENUM_AND_AMOUNT(ELLMTag::Meshes, -(int64)RequestedSize, ELLMTracker::Default, ELLMAllocType::None);
+		LLM_SCOPED_PAUSE_TRACKING_WITH_ENUM_AND_AMOUNT(ELLMTag::VertexData, -(int64)RequestedSize, ELLMTracker::Default, ELLMAllocType::None);
 		LLM_SCOPED_PAUSE_TRACKING_WITH_ENUM_AND_AMOUNT(ELLMTag::GraphicsPlatform, -(int64)RequestedSize, ELLMTracker::Platform, ELLMAllocType::None);
 #endif
 	}
diff --git a/Engine/Source/Runtime/RHI/Public/RHICommandList.h b/Engine/Source/Runtime/RHI/Public/RHICommandList.h
--- a/Engine/Source/Runtime/RHI/Public/RHICommandList.h
+++ b/Engine/Source/Runtime/RHI/Public/RHICommandList.h
@@ -3624,19 +3624,19 @@ public:
 	
 	FORCEINLINE FIndexBufferRHIRef CreateAndLockIndexBuffer(uint32 Stride, uint32 Size, uint32 InUsage, FRHIResourceCreateInfo& CreateInfo, void*& OutDataBuffer)
 	{
-		LLM_SCOPE(ELLMTag::Meshes);
+		LLM_SCOPE(ELLMTag::VertexData);
 		return GDynamicRHI->CreateAndLockIndexBuffer_RenderThread(*this, Stride, Size, InUsage, CreateInfo, OutDataBuffer);
 	}
 	
 	FORCEINLINE FIndexBufferRHIRef CreateIndexBuffer(uint32 Stride, uint32 Size, uint32 InUsage, FRHIResourceCreateInfo& CreateInfo)
 	{
-		LLM_SCOPE(ELLMTag::Meshes);
+		LLM_SCOPE(ELLMTag::VertexData);
 		return GDynamicRHI->CreateIndexBuffer_RenderThread(*this, Stride, Size, InUsage, CreateInfo);
 	}
 	
 	FORCEINLINE void* LockIndexBuffer(FIndexBufferRHIParamRef IndexBuffer, uint32 Offset, uint32 SizeRHI, EResourceLockMode LockMode)
 	{
-		LLM_SCOPE(ELLMTag::Meshes);
+		LLM_SCOPE(ELLMTag::VertexData);
 		return GDynamicRHI->LockIndexBuffer_RenderThread(*this, IndexBuffer, Offset, SizeRHI, LockMode);
 	}
 	
@@ -3647,7 +3647,7 @@ public:
 	
 	FORCEINLINE void* LockStagingBuffer(FStagingBufferRHIParamRef StagingBuffer, uint32 Offset, uint32 SizeRHI)
 	{
-		LLM_SCOPE(ELLMTag::Meshes);
+		LLM_SCOPE(ELLMTag::VertexData);
 		return GDynamicRHI->LockStagingBuffer_RenderThread(*this, StagingBuffer, Offset, SizeRHI);
 	}
 	
@@ -3658,19 +3658,19 @@ public:
 	
 	FORCEINLINE FVertexBufferRHIRef CreateAndLockVertexBuffer(uint32 Size, uint32 InUsage, FRHIResourceCreateInfo& CreateInfo, void*& OutDataBuffer)
 	{
-		LLM_SCOPE(ELLMTag::Meshes);
+		LLM_SCOPE(ELLMTag::VertexData);
 		return GDynamicRHI->CreateAndLockVertexBuffer_RenderThread(*this, Size, InUsage, CreateInfo, OutDataBuffer);
 	}
 
 	FORCEINLINE FVertexBufferRHIRef CreateVertexBuffer(uint32 Size, uint32 InUsage, FRHIResourceCreateInfo& CreateInfo)
 	{
-		LLM_SCOPE(ELLMTag::Meshes);
+		LLM_SCOPE(ELLMTag::VertexData);
 		return GDynamicRHI->CreateVertexBuffer_RenderThread(*this, Size, InUsage, CreateInfo);
 	}
 	
 	FORCEINLINE void* LockVertexBuffer(FVertexBufferRHIParamRef VertexBuffer, uint32 Offset, uint32 SizeRHI, EResourceLockMode LockMode)
 	{
-		LLM_SCOPE(ELLMTag::Meshes);
+		LLM_SCOPE(ELLMTag::VertexData);
 		return GDynamicRHI->LockVertexBuffer_RenderThread(*this, VertexBuffer, Offset, SizeRHI, LockMode);
 	}
 	
diff --git a/Engine/Source/Runtime/Windows/D3D11RHI/Private/D3D11RHI.cpp b/Engine/Source/Runtime/Windows/D3D11RHI/Private/D3D11RHI.cpp
--- a/Engine/Source/Runtime/Windows/D3D11RHI/Private/D3D11RHI.cpp
+++ b/Engine/Source/Runtime/Windows/D3D11RHI/Private/D3D11RHI.cpp
@@ -777,7 +777,7 @@ void UpdateBufferStats(TRefCountPtr<ID3D11Buffer> Buffer, bool bAllocating)
 		// this is a work-around on Windows. Due to the fact that there is no way
 		// to hook the actual d3d allocations we can't track the memory in the normal way.
 		// Instead we simply tell LLM the size of these resources.
-		LLM_SCOPED_PAUSE_TRACKING_WITH_ENUM_AND_AMOUNT(ELLMTag::Meshes, Desc.ByteWidth, ELLMTracker::Default, ELLMAllocType::None);
+		LLM_SCOPED_PAUSE_TRACKING_WITH_ENUM_AND_AMOUNT(ELLMTag::VertexData, Desc.ByteWidth, ELLMTracker::Default, ELLMAllocType::None);
 		LLM_SCOPED_PAUSE_TRACKING_WITH_ENUM_AND_AMOUNT(ELLMTag::GraphicsPlatform, Desc.ByteWidth, ELLMTracker::Platform, ELLMAllocType::None);
 #endif
 	}
@@ -804,7 +804,7 @@ void UpdateBufferStats(TRefCountPtr<ID3D11Buffer> Buffer, bool bAllocating)
 		// this is a work-around on Windows. Due to the fact that there is no way
 		// to hook the actual d3d allocations we can't track the memory in the normal way.
 		// Instead we simply tell LLM the size of these resources.
-		LLM_SCOPED_PAUSE_TRACKING_WITH_ENUM_AND_AMOUNT(ELLMTag::Meshes, -(int64)Desc.ByteWidth, ELLMTracker::Default, ELLMAllocType::None);
+		LLM_SCOPED_PAUSE_TRACKING_WITH_ENUM_AND_AMOUNT(ELLMTag::VertexData, -(int64)Desc.ByteWidth, ELLMTracker::Default, ELLMAllocType::None);
 		LLM_SCOPED_PAUSE_TRACKING_WITH_ENUM_AND_AMOUNT(ELLMTag::GraphicsPlatform, -(int64)Desc.ByteWidth, ELLMTracker::Platform, ELLMAllocType::None);
 #endif
 	}
