Added "Hang" crash context type sent by desktop platforms using the ReportHang path.

diff --git a/Engine/Source/Runtime/Core/Private/GenericPlatform/GenericPlatformCrashContext.cpp b/Engine/Source/Runtime/Core/Private/GenericPlatform/GenericPlatformCrashContext.cpp
--- a/Engine/Source/Runtime/Core/Private/GenericPlatform/GenericPlatformCrashContext.cpp
+++ b/Engine/Source/Runtime/Core/Private/GenericPlatform/GenericPlatformCrashContext.cpp
@@ -55,6 +55,7 @@ const FString FGenericCrashContext::CrashTypeCrash = TEXT("Crash");
 const FString FGenericCrashContext::CrashTypeAssert = TEXT("Assert");
 const FString FGenericCrashContext::CrashTypeEnsure = TEXT("Ensure");
 const FString FGenericCrashContext::CrashTypeGPU = TEXT("GPUCrash");
+const FString FGenericCrashContext::CrashTypeHang = TEXT("Hang");
 
 const FString FGenericCrashContext::EngineModeExUnknown = TEXT("Unset");
 const FString FGenericCrashContext::EngineModeExDirty = TEXT("Dirty");
@@ -604,6 +605,8 @@ const TCHAR* FGenericCrashContext::GetCrashTypeString(ECrashContextType Type)
 {
 	switch (Type)
 	{
+	case ECrashContextType::Hang:
+		return *CrashTypeHang;
 	case ECrashContextType::GPUCrash:
 		return *CrashTypeGPU;
 	case ECrashContextType::Ensure:
diff --git a/Engine/Source/Runtime/Core/Private/Mac/MacPlatformMisc.cpp b/Engine/Source/Runtime/Core/Private/Mac/MacPlatformMisc.cpp
--- a/Engine/Source/Runtime/Core/Private/Mac/MacPlatformMisc.cpp
+++ b/Engine/Source/Runtime/Core/Private/Mac/MacPlatformMisc.cpp
@@ -1960,7 +1960,7 @@ void ReportHang(const TCHAR* ErrorMessage, const uint64* StackFrames, int32 NumS
 	{
 		bReentranceGuard = true;
 
-		FMacCrashContext EnsureContext(ECrashContextType::Ensure, ErrorMessage);
+		FMacCrashContext EnsureContext(ECrashContextType::Hang, ErrorMessage);
 		EnsureContext.SetPortableCallStack(StackFrames, NumStackFrames);
 		EnsureContext.GenerateEnsureInfoAndLaunchReporter();
 
diff --git a/Engine/Source/Runtime/Core/Private/Unix/UnixPlatformStackWalk.cpp b/Engine/Source/Runtime/Core/Private/Unix/UnixPlatformStackWalk.cpp
--- a/Engine/Source/Runtime/Core/Private/Unix/UnixPlatformStackWalk.cpp
+++ b/Engine/Source/Runtime/Core/Private/Unix/UnixPlatformStackWalk.cpp
@@ -852,7 +852,7 @@ void ReportHang(const TCHAR* ErrorMessage, const uint64* StackFrames, int32 NumS
 	{
 		bReentranceGuard = true;
 
-		FUnixCrashContext EnsureContext(ECrashContextType::Ensure, ErrorMessage);
+		FUnixCrashContext EnsureContext(ECrashContextType::Hang, ErrorMessage);
 		EnsureContext.SetPortableCallStack(StackFrames, NumStackFrames);
 		EnsureContext.GenerateCrashInfoAndLaunchReporter(true);
 
diff --git a/Engine/Source/Runtime/Core/Private/Windows/WindowsPlatformCrashContext.cpp b/Engine/Source/Runtime/Core/Private/Windows/WindowsPlatformCrashContext.cpp
--- a/Engine/Source/Runtime/Core/Private/Windows/WindowsPlatformCrashContext.cpp
+++ b/Engine/Source/Runtime/Core/Private/Windows/WindowsPlatformCrashContext.cpp
@@ -570,7 +570,7 @@ void ReportHang(const TCHAR* ErrorMessage, const uint64* StackFrames, int32 NumS
 		return;
 	}
 
-	FWindowsPlatformCrashContext CrashContext(ECrashContextType::Ensure, ErrorMessage);
+	FWindowsPlatformCrashContext CrashContext(ECrashContextType::Hang, ErrorMessage);
 	CrashContext.SetPortableCallStack(StackFrames, NumStackFrames);
 	CrashContext.SetCrashedThreadId(HungThreadId);
 	CrashContext.CaptureAllThreadContexts();
diff --git a/Engine/Source/Runtime/Core/Public/GenericPlatform/GenericPlatformCrashContext.h b/Engine/Source/Runtime/Core/Public/GenericPlatform/GenericPlatformCrashContext.h
--- a/Engine/Source/Runtime/Core/Public/GenericPlatform/GenericPlatformCrashContext.h
+++ b/Engine/Source/Runtime/Core/Public/GenericPlatform/GenericPlatformCrashContext.h
@@ -88,6 +88,7 @@ enum class ECrashContextType
 	Assert,
 	Ensure,
 	GPUCrash,
+	Hang,
 
 	Max
 };
@@ -129,6 +130,7 @@ public:
 	static const FString CrashTypeAssert;
 	static const FString CrashTypeEnsure;
 	static const FString CrashTypeGPU;
+	static const FString CrashTypeHang;
 
 	static const FString EngineModeExUnknown;
 	static const FString EngineModeExDirty;
