Merging from DK to Main

https://p4-swarm.epicgames.net/reviews/5435295
https://p4-swarm.epicgames.net/reviews/5435326

diff --git a/Engine/Plugins/Runtime/AR/Apple/AppleARKit/AppleARKit.uplugin b/Engine/Plugins/Runtime/AR/Apple/AppleARKit/AppleARKit.uplugin
--- a/Engine/Plugins/Runtime/AR/Apple/AppleARKit/AppleARKit.uplugin
+++ b/Engine/Plugins/Runtime/AR/Apple/AppleARKit/AppleARKit.uplugin
@@ -13,6 +13,9 @@
 	"CanContainContent": true,
 	"IsBetaVersion": false,
 	"EnabledByDefault": false,
+	"SupportedTargetPlatforms" : [
+		"IOS"
+	],
 	"Modules": [
 		{
 			"Name": "AppleARKit",
diff --git a/Engine/Plugins/Runtime/AR/Apple/AppleARKitFaceSupport/AppleARKitFaceSupport.uplugin b/Engine/Plugins/Runtime/AR/Apple/AppleARKitFaceSupport/AppleARKitFaceSupport.uplugin
--- a/Engine/Plugins/Runtime/AR/Apple/AppleARKitFaceSupport/AppleARKitFaceSupport.uplugin
+++ b/Engine/Plugins/Runtime/AR/Apple/AppleARKitFaceSupport/AppleARKitFaceSupport.uplugin
@@ -13,6 +13,9 @@
 	"CanContainContent": true,
 	"IsBetaVersion": false,
 	"EnabledByDefault": false,
+	"SupportedTargetPlatforms" : [
+		"IOS"
+	],
 	"Modules":
 	[
 		{
diff --git a/Engine/Source/Runtime/Core/Public/GenericPlatform/GenericPlatformInstallBundleManager.h b/Engine/Source/Runtime/Core/Public/GenericPlatform/GenericPlatformInstallBundleManager.h
--- a/Engine/Source/Runtime/Core/Public/GenericPlatform/GenericPlatformInstallBundleManager.h
+++ b/Engine/Source/Runtime/Core/Public/GenericPlatform/GenericPlatformInstallBundleManager.h
@@ -50,6 +50,7 @@ enum class EInstallBundleModuleInitResult : int
 	DistributionRootDownloadError,
 	ManifestCreationError,
 	ManifestDownloadError,
+	BackgroundDownloadsIniDownloadError,
 	NoInternetConnectionError,
 	Count
 };
@@ -66,6 +67,7 @@ inline const TCHAR* GetInstallBundleModuleInitResultString(EInstallBundleModuleI
 		TEXT("DistributionRootDownloadError"),
 		TEXT("ManifestCreationError"),
 		TEXT("ManifestDownloadError"),
+		TEXT("BackgroundDownloadsIniDownloadError"),
 		TEXT("NoInternetConnectionError"),
 	};
 	static_assert(static_cast<UnderType>(EInstallBundleModuleInitResult::Count) == ARRAY_COUNT(Strings), "");
@@ -143,9 +145,9 @@ enum class EInstallBundleRequestInfoFlags : int32
 	EnqueuedBundlesForRemoval		= (1 << 1),
 	SkippedAlreadyMountedBundles	= (1 << 2),
 	SkippedBundlesQueuedForRemoval	= (1 << 3),
-	SkippedBundlesQueuedForInstall  = (1 << 5), // Only valid for removal requests
-	SkippedUnknownBundles			= (1 << 6),
-	InitializationError				= (1 << 7), // Can't enqueue because the bundle manager failed to initialize
+	SkippedBundlesQueuedForInstall  = (1 << 4), // Only valid for removal requests
+	SkippedUnknownBundles			= (1 << 5),
+	InitializationError				= (1 << 6), // Can't enqueue because the bundle manager failed to initialize
 };
 ENUM_CLASS_FLAGS(EInstallBundleRequestInfoFlags);
 
@@ -156,6 +158,13 @@ struct FInstallBundleRequestInfo
 	TArray<FName> BundlesQueuedForRemoval;
 };
 
+enum class EInstallBundleCancelFlags : int32
+{
+	None		= 0,
+	Resumable	= (1 << 0),
+};
+ENUM_CLASS_FLAGS(EInstallBundleCancelFlags);
+
 enum class EInstallBundleManagerInitErrorHandlerResult
 {
 	NotHandled, // Defer to the next handler
@@ -195,11 +204,15 @@ public:
 
 	virtual void RequestRemoveBundleOnNextInit(FName BundleName) = 0;
 
-	virtual void CancelBundle(FName BundleName) = 0;
+	virtual void CancelBundle(FName BundleName, EInstallBundleCancelFlags Flags) = 0;
+
+	virtual void CancelAllBundles(EInstallBundleCancelFlags Flags) = 0;
 
 	virtual TOptional<FInstallBundleStatus> GetBundleProgress(FName BundleName) const = 0;
 
 	virtual bool IsNullInterface() const = 0;
+
+	virtual void SetErrorSimulationCommands(const FString& CommandLine) {}
 };
 
 class IPlatformInstallBundleManagerModule : public IModuleInterface
diff --git a/Engine/Source/Runtime/Core/Public/Windows/WindowsPlatformCompilerSetup.h b/Engine/Source/Runtime/Core/Public/Windows/WindowsPlatformCompilerSetup.h
--- a/Engine/Source/Runtime/Core/Public/Windows/WindowsPlatformCompilerSetup.h
+++ b/Engine/Source/Runtime/Core/Public/Windows/WindowsPlatformCompilerSetup.h
@@ -166,10 +166,7 @@ static_assert(_MSC_VER != 1914 && _MSC_VER != 1915, "Visual Studio 2017 versions
 #pragma warning(disable: 4511) // 'class' : copy constructor could not be generated													https://docs.microsoft.com/en-us/cpp/error-messages/compiler-warnings/compiler-warning-level-3-c4511
 #pragma warning(disable: 4512) // 'class' : assignment operator could not be generated												https://docs.microsoft.com/en-us/cpp/error-messages/compiler-warnings/compiler-warning-level-4-c4512
 #pragma warning(disable: 4514) // 'function' : unreferenced inline function has been removed										https://docs.microsoft.com/en-us/cpp/error-messages/compiler-warnings/compiler-warning-level-4-c4514
-#if UE_BUILD_DEBUG
-// xstring.h causes this warning in debug builds
-#pragma warning(disable: 4548) // expression before comma has no effect; expected expression with side-effect						https://docs.microsoft.com/en-us/cpp/error-messages/compiler-warnings/compiler-warning-level-1-c4548
-#endif
+
 #pragma warning(disable: 4592) // 'function': 'constexpr' call evaluation failed; function will be called at run-time				https://docs.microsoft.com/en-us/cpp/error-messages/compiler-warnings/compiler-warnings-c4400-through-c4599
 #pragma warning(disable: 4599) // 'flag path': command line argument number number does not match precompiled header				https://docs.microsoft.com/en-us/cpp/error-messages/compiler-warnings/compiler-warnings-c4400-through-c4599
 
diff --git a/Engine/Source/Runtime/NullInstallBundleManager/Source/NullInstallBundleManager.cpp b/Engine/Source/Runtime/NullInstallBundleManager/Source/NullInstallBundleManager.cpp
--- a/Engine/Source/Runtime/NullInstallBundleManager/Source/NullInstallBundleManager.cpp
+++ b/Engine/Source/Runtime/NullInstallBundleManager/Source/NullInstallBundleManager.cpp
@@ -64,7 +64,12 @@ class FNullInstallBundleManager : public IPlatformInstallBundleManager
 
 	}
 
-	virtual void CancelBundle(FName BundleName) override
+	virtual void CancelBundle(FName BundleName, EInstallBundleCancelFlags Flags) override
+	{
+
+	}
+
+	virtual void CancelAllBundles(EInstallBundleCancelFlags Flags) override
 	{
 
 	}
