Cherry-pick 5459352: Implement missing non-SSE Math FPU Vector functions

  VectorMask_LT(const VectorRegister& Vec1, const VectorRegister& Vec2)
  VectorMask_LE(const VectorRegister& Vec1, const VectorRegister& Vec2)
  VectorMask_GT(const VectorRegister& Vec1, const VectorRegister& Vec2)
  VectorMask_GE(const VectorRegister& Vec1, const VectorRegister& Vec2)
  VectorMask_EQ(const VectorRegister& Vec1, const VectorRegister& Vec2)
  VectorMask_NE(const VectorRegister& Vec1, const VectorRegister& Vec2)
  VectorMaskBits(const VectorRegister& Vec1)

Fixes ARM64 builds.

diff --git a/Engine/Source/Runtime/Core/Public/Math/UnrealMathFPU.h b/Engine/Source/Runtime/Core/Public/Math/UnrealMathFPU.h
--- a/Engine/Source/Runtime/Core/Public/Math/UnrealMathFPU.h
+++ b/Engine/Source/Runtime/Core/Public/Math/UnrealMathFPU.h
@@ -1636,6 +1636,60 @@ FORCEINLINE VectorRegisterInt VectorIntLoad1(const void* Ptr)
 		IntSplat);
 }
 
+/**
+ * These functions return a vector mask to indicate which components pass the comparison.
+ * Each component is 0xffffffff if it passes, 0x00000000 if it fails.
+ *
+ * @param Vec1			1st source vector
+ * @param Vec2			2nd source vector
+ * @return				Vector with a mask for each component.
+ */
+FORCEINLINE VectorRegister VectorMask_LT(const VectorRegister& Vec1, const VectorRegister& Vec2)
+{
+	return VectorCompareLT(Vec1, Vec2);
+}
+
+FORCEINLINE VectorRegister VectorMask_LE(const VectorRegister& Vec1, const VectorRegister& Vec2)
+{
+	return VectorCompareLE(Vec1, Vec2);
+}
+
+FORCEINLINE VectorRegister VectorMask_GT(const VectorRegister& Vec1, const VectorRegister& Vec2)
+{
+	return VectorCompareGT(Vec1, Vec2);
+}
+
+FORCEINLINE VectorRegister VectorMask_GE(const VectorRegister& Vec1, const VectorRegister& Vec2)
+{
+	return VectorCompareGE(Vec1, Vec2);
+}
+
+FORCEINLINE VectorRegister VectorMask_EQ(const VectorRegister& Vec1, const VectorRegister& Vec2)
+{
+	return VectorCompareEQ(Vec1, Vec2);
+}
+
+FORCEINLINE VectorRegister VectorMask_NE(const VectorRegister& Vec1, const VectorRegister& Vec2)
+{
+	return VectorCompareNE(Vec1, Vec2);
+}
+
+/**
+ * Returns an integer bit-mask (0x00 - 0x0f) based on the sign-bit for each component in a vector.
+ *
+ * @param VecMask		Vector
+ * @return				Bit 0 = sign(VecMask.x), Bit 1 = sign(VecMask.y), Bit 2 = sign(VecMask.z), Bit 3 = sign(VecMask.w)
+ */
+FORCEINLINE int32_t VectorMaskBits(const VectorRegister& Vec1)
+{
+	uint32* V1 = (uint32*)(&(Vec1.V[0]));
+
+	return (V1[0] >> 31) |
+		  ((V1[1] >> 30) & 2) |
+		  ((V1[2] >> 29) & 4) |
+		  ((V1[3] >> 28) & 8);
+}
+
 // To be continued...
 
 
