- Fixed a licensee reported crash of LLM during shutdown

diff --git a/Engine/Source/Runtime/Core/Private/HAL/LowLevelMemTracker.cpp b/Engine/Source/Runtime/Core/Private/HAL/LowLevelMemTracker.cpp
--- a/Engine/Source/Runtime/Core/Private/HAL/LowLevelMemTracker.cpp
+++ b/Engine/Source/Runtime/Core/Private/HAL/LowLevelMemTracker.cpp
@@ -538,6 +538,9 @@ FLowLevelMemTracker::FLowLevelMemTracker()
 
 FLowLevelMemTracker::~FLowLevelMemTracker()
 {
+	// Ensure that we skip any further tracking since it will fail after this destructor
+	bIsDisabled = true;
+	
 	for (int32 TrackerIndex = 0; TrackerIndex < (int32)ELLMTracker::Max; ++TrackerIndex)
 	{
 		Trackers[TrackerIndex]->~FLLMTracker();
