Fixing static analysis warning.

diff --git a/Engine/Source/Runtime/Core/Private/Windows/WindowsPlatformCrashContext.cpp b/Engine/Source/Runtime/Core/Private/Windows/WindowsPlatformCrashContext.cpp
--- a/Engine/Source/Runtime/Core/Private/Windows/WindowsPlatformCrashContext.cpp
+++ b/Engine/Source/Runtime/Core/Private/Windows/WindowsPlatformCrashContext.cpp
@@ -318,11 +318,13 @@ int32 ReportCrashUsingCrashReportClient(FWindowsPlatformCrashContext& InContext,
 		GConfig->GetBool(TEXT("/Script/UnrealEd.CrashReportsPrivacySettings"), TEXT("bSendUnattendedBugReports"), bSendUnattendedBugReports, GEditorSettingsIni);
 	}
 
-	if (BuildSettings::IsLicenseeVersion() && !UE_EDITOR)
+#if !UE_EDITOR
+	if (BuildSettings::IsLicenseeVersion())
 	{
 		// do not send unattended reports in licensees' builds except for the editor, where it is governed by the above setting
 		bSendUnattendedBugReports = false;
 	}
+#endif
 
 	if (bNoDialog && !bSendUnattendedBugReports)
 	{
