Edigrated rollup of all engine RSA changes that were done for 4.22. Pak signature checking is now much more competent.

#rnx

diff --git a/Engine/Source/Programs/UnrealBuildTool/System/EncryptionAndSigning.cs b/Engine/Source/Programs/UnrealBuildTool/System/EncryptionAndSigning.cs
--- a/Engine/Source/Programs/UnrealBuildTool/System/EncryptionAndSigning.cs
+++ b/Engine/Source/Programs/UnrealBuildTool/System/EncryptionAndSigning.cs
@@ -61,6 +61,18 @@ namespace UnrealBuildTool
 			{
 				return PublicKey != null && PrivateKey != null && PublicKey.IsValid() && PrivateKey.IsValid();
 			}
+
+			/// <summary>
+			/// Returns TRUE if this is a short key from the old 256-bit system
+			/// </summary>
+			public bool IsUnsecureLegacyKey()
+			{
+				int LongestKey = PublicKey.Exponent.Length;
+				LongestKey = Math.Max(LongestKey, PublicKey.Modulus.Length);
+				LongestKey = Math.Max(LongestKey, PrivateKey.Exponent.Length);
+				LongestKey = Math.Max(LongestKey, PrivateKey.Modulus.Length);
+				return LongestKey <= 64;
+			}
 		}
 
 		/// <summary>
diff --git a/Engine/Source/Runtime/Core/Public/Misc/RSA.h b/Engine/Source/Runtime/Core/Public/Misc/RSA.h
new file mode 100644
--- /dev/null
+++ b/Engine/Source/Runtime/Core/Public/Misc/RSA.h
@@ -0,0 +1,84 @@
+// Copyright 1998-2019 Epic Games, Inc. All Rights Reserved.
+#pragma once
+
+#include "CoreTypes.h"
+#include "Math/BigInt.h"
+
+template<int KEY_SIZE>
+struct FRSAKey
+{
+	enum { KeySize = KEY_SIZE };
+	enum { KeySizeInBytes = KeySize / 8 };
+	enum { MaxDataSize = KeySizeInBytes - 7 };
+	static_assert(KEY_SIZE % 8 == 0, "Key sizes must be multiples of 8");
+	typedef TBigInt<KeySize * 2, false> TIntType;
+	TIntType Exponent;
+	TIntType Modulus;
+
+	FRSAKey()
+	{
+
+	}
+
+	FRSAKey(const TArray<uint8>& InExponent, const TArray<uint8>& InModulus)
+		: Exponent(InExponent.GetData(), InExponent.Num())
+		, Modulus(InModulus.GetData(), InModulus.Num())
+	{
+	}
+
+	bool IsValid() const
+	{
+		return !Exponent.IsZero() && !Modulus.IsZero();
+	}
+
+	static void Encrypt(const TIntType& InSource, TIntType& OutEncrypted, const FRSAKey& InKey)
+	{
+		OutEncrypted = FEncryption::ModularPow(InSource, InKey.Exponent, InKey.Modulus);
+	}
+
+	static void Encrypt(const uint8* InSourceBuffer, int64 InSourceBufferSizeInBytes, TArray<uint8>& OutEncrypted, const FRSAKey& InKey)
+	{
+		if (InSourceBufferSizeInBytes <= MaxDataSize)
+		{
+			TIntType Source(InSourceBuffer, InSourceBufferSizeInBytes);
+			// TODO: Pad Source data!
+			TIntType Encrypted = FEncryption::ModularPow(Source, InKey.Exponent, InKey.Modulus);
+			OutEncrypted.SetNum(KeySizeInBytes);
+			FMemory::Memcpy(OutEncrypted.GetData(), Encrypted.GetBits(), KeySizeInBytes);
+		}
+		else
+		{
+			OutEncrypted.Empty();
+		}
+	}
+
+	static void Encrypt(const TArray<uint8> InSource, TArray<uint8>& OutEncrypted, const FRSAKey& InKey)
+	{
+		Encrypt(InSource.GetData(), InSource.Num(), OutEncrypted, InKey);
+	}
+
+	static void Decrypt(const TIntType& InEncrypted, TIntType& OutDecrypted, const FRSAKey& InKey)
+	{
+		OutDecrypted = FEncryption::ModularPow(InEncrypted, InKey.Exponent, InKey.Modulus);
+	}
+
+	static void Decrypt(const TArray<uint8>& InEncrypted, uint8* InOutputBuffer, int32 InOutputBufferSizeInBytes, const FRSAKey& InKey)
+	{
+		if (InEncrypted.Num() <= KeySizeInBytes)
+		{
+			TIntType Source(InEncrypted.GetData(), InEncrypted.Num());
+			TIntType Decrypted = FEncryption::ModularPow(Source, InKey.Exponent, InKey.Modulus);
+			FMemory::Memcpy(InOutputBuffer, Decrypted.GetBits(), FMath::Min((int32)MaxDataSize, InOutputBufferSizeInBytes));
+		}
+		else
+		{
+			FMemory::Memset(InOutputBuffer, 0, InOutputBufferSizeInBytes);
+		}
+	}
+	
+	static void Decrypt(const TArray<uint8>& InEncrypted, TArray<uint8>& OutDecrypted, const FRSAKey& InKey)
+	{
+		OutDecrypted.SetNumUninitialized(MaxDataSize);
+		Decrypt(InEncrypted, OutDecrypted.GetData(), MaxDataSize, InKey);
+	}
+};
\ No newline at end of file
diff --git a/Engine/Source/Runtime/PakFile/Private/IPlatformFilePak.cpp b/Engine/Source/Runtime/PakFile/Private/IPlatformFilePak.cpp
--- a/Engine/Source/Runtime/PakFile/Private/IPlatformFilePak.cpp
+++ b/Engine/Source/Runtime/PakFile/Private/IPlatformFilePak.cpp
@@ -5408,8 +5408,8 @@ bool FPakPlatformFile::Initialize(IPlatformFile* Inner, const TCHAR* CmdLine)
 #if !(IS_PROGRAM || WITH_EDITOR)
 	FCoreDelegates::OnFEngineLoopInitComplete.AddLambda([this] {
 			FPlatformMisc::LowLevelOutputDebugStringf(TEXT("Checking Pak Config"));
-			bool bUnloadPakEntryFilenamesIfPossible = false;
-			GConfig->GetBool(TEXT("Pak"), TEXT("UnloadPakEntryFilenamesIfPossible"), bUnloadPakEntryFilenamesIfPossible, GEngineIni);
+			bool bUnloadPakEntryFilenamesIfPossible = true;
+			//GConfig->GetBool(TEXT("Pak"), TEXT("UnloadPakEntryFilenamesIfPossible"), bUnloadPakEntryFilenamesIfPossible, GEngineIni);
 
 			if (bUnloadPakEntryFilenamesIfPossible)
 			{
@@ -5428,8 +5428,8 @@ bool FPakPlatformFile::Initialize(IPlatformFile* Inner, const TCHAR* CmdLine)
 				PakPlatformFile->UnloadPakEntryFilenames(&DirectoryRootsToKeep);
 			}
 
-			bool bShrinkPakEntriesMemoryUsage = false;
-			GConfig->GetBool(TEXT("Pak"), TEXT("ShrinkPakEntriesMemoryUsage"), bShrinkPakEntriesMemoryUsage, GEngineIni);
+			bool bShrinkPakEntriesMemoryUsage = true;
+			//GConfig->GetBool(TEXT("Pak"), TEXT("ShrinkPakEntriesMemoryUsage"), bShrinkPakEntriesMemoryUsage, GEngineIni);
 			if (bShrinkPakEntriesMemoryUsage)
 			{
 				FPakPlatformFile* PakPlatformFile = (FPakPlatformFile*)(FPlatformFileManager::Get().FindPlatformFile(FPakPlatformFile::GetTypeName()));
