Fixed plugin config overrides for out of tree projects, also route plugin overrides properly when specifying ini on the command line

diff --git a/Engine/Source/Runtime/Core/Private/Misc/ConfigCacheIni.cpp b/Engine/Source/Runtime/Core/Private/Misc/ConfigCacheIni.cpp
--- a/Engine/Source/Runtime/Core/Private/Misc/ConfigCacheIni.cpp
+++ b/Engine/Source/Runtime/Core/Private/Misc/ConfigCacheIni.cpp
@@ -3228,16 +3228,7 @@ static void GetSourceIniHierarchyFilenames(const TCHAR* InBaseIniName, const TCH
 	OutHierarchy.KeySort([](const EConfigFileHierarchy& A, const EConfigFileHierarchy& B) { return (A < B); });
 }
 
-/**
- * Calculates the name of a dest (generated) .ini file for a given base (ie Engine, Game, etc)
- * 
- * @param IniBaseName Base name of the .ini (Engine, Game)
- * @param PlatformName Name of the platform to get the .ini path for (nullptr means to use the current platform)
- * @param GeneratedConfigDir The base folder that will contain the generated config files.
- * 
- * @return Standardized .ini filename
- */
-static FString GetDestIniFilename(const TCHAR* BaseIniName, const TCHAR* PlatformName, const TCHAR* GeneratedConfigDir)
+FString FConfigCacheIni::GetDestIniFilename(const TCHAR* BaseIniName, const TCHAR* PlatformName, const TCHAR* GeneratedConfigDir)
 {
 	// figure out what to look for on the commandline for an override
 	FString CommandLineSwitch = FString::Printf(TEXT("%sINI="), BaseIniName);
@@ -3248,14 +3239,13 @@ static FString GetDestIniFilename(const TCHAR* BaseIniName, const TCHAR* Platfor
 	{
 		FString Name(PlatformName ? PlatformName : ANSI_TO_TCHAR(FPlatformProperties::PlatformName()));
 
-		FString BaseIniNameString = BaseIniName;
-		if (BaseIniNameString.Contains(GeneratedConfigDir))
+		// if the BaseIniName doens't contain the config dir, put it all together
+		if (FCString::Stristr(BaseIniName, GeneratedConfigDir) != nullptr)
 		{
-			IniFilename = BaseIniNameString;
+			IniFilename = BaseIniName;
 		}
 		else
 		{
-			// put it all together
 			IniFilename = FString::Printf(TEXT("%s%s/%s.ini"), GeneratedConfigDir, *Name, BaseIniName);
 		}
 	}
diff --git a/Engine/Source/Runtime/Core/Public/Misc/ConfigCacheIni.h b/Engine/Source/Runtime/Core/Public/Misc/ConfigCacheIni.h
--- a/Engine/Source/Runtime/Core/Public/Misc/ConfigCacheIni.h
+++ b/Engine/Source/Runtime/Core/Public/Misc/ConfigCacheIni.h
@@ -821,6 +821,17 @@ public:
 	 */
 	static void InitializeConfigSystem();
 
+	/**
+	 * Calculates the name of a dest (generated) .ini file for a given base (ie Engine, Game, etc)
+	 *
+	 * @param IniBaseName Base name of the .ini (Engine, Game)
+	 * @param PlatformName Name of the platform to get the .ini path for (nullptr means to use the current platform)
+	 * @param GeneratedConfigDir The base folder that will contain the generated config files.
+	 *
+	 * @return Standardized .ini filename
+	 */
+	static FString GetDestIniFilename(const TCHAR* BaseIniName, const TCHAR* PlatformName, const TCHAR* GeneratedConfigDir);
+
 	/**
 	 * Loads and generates a destination ini file and adds it to GConfig:
 	 *   - Looking on commandline for override source/dest .ini filenames
@@ -844,7 +855,7 @@ public:
 	/**
 	 * Load an ini file directly into an FConfigFile, and nothing is written to GConfig or disk. 
 	 * The passed in .ini name can be a "base" (Engine, Game) which will be modified by platform and/or commandline override,
-	 * or it can be a full ini filenname (ie WrangleContent) loaded from the Source config directory
+	 * or it can be a full ini filename (ie WrangleContent) loaded from the Source config directory
 	 *
 	 * @param ConfigFile The output object to fill
 	 * @param IniName Either a Base ini name (Engine) or a full ini name (WrangleContent). NO PATH OR EXTENSION SHOULD BE USED!
@@ -858,7 +869,7 @@ public:
 	/**
 	 * Load an ini file directly into an FConfigFile from the specified config folders, optionally writing to disk. 
 	 * The passed in .ini name can be a "base" (Engine, Game) which will be modified by platform and/or commandline override,
-	 * or it can be a full ini filenname (ie WrangleContent) loaded from the Source config directory
+	 * or it can be a full ini filename (ie WrangleContent) loaded from the Source config directory
 	 *
 	 * @param ConfigFile The output object to fill
 	 * @param IniName Either a Base ini name (Engine) or a full ini name (WrangleContent). NO PATH OR EXTENSION SHOULD BE USED!
diff --git a/Engine/Source/Runtime/Projects/Private/PluginManager.cpp b/Engine/Source/Runtime/Projects/Private/PluginManager.cpp
--- a/Engine/Source/Runtime/Projects/Private/PluginManager.cpp
+++ b/Engine/Source/Runtime/Projects/Private/PluginManager.cpp
@@ -604,7 +604,8 @@ bool FPluginManager::ConfigureEnabledPlugins()
 					for (const FString& ConfigFile : PluginConfigs)
 					{
 						FString PlaformName = FPlatformProperties::PlatformName();
-						PluginConfigFilename = FString::Printf(TEXT("%s%s/%s.ini"), *FPaths::GeneratedConfigDir(), *PlaformName, *FPaths::GetBaseFilename(ConfigFile));
+						// Use GetDestIniFilename to find the proper config file to combine into, since it manages command line overrides and path sanitization
+						PluginConfigFilename = FConfigCacheIni::GetDestIniFilename(*FPaths::GetBaseFilename(ConfigFile), *PlaformName, *FPaths::GeneratedConfigDir());
 						FConfigFile* FoundConfig = GConfig->Find(PluginConfigFilename, false);
 						if (FoundConfig != nullptr)
 						{
