Ensure that the ApplicationWillTerminate delegate is always called from with in Windows' Control-C handler.

(peafour-eddy of //UE4/Dev-Core/[REDACTED_HANDLE] by [REDACTED_PERSON])

diff --git a/Engine/Source/Runtime/Core/Private/Windows/WindowsPlatformMisc.cpp b/Engine/Source/Runtime/Core/Private/Windows/WindowsPlatformMisc.cpp
--- a/Engine/Source/Runtime/Core/Private/Windows/WindowsPlatformMisc.cpp
+++ b/Engine/Source/Runtime/Core/Private/Windows/WindowsPlatformMisc.cpp
@@ -584,14 +584,24 @@ void FWindowsPlatformMisc::PlatformInit()
  */
 static BOOL WINAPI ConsoleCtrlHandler(DWORD CtrlType)
 {
+	// Broadcast the termination the first time through.
+	bool IsRequestingExit = GIsRequestingExit;
+	static bool AppTermDelegateBroadcast = false;
+	if (!AppTermDelegateBroadcast)
+	{
+		GIsRequestingExit = true;
+		FCoreDelegates::ApplicationWillTerminateDelegate.Broadcast();
+		AppTermDelegateBroadcast = true;
+	}
+
 	// Only "two-step Ctrl-C" if the termination event is Ctrl-C and the process
 	// is considered interactive. Hard-terminate on all other cases.
 	if (CtrlType != CTRL_C_EVENT || FApp::IsUnattended())
 	{
-		GIsRequestingExit = true;
+		IsRequestingExit = true;
 	}
 
-	if (!GIsRequestingExit)
+	if (!IsRequestingExit)
 	{
 		UE_LOG(LogCore, Warning, TEXT("*** INTERRUPTED *** : SHUTTING DOWN"));
 		UE_LOG(LogCore, Warning, TEXT("*** INTERRUPTED *** : CTRL-C TO FORCE QUIT"));
@@ -611,11 +621,10 @@ static BOOL WINAPI ConsoleCtrlHandler(DWORD CtrlType)
 		GError->Flush();
 	}
 
-	if (!GIsRequestingExit)
+	if (!IsRequestingExit)
 	{
-		// Notify anyone listening that we're about to terminate
-		GIsRequestingExit = true;
-		FCoreDelegates::ApplicationWillTerminateDelegate.Broadcast();
+		// We'll two-step Ctrl-C events to give processes like servers time to
+		// correctly terminate. Note the GIsRequestingExit is true now.
 		return true;
 	}
 
