Fixed uninitialized member for FAutoScopedDurationTimer
#rnx

diff --git a/Engine/Source/Runtime/Core/Public/ProfilingDebugging/ScopedTimers.h b/Engine/Source/Runtime/Core/Public/ProfilingDebugging/ScopedTimers.h
--- a/Engine/Source/Runtime/Core/Public/ProfilingDebugging/ScopedTimers.h
+++ b/Engine/Source/Runtime/Core/Public/ProfilingDebugging/ScopedTimers.h
@@ -72,6 +72,7 @@ class FAutoScopedDurationTimer : public FScopedDurationTimer
 public:
 	FAutoScopedDurationTimer()
 		: FScopedDurationTimer(AccumulatorValue)
+		, AccumulatorValue(0)
 	{
 	}
 
