Changed hang detection to only report when there has been at least one real heartbeat on the thread after a hang was detected. Previously the hang detection itself was treated as another heartbeat so 30s later the detection could trip again if there was a change in callstack. Commonly we'd see a user idle the application deep in 3rd party code which would flip between similar callstacks thus sending hours worth of reports. The tradeoff here is losing potentially multiple hangs within a single frame but that's very atypical behavior and would likely be considered a hard failure by the user.
Fixed malformed strings in hang report dump from uninitialized data.

diff --git a/Engine/Source/Runtime/Core/Private/HAL/ThreadHeartBeat.cpp b/Engine/Source/Runtime/Core/Private/HAL/ThreadHeartBeat.cpp
--- a/Engine/Source/Runtime/Core/Private/HAL/ThreadHeartBeat.cpp
+++ b/Engine/Source/Runtime/Core/Private/HAL/ThreadHeartBeat.cpp
@@ -200,6 +200,7 @@ void FORCENOINLINE FThreadHeartBeat::OnHang(double HangDuration, uint32 ThreadTh
 		for (int32 Idx = 0; Idx < NumStackFrames; Idx++)
 		{
 			ANSICHAR Buffer[1024];
+			Buffer[0] = '\0';
 			FPlatformStackWalk::ProgramCounterToHumanReadableString(Idx, StackFrames[Idx], Buffer, sizeof(Buffer));
 			StackLines.Add(Buffer);
 		}
@@ -389,12 +390,14 @@ uint32 FThreadHeartBeat::CheckHeartBeat(double& OutHangDuration)
 		if (ConfigHangDuration > 0.0)
 		{
 			// Check heartbeat for all threads and return thread ID of the thread that hung.
+			// Note: We only return a thread id for a thread that has updated since the last hang, i.e. is still alive
+			// This avoids the case where a user may be in a deep and minorly varying callstack and flood us with reports
 			for (TPair<uint32, FHeartBeatInfo>& LastHeartBeat : ThreadHeartBeat)
 			{
 				FHeartBeatInfo& HeartBeatInfo = LastHeartBeat.Value;
-				if (HeartBeatInfo.SuspendedCount == 0 && (CurrentTime - HeartBeatInfo.LastHeartBeatTime) > HeartBeatInfo.HangDuration)
+				if (HeartBeatInfo.SuspendedCount == 0 && (CurrentTime - HeartBeatInfo.LastHeartBeatTime) > HeartBeatInfo.HangDuration && HeartBeatInfo.LastHeartBeatTime >= HeartBeatInfo.LastHangTime)
 				{
-					HeartBeatInfo.LastHeartBeatTime = CurrentTime;
+					HeartBeatInfo.LastHangTime = CurrentTime;
 					OutHangDuration = HeartBeatInfo.HangDuration;
 					return LastHeartBeat.Key;
 				}
diff --git a/Engine/Source/Runtime/Core/Public/HAL/ThreadHeartBeat.h b/Engine/Source/Runtime/Core/Public/HAL/ThreadHeartBeat.h
--- a/Engine/Source/Runtime/Core/Public/HAL/ThreadHeartBeat.h
+++ b/Engine/Source/Runtime/Core/Public/HAL/ThreadHeartBeat.h
@@ -47,12 +47,15 @@ class CORE_API FThreadHeartBeat : public FRunnable
 	{
 		FHeartBeatInfo()
 			: LastHeartBeatTime(0.0)
+			, LastHangTime(0.0)
 			, SuspendedCount(0)
 			, HangDuration(0)
 		{}
 
 		/** Time we last received a heartbeat for the current thread */
 		double LastHeartBeatTime;
+		/** Time we last detected a hang due to lack of heartbeats for the current thread */
+		double LastHangTime;
 		/** Suspended counter */
 		int32 SuspendedCount;
 		/** The timeout for this thread */
