Linux: cache larger OS allocs (UE-67825).

- Some of RHI buffers are locked very frequently and their size is larger than the largerst we pool. Passing them through directly to platform functions result in a subpar performance.
- On other platforms, FCachedOSPageAllocator is already used for all allocs above those in FMallocBinned(2), so this addition is consistent.

diff --git a/Engine/Source/Runtime/Core/Private/HAL/Allocators/PooledVirtualMemoryAllocator.cpp b/Engine/Source/Runtime/Core/Private/HAL/Allocators/PooledVirtualMemoryAllocator.cpp
--- a/Engine/Source/Runtime/Core/Private/HAL/Allocators/PooledVirtualMemoryAllocator.cpp
+++ b/Engine/Source/Runtime/Core/Private/HAL/Allocators/PooledVirtualMemoryAllocator.cpp
@@ -43,7 +43,10 @@ void* FPooledVirtualMemoryAllocator::Allocate(SIZE_T Size)
 {
 	if (Size > Limits::MaxAllocationSizeToPool)
 	{
-		return FPlatformMemory::BinnedAllocFromOS(Size);
+		// do not report to LLM here, the platform functions will do that
+		FScopeLock Lock(&OsAllocatorCacheLock);
+		void* Ptr = OsAllocatorCache.Allocate(Size);
+		return Ptr;
 	}
 	else
 	{
@@ -97,7 +100,9 @@ void FPooledVirtualMemoryAllocator::Free(void* Ptr, SIZE_T Size)
 {
 	if (Size > Limits::MaxAllocationSizeToPool)
 	{
-		return FPlatformMemory::BinnedFreeToOS(Ptr, Size);
+		// do not report to LLM here, the platform functions will do that
+		FScopeLock Lock(&OsAllocatorCacheLock);
+		OsAllocatorCache.Free(Ptr, Size);
 	}
 	else
 	{
@@ -219,7 +224,10 @@ void FPooledVirtualMemoryAllocator::DestroyPool(FPoolDescriptorBase* Pool)
 
 void FPooledVirtualMemoryAllocator::FreeAll()
 {
-	// Currently, there's nothing to trim.
+	FScopeLock Lock(&OsAllocatorCacheLock);
+	OsAllocatorCache.FreeAll();
+
+	// Currently, there's nothing else to trim.
 	// We could avoid deleting pools on Free() and instead keep them in a separate list to delete on FreeAll() (unless they're reused before that).
 	// That would be a speed optimization and not a size optimization so I'm not going for this at this point, this method is speedy enough.
 };
diff --git a/Engine/Source/Runtime/Core/Public/HAL/Allocators/PooledVirtualMemoryAllocator.h b/Engine/Source/Runtime/Core/Public/HAL/Allocators/PooledVirtualMemoryAllocator.h
--- a/Engine/Source/Runtime/Core/Public/HAL/Allocators/PooledVirtualMemoryAllocator.h
+++ b/Engine/Source/Runtime/Core/Public/HAL/Allocators/PooledVirtualMemoryAllocator.h
@@ -4,6 +4,7 @@
 
 #include "CoreTypes.h"
 #include "HAL/CriticalSection.h"
+#include "HAL/Allocators/CachedOSPageAllocator.h"
 
 /**
  * This class pools OS allocations made from FMallocBinned2.
@@ -17,7 +18,7 @@
  * of pools that hold varied number of same-sized allocations (called blocks). Each bucket starts empty (linked list is empty).
  * 
  * Whenever an allocation request arrives, it is first bucketed based on its size (if it is larger than the largest bucket,
- * it is passed through to a platform function BinnedAllocFromOS()). Then, the linked list of that bucket is walked, and 
+ * it is passed through to a caching OS allocator). Then, the linked list of that bucket is walked, and 
  * the allocation is fulfilled by the first pool that has empty "blocks". If there are no such pool, a new pool is created 
  * (with number of blocks being possibly larger than in the existing list, if any), this pool becomes the new head of the 
  * list and the allocation happens there.
@@ -54,6 +55,7 @@
  * and we cannot "trim" anything here. Also, it does not make sense to put the global cap on the pooled memory
  * since BinnedAllocFromOS() can support only a limited number of allocations on some platforms.
  * 
+ * CachedOSPageAllocator sits "below" this and is used for allocs larger than the largest bucketed.
  */
 struct FPooledVirtualMemoryAllocator
 {
@@ -81,7 +83,10 @@ private:
 	enum Limits
 	{
 		NumAllocationSizeClasses	= 64,
-		MaxAllocationSizeToPool		= NumAllocationSizeClasses * 65536
+		MaxAllocationSizeToPool		= NumAllocationSizeClasses * 65536,
+
+		MaxOSAllocCacheSize			= 64 * 1024 * 1024,
+		MaxOSAllocsCached			= 64
 	};
 
 	/**
@@ -134,4 +139,10 @@ private:
 
 	/** Destroys a pool */
 	void DestroyPool(FPoolDescriptorBase* Pool);
+
+	/** Lock for accessing the caching allocator for larger allocs */
+	FCriticalSection OsAllocatorCacheLock;
+
+	/** Caching allocator for larger allocs */
+	TCachedOSPageAllocator<MaxOSAllocsCached, MaxOSAllocCacheSize> OsAllocatorCache;
 };
