[LLM] shutdown fix for OnLowLevelAllocMoved

diff --git a/Engine/Source/Runtime/Core/Private/HAL/LowLevelMemTracker.cpp b/Engine/Source/Runtime/Core/Private/HAL/LowLevelMemTracker.cpp
--- a/Engine/Source/Runtime/Core/Private/HAL/LowLevelMemTracker.cpp
+++ b/Engine/Source/Runtime/Core/Private/HAL/LowLevelMemTracker.cpp
@@ -1259,6 +1259,11 @@ void FLLMTracker::TrackFree(const void* Ptr, ELLMTracker Tracker, ELLMAllocType
 
 void FLLMTracker::OnAllocMoved(const void* Dest, const void* Source)
 {
+	if (GIsRequestingExit)
+	{
+		return;
+	}
+
 	LLMMap::Values Values = GetAllocationMap().Remove(Source);
 	GetAllocationMap().Add(Dest, Values.Value1, Values.Value2);
 }
