[LLM] fix for OnAllocMoved in external profilers

diff --git a/Engine/Source/Runtime/Core/Private/HAL/LowLevelMemTracker.cpp b/Engine/Source/Runtime/Core/Private/HAL/LowLevelMemTracker.cpp
--- a/Engine/Source/Runtime/Core/Private/HAL/LowLevelMemTracker.cpp
+++ b/Engine/Source/Runtime/Core/Private/HAL/LowLevelMemTracker.cpp
@@ -310,7 +310,7 @@ public:
 #endif
 	void TrackAllocation(const void* Ptr, uint64 Size, ELLMTag DefaultTag, ELLMTracker Tracker, ELLMAllocType AllocType);
 	void TrackFree(const void* Ptr, ELLMTracker Tracker, ELLMAllocType AllocType);
-	void OnAllocMoved(const void* Dest, const void* Source);
+	void OnAllocMoved(const void* Dest, const void* Source, uint64& OutSize, int64& OutTag);
 
 	void TrackMemory(int64 Tag, int64 Amount);
 
@@ -783,12 +783,31 @@ FLLMTracker* FLowLevelMemTracker::GetTracker(ELLMTracker Tracker)
 
 void FLowLevelMemTracker::OnLowLevelAllocMoved(ELLMTracker Tracker, const void* Dest, const void* Source)
 {
-	if (bIsDisabled)
+	if (bIsDisabled || GIsRequestingExit)
 	{
 		return;
 	}
 
-	GetTracker(Tracker)->OnAllocMoved(Dest, Source);
+	//update the allocation map
+	uint64 Size;
+	int64 Tag;
+	GetTracker(Tracker)->OnAllocMoved(Dest, Source, Size, Tag);
+
+	// update external memory trackers (ideally would want a proper 'move' option on these)
+	if (Tracker == ELLMTracker::Default)
+	{
+		FPlatformMemory::OnLowLevelMemory_Free(Source, Size, Tag);
+		FPlatformMemory::OnLowLevelMemory_Alloc(Dest, Size, Tag);
+	}
+
+#ifdef ENABLE_MEMPRO
+	if (MemProProfiler::bStart && Tracker == ELLMTracker::Default && (MemProProfiler::TrackTag == ELLMTag::GenericTagCount || MemProProfiler::TrackTag == (ELLMTag)Tag))
+	{
+		MEMPRO_TRACK_FREE((void*)Source);
+		MEMPRO_TRACK_ALLOC((void*)Dest, (size_t)Size);
+	}
+#endif
+
 }
 
 bool FLowLevelMemTracker::Exec(const TCHAR* Cmd, FOutputDevice& Ar)
@@ -1257,15 +1276,20 @@ void FLLMTracker::TrackFree(const void* Ptr, ELLMTracker Tracker, ELLMAllocType
 #endif
 }
 
-void FLLMTracker::OnAllocMoved(const void* Dest, const void* Source)
+void FLLMTracker::OnAllocMoved(const void* Dest, const void* Source, uint64& OutSize, int64& OutTag)
 {
-	if (GIsRequestingExit)
-	{
-		return;
-	}
-
 	LLMMap::Values Values = GetAllocationMap().Remove(Source);
 	GetAllocationMap().Add(Dest, Values.Value1, Values.Value2);
+
+
+	const FLLMTracker::FLowLevelAllocInfo& AllocInfo = Values.Value2;
+#if LLM_USE_ALLOC_INFO_STRUCT
+	OutTag = AllocInfo.Tag;
+#else
+	OutTag = (int64)AllocInfo;
+#endif
+
+	OutSize = Values.Value1;
 }
 
 void FLLMTracker::TrackMemory(int64 Tag, int64 Amount)
