Fix task graph stall timing difference when STATS is enabled or disabled.

diff --git a/Engine/Source/Runtime/Core/Private/Async/TaskGraph.cpp b/Engine/Source/Runtime/Core/Private/Async/TaskGraph.cpp
--- a/Engine/Source/Runtime/Core/Private/Async/TaskGraph.cpp
+++ b/Engine/Source/Runtime/Core/Private/Async/TaskGraph.cpp
@@ -609,11 +609,9 @@ public:
 		}
 #endif
 
+		TStatId StatName;
 		TStatId StallStatId;
 		bool bCountAsStall = false;
-#if STATS
-		TStatId StatName;
-		FCycleCounter ProcessingTasks;
 		if (ThreadId == ENamedThreads::GameThread)
 		{
 			StatName = GET_STATID(STAT_TaskGraph_GameTasks);
@@ -629,12 +627,19 @@ public:
 			}
 			// else StatName = none, we need to let the scope empty so that the render thread submits tasks in a timely manner. 
 		}
+#if STATS
 		else if (ThreadId != ENamedThreads::StatsThread)
+#else
+		else
+#endif
 		{
 			StatName = GET_STATID(STAT_TaskGraph_OtherTasks);
 			StallStatId = GET_STATID(STAT_TaskGraph_OtherStalls);
 			bCountAsStall = true;
 		}
+
+#if STATS
+		FCycleCounter ProcessingTasks;
 		bool bTasksOpen = false;
 		if (FThreadStats::IsCollectingData(StatName))
 		{
