Adding IsRunningOnBattery to the windows platform crash properties, to get further feedback about the condition of the machine crashing.  Also adding the ability to query if they're on a laptop at all, but not enabled for now.

Robert.Manuszewski

diff --git a/Engine/Source/Runtime/Core/Private/Windows/WindowsPlatformCrashContext.cpp b/Engine/Source/Runtime/Core/Private/Windows/WindowsPlatformCrashContext.cpp
--- a/Engine/Source/Runtime/Core/Private/Windows/WindowsPlatformCrashContext.cpp
+++ b/Engine/Source/Runtime/Core/Private/Windows/WindowsPlatformCrashContext.cpp
@@ -139,6 +139,8 @@ void FWindowsPlatformCrashContext::AddPlatformSpecificProperties() const
 	AddCrashProperty(TEXT("PlatformIsRunningWindows"), 1);
 	// On windows track the crash type
 	AddCrashProperty(TEXT("PlatformCallbackResult"), GetCrashType());
+
+	AddCrashProperty(TEXT("IsRunningOnBattery"), FPlatformMisc::IsRunningOnBattery());
 }
 
 bool FWindowsPlatformCrashContext::GetPlatformAllThreadContextsString(FString& OutStr) const
