[LLM] support for the new version of MemPro.

diff --git a/Engine/Source/Runtime/Core/Private/HAL/LowLevelMemTracker.cpp b/Engine/Source/Runtime/Core/Private/HAL/LowLevelMemTracker.cpp
--- a/Engine/Source/Runtime/Core/Private/HAL/LowLevelMemTracker.cpp
+++ b/Engine/Source/Runtime/Core/Private/HAL/LowLevelMemTracker.cpp
@@ -10,23 +10,12 @@
 #include "HAL/IConsoleManager.h"
 
 #if ENABLE_LOW_LEVEL_MEM_TRACKER
-
-// uncomment this to use MemPro (note: MemPro.cpp/need to be added to the project for this to work)
-//#define ENABLE_MEMPRO
+#include "MemPro/MemProProfiler.h"
 
 // There is a little memory and cpu overhead in tracking peak memory but it is generally more useful than current memory.
 // Disable if you need a little more memory or speed
 #define LLM_TRACK_PEAK_MEMORY 0		// currently disabled because there was a problem with tracking peaks from multiple threads.
 
-#ifdef ENABLE_MEMPRO
-namespace MemProProfiler
-{
-	ELLMTag TrackTag = ELLMTag::RHIMisc;		// GenericTagCount to track all allocs
-	bool bStart = true;
-}
-#include "MemPro.hpp"
-#endif //ENABLE_MEMPRO
-
 TAutoConsoleVariable<int32> CVarLLMWriteInterval(
 	TEXT("LLM.LLMWriteInterval"),
 	5,
@@ -585,6 +574,11 @@ void FLowLevelMemTracker::UpdateStatsPerFrame(const TCHAR* LogName)
 		GetTracker(ELLMTracker::Default)->SetTotalTags(ELLMTag::UntaggedTotal, ELLMTag::TrackedTotal);
 
 		bFirstTimeUpdating = false;
+
+#if MEMPRO_ENABLED
+		FMemProProfiler::PostInit();
+#endif
+
 	}
 
 	// update the trackers
@@ -800,8 +794,8 @@ void FLowLevelMemTracker::OnLowLevelAllocMoved(ELLMTracker Tracker, const void*
 		FPlatformMemory::OnLowLevelMemory_Alloc(Dest, Size, Tag);
 	}
 
-#ifdef ENABLE_MEMPRO
-	if (MemProProfiler::bStart && Tracker == ELLMTracker::Default && (MemProProfiler::TrackTag == ELLMTag::GenericTagCount || MemProProfiler::TrackTag == (ELLMTag)Tag))
+#if MEMPRO_ENABLED
+	if (FMemProProfiler::IsTrackingTag( (ELLMTag)Tag) )
 	{
 		MEMPRO_TRACK_FREE((void*)Source);
 		MEMPRO_TRACK_ALLOC((void*)Dest, (size_t)Size);
@@ -1611,8 +1605,8 @@ void FLLMTracker::FLLMThreadState::TrackAllocation(const void* Ptr, uint64 Size,
 		FPlatformMemory::OnLowLevelMemory_Alloc(Ptr, Size, Tag);
 	}
 	
-#ifdef ENABLE_MEMPRO
-	if (MemProProfiler::bStart && Tracker == ELLMTracker::Default && (MemProProfiler::TrackTag == ELLMTag::GenericTagCount || MemProProfiler::TrackTag == (ELLMTag)Tag))
+#if MEMPRO_ENABLED
+	if (FMemProProfiler::IsTrackingTag( (ELLMTag)Tag) )
 	{
 		MEMPRO_TRACK_ALLOC((void*)Ptr, (size_t)Size);
 	}
@@ -1632,8 +1626,8 @@ void FLLMTracker::FLLMThreadState::TrackFree(const void* Ptr, int64 Tag, uint64
 		FPlatformMemory::OnLowLevelMemory_Free(Ptr, Size, Tag);
 	}
 
-#ifdef ENABLE_MEMPRO
-	if (MemProProfiler::bStart && Tracker == ELLMTracker::Default && (MemProProfiler::TrackTag == ELLMTag::GenericTagCount || MemProProfiler::TrackTag == (ELLMTag)Tag))
+#if MEMPRO_ENABLED
+	if (FMemProProfiler::IsTrackingTag( (ELLMTag)Tag) )
 	{
 		MEMPRO_TRACK_FREE((void*)Ptr);
 	}
diff --git a/Engine/Source/Runtime/Core/Private/MemPro/MemProProfiler.cpp b/Engine/Source/Runtime/Core/Private/MemPro/MemProProfiler.cpp
new file mode 100644
--- /dev/null
+++ b/Engine/Source/Runtime/Core/Private/MemPro/MemProProfiler.cpp
@@ -0,0 +1,133 @@
+// Copyright 1998-2019 Epic Games, Inc. All Rights Reserved.
+
+#include "MemPro/MemProProfiler.h"
+
+#if MEMPRO_ENABLED
+#include "Misc/CString.h"
+#include "Misc/CoreDelegates.h"
+#include "HAL/IConsoleManager.h"
+
+DEFINE_LOG_CATEGORY_STATIC(LogMemPro, Log, All);
+
+/* NB. you can enable MemPro tracking after engine init by adding something like this to the command line:
+ *    -execcmds="MemPro.LLMTag RHIMisc, MemPro.Enabled 1"
+ */
+
+
+
+/* 
+ * Main runtime switch for MemPro support
+ */
+int32 GMemProEnabled = 0; //edit this and set to 1 to track from startup. You probably want to edit GMemProTrackTag too
+static FAutoConsoleVariableRef CVarMemProEnable(
+	TEXT("MemPro.Enabled"),
+	GMemProEnabled,
+	TEXT("Enable MemPro memory tracking.\n"),
+	ECVF_Default
+);
+
+
+/* 
+ * the LLM tag to track in MemPro, or ELLMTag::GenericTagCount to track all
+ */
+#if ENABLE_LOW_LEVEL_MEM_TRACKER
+
+	ELLMTag GMemProTrackTag = ELLMTag::EngineMisc;
+
+#endif //ENABLE_LOW_LEVEL_MEM_TRACKER
+
+
+
+/* 
+ * helper function to track a tag
+ */
+#if ENABLE_LOW_LEVEL_MEM_TRACKER
+void FMemProProfiler::TrackTag( ELLMTag Tag )
+{
+	GMemProTrackTag = Tag;
+}
+#endif //ENABLE_LOW_LEVEL_MEM_TRACKER
+
+
+/* 
+ * helper function to track a tag given its name
+ */
+#if ENABLE_LOW_LEVEL_MEM_TRACKER
+void FMemProProfiler::TrackTagByName( const TCHAR* TagName )
+{
+	//sanity check
+	if( TagName == nullptr || FCString::Strlen(TagName) == 0 )
+	{
+		UE_LOG( LogMemPro, Display, TEXT("please specify an LLM tag or * to track all") );
+		return;
+	}
+
+	//check whether they want to track all tags
+	if( FCString::Strcmp(TagName, TEXT("*") ) == 0 )
+	{
+		TrackTag( ELLMTag::GenericTagCount );
+		UE_LOG( LogMemPro, Display, TEXT("tracking all LLM tags" ) );
+	}
+	else
+	{
+		//find the specific tag to track
+		uint64 TagIndex = (uint64)ELLMTag::Paused;
+		if ( FLowLevelMemTracker::Get().FindTagByName(TagName, TagIndex) && TagIndex < LLM_TAG_COUNT )
+		{
+			TrackTag( (ELLMTag)TagIndex );
+			UE_LOG( LogMemPro, Display, TEXT("tracking LLM tag \'%s\'" ), TagName );
+		}
+		else
+		{
+			UE_LOG( LogMemPro, Display, TEXT("Unknown LLM tag \'%s\'" ), TagName );
+		}
+	}
+}
+#endif //ENABLE_LOW_LEVEL_MEM_TRACKER
+
+
+/*
+ * console command to get mempro to track a specific LLM tag
+ */
+#if ENABLE_LOW_LEVEL_MEM_TRACKER
+static FAutoConsoleCommand MemProTrackLLMTag(
+	TEXT("MemPro.LLMTag"),
+	TEXT("Capture a specific LLM tag with MemPro"),
+	FConsoleCommandWithArgsDelegate::CreateLambda( [](const TArray<FString>& Args )
+	{
+		FMemProProfiler::TrackTagByName( (Args.Num() == 0) ? nullptr : *Args[0] );
+	})
+);
+#endif //ENABLE_LOW_LEVEL_MEM_TRACKER
+
+
+/*
+ * query the port that MemPro might be using so other development tools can steer clear if necessary
+ */
+bool FMemProProfiler::IsUsingPort( uint32 Port )
+{
+#if defined(MEMPRO_WRITE_DUMP)
+	return false;
+#else
+	return Port == FCStringAnsi::Atoi(MEMPRO_PORT);
+#endif
+}
+
+
+/*
+ * initialisation for MemPro. IT may have
+ */
+
+void FMemProProfiler::PostInit()
+{
+	//shutdown MemPro when the engine is shutting down so that the send thread terminates cleanly
+	FCoreDelegates::OnPreExit.AddLambda( []()
+	{
+		GMemProEnabled = 0;
+		MemPro::Disconnect();
+		MemPro::Shutdown();
+	});
+
+}
+
+#endif //MEMPRO_ENABLED
diff --git a/Engine/Source/Runtime/Core/Public/MemPro/MemProProfiler.h b/Engine/Source/Runtime/Core/Public/MemPro/MemProProfiler.h
new file mode 100644
--- /dev/null
+++ b/Engine/Source/Runtime/Core/Public/MemPro/MemProProfiler.h
@@ -0,0 +1,46 @@
+// Copyright 1998-2019 Epic Games, Inc. All Rights Reserved.
+
+#pragma once
+
+
+/* 
+ * To enable MemPro support in the engine, add GlobalDefinitions.Add("MEMPRO_ENABLED=1"); to your game's .Target.cs build rules file for supported platforms.
+ * (note: MemPro.cpp/.h need to be added to the project for this to work)
+ */
+#if !defined(MEMPRO_ENABLED)
+	#define MEMPRO_ENABLED 0		
+#endif
+
+
+
+#if MEMPRO_ENABLED
+#include "CoreGlobals.h"
+#include "HAL/LowLevelMemTracker.h"
+#include "MemPro/MemPro.h"
+
+class CORE_API FMemProProfiler
+{
+public:
+	static void PostInit();
+
+	static bool IsUsingPort( uint32 Port );
+
+	static inline bool IsStarted()
+	{
+		extern int32 GMemProEnabled;
+		return (GMemProEnabled != 0) && !GIsRequestingExit;
+	}
+
+#if ENABLE_LOW_LEVEL_MEM_TRACKER
+	static inline bool IsTrackingTag( ELLMTag Tag )
+	{
+		extern ELLMTag GMemProTrackTag;
+		return IsStarted() && (GMemProTrackTag != ELLMTag::Paused) && ((Tag == GMemProTrackTag) || (GMemProTrackTag == ELLMTag::GenericTagCount));
+	}
+
+	static void TrackTag( ELLMTag Tag );
+	static void TrackTagByName( const TCHAR* TagName );
+#endif //ENABLE_LOW_LEVEL_MEM_TRACKER
+};
+
+#endif //MEMPRO_ENABLED
