[LLM] add an LLM_DUMP_TAG macro that logs the current tag and file location in order to identify mis-tagged items.

NB. it is not expected that code will be checked in with this macro in use - just during debugging/taggig sessions.

diff --git a/Engine/Source/Runtime/Core/Private/HAL/LowLevelMemTracker.cpp b/Engine/Source/Runtime/Core/Private/HAL/LowLevelMemTracker.cpp
--- a/Engine/Source/Runtime/Core/Private/HAL/LowLevelMemTracker.cpp
+++ b/Engine/Source/Runtime/Core/Private/HAL/LowLevelMemTracker.cpp
@@ -928,6 +928,16 @@ int64 FLowLevelMemTracker::GetActiveTag(ELLMTracker Tracker)
     return GetTracker(Tracker)->GetActiveTag();
 }
 
+void FLowLevelMemTracker::DumpTag( ELLMTracker Tracker, const char* FileName, int LineNumber )
+{
+	int64 Tag = GetActiveTag(Tracker);
+	const TCHAR* TagName = FindTagName(Tag);
+
+	FPlatformMisc::LowLevelOutputDebugStringf( TEXT("LLM TAG: %s (%lld) @ %s:%d\n"), TagName ? TagName : TEXT("<unknown>"), Tag, FileName ? ANSI_TO_TCHAR(FileName) : TEXT("?"), LineNumber );
+}
+
+
+
 
 FLLMScope::FLLMScope(FName StatIDName, ELLMTagSet Set, ELLMTracker Tracker)
 {
diff --git a/Engine/Source/Runtime/Core/Public/HAL/LowLevelMemTracker.h b/Engine/Source/Runtime/Core/Public/HAL/LowLevelMemTracker.h
--- a/Engine/Source/Runtime/Core/Public/HAL/LowLevelMemTracker.h
+++ b/Engine/Source/Runtime/Core/Public/HAL/LowLevelMemTracker.h
@@ -321,6 +321,13 @@ extern FName LLMGetTagStat(ELLMTag Tag);
 	#define LLM_PUSH_STATS_FOR_ASSET_TAGS()
 #endif
 
+/**
+ * LLM tag dumping, to help with identifying mis-tagged items. Probably don't want to check in with these in use!
+ */
+#define LLM_DUMP_TAG()  FLowLevelMemTracker::Get().DumpTag(ELLMTracker::Default,__FILE__,__LINE__);
+#define LLM_DUMP_PLATFORM_TAG()  FLowLevelMemTracker::Get().DumpTag(ELLMTracker::Platform,__FILE__,__LINE__);
+
+
 typedef void*(*LLMAllocFunction)(size_t);
 typedef void(*LLMFreeFunction)(void*, size_t);
 
@@ -449,6 +456,9 @@ public:
 	// Get the amount of memory for a tag from the given tracker
 	int64 GetTagAmountForTracker(ELLMTracker Tracker, ELLMTag Tag);
 
+	// Dump the current tag for the given tracker to the output
+	void DumpTag( ELLMTracker Tracker, const char* FileName, int LineNumber );
+
 private:
 	FLowLevelMemTracker();
 
@@ -556,4 +566,7 @@ protected:
 	#define LLM_SCOPED_PAUSE_TRACKING_WITH_ENUM_AND_AMOUNT(...)
 	#define LLM_SCOPED_PAUSE_TRACKING_WITH_STAT_AND_AMOUNT(...)
 	#define LLM_PUSH_STATS_FOR_ASSET_TAGS()
+	#define LLM_DUMP_TAG()
+	#define LLM_DUMP_PLATFORM_TAG()
+
 #endif		// #if ENABLE_LOW_LEVEL_MEM_TRACKER
