[LLM] mempro shutdown fixes

diff --git a/Engine/Source/Runtime/Core/Private/MemPro/MemProProfiler.cpp b/Engine/Source/Runtime/Core/Private/MemPro/MemProProfiler.cpp
--- a/Engine/Source/Runtime/Core/Private/MemPro/MemProProfiler.cpp
+++ b/Engine/Source/Runtime/Core/Private/MemPro/MemProProfiler.cpp
@@ -123,9 +123,10 @@ void FMemProProfiler::PostInit()
 	//shutdown MemPro when the engine is shutting down so that the send thread terminates cleanly
 	FCoreDelegates::OnPreExit.AddLambda( []()
 	{
+		MemPro::SetPaused(true);
 		GMemProEnabled = 0;
 		MemPro::Disconnect();
-		MemPro::Shutdown();
+		//MemPro::Shutdown(); ...disabled for now as I was getting hangs on shutdown.
 	});
 
 }
