[LLM] adding new SceneRender tag, moving up to 50 MB from EngineMisc to new SceneRender tag.

diff --git a/Engine/Source/Runtime/Core/Private/HAL/LowLevelMemTracker.cpp b/Engine/Source/Runtime/Core/Private/HAL/LowLevelMemTracker.cpp
--- a/Engine/Source/Runtime/Core/Private/HAL/LowLevelMemTracker.cpp
+++ b/Engine/Source/Runtime/Core/Private/HAL/LowLevelMemTracker.cpp
@@ -70,6 +70,7 @@ DECLARE_LLM_MEMORY_STAT(TEXT("PSO"), STAT_PSOLLM, STATGROUP_LLMFULL);
 DECLARE_LLM_MEMORY_STAT(TEXT("Textures"), STAT_TexturesLLM, STATGROUP_LLMFULL);
 DECLARE_LLM_MEMORY_STAT(TEXT("TextureMetaData"), STAT_TextureMetaDataLLM, STATGROUP_LLMFULL);
 DECLARE_LLM_MEMORY_STAT(TEXT("Render Targets"), STAT_RenderTargetsLLM, STATGROUP_LLMFULL);
+DECLARE_LLM_MEMORY_STAT(TEXT("SceneRender"), STAT_SceneRenderLLM, STATGROUP_LLMFULL);
 DECLARE_LLM_MEMORY_STAT(TEXT("RHIMisc"), STAT_RHIMiscLLM, STATGROUP_LLMFULL);
 DECLARE_LLM_MEMORY_STAT(TEXT("PhysX TriMesh"), STAT_PhysXTriMeshLLM, STATGROUP_LLMFULL);
 DECLARE_LLM_MEMORY_STAT(TEXT("PhysX ConvexMesh"), STAT_PhysXConvexMeshLLM, STATGROUP_LLMFULL);
diff --git a/Engine/Source/Runtime/Core/Public/HAL/LowLevelMemTracker.h b/Engine/Source/Runtime/Core/Public/HAL/LowLevelMemTracker.h
--- a/Engine/Source/Runtime/Core/Public/HAL/LowLevelMemTracker.h
+++ b/Engine/Source/Runtime/Core/Public/HAL/LowLevelMemTracker.h
@@ -189,6 +189,7 @@ enum class ELLMTagSet : uint8
 	macro(Textures,								"Textures",						GET_STATFNAME(STAT_TexturesLLM),							GET_STATFNAME(STAT_TexturesSummaryLLM),			-1)\
 	macro(TextureMetaData,						"TextureMetaData",				GET_STATFNAME(STAT_TextureMetaDataLLM),						GET_STATFNAME(STAT_TexturesSummaryLLM),			-1)\
 	macro(RenderTargets,						"RenderTargets",				GET_STATFNAME(STAT_RenderTargetsLLM),						GET_STATFNAME(STAT_EngineSummaryLLM),			-1)\
+	macro(SceneRender,							"SceneRender",					GET_STATFNAME(STAT_SceneRenderLLM),							GET_STATFNAME(STAT_EngineSummaryLLM),			-1)\
 	macro(RHIMisc,								"RHIMisc",						GET_STATFNAME(STAT_RHIMiscLLM),								GET_STATFNAME(STAT_EngineSummaryLLM),			-1)\
 	macro(AsyncLoading,							"AsyncLoading",					GET_STATFNAME(STAT_AsyncLoadingLLM),						GET_STATFNAME(STAT_EngineSummaryLLM),			-1)\
 	macro(UObject,								"UObject",						GET_STATFNAME(STAT_UObjectLLM),								GET_STATFNAME(STAT_UObjectSummaryLLM),			-1)\
diff --git a/Engine/Source/Runtime/Renderer/Private/SceneRendering.cpp b/Engine/Source/Runtime/Renderer/Private/SceneRendering.cpp
--- a/Engine/Source/Runtime/Renderer/Private/SceneRendering.cpp
+++ b/Engine/Source/Runtime/Renderer/Private/SceneRendering.cpp
@@ -49,6 +49,7 @@
 #include "VisualizeTexture.h"
 #include "VisualizeTexturePresent.h"
 #include "MeshDrawCommands.h"
+#include "HAL/LowLevelMemTracker.h"
 
 /*-----------------------------------------------------------------------------
 	Globals
@@ -3084,6 +3085,8 @@ static TAutoConsoleVariable<int32> CVarDelaySceneRenderCompletion(
  */
 static void RenderViewFamily_RenderThread(FRHICommandListImmediate& RHICmdList, FSceneRenderer* SceneRenderer)
 {
+	LLM_SCOPE(ELLMTag::SceneRender);
+
 	{
 		QUICK_SCOPE_CYCLE_COUNTER(STAT_DelaySceneRenderCompletion_TaskWait);
 		FRHICommandListExecutor::GetImmediateCommandList().ImmediateFlush(EImmediateFlushType::WaitForOutstandingTasksOnly);
