Reduce the amount of stack memory used in Mac crash handling code to fit within the new limits in macOS 10.14.4 and solve the issues with launching CrashReportClient for some crashes on this version of the OS.

diff --git a/Engine/Source/Runtime/Core/Private/Apple/ApplePlatformStackWalk.cpp b/Engine/Source/Runtime/Core/Private/Apple/ApplePlatformStackWalk.cpp
--- a/Engine/Source/Runtime/Core/Private/Apple/ApplePlatformStackWalk.cpp
+++ b/Engine/Source/Runtime/Core/Private/Apple/ApplePlatformStackWalk.cpp
@@ -119,7 +119,7 @@ static void AsyncSafeProgramCounterToSymbolInfo( uint64 ProgramCounter, FProgram
 	// Increased the size of the demangle destination to reduce the chances that abi::__cxa_demangle will allocate
 	// this causes the app to hang as malloc isn't signal handler safe. Ideally we wouldn't call this function in a handler.
 	size_t DemangledNameLen = 65536;
-	ANSICHAR DemangledNameBuffer[65536]= {0};
+	static ANSICHAR DemangledNameBuffer[65536]= {0};
 	DemangledName = abi::__cxa_demangle(DylibInfo.dli_sname, DemangledNameBuffer, &DemangledNameLen, &Status);
 	
 	if (DemangledName)
diff --git a/Engine/Source/Runtime/Core/Private/Mac/MacPlatformMisc.cpp b/Engine/Source/Runtime/Core/Private/Mac/MacPlatformMisc.cpp
--- a/Engine/Source/Runtime/Core/Private/Mac/MacPlatformMisc.cpp
+++ b/Engine/Source/Runtime/Core/Private/Mac/MacPlatformMisc.cpp
@@ -1495,6 +1495,7 @@ static uint32 GMacStackIgnoreDepth = 6;
 /** Message for the assert triggered on this thread */
 thread_local const TCHAR* GCrashErrorMessage = nullptr;
 thread_local ECrashContextType GCrashErrorType = ECrashContextType::Crash;
+thread_local uint8* GCrashContextMemory[sizeof(FMacCrashContext)];
 
 /** True system-specific crash handler that gets called first */
 static void PlatformCrashHandler(int32 Signal, siginfo_t* Info, void* Context)
@@ -1515,23 +1516,23 @@ static void PlatformCrashHandler(int32 Signal, siginfo_t* Info, void* Context)
 		Type = GCrashErrorType;
 		ErrorMessage = GCrashErrorMessage;
 	}
-	
-	FMacCrashContext CrashContext(Type, ErrorMessage);
-	CrashContext.IgnoreDepth = GMacStackIgnoreDepth;
-	CrashContext.InitFromSignal(Signal, Info, Context);
-	
+
+	FMacCrashContext* CrashContext = new (GCrashContextMemory) FMacCrashContext(Type, ErrorMessage);
+	CrashContext->IgnoreDepth = GMacStackIgnoreDepth;
+	CrashContext->InitFromSignal(Signal, Info, Context);
+
 	// Switch to crash handler malloc to avoid malloc reentrancy
 	check(GCrashMalloc);
-	GCrashMalloc->Enable(&CrashContext, FPlatformTLS::GetCurrentThreadId());
-	
+	GCrashMalloc->Enable(CrashContext, FPlatformTLS::GetCurrentThreadId());
+
 	if (GCrashHandlerPointer)
 	{
-		GCrashHandlerPointer(CrashContext);
+		GCrashHandlerPointer(*CrashContext);
 	}
 	else
 	{
 		// call default one
-		DefaultCrashHandler(CrashContext);
+		DefaultCrashHandler(*CrashContext);
 	}
 }
 
