[LLM] async work tasks inherit the LLM tag scope of the thread that created them. Moves some items out of Untracked:

in my test replay run, notable items were:
~6MB Shaders
~15MB UObject
~2MB StreamingManager

diff --git a/Engine/Source/Runtime/Core/Public/Async/AsyncWork.h b/Engine/Source/Runtime/Core/Public/Async/AsyncWork.h
--- a/Engine/Source/Runtime/Core/Public/Async/AsyncWork.h
+++ b/Engine/Source/Runtime/Core/Public/Async/AsyncWork.h
@@ -13,6 +13,7 @@
 #include "Stats/Stats.h"
 #include "HAL/Event.h"
 #include "HAL/PlatformProcess.h"
+#include "HAL/LowLevelMemTracker.h"
 #include "Misc/IQueuedWork.h"
 #include "Misc/QueuedThreadPool.h"
 
@@ -60,12 +61,16 @@ class FAutoDeleteAsyncTask
 {
 	/** User job embedded in this task */
 	TTask Task;
+	/** optional LLM tag */
+	LLM(ELLMTag InheritedLLMTag);
 
 	/* Generic start function, not called directly
 	 * @param bForceSynchronous if true, this job will be started synchronously, now, on this thread
 	 **/
 	void Start(bool bForceSynchronous, FQueuedThreadPool* InQueuedPool)
 	{
+		LLM(InheritedLLMTag = FLowLevelMemTracker::bIsDisabled ? ELLMTag::Untagged : (ELLMTag)FLowLevelMemTracker::Get().GetActiveTag(ELLMTracker::Default));
+
 		FPlatformMisc::MemoryBarrier();
 		FQueuedThreadPool* QueuedPool = InQueuedPool;
 		if (bForceSynchronous)
@@ -88,6 +93,7 @@ class FAutoDeleteAsyncTask
 	 **/
 	void DoWork()
 	{
+		LLM_SCOPE(InheritedLLMTag);
 		FScopeCycleCounter Scope(Task.GetStatId(), true);
 
 		Task.DoWork();
@@ -209,6 +215,8 @@ class FAsyncTask
 	FEvent*				DoneEvent;
 	/** Pool we are queued into, maintained by the calling thread */
 	FQueuedThreadPool*	QueuedPool;
+	/** optional LLM tag */
+	LLM(ELLMTag InheritedLLMTag);
 
 	/* Internal function to destroy the completion event
 	**/
@@ -225,6 +233,7 @@ class FAsyncTask
 	{
 		FScopeCycleCounter Scope( Task.GetStatId(), true );
 		DECLARE_SCOPE_CYCLE_COUNTER( TEXT( "FAsyncTask::Start" ), STAT_FAsyncTask_Start, STATGROUP_ThreadPoolAsyncTasks );
+		LLM(InheritedLLMTag = FLowLevelMemTracker::bIsDisabled ? ELLMTag::Untagged : (ELLMTag)FLowLevelMemTracker::Get().GetActiveTag(ELLMTracker::Default));
 
 		FPlatformMisc::MemoryBarrier();
 		CheckIdle();  // can't start a job twice without it being completed first
@@ -256,6 +265,7 @@ class FAsyncTask
 	**/
 	void DoWork()
 	{	
+		LLM_SCOPE(InheritedLLMTag);
 		FScopeCycleCounter Scope(Task.GetStatId(), true); 
 
 		Task.DoWork();		
