Add EngineData to crash platform context to allow engine modules to append their own information such as active RHI and feature level.

diff --git a/Engine/Source/Runtime/Core/Private/GenericPlatform/GenericPlatformCrashContext.cpp b/Engine/Source/Runtime/Core/Private/GenericPlatform/GenericPlatformCrashContext.cpp
--- a/Engine/Source/Runtime/Core/Private/GenericPlatform/GenericPlatformCrashContext.cpp
+++ b/Engine/Source/Runtime/Core/Private/GenericPlatform/GenericPlatformCrashContext.cpp
@@ -46,6 +46,7 @@ const FString FGenericCrashContext::CrashGUIDRootPrefix = TEXT("UE4CC-");
 const FString FGenericCrashContext::CrashContextExtension = TEXT(".runtime-xml");
 const FString FGenericCrashContext::RuntimePropertiesTag = TEXT( "RuntimeProperties" );
 const FString FGenericCrashContext::PlatformPropertiesTag = TEXT( "PlatformProperties" );
+const FString FGenericCrashContext::EngineDataTag = TEXT( "EngineData" );
 const FString FGenericCrashContext::GameDataTag = TEXT( "GameData" );
 const FString FGenericCrashContext::EnabledPluginsTag = TEXT("EnabledPlugins");
 const FString FGenericCrashContext::UE4MinidumpName = TEXT( "UE4Minidump.dmp" );
@@ -102,6 +103,7 @@ namespace NCachedCrashContextProperties
 	static FString CrashReportClientRichText;
 	static FString GameStateName;
 	static TArray<FString> EnabledPluginsList;
+	static TMap<FString, FString> EngineData;
 	static TMap<FString, FString> GameData;
 }
 
@@ -409,6 +411,14 @@ void FGenericCrashContext::SerializeContentToBuffer() const
 	AddPlatformSpecificProperties();
 	EndSection( *PlatformPropertiesTag );
 
+	// Add the engine data
+	BeginSection( *EngineDataTag );
+	for (const TPair<FString, FString>& Pair : NCachedCrashContextProperties::EngineData)
+	{
+		AddCrashProperty(*Pair.Key, *Pair.Value);
+	}
+	EndSection( *EngineDataTag );
+
 	// Add the game data
 	BeginSection( *GameDataTag );
 	for (const TPair<FString, FString>& Pair : NCachedCrashContextProperties::GameData)
@@ -669,6 +679,39 @@ void FGenericCrashContext::PurgeOldCrashConfig()
 	}
 }
 
+void FGenericCrashContext::ResetEngineData()
+{
+	NCachedCrashContextProperties::EngineData.Reset();
+}
+
+void FGenericCrashContext::SetEngineData(const FString& Key, const FString& Value)
+{
+	if (Value.Len() == 0)
+	{
+		// for testing purposes, only log values when they change, but don't pay the lookup price normally.
+		UE_SUPPRESS(LogCrashContext, VeryVerbose, 
+		{
+			if (NCachedCrashContextProperties::EngineData.Find(Key))
+			{
+				UE_LOG(LogCrashContext, VeryVerbose, TEXT("FGenericCrashContext::SetEngineData(%s, <RemoveKey>)"), *Key);
+			}
+		});
+		NCachedCrashContextProperties::EngineData.Remove(Key);
+	}
+	else
+	{
+		FString& OldVal = NCachedCrashContextProperties::EngineData.FindOrAdd(Key);
+		UE_SUPPRESS(LogCrashContext, VeryVerbose, 
+		{
+			if (OldVal != Value)
+			{
+				UE_LOG(LogCrashContext, VeryVerbose, TEXT("FGenericCrashContext::SetEngineData(%s, %s)"), *Key, *Value);
+			}
+		});
+		OldVal = Value;
+	}
+}
+
 void FGenericCrashContext::ResetGameData()
 {
 	NCachedCrashContextProperties::GameData.Reset();
diff --git a/Engine/Source/Runtime/Core/Public/GenericPlatform/GenericPlatformCrashContext.h b/Engine/Source/Runtime/Core/Public/GenericPlatform/GenericPlatformCrashContext.h
--- a/Engine/Source/Runtime/Core/Public/GenericPlatform/GenericPlatformCrashContext.h
+++ b/Engine/Source/Runtime/Core/Public/GenericPlatform/GenericPlatformCrashContext.h
@@ -120,6 +120,7 @@ public:
 	static const FString CrashContextExtension;
 	static const FString RuntimePropertiesTag;
 	static const FString PlatformPropertiesTag;
+	static const FString EngineDataTag;
 	static const FString GameDataTag;
 	static const FString EnabledPluginsTag;
 	static const FString UE4MinidumpName;
@@ -213,6 +214,12 @@ public:
 	/** Helper to clean out old files in the crash report client config folder. */
 	static void PurgeOldCrashConfig();
 
+	/** Clears the engine data dictionary */
+	static void ResetEngineData();
+
+	/** Updates (or adds if not already present) arbitrary engine data to the crash context (will remove the key if passed an empty string) */
+	static void SetEngineData(const FString& Key, const FString& Value);
+
 	/** Clears the game data dictionary */
 	static void ResetGameData();
 
diff --git a/Engine/Source/Runtime/RHI/Private/DynamicRHI.cpp b/Engine/Source/Runtime/RHI/Private/DynamicRHI.cpp
--- a/Engine/Source/Runtime/RHI/Private/DynamicRHI.cpp
+++ b/Engine/Source/Runtime/RHI/Private/DynamicRHI.cpp
@@ -12,6 +12,7 @@
 #include "RHI.h"
 #include "Modules/ModuleManager.h"
 #include "GenericPlatform/GenericPlatformDriver.h"
+#include "GenericPlatform/GenericPlatformCrashContext.h"
 #include "PipelineStateCache.h"
 
 #ifndef PLATFORM_ALLOW_NULL_RHI
@@ -53,6 +54,9 @@ void InitNullRHI()
 	GRHICommandList.GetImmediateAsyncComputeCommandList().SetComputeContext(GDynamicRHI->RHIGetDefaultAsyncComputeContext());
 	GUsingNullRHI = true;
 	GRHISupportsTextureStreaming = false;
+
+	// Update the crash context analytics
+	FGenericCrashContext::SetEngineData(TEXT("RHI.RHIName"), TEXT("NullRHI"));
 }
 
 #if PLATFORM_WINDOWS
@@ -207,15 +211,24 @@ void RHIInit(bool bHasEditorToken)
 				GRHICommandList.GetImmediateCommandList().SetContext(GDynamicRHI->RHIGetDefaultContext());
 				GRHICommandList.GetImmediateAsyncComputeCommandList().SetComputeContext(GDynamicRHI->RHIGetDefaultAsyncComputeContext());
 
+				FString FeatureLevelString;
+				GetFeatureLevelName(GMaxRHIFeatureLevel, FeatureLevelString);
+
 				if (bHasEditorToken && GMaxRHIFeatureLevel < ERHIFeatureLevel::SM5)
 				{
-					FString FeatureLevelString;
-					GetFeatureLevelName(GMaxRHIFeatureLevel, FeatureLevelString);
 					FString ShaderPlatformString = LegacyShaderPlatformToShaderFormat(GetFeatureLevelShaderPlatform(GMaxRHIFeatureLevel)).ToString();
 					FString Error = FString::Printf(TEXT("A Feature Level 5 video card is required to run the editor.\nAvailableFeatureLevel = %s, ShaderPlatform = %s"), *FeatureLevelString, *ShaderPlatformString);
 					FMessageDialog::Open(EAppMsgType::Ok, FText::FromString(Error));
 					FPlatformMisc::RequestExit(1);
 				}
+				
+				// Update the crash context analytics
+				FGenericCrashContext::SetEngineData(TEXT("RHI.RHIName"), GDynamicRHI ? GDynamicRHI->GetName() : TEXT("Unknown"));
+				FGenericCrashContext::SetEngineData(TEXT("RHI.AdapterName"), GRHIAdapterName);
+				FGenericCrashContext::SetEngineData(TEXT("RHI.UserDriverVersion"), GRHIAdapterUserDriverVersion);
+				FGenericCrashContext::SetEngineData(TEXT("RHI.InternalDriverVersion"), GRHIAdapterInternalDriverVersion);
+				FGenericCrashContext::SetEngineData(TEXT("RHI.DriverDate"), GRHIAdapterDriverDate);
+				FGenericCrashContext::SetEngineData(TEXT("RHI.FeatureLevel"), FeatureLevelString);
 			}
 #if PLATFORM_ALLOW_NULL_RHI
 			else
