Add back bAgreeToCrashUpload setting to check if we should implicitly upload unattended crashes

With the 4.22 Engine merged over we moved to using an Editor setting which broke crash sending

diff --git a/Engine/Source/Runtime/Core/Private/Unix/UnixPlatformCrashContext.cpp b/Engine/Source/Runtime/Core/Private/Unix/UnixPlatformCrashContext.cpp
--- a/Engine/Source/Runtime/Core/Private/Unix/UnixPlatformCrashContext.cpp
+++ b/Engine/Source/Runtime/Core/Private/Unix/UnixPlatformCrashContext.cpp
@@ -395,15 +395,20 @@ void FUnixCrashContext::GenerateCrashInfoAndLaunchReporter(bool bReportingNonCra
 		CrashGuid = FGuid::NewGuid().ToString();
 	}
 
-	/* Table showing the desired behavior when wanting to show a pop up or not. As well as avoiding starting CRC
-	 *  based on an *.ini setting bSendUnattendedBugReports and whether or not we are unattended
+
+	/* Table showing the desired behavior when wanting to start the CRC or not.
+	 *  based on an *.ini setting for bSendUnattendedBugReports or bAgreeToCrashUpload and whether or not we are unattended
 	 *
-	 *  Unattended | AgreeToUpload || Show Popup | Start CRC
-	 *  ----------------------------------------------------
-	 *      1      |       1       ||     0      |    1
-	 *      1      |       0       ||     0      |    0
-	 *      0      |       1       ||     0      |    1
-	 *      0      |       0       ||     1      |    1
+	 *  Unattended | AgreeToUpload | SendUnattendedBug || Start CRC
+	 *  ------------------------------------------------------------
+	 *      1      |       1       |         1         ||     1
+	 *      1      |       1       |         0         ||     1
+	 *      1      |       0       |         1         ||     1
+	 *      1      |       0       |         0         ||     0
+	 *      0      |       1       |         1         ||     1
+	 *      0      |       1       |         0         ||     1
+	 *      0      |       0       |         1         ||     1
+	 *      0      |       0       |         0         ||     1
 	 */
 
 	// Suppress the user input dialog if we're running in unattended mode
@@ -428,13 +433,21 @@ void FUnixCrashContext::GenerateCrashInfoAndLaunchReporter(bool bReportingNonCra
 		GConfig->GetBool(TEXT("/Script/UnrealEd.CrashReportsPrivacySettings"), TEXT("bSendUnattendedBugReports"), bSendUnattendedBugReports, GEditorSettingsIni);
 	}
 
+	// If we are not an editor but still want to agree to upload for non-licensee check the settings
+	bool bAgreeToCrashUpload = false;
+	if (!UE_EDITOR && GConfig)
+	{
+		GConfig->GetBool(TEXT("CrashReportClient"), TEXT("bAgreeToCrashUpload"), bAgreeToCrashUpload, GEngineIni);
+	}
+
 	if (BuildSettings::IsLicenseeVersion() && !UE_EDITOR)
 	{
 		// do not send unattended reports in licensees' builds except for the editor, where it is governed by the above setting
 		bSendUnattendedBugReports = false;
+		bAgreeToCrashUpload = false;
 	}
 
-	bool bSkipCRC = bUnattended && !bSendUnattendedBugReports;
+	bool bSkipCRC = bUnattended && !bSendUnattendedBugReports && !bAgreeToCrashUpload;
 
 	if (!bSkipCRC)
 	{
