Do not check the stats on a file that does not exist yet. We can check the dir instead to check if we are low mem or not

diff --git a/Engine/Source/Runtime/Core/Private/Unix/UnixPlatformCrashContext.cpp b/Engine/Source/Runtime/Core/Private/Unix/UnixPlatformCrashContext.cpp
--- a/Engine/Source/Runtime/Core/Private/Unix/UnixPlatformCrashContext.cpp
+++ b/Engine/Source/Runtime/Core/Private/Unix/UnixPlatformCrashContext.cpp
@@ -492,7 +492,7 @@ void FUnixCrashContext::GenerateCrashInfoAndLaunchReporter(bool bReportingNonCra
 			uint64 TotalDiskSpace = 0;
 			uint64 TotalDiskFreeSpace = 0;
 			bool bLowDriveSpace = false;
-			if (FPlatformMisc::GetDiskTotalAndFreeSpace(*LogDstAbsolute, TotalDiskSpace, TotalDiskFreeSpace))
+			if (FPlatformMisc::GetDiskTotalAndFreeSpace(*CrashInfoAbsolute, TotalDiskSpace, TotalDiskFreeSpace))
 			{
 				if (TotalDiskFreeSpace < MinDriveSpaceForCrashLog)
 				{
