[Switch] Add detection and reporting of live crashes due to hang detector triggering.

diff --git a/Engine/Source/Runtime/Core/Private/HAL/ThreadHeartBeat.cpp b/Engine/Source/Runtime/Core/Private/HAL/ThreadHeartBeat.cpp
--- a/Engine/Source/Runtime/Core/Private/HAL/ThreadHeartBeat.cpp
+++ b/Engine/Source/Runtime/Core/Private/HAL/ThreadHeartBeat.cpp
@@ -16,6 +16,9 @@
 #include "Stats/Stats.h"
 #include "Async/TaskGraphInterfaces.h"
 
+#if PLATFORM_SWITCH
+#include "Switch/SwitchPlatformCrashContext.h"
+#endif
 // When enabled, the heart beat thread will call abort() when a hang
 // is detected, rather than performing stack back-traces and logging.
 #define MINIMAL_FATAL_HANG_DETECTION	((PLATFORM_PS4 || PLATFORM_XBOXONE || PLATFORM_SWITCH) && 1)
@@ -67,7 +70,7 @@ FThreadHeartBeat::FThreadHeartBeat()
 	, CurrentPresentDuration(0)
 	, HangDurationMultiplier(1.0)
 	, LastHangCallstackCRC(0)
-	, LastHungThreadId(0)
+	, LastHungThreadId(InvalidThreadId)
 	, Clock(HangDetectorClock_MaxTimeStep_MS / 1000)
 {
 	// Start with the frame-present based hang detection disabled. This will be automatically enabled on
@@ -149,7 +152,9 @@ void FORCENOINLINE FThreadHeartBeat::OnPresentHang(double HangDuration)
 #if MINIMAL_FATAL_HANG_DETECTION
 
 	LastHungThreadId = FThreadHeartBeat::PresentThreadId;
-
+#if PLATFORM_SWITCH
+	FPlatformCrashContext::UpdateDynamicData();
+#endif
 	// We want to avoid all memory allocations if a hang is detected.
 	// Force a crash in a way that will generate a crash report.
 
@@ -171,7 +176,9 @@ void FORCENOINLINE FThreadHeartBeat::OnHang(double HangDuration, uint32 ThreadTh
 #if MINIMAL_FATAL_HANG_DETECTION
 
 	LastHungThreadId = ThreadThatHung;
-
+#if PLATFORM_SWITCH
+	FPlatformCrashContext::UpdateDynamicData();
+#endif
 	// We want to avoid all memory allocations if a hang is detected.
 	// Force a crash in a way that will generate a crash report.
 
diff --git a/Engine/Source/Runtime/Core/Public/HAL/ThreadHeartBeat.h b/Engine/Source/Runtime/Core/Public/HAL/ThreadHeartBeat.h
--- a/Engine/Source/Runtime/Core/Public/HAL/ThreadHeartBeat.h
+++ b/Engine/Source/Runtime/Core/Public/HAL/ThreadHeartBeat.h
@@ -160,6 +160,12 @@ public:
 	 */
 	void SetDurationMultiplier(double NewMultiplier);
 
+	/*
+	* Get the Id of the last thread to trigger the hang detector.
+	* Returns InvalidThreadId if hang detector has not been triggered.
+	*/
+	uint32 GetLastHungThreadId() const { return LastHungThreadId; }
+
 	//~ Begin FRunnable Interface.
 	virtual bool Init();
 	virtual uint32 Run();
