Prevent MallocAnsi using _aligned_malloc on VC toolchains because _aligned_msize requires the correct alignment to be passed in, otherwise it returns garbage (seems to be a change in newer versions of VC)

diff --git a/Engine/Source/Runtime/Core/Public/HAL/MallocAnsi.h b/Engine/Source/Runtime/Core/Public/HAL/MallocAnsi.h
--- a/Engine/Source/Runtime/Core/Public/HAL/MallocAnsi.h
+++ b/Engine/Source/Runtime/Core/Public/HAL/MallocAnsi.h
@@ -5,7 +5,7 @@
 #include "CoreTypes.h"
 #include "HAL/UnrealMemory.h"
 
-#if defined(_MSC_VER) || PLATFORM_MAC || PLATFORM_IOS
+#if PLATFORM_MAC || PLATFORM_IOS
 	#define USE_ALIGNED_MALLOC 1
 #else
 	#define USE_ALIGNED_MALLOC 0
