Add bundle manager flag to check for internet connection.

diff --git a/Engine/Source/Runtime/Core/Public/GenericPlatform/GenericPlatformInstallBundleManager.h b/Engine/Source/Runtime/Core/Public/GenericPlatform/GenericPlatformInstallBundleManager.h
--- a/Engine/Source/Runtime/Core/Public/GenericPlatform/GenericPlatformInstallBundleManager.h
+++ b/Engine/Source/Runtime/Core/Public/GenericPlatform/GenericPlatformInstallBundleManager.h
@@ -78,7 +78,9 @@ inline const TCHAR* GetInstallBundleModuleInitResultString(EInstallBundleModuleI
 enum class EInstallBundleRequestFlags : uint32
 {
 	None = 0,
-	CheckForCellularDataUsage = (1 << 0),
+	CheckForInternetConnection = (1 << 0),
+	CheckForCellularDataUsage = (1 << 1),
+	Defaults = CheckForInternetConnection,
 };
 ENUM_CLASS_FLAGS(EInstallBundleRequestFlags)
 
