ThreadCycles could be written while the async write task was running causing the archive to write too many or too few elements for that set.

(peafour-cherrypick of //UE4/Dev-Core/[REDACTED_HANDLE] by [REDACTED_PERSON])

diff --git a/Engine/Source/Runtime/Core/Private/Stats/StatsFile.cpp b/Engine/Source/Runtime/Core/Private/Stats/StatsFile.cpp
--- a/Engine/Source/Runtime/Core/Private/Stats/StatsFile.cpp
+++ b/Engine/Source/Runtime/Core/Private/Stats/StatsFile.cpp
@@ -367,7 +367,7 @@ void IStatsWriteFile::Finalize()
 	Ar << Header;
 }
 
-void IStatsWriteFile::SendTask()
+void IStatsWriteFile::WaitTask()
 {
 	if (AsyncTask)
 	{
@@ -376,6 +376,12 @@ void IStatsWriteFile::SendTask()
 		delete AsyncTask;
 		AsyncTask = nullptr;
 	}
+}
+
+void IStatsWriteFile::SendTask()
+{
+	WaitTask();
+
 	if (OutData.Num())
 	{
 		AsyncTask = new FAsyncTask<FAsyncStatsWrite>( this );
@@ -414,6 +420,8 @@ void FStatsWriteFile::WriteFrame( int64 TargetFrame )
 	WriteCondensedMessages( Ar, TargetFrame );
 
 	// Get cycles for all threads, so we can use that data to generate the mini-view.
+	WaitTask();
+
 	const FStatsThreadState& Stats = FStatsThreadState::GetLocalState();
 	for (const auto& It : Stats.Threads)
 	{
diff --git a/Engine/Source/Runtime/Core/Public/Stats/StatsFile.h b/Engine/Source/Runtime/Core/Public/Stats/StatsFile.h
--- a/Engine/Source/Runtime/Core/Public/Stats/StatsFile.h
+++ b/Engine/Source/Runtime/Core/Public/Stats/StatsFile.h
@@ -502,6 +502,9 @@ protected:
 	/**	Finalizes writing to the file. */
 	void Finalize();
 
+	/** Waits on the async send task to complete. */
+	void WaitTask();
+
 	/** Sends the data to the file via async task. */
 	void SendTask();
 };
