Modify framepro's connection changed callback to pass the filename

diff --git a/Engine/Source/Runtime/Core/Private/FramePro/FramePro.cpp b/Engine/Source/Runtime/Core/Private/FramePro/FramePro.cpp
--- a/Engine/Source/Runtime/Core/Private/FramePro/FramePro.cpp
+++ b/Engine/Source/Runtime/Core/Private/FramePro/FramePro.cpp
@@ -3862,8 +3862,9 @@ namespace FramePro
 
 		bool HasSetThreadName(int thread_id) const;
 
-		void OnConnectionChanged(bool connected) const;
-
+// START EPIC
+		void OnConnectionChanged(const FString& ProfileName, bool connected) const;
+// END EPIC
 		int GetConnectionChangedCallbackIndex(ConnectionChangedCallback p_callback);
 
 		size_t GetMemoryUsage() const;
@@ -6138,7 +6139,14 @@ namespace FramePro
 
 		ClearGlobalHiResTimers();
 
-		OnConnectionChanged(true);
+		// START EPIC
+		FString ArchiveName;
+		if (mp_RecordingFile)
+		{
+			ArchiveName = mp_RecordingFile->GetArchiveName();
+		}
+		OnConnectionChanged(ArchiveName, true);
+		// END EPIC
 	}
 
 	//------------------------------------------------------------------------
@@ -6385,10 +6393,15 @@ namespace FramePro
 
 		m_InitialiseConnectionNextFrame = false;
 
+		// START EPIC
+		FString RecordingFileName;
+		// END EPIC
+
         {
             if(mp_RecordingFile)
             {
                 // START EPIC
+				RecordingFileName = mp_RecordingFile->GetArchiveName();
                 mp_RecordingFile->Close();
                 delete mp_RecordingFile;
                 // END EPIC
@@ -6400,7 +6413,9 @@ namespace FramePro
 		// start listening for new connections
 		StartConnectThread();
 #endif
-		OnConnectionChanged(false);
+		// START EPIC
+		OnConnectionChanged(RecordingFileName, false);
+		// END EPIC
 	}
 
 	//------------------------------------------------------------------------
@@ -6742,7 +6757,9 @@ namespace FramePro
 
 		// call immediately if already connected
 		if(g_Connected)
-			p_callback(true, p_context);
+			// START EPIC
+			p_callback(true, mp_RecordingFile ? mp_RecordingFile->GetArchiveName() : FString(), p_context);
+			// END EPIC
 
 		if(GetConnectionChangedCallbackIndex(p_callback) == -1)
 		{
@@ -6767,14 +6784,19 @@ namespace FramePro
 	}
 
 	//------------------------------------------------------------------------
-	void FrameProSession::OnConnectionChanged(bool connected) const
+	// START EPIC
+	void FrameProSession::OnConnectionChanged(const FString& ProfileName, bool connected) const
+	// END EPIC
 	{
+
 		CriticalSectionScope lock(m_ConnectionChangedCriticalSection);
 
 		for(int i=0; i<m_Connectionchangedcallbacks.GetCount(); ++i)
 		{
 			const ConnectionChangedcallbackInfo& data = m_Connectionchangedcallbacks[i];
-			data.mp_Callback(connected, data.mp_Context);
+			// START EPIC
+			data.mp_Callback(connected, ProfileName, data.mp_Context);
+			// END EPIC
 		}
 	}
 
diff --git a/Engine/Source/Runtime/Core/Public/FramePro/FramePro.h b/Engine/Source/Runtime/Core/Public/FramePro/FramePro.h
--- a/Engine/Source/Runtime/Core/Public/FramePro/FramePro.h
+++ b/Engine/Source/Runtime/Core/Public/FramePro/FramePro.h
@@ -1033,7 +1033,9 @@ namespace FramePro
 namespace FramePro
 {
 	// callbacks
-	typedef void(*ConnectionChangedCallback)(bool connected, void* user_data);
+	// START EPIC
+	typedef void(*ConnectionChangedCallback)(bool connected, const FString& ProfileName, void* user_data);
+	// END EPIC
 	typedef bool(*ConditionalParentScopeCallback)(const char* p_name, int64 start_time, int64 end_time, int64 ticks_per_second);
 
 	FRAMEPRO_API extern RelaxedAtomic<bool> g_Connected;
