Install Bundle Manager
- GetContentState() now gatheres free space and install size.
- Fix stats assert when bundles are canceled.

diff --git a/Engine/Source/Runtime/Core/Public/GenericPlatform/GenericPlatformInstallBundleManager.h b/Engine/Source/Runtime/Core/Public/GenericPlatform/GenericPlatformInstallBundleManager.h
--- a/Engine/Source/Runtime/Core/Public/GenericPlatform/GenericPlatformInstallBundleManager.h
+++ b/Engine/Source/Runtime/Core/Public/GenericPlatform/GenericPlatformInstallBundleManager.h
@@ -138,6 +138,8 @@ struct FInstallBundleContentState
 {
 	EInstallBundleContentState State = EInstallBundleContentState::InitializationError;
 	uint64 DownloadSize = 0;
+	uint64 InstallSize = 0;	
+	uint64 FreeSpace = 0;
 };
 
 enum class EInstallBundleRequestInfoFlags : int32
