Early Startup Patcher
- Do an intial check to see if there is enough room for the entire install.

diff --git a/Engine/Source/Runtime/Core/Public/GenericPlatform/GenericPlatformInstallBundleManager.h b/Engine/Source/Runtime/Core/Public/GenericPlatform/GenericPlatformInstallBundleManager.h
--- a/Engine/Source/Runtime/Core/Public/GenericPlatform/GenericPlatformInstallBundleManager.h
+++ b/Engine/Source/Runtime/Core/Public/GenericPlatform/GenericPlatformInstallBundleManager.h
@@ -80,7 +80,8 @@ enum class EInstallBundleRequestFlags : uint32
 	None = 0,
 	CheckForInternetConnection = (1 << 0),
 	CheckForCellularDataUsage = (1 << 1),
-	Defaults = CheckForInternetConnection,
+	UseBackgroundDownloads = (1 << 2),
+	Defaults = CheckForInternetConnection | UseBackgroundDownloads,
 };
 ENUM_CLASS_FLAGS(EInstallBundleRequestFlags)
 
