Refactored the cleanupMigration code path to take into account the low disk space situation.

diff --git a/Engine/Source/Runtime/Core/Public/GenericPlatform/GenericPlatformInstallBundleManager.h b/Engine/Source/Runtime/Core/Public/GenericPlatform/GenericPlatformInstallBundleManager.h
--- a/Engine/Source/Runtime/Core/Public/GenericPlatform/GenericPlatformInstallBundleManager.h
+++ b/Engine/Source/Runtime/Core/Public/GenericPlatform/GenericPlatformInstallBundleManager.h
@@ -219,7 +219,8 @@ public:
 
 	virtual void SetErrorSimulationCommands(const FString& CommandLine) {}
 
-	virtual void CleanupMigrationDirectory() const {}
+	// return true if we actually cleaned up a migration directory
+	virtual bool CleanupMigrationDirectory() { return false;  }
 };
 
 class IPlatformInstallBundleManagerModule : public IModuleInterface
