Install Bundle Manager / ESP
- Pause background downloads instead of canceling bundle install when failing over to cellular.  This should prevent losing active background http progress on IOS.  It doesn't help on Android right now.
- All the check cellular flag to be cleared on a bundle after it is already enqued in order to restart the download on cellular if the user OKs it.

diff --git a/Engine/Source/Runtime/Core/Public/GenericPlatform/GenericPlatformInstallBundleManager.h b/Engine/Source/Runtime/Core/Public/GenericPlatform/GenericPlatformInstallBundleManager.h
--- a/Engine/Source/Runtime/Core/Public/GenericPlatform/GenericPlatformInstallBundleManager.h
+++ b/Engine/Source/Runtime/Core/Public/GenericPlatform/GenericPlatformInstallBundleManager.h
@@ -215,6 +215,8 @@ public:
 
 	virtual TOptional<FInstallBundleStatus> GetBundleProgress(FName BundleName) const = 0;
 
+	virtual void UpdateContentRequestFlags(FName BundleName, EInstallBundleRequestFlags AddFlags, EInstallBundleRequestFlags RemoveFlags) = 0;
+
 	virtual bool IsNullInterface() const = 0;
 
 	virtual void SetErrorSimulationCommands(const FString& CommandLine) {}
diff --git a/Engine/Source/Runtime/NullInstallBundleManager/Source/NullInstallBundleManager.cpp b/Engine/Source/Runtime/NullInstallBundleManager/Source/NullInstallBundleManager.cpp
--- a/Engine/Source/Runtime/NullInstallBundleManager/Source/NullInstallBundleManager.cpp
+++ b/Engine/Source/Runtime/NullInstallBundleManager/Source/NullInstallBundleManager.cpp
@@ -79,6 +79,11 @@ class FNullInstallBundleManager : public IPlatformInstallBundleManager
 		return TOptional<FInstallBundleStatus>();
 	}
 
+	virtual void UpdateContentRequestFlags(FName BundleName, EInstallBundleRequestFlags AddFlags, EInstallBundleRequestFlags RemoveFlags) override
+	{
+
+	}
+
 	virtual bool IsNullInterface() const override
 	{
 		return true;
