Remove .complete file functionality for csv profiles and instead create a delegate that can be used to register for when a csv profile is done writing

diff --git a/Engine/Source/Runtime/Core/Private/ProfilingDebugging/CsvProfiler.cpp b/Engine/Source/Runtime/Core/Private/ProfilingDebugging/CsvProfiler.cpp
--- a/Engine/Source/Runtime/Core/Private/ProfilingDebugging/CsvProfiler.cpp
+++ b/Engine/Source/Runtime/Core/Private/ProfilingDebugging/CsvProfiler.cpp
@@ -2587,12 +2587,7 @@ void FCsvProfiler::FinalizeCsvFile()
 	UE_LOG(LogCsvProfiler, Display, TEXT("  Frames : %d"), CaptureEndFrameCount);
 	UE_LOG(LogCsvProfiler, Display, TEXT("  Peak memory usage  : %.2fMB"), float(MemoryBytesAtEndOfCapture) / (1024.0f * 1024.0f));
 
-	if (EnumHasAnyFlags(CurrentFlags, ECsvProfilerFlags::WriteCompletionFile))
-	{
-		// Create a zero-byte file to signal completion.
-		FString CompletionFilename = OutputFilename + TEXT(".complete");
-		delete IFileManager::Get().CreateFileWriter(*CompletionFilename);
-	}
+	OnCSVProfileFinished().Broadcast(OutputFilename);
 
 	float FinalizeDuration = float(FPlatformTime::Seconds() - FinalizeStartTime);
 	UE_LOG(LogCsvProfiler, Display, TEXT("  CSV finalize time : %.3f seconds"), FinalizeDuration);
diff --git a/Engine/Source/Runtime/Core/Public/ProfilingDebugging/CsvProfiler.h b/Engine/Source/Runtime/Core/Public/ProfilingDebugging/CsvProfiler.h
--- a/Engine/Source/Runtime/Core/Public/ProfilingDebugging/CsvProfiler.h
+++ b/Engine/Source/Runtime/Core/Public/ProfilingDebugging/CsvProfiler.h
@@ -249,6 +249,9 @@ public:
 
 	CORE_API FString GetOutputFilename() const { return OutputFilename; }
 
+	DECLARE_MULTICAST_DELEGATE_OneParam(FOnCSVProfileFinished, const FString& /*Filename */);
+	FOnCSVProfileFinished& OnCSVProfileFinished() { return OnCSVProfileFinishedDelegate; }
+
 private:
 	CORE_API static void VARARGS RecordEventfInternal(int32 CategoryIndex, const TCHAR* Fmt, ...);
 
@@ -283,6 +286,8 @@ private:
 	class FCsvStreamWriter* CsvWriter;
 
 	ECsvProfilerFlags CurrentFlags;
+
+	FOnCSVProfileFinished OnCSVProfileFinishedDelegate;
 };
 
 class FScopedCsvStat
