Break apart different asserts from CrashLytics
note: remove this when crashreporter is enabled

diff --git a/Engine/Source/Runtime/Apple/MetalRHI/Private/MetalCommandList.cpp b/Engine/Source/Runtime/Apple/MetalRHI/Private/MetalCommandList.cpp
--- a/Engine/Source/Runtime/Apple/MetalRHI/Private/MetalCommandList.cpp
+++ b/Engine/Source/Runtime/Apple/MetalRHI/Private/MetalCommandList.cpp
@@ -206,7 +206,12 @@ static void ReportMetalCommandBufferFailure(mtlpp::CommandBuffer const& Complete
 
         FMetalCommandBufferDebugHelpers::DumpResources(CompletedBuffer.GetPtr());
 		
+#if PLATFORM_IOS
+        UE_LOG(LogMetal, Warning, TEXT("Command Buffer %s Failed with %s Error! Error Domain: %s Code: %d Description %s %s %s"), *LabelString, ErrorType, *DomainString, Code, *ErrorString, *FailureString, *RecoveryString);
+        FIOSPlatformMisc::GPUAssert();
+#else
 		UE_LOG(LogMetal, Fatal, TEXT("Command Buffer %s Failed with %s Error! Error Domain: %s Code: %d Description %s %s %s"), *LabelString, ErrorType, *DomainString, Code, *ErrorString, *FailureString, *RecoveryString);
+#endif
     }
 }
 
diff --git a/Engine/Source/Runtime/Apple/MetalRHI/Private/MetalRHIPrivate.h b/Engine/Source/Runtime/Apple/MetalRHI/Private/MetalRHIPrivate.h
--- a/Engine/Source/Runtime/Apple/MetalRHI/Private/MetalRHIPrivate.h
+++ b/Engine/Source/Runtime/Apple/MetalRHI/Private/MetalRHIPrivate.h
@@ -137,7 +137,11 @@ extern bool GMetalSupportsTileShaders;
 
 #define METAL_NEW_NONNULL_DECL (__clang_major__ >= 9)
 
+#if PLATFORM_IOS
+#define METAL_FATAL_ERROR(Format, ...)  { UE_LOG(LogMetal, Warning, Format, __VA_ARGS__); FIOSPlatformMisc::MetalAssert(); }
+#else
 #define METAL_FATAL_ERROR(Format, ...)	UE_LOG(LogMetal, Fatal, Format, __VA_ARGS__)
+#endif
 #define METAL_FATAL_ASSERT(Condition, Format, ...) if (!(Condition)) { METAL_FATAL_ERROR(Format, __VA_ARGS__); }
 
 struct FMetalDebugInfo
diff --git a/Engine/Source/Runtime/Core/Private/IOS/IOSPlatformMemory.cpp b/Engine/Source/Runtime/Core/Private/IOS/IOSPlatformMemory.cpp
--- a/Engine/Source/Runtime/Core/Private/IOS/IOSPlatformMemory.cpp
+++ b/Engine/Source/Runtime/Core/Private/IOS/IOSPlatformMemory.cpp
@@ -8,3 +8,51 @@
 #include "Misc/CoreStats.h"
 #include "HAL/MallocBinned.h"
 #include "HAL/MallocAnsi.h"
+#include "GenericPlatform/GenericPlatformMemoryPoolStats.h"
+#include "Misc/CoreDelegates.h"
+
+void FIOSPlatformMemory::OnOutOfMemory(uint64 Size, uint32 Alignment)
+{
+    // Update memory stats before we enter the crash handler.
+    OOMAllocationSize = Size;
+    OOMAllocationAlignment = Alignment;
+    
+    // only call this code one time - if already OOM, abort
+    if (bIsOOM)
+    {
+        return;
+    }
+    bIsOOM = true;
+    
+    FPlatformMemoryStats PlatformMemoryStats = FPlatformMemory::GetStats();
+    if (BackupOOMMemoryPool)
+    {
+        FPlatformMemory::BinnedFreeToOS(BackupOOMMemoryPool, FPlatformMemory::GetBackMemoryPoolSize());
+        UE_LOG(LogMemory, Warning, TEXT("Freeing %d bytes from backup pool to handle out of memory."), FPlatformMemory::GetBackMemoryPoolSize());
+        
+        LLM(FLowLevelMemTracker::Get().OnLowLevelFree(ELLMTracker::Default, BackupOOMMemoryPool));
+    }
+    
+    UE_LOG(LogMemory, Warning, TEXT("MemoryStats:")\
+           TEXT("\n\tAvailablePhysical %llu")\
+           TEXT("\n\t AvailableVirtual %llu")\
+           TEXT("\n\t     UsedPhysical %llu")\
+           TEXT("\n\t PeakUsedPhysical %llu")\
+           TEXT("\n\t      UsedVirtual %llu")\
+           TEXT("\n\t  PeakUsedVirtual %llu"),
+           (uint64)PlatformMemoryStats.AvailablePhysical,
+           (uint64)PlatformMemoryStats.AvailableVirtual,
+           (uint64)PlatformMemoryStats.UsedPhysical,
+           (uint64)PlatformMemoryStats.PeakUsedPhysical,
+           (uint64)PlatformMemoryStats.UsedVirtual,
+           (uint64)PlatformMemoryStats.PeakUsedVirtual);
+    
+    // let any registered handlers go
+    FCoreDelegates::GetMemoryTrimDelegate().Broadcast();
+    
+    UE_LOG(LogMemory, Warning, TEXT("Ran out of memory allocating %llu bytes with alignment %u"), Size, Alignment);
+    
+    // make this a fatal error that ends here not in the log
+    // changed to 3 from NULL because clang noticed writing to NULL and warned about it
+    *(int32 *)3 = 123;
+}
diff --git a/Engine/Source/Runtime/Core/Private/IOS/IOSPlatformMisc.cpp b/Engine/Source/Runtime/Core/Private/IOS/IOSPlatformMisc.cpp
--- a/Engine/Source/Runtime/Core/Private/IOS/IOSPlatformMisc.cpp
+++ b/Engine/Source/Runtime/Core/Private/IOS/IOSPlatformMisc.cpp
@@ -1682,6 +1682,20 @@ bool FIOSPlatformMisc::HasSeparateChannelForDebugOutput()
     return FPlatformMisc::IsDebuggerPresent();
 }
 
+void FIOSPlatformMisc::GPUAssert()
+{
+    // make this a fatal error that ends here not in the log
+    // changed to 3 from NULL because clang noticed writing to NULL and warned about it
+    *(int32 *)13 = 123;
+}
+
+void FIOSPlatformMisc::MetalAssert()
+{
+    // make this a fatal error that ends here not in the log
+    // changed to 3 from NULL because clang noticed writing to NULL and warned about it
+    *(int32 *)7 = 123;
+}
+
 FIOSCrashContext::FIOSCrashContext(ECrashContextType InType, const TCHAR* InErrorMessage)
 	: FApplePlatformCrashContext(InType, InErrorMessage)
 {
diff --git a/Engine/Source/Runtime/Core/Public/IOS/IOSPlatformMemory.h b/Engine/Source/Runtime/Core/Public/IOS/IOSPlatformMemory.h
--- a/Engine/Source/Runtime/Core/Public/IOS/IOSPlatformMemory.h
+++ b/Engine/Source/Runtime/Core/Public/IOS/IOSPlatformMemory.h
@@ -14,6 +14,8 @@
 **/
 struct CORE_API FIOSPlatformMemory : public FApplePlatformMemory
 {
+    // added this for now because Crashlytics doesn't properly break up different callstacks all ending in UE_LOG(LogXXX, Fatal, ...)
+    static CA_NO_RETURN void OnOutOfMemory(uint64 Size, uint32 Alignment);
 };
 
 typedef FIOSPlatformMemory FPlatformMemory;
diff --git a/Engine/Source/Runtime/Core/Public/IOS/IOSPlatformMisc.h b/Engine/Source/Runtime/Core/Public/IOS/IOSPlatformMisc.h
--- a/Engine/Source/Runtime/Core/Public/IOS/IOSPlatformMisc.h
+++ b/Engine/Source/Runtime/Core/Public/IOS/IOSPlatformMisc.h
@@ -214,6 +214,10 @@ struct CORE_API FIOSPlatformMisc : public FApplePlatformMisc
 	}
 
 	static bool RequestDeviceCheckToken(TFunction<void(const TArray<uint8>&)> QuerySucceededFunc, TFunction<void(const FString&, const FString&)> QueryFailedFunc);
+    
+    // added these for now because Crashlytics doesn't properly break up different callstacks all ending in UE_LOG(LogXXX, Fatal, ...)
+    static CA_NO_RETURN void GPUAssert();
+    static CA_NO_RETURN void MetalAssert();
 };
 
 typedef FIOSPlatformMisc FPlatformMisc;
