Compile out some csv profiler ensures in shipping builds

diff --git a/Engine/Source/Runtime/Core/Private/ProfilingDebugging/CsvProfiler.cpp b/Engine/Source/Runtime/Core/Private/ProfilingDebugging/CsvProfiler.cpp
--- a/Engine/Source/Runtime/Core/Private/ProfilingDebugging/CsvProfiler.cpp
+++ b/Engine/Source/Runtime/Core/Private/ProfilingDebugging/CsvProfiler.cpp
@@ -2013,7 +2013,9 @@ void FCsvProfilerThreadDataProcessor::Process(FCsvProcessThreadDataStats& OutSta
 
 	if (ThreadMarkers.Num() > 0)
 	{
+#if !UE_BUILD_SHIPPING
 		ensure(ThreadMarkers[0].GetTimestamp() >= LastProcessedTimestamp);
+#endif
 		LastProcessedTimestamp = ThreadMarkers.Last().GetTimestamp();
 	}
 
@@ -2108,8 +2110,10 @@ void FCsvProfilerThreadDataProcessor::Process(FCsvProcessThreadDataStats& OutSta
 				// AEnd would be missing 
 				if (FrameNumber >= 0 && bFoundStart)
 				{
+#if !UE_BUILD_SHIPPING
 					ensure(Marker.RawStatID == StartMarker.RawStatID);
 					ensure(Marker.GetTimestamp() >= StartMarker.GetTimestamp());
+#endif
 					if (Marker.GetTimestamp() > StartMarker.GetTimestamp())
 					{
 						uint64 ElapsedCycles = Marker.GetTimestamp() - StartMarker.GetTimestamp();
