UE-71899 Error on tvOS app start: Error excluding PersistentDownloadDir from backup

#tvOS
#4.22.1

This was fixed by [REDACTED_PERSON] with CL4597217 in //UE4/Release-4.21. The fix didn't get into //UE4/Release-Staging-4.21.

diff --git a/Engine/Source/Runtime/Core/Private/IOS/IOSPlatformMisc.cpp b/Engine/Source/Runtime/Core/Private/IOS/IOSPlatformMisc.cpp
--- a/Engine/Source/Runtime/Core/Private/IOS/IOSPlatformMisc.cpp
+++ b/Engine/Source/Runtime/Core/Private/IOS/IOSPlatformMisc.cpp
@@ -199,6 +199,8 @@ const TCHAR* FIOSPlatformMisc::GamePersistentDownloadDir()
 #endif
         Result = DownloadPath + Result;
         NSURL* URL = [NSURL fileURLWithPath : Result.GetNSString()];
+#if !PLATFORM_TVOS		
+		// this folder is expected to not exist on TVOS 
         if (![[NSFileManager defaultManager] fileExistsAtPath:[URL path]])
         {
             [[NSFileManager defaultManager] createDirectoryAtURL:URL withIntermediateDirectories : YES attributes : nil error : nil];
@@ -211,6 +213,7 @@ const TCHAR* FIOSPlatformMisc::GamePersistentDownloadDir()
         {
             NSLog(@"Error excluding %@ from backup %@",[URL lastPathComponent], error);
         }
+#endif		
     }
     return *GamePersistentDownloadDir;
 }
