Add back bAgreeToCrashUpload setting to check if we should implicitly upload unattended crashes for non-editor uses

diff --git a/Engine/Source/Runtime/Core/Private/Unix/UnixPlatformCrashContext.cpp b/Engine/Source/Runtime/Core/Private/Unix/UnixPlatformCrashContext.cpp
--- a/Engine/Source/Runtime/Core/Private/Unix/UnixPlatformCrashContext.cpp
+++ b/Engine/Source/Runtime/Core/Private/Unix/UnixPlatformCrashContext.cpp
@@ -409,6 +409,7 @@ void FUnixCrashContext::GenerateCrashInfoAndLaunchReporter(bool bReportingNonCra
 	 *      0      |       1       |         0         ||     1
 	 *      0      |       0       |         1         ||     1
 	 *      0      |       0       |         0         ||     1
+	 *
 	 */
 
 	// Suppress the user input dialog if we're running in unattended mode
@@ -568,11 +569,6 @@ void FUnixCrashContext::GenerateCrashInfoAndLaunchReporter(bool bReportingNonCra
 				CrashReportClientArguments += TEXT(" -Unattended ");
 			}
 
-			if (bSendUnattendedBugReports)
-			{
-				CrashReportClientArguments += TEXT(" -SkipPopup ");
-			}
-
 			// Whether to clean up crash reports after send
 			if (IsRunningDedicatedServer() && FParse::Param(FCommandLine::Get(), TEXT("CleanCrashReports")))
 			{
