Added FText::Join as a way to combine delimited text in a localization friendly way

Updated FTextBuilder to use this to join several generated lines together, rather than building a single FString which loses all the localization information.

Note: FTextBuilder will now consistently add a line when any of its AppendLine functions are called, whereas previously it may have done nothing if the builder was empty and indentation was disabled.

diff --git a/Engine/Source/Runtime/Core/Private/Internationalization/Text.cpp b/Engine/Source/Runtime/Core/Private/Internationalization/Text.cpp
--- a/Engine/Source/Runtime/Core/Private/Internationalization/Text.cpp
+++ b/Engine/Source/Runtime/Core/Private/Internationalization/Text.cpp
@@ -1,6 +1,7 @@
 // Copyright 1998-2019 Epic Games, Inc. All Rights Reserved.
 
 #include "Internationalization/Text.h"
+#include "Algo/Transform.h"
 #include "Misc/Parse.h"
 #include "UObject/ObjectVersion.h"
 #include "UObject/DebugSerializationFlags.h"
@@ -470,6 +471,58 @@ FText FText::FormatOrderedImpl(FTextFormat&& Fmt, FFormatOrderedArguments&& InAr
 	return FTextFormatter::Format(MoveTemp(Fmt), MoveTemp(InArguments), false, false);
 }
 
+template <typename T>
+FText TextJoinImpl(const FText& Delimiter, const TArray<T>& Args)
+{
+	if (Args.Num() > 0)
+	{
+		FText NamedFmtPattern;
+		FFormatNamedArguments NamedArgs;
+		{
+			FString FmtPattern;
+
+			const int32 ArgsCount = Args.Num();
+			NamedArgs.Reserve(ArgsCount + 1);
+			NamedArgs.Add(TEXT("Delimiter"), Delimiter);
+
+			for (int32 ArgIndex = 0; ArgIndex < ArgsCount; ++ArgIndex)
+			{
+				const T& Arg = Args[ArgIndex];
+				NamedArgs.Add(FString::Printf(TEXT("%d"), ArgIndex), Arg);
+
+				FmtPattern += TEXT('{');
+				FmtPattern.AppendInt(ArgIndex);
+				FmtPattern += TEXT('}');
+
+				if (ArgIndex + 1 < ArgsCount)
+				{
+					FmtPattern += TEXT("{Delimiter}");
+				}
+			}
+
+			NamedFmtPattern = FText::AsCultureInvariant(MoveTemp(FmtPattern));
+		}
+
+		return FTextFormatter::Format(MoveTemp(NamedFmtPattern), MoveTemp(NamedArgs), false, false);
+	}
+
+	return FText::GetEmpty();
+}
+
+FText FText::Join(const FText& Delimiter, const FFormatOrderedArguments& Args)
+{
+	return TextJoinImpl(Delimiter, Args);
+}
+
+FText FText::Join(const FText& Delimiter, const TArray<FText>& Args)
+{
+	if (Args.Num() == 1)
+	{
+		return Args[0];
+	}
+	return TextJoinImpl(Delimiter, Args);
+}
+
 FText FText::FromTextGenerator(const TSharedRef<ITextGenerator>& TextGenerator)
 {
 	FString ResultString = TextGenerator->BuildLocalizedDisplayString();
@@ -1708,6 +1761,99 @@ bool FTextStringHelper::IsComplexText(const TCHAR* Buffer)
 #undef LOC_DEFINE_REGION
 }
 
+void FTextBuilder::Indent()
+{
+	++IndentCount;
+}
+
+void FTextBuilder::Unindent()
+{
+	--IndentCount;
+}
+
+void FTextBuilder::AppendLine()
+{
+	BuildAndAppendLine(FText());
+}
+
+void FTextBuilder::AppendLine(const FText& Text)
+{
+	BuildAndAppendLine(CopyTemp(Text));
+}
+
+void FTextBuilder::AppendLine(const FString& String)
+{
+	BuildAndAppendLine(CopyTemp(String));
+}
+
+void FTextBuilder::AppendLine(const FName& Name)
+{
+	BuildAndAppendLine(Name.ToString());
+}
+
+void FTextBuilder::AppendLineFormat(const FTextFormat& Pattern, const FFormatNamedArguments& Arguments)
+{
+	BuildAndAppendLine(FText::Format(Pattern, Arguments));
+}
+
+void FTextBuilder::AppendLineFormat(const FTextFormat& Pattern, const FFormatOrderedArguments& Arguments)
+{
+	BuildAndAppendLine(FText::Format(Pattern, Arguments));
+}
+
+void FTextBuilder::Clear()
+{
+	Lines.Reset();
+}
+
+bool FTextBuilder::IsEmpty()
+{
+	return Lines.Num() == 0;
+}
+
+FText FTextBuilder::ToText() const
+{
+	return FText::Join(FText::AsCultureInvariant(LINE_TERMINATOR), Lines);
+}
+
+void FTextBuilder::BuildAndAppendLine(FString&& Data)
+{
+	if (IndentCount <= 0)
+	{
+		Lines.Emplace(FText::AsCultureInvariant(MoveTemp(Data)));
+	}
+	else
+	{
+		FString IndentedLine;
+		IndentedLine.Reserve((4 * IndentCount) + Data.Len());
+		for (int32 Index = 0; Index < IndentCount; ++Index)
+		{
+			IndentedLine += TEXT("    ");
+		}
+		IndentedLine += Data;
+		Lines.Emplace(FText::AsCultureInvariant(MoveTemp(IndentedLine)));
+	}
+}
+
+void FTextBuilder::BuildAndAppendLine(FText&& Data)
+{
+	if (IndentCount <= 0)
+	{
+		Lines.Emplace(MoveTemp(Data));
+	}
+	else
+	{
+		FString IndentedFmtPattern;
+		IndentedFmtPattern.Reserve((4 * IndentCount) + 3);
+		for (int32 Index = 0; Index < IndentCount; ++Index)
+		{
+			IndentedFmtPattern += TEXT("    ");
+		}
+		IndentedFmtPattern += TEXT("{0}");
+		Lines.Emplace(FText::Format(FText::AsCultureInvariant(MoveTemp(IndentedFmtPattern)), MoveTemp(Data)));
+	}
+}
+
 bool LexTryParseString(ETextGender& OutValue, const TCHAR* Buffer)
 {
 #define ENUM_CASE_FROM_STRING(Enum) if (FCString::Stricmp(Buffer, TEXT(#Enum)) == 0) { OutValue = ETextGender::Enum; return true; }
diff --git a/Engine/Source/Runtime/Core/Public/Internationalization/Text.h b/Engine/Source/Runtime/Core/Public/Internationalization/Text.h
--- a/Engine/Source/Runtime/Core/Public/Internationalization/Text.h
+++ b/Engine/Source/Runtime/Core/Public/Internationalization/Text.h
@@ -560,6 +560,34 @@ public:
 	template < typename... TArguments >
 	static FText FormatOrdered( FTextFormat Fmt, TArguments&&... Args );
 
+	/**
+	 * Join an arbitrary list of formattable values together, separated by the given delimiter
+	 * @note Internally this uses FText::Format with a generated culture invariant format pattern
+	 *
+	 * @param Delimiter The delimiter to insert between the items
+	 * @param Args An array of formattable values to join together
+	 * @return The joined FText
+	 */
+	static FText Join(const FText& Delimiter, const FFormatOrderedArguments& Args);
+	static FText Join(const FText& Delimiter, const TArray<FText>& Args);
+
+	/**
+	 * Join an arbitrary list of formattable items together, separated by the given delimiter
+	 * @note Internally this uses FText::Format with a generated culture invariant format pattern
+	 *
+	 * @param Delimiter The delimiter to insert between the items
+	 * @param Args A variadic list of values to join together
+	 * @return The joined FText
+	 */
+	template <typename... ArgTypes>
+	static FORCEINLINE FText Join(const FText& Delimiter, ArgTypes... Args)
+	{
+		static_assert(TAnd<TIsConstructible<FFormatArgumentValue, ArgTypes>...>::Value, "Invalid argument type passed to FText::Join");
+		static_assert(sizeof...(Args) > 0, "FText::Join expects at least one non-format argument"); // we do this to ensure that people don't call Join for no good reason
+
+		return Join(Delimiter, FFormatOrderedArguments{ MoveTemp(Args)... });
+	}
+
 	/**
 	 * Produces a custom-generated FText. Can be used for objects that produce text dependent on localized strings but
 	 * that do not fit the standard formats.
@@ -1092,100 +1120,79 @@ private:
 class CORE_API FTextBuilder
 {
 public:
-	FTextBuilder()
-		: IndentCount(0)
-	{
-	}
-
-	void Indent()
-	{
-		++IndentCount;
-	}
-
-	void Unindent()
-	{
-		--IndentCount;
-	}
-
-	void AppendLine()
-	{
-		if (!Report.IsEmpty())
-		{
-			Report += LINE_TERMINATOR;
-		}
-
-		for (int32 Index = 0; Index < IndentCount; Index++)
-		{
-			Report += TEXT("    ");
-		}
-	}
+	/**
+	 * Increase the running indentation of the builder.
+	 */
+	void Indent();
 
-	void AppendLine(const FText& Text)
-	{
-		AppendLine();
-		Report += Text.ToString();
-	}
+	/**
+	 * Decrease the running indentation of the builder.
+	 */
+	void Unindent();
 
-	void AppendLine(const FString& String)
-	{
-		AppendLine();
-		Report += String;
-	}
+	/**
+	 * Append an empty line to the builder, indented by the running indentation of the builder.
+	 */
+	void AppendLine();
 
-	void AppendLine(const FName& Name)
-	{
-		AppendLine();
-		Report += Name.ToString();
-	}
+	/**
+	 * Append the given text line to the builder, indented by the running indentation of the builder.
+	 */
+	void AppendLine(const FText& Text);
 
-	void AppendLineFormat(const FTextFormat& Pattern, const FFormatNamedArguments& Arguments)
-	{
-		AppendLine(FText::Format(Pattern, Arguments));
-	}
+	/**
+	 * Append the given string line to the builder, indented by the running indentation of the builder.
+	 */
+	void AppendLine(const FString& String);
 
-	void AppendLineFormat(const FTextFormat& Pattern, const FFormatOrderedArguments& Arguments)
-	{
-		AppendLine(FText::Format(Pattern, Arguments));
-	}
+	/**
+	 * Append the given name line to the builder, indented by the running indentation of the builder.
+	 */
+	void AppendLine(const FName& Name);
 
-	void AppendLineFormat(const FTextFormat& Fmt, const FFormatArgumentValue& v1)
-	{
-		AppendLine(FText::Format(Fmt, v1));
-	}
+	/**
+	 * Append the given formatted text line to the builder, indented by the running indentation of the builder.
+	 */
+	void AppendLineFormat(const FTextFormat& Pattern, const FFormatNamedArguments& Arguments);
 
-	void AppendLineFormat(const FTextFormat& Fmt, const FFormatArgumentValue& v1, const FFormatArgumentValue& v2)
-	{
-		AppendLine(FText::Format(Fmt, v1, v2));
-	}
+	/**
+	 * Append the given formatted text line to the builder, indented by the running indentation of the builder.
+	 */
+	void AppendLineFormat(const FTextFormat& Pattern, const FFormatOrderedArguments& Arguments);
 
-	void AppendLineFormat(const FTextFormat& Fmt, const FFormatArgumentValue& v1, const FFormatArgumentValue& v2, const FFormatArgumentValue& v3)
+	/**
+	 * Append the given formatted text line to the builder, indented by the running indentation of the builder.
+	 */
+	template <typename... ArgTypes>
+	FORCEINLINE void AppendLineFormat(FTextFormat Pattern, ArgTypes... Args)
 	{
-		AppendLine(FText::Format(Fmt, v1, v2, v3));
-	}
+		static_assert(TAnd<TIsConstructible<FFormatArgumentValue, ArgTypes>...>::Value, "Invalid argument type passed to FTextBuilder::AppendLineFormat");
+		static_assert(sizeof...(Args) > 0, "FTextBuilder::AppendLineFormat expects at least one non-format argument"); // we do this to ensure that people don't call AppendLineFormat for no good reason
 
-	void AppendLineFormat(const FTextFormat& Fmt, const FFormatArgumentValue& v1, const FFormatArgumentValue& v2, const FFormatArgumentValue& v3, const FFormatArgumentValue& v4)
-	{
-		AppendLine(FText::Format(Fmt, v1, v2, v3, v4));
+		BuildAndAppendLine(FText::Format(MoveTemp(Pattern), FFormatOrderedArguments{ MoveTemp(Args)... }));
 	}
 
-	void Clear()
-	{
-		Report.Empty();
-	}
+	/**
+	 * Clear the builder and reset it to its default state.
+	 */
+	void Clear();
 
-	bool IsEmpty()
-	{
-		return Report.IsEmpty();
-	}
+	/**
+	 * Check to see if the builder has any data.
+	 */
+	bool IsEmpty();
 
-	FText ToText() const
-	{
-		return FText::FromString(Report);
-	}
+	/**
+	 * Build the current set of input into a FText.
+	 */
+	FText ToText() const;
 
 private:
-	FString Report;
-	int32 IndentCount;
+	void BuildAndAppendLine(FString&& Data);
+	void BuildAndAppendLine(FText&& Data);
+
+	TArray<FText> Lines;
+	int32 IndentCount = 0;
 };
 
 class CORE_API FScopedTextIdentityPreserver
