Install Bundle Manager
- Add callback for when bundles get paused due to being on cellular.
- Handle corner cause of needing to pause the BPT installer in the case that we did not fetch all chunkDBs prior to starting the install.

diff --git a/Engine/Source/Runtime/Core/Private/GenericPlatform/GenericPlatformInstallBundleManager.cpp b/Engine/Source/Runtime/Core/Private/GenericPlatform/GenericPlatformInstallBundleManager.cpp
--- a/Engine/Source/Runtime/Core/Private/GenericPlatform/GenericPlatformInstallBundleManager.cpp
+++ b/Engine/Source/Runtime/Core/Private/GenericPlatform/GenericPlatformInstallBundleManager.cpp
@@ -6,3 +6,5 @@ FInstallBundleCompleteMultiDelegate IPlatformInstallBundleManager::InstallBundle
 
 FInstallBundleCompleteMultiDelegate IPlatformInstallBundleManager::RemoveBundleCompleteDelegate;
 
+FInstallBundlePausedMultiDelegate IPlatformInstallBundleManager::PausedBundleDelegate;
+
diff --git a/Engine/Source/Runtime/Core/Public/GenericPlatform/GenericPlatformInstallBundleManager.h b/Engine/Source/Runtime/Core/Public/GenericPlatform/GenericPlatformInstallBundleManager.h
--- a/Engine/Source/Runtime/Core/Public/GenericPlatform/GenericPlatformInstallBundleManager.h
+++ b/Engine/Source/Runtime/Core/Public/GenericPlatform/GenericPlatformInstallBundleManager.h
@@ -9,6 +9,40 @@
 #include "Logging/LogMacros.h"
 #include "Internationalization/Text.h"
 
+enum class EInstallBundleModuleInitResult : int
+{
+	OK,
+	BuildMetaDataNotFound,
+	BuildMetaDataParsingError,
+	DistributionRootParseError,
+	DistributionRootDownloadError,
+	ManifestCreationError,
+	ManifestDownloadError,
+	BackgroundDownloadsIniDownloadError,
+	NoInternetConnectionError,
+	Count
+};
+
+inline const TCHAR* GetInstallBundleModuleInitResultString(EInstallBundleModuleInitResult Result)
+{
+	using UnderType = __underlying_type(EInstallBundleModuleInitResult);
+	static const TCHAR* Strings[] =
+	{
+		TEXT("OK"),
+		TEXT("BuildMetaDataNotFound"),
+		TEXT("BuildMetaDataParsingError"),
+		TEXT("DistributionRootParseError"),
+		TEXT("DistributionRootDownloadError"),
+		TEXT("ManifestCreationError"),
+		TEXT("ManifestDownloadError"),
+		TEXT("BackgroundDownloadsIniDownloadError"),
+		TEXT("NoInternetConnectionError"),
+	};
+	static_assert(static_cast<UnderType>(EInstallBundleModuleInitResult::Count) == ARRAY_COUNT(Strings), "");
+
+	return Strings[static_cast<UnderType>(Result)];
+}
+
 enum class EInstallBundleResult : int
 {
 	OK,
@@ -41,39 +75,12 @@
 	return Strings[static_cast<UnderType>(Result)];
 }
 
-enum class EInstallBundleModuleInitResult : int
+enum class EInstallBundlePauseFlags : uint32
 {
-	OK,
-	BuildMetaDataNotFound,
-	BuildMetaDataParsingError,
-	DistributionRootParseError,
-	DistributionRootDownloadError,
-	ManifestCreationError,
-	ManifestDownloadError,
-	BackgroundDownloadsIniDownloadError,
-	NoInternetConnectionError,
-	Count
+	None = 0,
+	OnCellularNetwork = (1 << 0),
 };
+ENUM_CLASS_FLAGS(EInstallBundlePauseFlags)
 
-inline const TCHAR* GetInstallBundleModuleInitResultString(EInstallBundleModuleInitResult Result)
-{
-	using UnderType = __underlying_type(EInstallBundleModuleInitResult);
-	static const TCHAR* Strings[] =
-	{
-		TEXT("OK"),
-		TEXT("BuildMetaDataNotFound"),
-		TEXT("BuildMetaDataParsingError"),
-		TEXT("DistributionRootParseError"),
-		TEXT("DistributionRootDownloadError"),
-		TEXT("ManifestCreationError"),
-		TEXT("ManifestDownloadError"),
-		TEXT("BackgroundDownloadsIniDownloadError"),
-		TEXT("NoInternetConnectionError"),
-	};
-	static_assert(static_cast<UnderType>(EInstallBundleModuleInitResult::Count) == ARRAY_COUNT(Strings), "");
-
-	return Strings[static_cast<UnderType>(Result)];
-}
-
 enum class EInstallBundleRequestFlags : uint32
 {
@@ -128,6 +135,12 @@
 	FString OptionalErrorCode;
 };
 
+struct FInstallBundlePauseInfo
+{
+	FName BundleName;
+	EInstallBundlePauseFlags PauseFlags = EInstallBundlePauseFlags::None;
+};
+
 enum class EInstallBundleContentState : int
 {
 	InitializationError,
@@ -180,6 +193,7 @@
 DECLARE_DELEGATE_RetVal_OneParam(EInstallBundleManagerInitErrorHandlerResult, FInstallBundleManagerInitErrorHandler, EInstallBundleModuleInitResult);
 
 DECLARE_MULTICAST_DELEGATE_OneParam(FInstallBundleCompleteMultiDelegate, FInstallBundleResultInfo);
+DECLARE_MULTICAST_DELEGATE_OneParam(FInstallBundlePausedMultiDelegate, FInstallBundlePauseInfo);
 
 DECLARE_DELEGATE_OneParam(FInstallBundleGetContentStateDelegate, FInstallBundleContentState);
 
@@ -188,6 +202,7 @@
 public:
 	static FInstallBundleCompleteMultiDelegate InstallBundleCompleteDelegate;
 	static FInstallBundleCompleteMultiDelegate RemoveBundleCompleteDelegate;
+	static FInstallBundlePausedMultiDelegate PausedBundleDelegate;
 
 	virtual ~IPlatformInstallBundleManager() {}
 
