#UE4 When using WaitAndFork, you can now supply an exit code in a child process to request that the parent close as well with WAIT_AND_FORK_PARENT_SHUTDOWN_EXIT_CODE

diff --git a/Engine/Source/Runtime/Core/Private/Unix/UnixPlatformProcess.cpp b/Engine/Source/Runtime/Core/Private/Unix/UnixPlatformProcess.cpp
--- a/Engine/Source/Runtime/Core/Private/Unix/UnixPlatformProcess.cpp
+++ b/Engine/Source/Runtime/Core/Private/Unix/UnixPlatformProcess.cpp
@@ -1208,6 +1208,7 @@ void FUnixPlatformProcess::TerminateProc( FProcHandle & ProcessHandle, bool Kill
  *
  * This is a function that halts execution and waits for signals to cause forked processes to be created and continue execution.
  * The parent process will return when GIsRequestingExit is true. SIGRTMIN+1 is used to cause a fork to happen.
+ * Optionally, the parent process will also return if any of the children close with an exit code equal to WAIT_AND_FORK_PARENT_SHUTDOWN_EXIT_CODE if it is set to non-zero.
  * If sigqueue is used, the payload int will be split into the upper and lower uint16 values. The upper value is a "cookie" and the
  *     lower value is an "index". These two values will be used to name the process using the pattern DS-<cookie>-<index>. This name
  *     can be used to uniquely discover the process that was spawned.
@@ -1221,6 +1222,9 @@ FGenericPlatformProcess::EWaitAndForkResult FUnixPlatformProcess::WaitAndFork()
 #define WAIT_AND_FORK_QUEUE_LENGTH 4096
 #define WAIT_AND_FORK_PARENT_SLEEP_DURATION 10
 #define WAIT_AND_FORK_CHILD_SPAWN_DELAY 0.125
+#ifndef WAIT_AND_FORK_PARENT_SHUTDOWN_EXIT_CODE
+	#define WAIT_AND_FORK_PARENT_SHUTDOWN_EXIT_CODE 0
+#endif
 
 	// Only works in -nothreading mode for now (probably best this way)
 	if (FPlatformProcess::SupportsMultithreading())
@@ -1403,7 +1407,8 @@ FGenericPlatformProcess::EWaitAndForkResult FUnixPlatformProcess::WaitAndFork()
 			{
 				const FPidAndSignal& ChildPidAndSignal = AllChildren[ChildIdx];
 
-				pid_t WaitResult = waitpid(ChildPidAndSignal.Pid, nullptr, WNOHANG);
+				int32 Status = 0;
+				pid_t WaitResult = waitpid(ChildPidAndSignal.Pid, &Status, WNOHANG);
 				if (WaitResult == -1)
 				{
 					int32 ErrNo = errno;
@@ -1411,7 +1416,13 @@ FGenericPlatformProcess::EWaitAndForkResult FUnixPlatformProcess::WaitAndFork()
 				}
 				else if (WaitResult != 0)
 				{
-					if (NumForks > 0 && ChildPidAndSignal.SignalValue > 0 && ChildPidAndSignal.SignalValue <= NumForks)
+					int32 ExitCode = WIFEXITED(Status) ? WEXITSTATUS(Status) : 0;
+					if (ExitCode != 0 && ExitCode == WAIT_AND_FORK_PARENT_SHUTDOWN_EXIT_CODE)
+					{
+						UE_LOG(LogHAL, Log, TEXT("[Parent] WaitAndFork child %d exited with return code %d, indicating that the parent process should shut down. Shutting down..."), ChildPidAndSignal.Pid, WAIT_AND_FORK_PARENT_SHUTDOWN_EXIT_CODE);
+						GIsRequestingExit = true;
+					}
+					else if (NumForks > 0 && ChildPidAndSignal.SignalValue > 0 && ChildPidAndSignal.SignalValue <= NumForks)
 					{
 						UE_LOG(LogHAL, Log, TEXT("[Parent] WaitAndFork child %d missing. This was NumForks child %d. Relaunching..."), ChildPidAndSignal.Pid, ChildPidAndSignal.SignalValue);
 						WaitAndForkSignalQueue.Enqueue(ChildPidAndSignal.SignalValue);
