[LLM] adding the option to allow platform-specific computed memory stats.

diff --git a/Engine/Source/Runtime/Core/Private/HAL/LowLevelMemTracker.cpp b/Engine/Source/Runtime/Core/Private/HAL/LowLevelMemTracker.cpp
--- a/Engine/Source/Runtime/Core/Private/HAL/LowLevelMemTracker.cpp
+++ b/Engine/Source/Runtime/Core/Private/HAL/LowLevelMemTracker.cpp
@@ -602,6 +602,9 @@ void FLowLevelMemTracker::UpdateStatsPerFrame(const TCHAR* LogName)
 	int64 Overhead = StaticOverhead + Allocator.GetTotal();
 	SET_MEMORY_STAT(STAT_LLMOverheadTotal, Overhead);
 
+	// get the platform to update any custom tags
+	FPlatformMemory::UpdateCustomLLMTags();
+
 	// calculate memory the platform thinks we have allocated, compared to what we have tracked, including the program memory
 	FPlatformMemoryStats PlatformStats = FPlatformMemory::GetStats();
 #if PLATFORM_ANDROID || PLATFORM_IOS
@@ -921,6 +924,11 @@ int64 FLowLevelMemTracker::GetTagAmountForTracker(ELLMTracker Tracker, ELLMTag T
 	return GetTracker(Tracker)->GetTagAmount(Tag);
 }
 
+void FLowLevelMemTracker::SetTagAmountForTracker(ELLMTracker Tracker, ELLMTag Tag, int64 Amount, bool bAddToTotal )
+{
+	GetTracker(Tracker)->SetTagAmount( Tag, Amount, bAddToTotal );
+}
+
 int64 FLowLevelMemTracker::GetActiveTag(ELLMTracker Tracker)
 {
     return GetTracker(Tracker)->GetActiveTag();
diff --git a/Engine/Source/Runtime/Core/Public/GenericPlatform/GenericPlatformMemory.h b/Engine/Source/Runtime/Core/Public/GenericPlatform/GenericPlatformMemory.h
--- a/Engine/Source/Runtime/Core/Public/GenericPlatform/GenericPlatformMemory.h
+++ b/Engine/Source/Runtime/Core/Public/GenericPlatform/GenericPlatformMemory.h
@@ -584,6 +584,12 @@ public:
 	FORCEINLINE static void OnLowLevelMemory_Alloc(void const* Pointer, uint64 Size, uint64 Tag) { }
 	FORCEINLINE static void OnLowLevelMemory_Free(void const* Pointer, uint64 Size, uint64 Tag) { }
 
+	/**
+	* Called once per frame when LLM is collating the data for the current frame.
+	* Can be used to set platform-specific calculated tag data via SetTagAmountForTracker
+	*/
+	static void UpdateCustomLLMTags() { };
+
 protected:
 	friend struct FGenericStatsUpdater;
 
diff --git a/Engine/Source/Runtime/Core/Public/HAL/LowLevelMemTracker.h b/Engine/Source/Runtime/Core/Public/HAL/LowLevelMemTracker.h
--- a/Engine/Source/Runtime/Core/Public/HAL/LowLevelMemTracker.h
+++ b/Engine/Source/Runtime/Core/Public/HAL/LowLevelMemTracker.h
@@ -454,6 +454,9 @@ public:
 	// Get the amount of memory for a tag from the given tracker
 	int64 GetTagAmountForTracker(ELLMTracker Tracker, ELLMTag Tag);
 
+	// Set the amount of memory for a tag for a given tracker, optionally updating the total tracked memory too
+	void SetTagAmountForTracker(ELLMTracker Tracker, ELLMTag Tag, int64 Amount, bool bAddToTotal );
+
 	// Dump the current tag for the given tracker to the output
 	void DumpTag( ELLMTracker Tracker, const char* FileName, int LineNumber );
 
