#android
#ios

diff --git a/Engine/Source/Runtime/Core/Private/Android/AndroidMisc.cpp b/Engine/Source/Runtime/Core/Private/Android/AndroidMisc.cpp
--- a/Engine/Source/Runtime/Core/Private/Android/AndroidMisc.cpp
+++ b/Engine/Source/Runtime/Core/Private/Android/AndroidMisc.cpp
@@ -45,6 +45,7 @@
 #if USE_ANDROID_JNI
 extern AAssetManager * AndroidThunkCpp_GetAssetManager();
 #endif
+extern int32 GAndroidPackageVersion;
 
 static int32 GAndroidTraceMarkersEnabled = 0;
 static FAutoConsoleVariableRef CAndroidTraceMarkersEnabled(
@@ -1168,6 +1169,10 @@ const FString FAndroidMisc::GetOSLanguage()
 	return OSLanguage;
 }
 
+const FString FAndroidMisc::GetProjectVersion() {
+	return FString::FromInt(GAndroidPackageVersion);
+}
+
 FString FAndroidMisc::GetDefaultLocale()
 {
 	return OSLanguage;
diff --git a/Engine/Source/Runtime/Core/Private/IOS/IOSPlatformMisc.cpp b/Engine/Source/Runtime/Core/Private/IOS/IOSPlatformMisc.cpp
--- a/Engine/Source/Runtime/Core/Private/IOS/IOSPlatformMisc.cpp
+++ b/Engine/Source/Runtime/Core/Private/IOS/IOSPlatformMisc.cpp
@@ -1272,6 +1272,13 @@ int32 FIOSPlatformMisc::IOSVersionCompare(uint8 Major, uint8 Minor, uint8 Revisi
 	return 0;
 }
 
+FString FIOSPlatformMisc::GetProjectVersion()
+{
+	NSDictionary* infoDictionary = [[NSBundle mainBundle] infoDictionary];
+	FString localVersionString = FString(infoDictionary[@"CFBundleShortVersionString"]);
+	return localVersionString;
+}
+
 bool FIOSPlatformMisc::RequestDeviceCheckToken(TFunction<void(const TArray<uint8>&)> QuerySucceededFunc, TFunction<void(const FString&, const FString&)> QueryFailedFunc)
 {
 	DCDevice* DeviceCheckDevice = [DCDevice currentDevice];
diff --git a/Engine/Source/Runtime/Core/Public/Android/AndroidMisc.h b/Engine/Source/Runtime/Core/Public/Android/AndroidMisc.h
--- a/Engine/Source/Runtime/Core/Public/Android/AndroidMisc.h
+++ b/Engine/Source/Runtime/Core/Public/Android/AndroidMisc.h
@@ -168,6 +168,7 @@ public:
 	static const FString GetDeviceModel();
 	static const FString GetOSLanguage();
 	static const FString GetDeviceBuildNumber();
+	static const FString GetProjectVersion();
 	static FString GetDefaultLocale();
 	static FString GetGPUFamily();
 	static FString GetGLVersion();
diff --git a/Engine/Source/Runtime/Core/Public/IOS/IOSPlatformMisc.h b/Engine/Source/Runtime/Core/Public/IOS/IOSPlatformMisc.h
--- a/Engine/Source/Runtime/Core/Public/IOS/IOSPlatformMisc.h
+++ b/Engine/Source/Runtime/Core/Public/IOS/IOSPlatformMisc.h
@@ -204,7 +204,8 @@ struct CORE_API FIOSPlatformMisc : public FApplePlatformMisc
 	static FString GetCPUBrand();
 	static void GetOSVersions(FString& out_OSVersionLabel, FString& out_OSSubVersionLabel);
 	static int32 IOSVersionCompare(uint8 Major, uint8 Minor, uint8 Revision);
-	
+	static FString GetProjectVersion();
+
 	static void SetGracefulTerminationHandler();
 	static void SetCrashHandler(void(*CrashHandler)(const FGenericCrashContext& Context));
 
