Bundle Install
-Don't lose background download progress when pausing downloads due to moving on and off cellular.
-Install bundle manager no longer cancels requests due to network connectivity changes.  Instead, it pauses them.
-Make sure UserCancelledError has priorority over any other bundle error.
-Add function for Background HTTP to return whether or not the generic implementation is being used.
--On generic platforms, bundle manager pauses background downloads because the OS will not keep them alive.

diff --git a/Engine/Source/Runtime/Core/Public/GenericPlatform/GenericPlatformInstallBundleManager.h b/Engine/Source/Runtime/Core/Public/GenericPlatform/GenericPlatformInstallBundleManager.h
--- a/Engine/Source/Runtime/Core/Public/GenericPlatform/GenericPlatformInstallBundleManager.h
+++ b/Engine/Source/Runtime/Core/Public/GenericPlatform/GenericPlatformInstallBundleManager.h
@@ -79,16 +79,16 @@ enum class EInstallBundlePauseFlags : uint32
 {
 	None = 0,
 	OnCellularNetwork = (1 << 0),
+	NoInternetConnection = (1 << 1),
 };
 ENUM_CLASS_FLAGS(EInstallBundlePauseFlags)
 
 enum class EInstallBundleRequestFlags : uint32
 {
 	None = 0,
-	CheckForInternetConnection = (1 << 0),
-	CheckForCellularDataUsage = (1 << 1),
-	UseBackgroundDownloads = (1 << 2),
-	Defaults = CheckForInternetConnection | UseBackgroundDownloads,
+	CheckForCellularDataUsage = (1 << 0),
+	UseBackgroundDownloads = (1 << 1),
+	Defaults = UseBackgroundDownloads,
 };
 ENUM_CLASS_FLAGS(EInstallBundleRequestFlags)
 
@@ -228,6 +228,8 @@ public:
 
 	virtual void CancelAllBundles(EInstallBundleCancelFlags Flags) = 0;
 
+	virtual void RequestPausedBundleCallback() const = 0;
+
 	virtual TOptional<FInstallBundleStatus> GetBundleProgress(FName BundleName) const = 0;
 
 	virtual void UpdateContentRequestFlags(FName BundleName, EInstallBundleRequestFlags AddFlags, EInstallBundleRequestFlags RemoveFlags) = 0;
diff --git a/Engine/Source/Runtime/NullInstallBundleManager/Source/NullInstallBundleManager.cpp b/Engine/Source/Runtime/NullInstallBundleManager/Source/NullInstallBundleManager.cpp
--- a/Engine/Source/Runtime/NullInstallBundleManager/Source/NullInstallBundleManager.cpp
+++ b/Engine/Source/Runtime/NullInstallBundleManager/Source/NullInstallBundleManager.cpp
@@ -74,6 +74,11 @@ class FNullInstallBundleManager : public IPlatformInstallBundleManager
 
 	}
 
+	virtual void RequestPausedBundleCallback() const override
+	{
+
+	}
+
 	virtual TOptional<FInstallBundleStatus> GetBundleProgress(FName BundleName) const override
 	{
 		return TOptional<FInstallBundleStatus>();
diff --git a/Engine/Source/Runtime/Online/BackgroundHTTP/Public/GenericPlatform/GenericPlatformBackgroundHttpManager.h b/Engine/Source/Runtime/Online/BackgroundHTTP/Public/GenericPlatform/GenericPlatformBackgroundHttpManager.h
--- a/Engine/Source/Runtime/Online/BackgroundHTTP/Public/GenericPlatform/GenericPlatformBackgroundHttpManager.h
+++ b/Engine/Source/Runtime/Online/BackgroundHTTP/Public/GenericPlatform/GenericPlatformBackgroundHttpManager.h
@@ -16,5 +16,6 @@ class BACKGROUNDHTTP_API FGenericPlatformBackgroundHttpManager
 	: public FBackgroundHttpManagerImpl
 {
 public:
+	virtual bool IsGenericImplementation() const override { return true; }
 	virtual ~FGenericPlatformBackgroundHttpManager();
 };
\ No newline at end of file
diff --git a/Engine/Source/Runtime/Online/BackgroundHTTP/Public/IOS/ApplePlatformBackgroundHttpManager.h b/Engine/Source/Runtime/Online/BackgroundHTTP/Public/IOS/ApplePlatformBackgroundHttpManager.h
--- a/Engine/Source/Runtime/Online/BackgroundHTTP/Public/IOS/ApplePlatformBackgroundHttpManager.h
+++ b/Engine/Source/Runtime/Online/BackgroundHTTP/Public/IOS/ApplePlatformBackgroundHttpManager.h
@@ -21,6 +21,7 @@ public:
 	virtual void Shutdown() override;
 	virtual void AddRequest(const FBackgroundHttpRequestPtr Request) override;
 	virtual void RemoveRequest(const FBackgroundHttpRequestPtr Request) override;
+	virtual bool IsGenericImplementation() const override { return false;  }
 
 	/**
 	* Constructor
diff --git a/Engine/Source/Runtime/Online/BackgroundHTTP/Public/Interfaces/IBackgroundHttpManager.h b/Engine/Source/Runtime/Online/BackgroundHTTP/Public/Interfaces/IBackgroundHttpManager.h
--- a/Engine/Source/Runtime/Online/BackgroundHTTP/Public/Interfaces/IBackgroundHttpManager.h
+++ b/Engine/Source/Runtime/Online/BackgroundHTTP/Public/Interfaces/IBackgroundHttpManager.h
@@ -44,6 +44,9 @@ public:
 	*/
 	virtual void CleanUpTemporaryFiles() = 0;
 
+	/** Returns whether or not this is a platform specific implementation */
+	virtual bool IsGenericImplementation() const = 0;
+
 protected:
 	/**
 	* Designed to be called internally by AddRequest to associate our incoming request with any 
