Clean up Bundle Manager result flags that are no longer possible.
Forward Bundle Pause status to JSBridge.

diff --git a/Engine/Source/Runtime/Core/Public/GenericPlatform/GenericPlatformInstallBundleManager.h b/Engine/Source/Runtime/Core/Public/GenericPlatform/GenericPlatformInstallBundleManager.h
--- a/Engine/Source/Runtime/Core/Public/GenericPlatform/GenericPlatformInstallBundleManager.h
+++ b/Engine/Source/Runtime/Core/Public/GenericPlatform/GenericPlatformInstallBundleManager.h
@@ -49,8 +49,6 @@ enum class EInstallBundleResult : int
 	FailedPrereqRequiresLatestClient,
 	InstallError,
 	InstallerOutOfDiskSpaceError,
-	OnCellularNetworkError,
-	NoInternetConnectionError,
 	UserCancelledError,
 	InitializationError,
 	Count,
@@ -65,8 +63,6 @@ inline const TCHAR* GetInstallBundleResultString(EInstallBundleResult Result)
 		TEXT("FailedPrereqRequiresLatestClient"),
 		TEXT("InstallError"),
 		TEXT("InstallerOutOfDiskSpaceError"),
-		TEXT("OnCellularNetworkError"),
-		TEXT("NoInternetConnectionError"),
 		TEXT("UserCancelledError"),
 		TEXT("InitializationError"),
 	};
@@ -83,6 +79,19 @@ enum class EInstallBundlePauseFlags : uint32
 };
 ENUM_CLASS_FLAGS(EInstallBundlePauseFlags)
 
+inline const TCHAR* GetInstallBundlePauseReason(EInstallBundlePauseFlags Flags)
+{
+	// Return the most appropriate reason given the flags
+
+	if (EnumHasAnyFlags(Flags, EInstallBundlePauseFlags::NoInternetConnection))
+		return TEXT("NoInternetConnection");
+	
+	if (EnumHasAnyFlags(Flags, EInstallBundlePauseFlags::OnCellularNetwork))
+		return TEXT("OnCellularNetwork");
+
+	return TEXT("");
+}
+
 enum class EInstallBundleRequestFlags : uint32
 {
 	None = 0,
@@ -119,6 +128,8 @@ struct FInstallBundleStatus
 
 	EInstallBundleStatus Status = EInstallBundleStatus::NotRequested;
 
+	EInstallBundlePauseFlags PauseFlags = EInstallBundlePauseFlags::None;
+
 	FText StatusText;
 
 	// Progress is only present if Status is Downloading or Installing
