User pausing and resuming of bundles.
Add transient user setting bResumeUserCanceledOrPausedDownload.  When set, it will allow user paused or canceled downloads to resume.

diff --git a/Engine/Source/Runtime/Core/Public/GenericPlatform/GenericPlatformInstallBundleManager.h b/Engine/Source/Runtime/Core/Public/GenericPlatform/GenericPlatformInstallBundleManager.h
--- a/Engine/Source/Runtime/Core/Public/GenericPlatform/GenericPlatformInstallBundleManager.h
+++ b/Engine/Source/Runtime/Core/Public/GenericPlatform/GenericPlatformInstallBundleManager.h
@@ -76,6 +76,7 @@ enum class EInstallBundlePauseFlags : uint32
 	None = 0,
 	OnCellularNetwork = (1 << 0),
 	NoInternetConnection = (1 << 1),
+	UserPaused	= (1 << 2)
 };
 ENUM_CLASS_FLAGS(EInstallBundlePauseFlags)
 
@@ -83,6 +84,9 @@ inline const TCHAR* GetInstallBundlePauseReason(EInstallBundlePauseFlags Flags)
 {
 	// Return the most appropriate reason given the flags
 
+	if (EnumHasAnyFlags(Flags, EInstallBundlePauseFlags::UserPaused))
+		return TEXT("UserPaused");
+
 	if (EnumHasAnyFlags(Flags, EInstallBundlePauseFlags::NoInternetConnection))
 		return TEXT("NoInternetConnection");
 	
@@ -239,6 +243,10 @@ public:
 
 	virtual void CancelAllBundles(EInstallBundleCancelFlags Flags) = 0;
 
+	virtual bool PauseBundle(FName BundleName) = 0;
+
+	virtual void ResumeBundle(FName BundleName) = 0;
+
 	virtual void RequestPausedBundleCallback() const = 0;
 
 	virtual TOptional<FInstallBundleStatus> GetBundleProgress(FName BundleName) const = 0;
diff --git a/Engine/Source/Runtime/NullInstallBundleManager/Source/NullInstallBundleManager.cpp b/Engine/Source/Runtime/NullInstallBundleManager/Source/NullInstallBundleManager.cpp
--- a/Engine/Source/Runtime/NullInstallBundleManager/Source/NullInstallBundleManager.cpp
+++ b/Engine/Source/Runtime/NullInstallBundleManager/Source/NullInstallBundleManager.cpp
@@ -74,6 +74,16 @@ class FNullInstallBundleManager : public IPlatformInstallBundleManager
 
 	}
 
+	virtual bool PauseBundle(FName BundleName) override
+	{
+		return false;
+	}
+
+	virtual void ResumeBundle(FName BundleName) override
+	{
+
+	}
+
 	virtual void RequestPausedBundleCallback() const override
 	{
 
