Add support for starting with the hitch detector suspended

diff --git a/Engine/Source/Runtime/Core/Private/HAL/ThreadHeartBeat.cpp b/Engine/Source/Runtime/Core/Private/HAL/ThreadHeartBeat.cpp
--- a/Engine/Source/Runtime/Core/Private/HAL/ThreadHeartBeat.cpp
+++ b/Engine/Source/Runtime/Core/Private/HAL/ThreadHeartBeat.cpp
@@ -616,7 +616,6 @@ void FGameThreadHitchHeartBeatThreaded::InitSettings()
 	{
 		bHasCmdLine = FParse::Value(FCommandLine::Get(), TEXT("hitchdetection="), CmdLine_HangDuration);
 		CmdLine_StackWalk = FParse::Param(FCommandLine::Get(), TEXT("hitchdetectionstackwalk"));
-		bFirst = false;
 	}
 
 	if (bHasCmdLine)
@@ -650,12 +649,35 @@ void FGameThreadHitchHeartBeatThreaded::InitSettings()
 			bWalkStackOnHitch = false;
 		}
 	}
+	
+	// Figure out whether to start suspended
+	bool bStartSuspended = false;
+	if (GConfig)
+	{
+		GConfig->GetBool(TEXT("Core.System"), TEXT("GameThreadHeartBeatStartSuspended"), bStartSuspended, GEngineIni);
+	}
+	if (bFirst)
+	{
+		if (FParse::Param(FCommandLine::Get(), TEXT("hitchdetectionstartsuspended")))
+		{
+			bStartSuspended = true;
+		}
+		else if (FParse::Param(FCommandLine::Get(), TEXT("hitchdetectionstartrunning")))
+		{
+			bStartSuspended = false;
+		}
+	}
+	if (bStartSuspended)
+	{
+		SuspendedCount = 1;
+	}
 
 	// Start the heart beat thread if it hasn't already been started.
 	if (Thread == nullptr && FPlatformProcess::SupportsMultithreading() && HangDuration > 0)
 	{
 		Thread = FRunnableThread::Create(this, TEXT("FGameThreadHitchHeartBeatThreaded"), 0, TPri_AboveNormal);
 	}
+	bFirst = false;
 #endif
 }
 
@@ -806,6 +828,17 @@ void FGameThreadHitchHeartBeatThreaded::ResumeHeartBeat()
 #endif
 }
 
+bool FGameThreadHitchHeartBeatThreaded::IsSuspended_Gamethread() const
+{
+#if USE_HITCH_DETECTION
+	check(IsInGameThread());
+	return (SuspendedCount > 0);
+#else
+	return false;
+#endif
+}
+
+
 double FGameThreadHitchHeartBeatThreaded::GetFrameStartTime()
 {
 	return FrameStartTime;
diff --git a/Engine/Source/Runtime/Core/Public/HAL/ThreadHeartBeat.h b/Engine/Source/Runtime/Core/Public/HAL/ThreadHeartBeat.h
--- a/Engine/Source/Runtime/Core/Public/HAL/ThreadHeartBeat.h
+++ b/Engine/Source/Runtime/Core/Public/HAL/ThreadHeartBeat.h
@@ -267,6 +267,11 @@ public:
 	*/
 	void ResumeHeartBeat();
 
+	/**
+	* Is the heartbeat suspended?
+	*/
+	bool IsSuspended_Gamethread() const;
+
 	// No-op, used in FUnixSignalGameHitchHeartBeat
 	void Restart() {}
 
