Rewrite packaging config parsing code, part 1
- Use parse code in engine to parse config files instead of reinventing the wheel.

#test tested all modified code pathes and made sure output remains the same.

diff --git a/Engine/Source/Runtime/Core/Private/GenericPlatform/GenericPlatformMisc.cpp b/Engine/Source/Runtime/Core/Private/GenericPlatform/GenericPlatformMisc.cpp
--- a/Engine/Source/Runtime/Core/Private/GenericPlatform/GenericPlatformMisc.cpp
+++ b/Engine/Source/Runtime/Core/Private/GenericPlatform/GenericPlatformMisc.cpp
@@ -1327,31 +1327,19 @@ void FGenericPlatformMisc::ParseChunkIdPakchunkIndexMapping(TArray<FString> Chun
 {
 	OutMapping.Empty();
 
-	const TCHAR* PropertyOldChunkIndex = TEXT("Old");
-	const TCHAR* PropertyNewChunkIndex = TEXT("New");
-	for (const FString& Entry : ChunkIndexMappingData)
+	const TCHAR* PropertyOldChunkIndex = TEXT("Old=");
+	const TCHAR* PropertyNewChunkIndex = TEXT("New=");
+	for (FString& Entry : ChunkIndexMappingData)
 	{
 		// Remove parentheses
-		FString EntryContent = Entry.Mid(1, Entry.Len() - 2);
-		TArray<FString> EntryProperties;
-		EntryContent.ParseIntoArray(EntryProperties, TEXT(","));
+		Entry.TrimStartAndEndInline();
+		Entry.ReplaceInline(TEXT("("), TEXT(""));
+		Entry.ReplaceInline(TEXT(")"), TEXT(""));
 
 		int32 ChunkId = -1;
 		int32 PakchunkIndex = -1;
-		for (const FString EntryProperty : EntryProperties)
-		{
-			TArray<FString> MappingKeyValue;
-			EntryProperty.ParseIntoArray(MappingKeyValue, TEXT("="));
-
-			if (MappingKeyValue[0].Equals(PropertyOldChunkIndex, ESearchCase::IgnoreCase))
-			{
-				LexTryParseString(ChunkId, *MappingKeyValue[1]);
-			}
-			else if (MappingKeyValue[0].Equals(PropertyNewChunkIndex, ESearchCase::IgnoreCase))
-			{
-				LexTryParseString(PakchunkIndex, *MappingKeyValue[1]);
-			}
-		}
+		FParse::Value(*Entry, PropertyOldChunkIndex, ChunkId);
+		FParse::Value(*Entry, PropertyNewChunkIndex, PakchunkIndex);
 
 		if (ChunkId != -1 && PakchunkIndex != -1 && ChunkId != PakchunkIndex && !OutMapping.Contains(ChunkId))
 		{
