Pump platform messages when clearing out task graphs that rely on other threads, e.g. GT GC event waiting on RT which is blocked in Present waiting for messages to be handled.
Gil.Gribb, Matt.Kuhlenschmidt, Jian.Ru

diff --git a/Engine/Source/Runtime/ApplicationCore/Private/Windows/WindowsPlatformApplicationMisc.cpp b/Engine/Source/Runtime/ApplicationCore/Private/Windows/WindowsPlatformApplicationMisc.cpp
--- a/Engine/Source/Runtime/ApplicationCore/Private/Windows/WindowsPlatformApplicationMisc.cpp
+++ b/Engine/Source/Runtime/ApplicationCore/Private/Windows/WindowsPlatformApplicationMisc.cpp
@@ -109,20 +109,12 @@ static void WinPumpMessages()
 	}
 }
 
-static void WinPumpSentMessages()
-{
-	MSG Msg;
-	PeekMessage(&Msg,NULL,0,0,PM_NOREMOVE | PM_QS_SENDMESSAGE);
-}
 
 void FWindowsPlatformApplicationMisc::PumpMessages(bool bFromMainLoop)
 {
 	if (!bFromMainLoop)
 	{
-		TGuardValue<bool> PumpMessageGuard( GPumpingMessagesOutsideOfMainLoop, true );
-		// Process pending windows messages, which is necessary to the rendering thread in some rare cases where D3D
-		// sends window messages (from IDXGISwapChain::Present) to the main thread owned viewport window.
-		WinPumpSentMessages();
+		FPlatformMisc::PumpMessagesOutsideMainLoop();
 		return;
 	}
 
diff --git a/Engine/Source/Runtime/Core/Private/Async/TaskGraph.cpp b/Engine/Source/Runtime/Core/Private/Async/TaskGraph.cpp
--- a/Engine/Source/Runtime/Core/Private/Async/TaskGraph.cpp
+++ b/Engine/Source/Runtime/Core/Private/Async/TaskGraph.cpp
@@ -1926,7 +1926,9 @@ void FTaskGraphInterface::BroadcastSlow_OnlyUseForSpecialPurposes(bool bDoTaskTh
 	{
 		Tasks.Add(TGraphTask<FBroadcastTask>::CreateTask().ConstructAndDispatchWhenReady(Callback, StartTime, TEXT("AudioT"), ENamedThreads::SetTaskPriority(ENamedThreads::AudioThread, ENamedThreads::HighTaskPriority), nullptr, nullptr, nullptr));
 	}
-	Tasks.Add(TGraphTask<FBroadcastTask>::CreateTask().ConstructAndDispatchWhenReady(Callback, StartTime, TEXT("GT"), ENamedThreads::GameThread_Local, nullptr, nullptr, nullptr));
+
+	Callback(ENamedThreads::GameThread_Local);
+
 	if (bDoTaskThreads)
 	{
 		check(MyEvent);
@@ -1954,13 +1956,31 @@ void FTaskGraphInterface::BroadcastSlow_OnlyUseForSpecialPurposes(bool bDoTaskTh
 	{
 		double StartTimeInner = FPlatformTime::Seconds();
 		QUICK_SCOPE_CYCLE_COUNTER(STAT_Broadcast_WaitForNamedThreads);
-		FTaskGraphInterface::Get().WaitUntilTasksComplete(Tasks, ENamedThreads::GameThread_Local);
+		
+		// Wait for all tasks to be complete.  Spin and pump messages to avoid deadlocks when other threads send messages and block until messages are processed
+		while (true)
 		{
-			float ThisTime = FPlatformTime::Seconds() - StartTimeInner;
-			if (ThisTime > 0.02f)
+			bool bAnyNotDone = false;
+			for (FGraphEventRef& Item : Tasks)
+			{
+				if (Item.GetReference() && !Item->IsComplete())
+				{
+					bAnyNotDone = true;
+					break;
+				}
+			}
+			if (!bAnyNotDone)
 			{
-				UE_CLOG(GPrintBroadcastWarnings, LogTaskGraph, Warning, TEXT("Task graph took %6.2fms to wait for named thread broadcast."), ThisTime * 1000.0f);
+				break;
 			}
+
+			FPlatformMisc::PumpMessagesOutsideMainLoop();
+		}
+		
+		float EndTimeInner = FPlatformTime::Seconds() - StartTimeInner;
+		if (EndTimeInner > 0.02f)
+		{
+			UE_CLOG(GPrintBroadcastWarnings, LogTaskGraph, Warning, TEXT("Task graph took %6.2fms to wait for named thread broadcast."), EndTimeInner * 1000.0f);
 		}
 	}
 	for (FEvent* TaskEvent : TaskEvents)
diff --git a/Engine/Source/Runtime/Core/Private/Windows/WindowsPlatformMisc.cpp b/Engine/Source/Runtime/Core/Private/Windows/WindowsPlatformMisc.cpp
--- a/Engine/Source/Runtime/Core/Private/Windows/WindowsPlatformMisc.cpp
+++ b/Engine/Source/Runtime/Core/Private/Windows/WindowsPlatformMisc.cpp
@@ -2694,3 +2694,13 @@ IPlatformChunkInstall* FWindowsPlatformMisc::GetPlatformChunkInstall()
 
 	return ChunkInstall;
 }
+
+void FWindowsPlatformMisc::PumpMessagesOutsideMainLoop()
+{
+	TGuardValue<bool> PumpMessageGuard(GPumpingMessagesOutsideOfMainLoop, true);
+	// Process pending windows messages, which is necessary to the rendering thread in some cases where D3D
+	// sends window messages (from IDXGISwapChain::Present) to the main thread owned viewport window.
+	MSG Msg;
+	PeekMessage(&Msg, NULL, 0, 0, PM_NOREMOVE | PM_QS_SENDMESSAGE);
+	return;
+}
diff --git a/Engine/Source/Runtime/Core/Public/GenericPlatform/GenericPlatformMisc.h b/Engine/Source/Runtime/Core/Public/GenericPlatform/GenericPlatformMisc.h
--- a/Engine/Source/Runtime/Core/Public/GenericPlatform/GenericPlatformMisc.h
+++ b/Engine/Source/Runtime/Core/Public/GenericPlatform/GenericPlatformMisc.h
@@ -1290,6 +1290,10 @@ public:
 		return false;
 	}
 
+	static void PumpMessagesOutsideMainLoop()
+	{
+	}
+
 protected:
 	/** Whether the user should be prompted to allow for a remote debugger to be attached */
 	static bool bShouldPromptForRemoteDebugging;
diff --git a/Engine/Source/Runtime/Core/Public/Windows/WindowsPlatformMisc.h b/Engine/Source/Runtime/Core/Public/Windows/WindowsPlatformMisc.h
--- a/Engine/Source/Runtime/Core/Public/Windows/WindowsPlatformMisc.h
+++ b/Engine/Source/Runtime/Core/Public/Windows/WindowsPlatformMisc.h
@@ -241,6 +241,9 @@ struct CORE_API FWindowsPlatformMisc
 	static EConvertibleLaptopMode GetConvertibleLaptopMode();
 
 	static IPlatformChunkInstall* GetPlatformChunkInstall();
+
+	static void PumpMessagesOutsideMainLoop();
+
 };
 
 
