Fix non-JNI Android compile issue

#ue4
#android

diff --git a/Engine/Source/Runtime/Core/Private/Android/AndroidMisc.cpp b/Engine/Source/Runtime/Core/Private/Android/AndroidMisc.cpp
--- a/Engine/Source/Runtime/Core/Private/Android/AndroidMisc.cpp
+++ b/Engine/Source/Runtime/Core/Private/Android/AndroidMisc.cpp
@@ -44,8 +44,10 @@
 
 #if USE_ANDROID_JNI
 extern AAssetManager * AndroidThunkCpp_GetAssetManager();
-#endif
 extern int32 GAndroidPackageVersion;
+#else
+#define GAndroidPackageVersion	1
+#endif
 
 static int32 GAndroidTraceMarkersEnabled = 0;
 static FAutoConsoleVariableRef CAndroidTraceMarkersEnabled(
