Fix CVar name

diff --git a/Engine/Source/Runtime/Core/Private/Apple/ApplePlatformStackWalk.cpp b/Engine/Source/Runtime/Core/Private/Apple/ApplePlatformStackWalk.cpp
--- a/Engine/Source/Runtime/Core/Private/Apple/ApplePlatformStackWalk.cpp
+++ b/Engine/Source/Runtime/Core/Private/Apple/ApplePlatformStackWalk.cpp
@@ -28,7 +28,7 @@ static TAutoConsoleVariable<int32> CVarApplePlatformThreadCallStackEnabled(
 	1,
 	TEXT("If > 0, then when requesting callstack info about a thread will raise a signal and gather that information."));
 static TAutoConsoleVariable<float> CVarApplePlatformThreadCallStackMaxWait(
-	TEXT("ApplePlatformThreadStackWalk.Enable"),
+	TEXT("ApplePlatformThreadStackWalk.MaxWait"),
 	60.0f,
 	TEXT("The number of seconds allowed to spin before killing the process, with the assumption the signal handler has hung."));
 #endif
