Do not set the ignore count based on the si_value from the signal info

diff --git a/Engine/Source/Runtime/Core/Private/Unix/UnixPlatformCrashContext.cpp b/Engine/Source/Runtime/Core/Private/Unix/UnixPlatformCrashContext.cpp
--- a/Engine/Source/Runtime/Core/Private/Unix/UnixPlatformCrashContext.cpp
+++ b/Engine/Source/Runtime/Core/Private/Unix/UnixPlatformCrashContext.cpp
@@ -101,12 +101,6 @@ void FUnixCrashContext::InitFromSignal(int32 InSignal, siginfo_t* InInfo, void*
 	Context = reinterpret_cast< ucontext_t* >( InContext );
 
 	FCString::Strcat(SignalDescription, ARRAY_COUNT( SignalDescription ) - 1, *DescribeSignal(Signal, Info, Context));
-
-	// Retrieve NumStackFramesToIgnore from signal data
-	if (Info && Info->si_value.sival_int != 0)
-	{
-		SetNumMinidumpFramesToIgnore(Info->si_value.sival_int);
-	}
 }
 
 void FUnixCrashContext::InitFromEnsureHandler(const TCHAR* EnsureMessage, const void* CrashAddress)
