Copy all the GPURestart logs into our .log file for non-sandboxed UE4 projects so that we don't need to gather sysdiagnose from Macs to report bugs to the GPU vendors. Only happens when the crash is caused by a GPU error.

diff --git a/Engine/Source/Runtime/Apple/MetalRHI/Private/MetalCommandList.cpp b/Engine/Source/Runtime/Apple/MetalRHI/Private/MetalCommandList.cpp
--- a/Engine/Source/Runtime/Apple/MetalRHI/Private/MetalCommandList.cpp
+++ b/Engine/Source/Runtime/Apple/MetalRHI/Private/MetalCommandList.cpp
@@ -34,8 +34,11 @@ FMetalCommandList::~FMetalCommandList(void)
 	
 #pragma mark - Public Command List Mutators -
 
+extern CORE_API bool GIsGPUCrashed;
 static void ReportMetalCommandBufferFailure(mtlpp::CommandBuffer const& CompletedBuffer, TCHAR const* ErrorType, bool bDoCheck=true)
 {
+	GIsGPUCrashed = true;
+	
 	NSString* Label = CompletedBuffer.GetLabel();
 	int32 Code = CompletedBuffer.GetError().GetCode();
 	NSString* Domain = CompletedBuffer.GetError().GetDomain();
diff --git a/Engine/Source/Runtime/Core/Private/Mac/MacPlatformMisc.cpp b/Engine/Source/Runtime/Core/Private/Mac/MacPlatformMisc.cpp
--- a/Engine/Source/Runtime/Core/Private/Mac/MacPlatformMisc.cpp
+++ b/Engine/Source/Runtime/Core/Private/Mac/MacPlatformMisc.cpp
@@ -258,9 +258,12 @@ struct FMacApplicationInfo
 		[PLCrashReportFile getCString:PLCrashReportPath maxLength:PATH_MAX encoding:NSUTF8StringEncoding];
 		
 		SystemLogSize = 0;
+		KernelErrorDir = nullptr;
 		if (!bIsSandboxed)
 		{
 			SystemLogSize = IFileManager::Get().FileSize(TEXT("/var/log/system.log"));
+			
+			KernelErrorDir = opendir("/Library/Logs/DiagnosticReports");
 		}
 		
 		if (!FPlatformMisc::IsDebuggerPresent() && FParse::Param(FCommandLine::Get(), TEXT("RedirectNSLog")))
@@ -317,6 +320,11 @@ struct FMacApplicationInfo
 			notify_cancel(PowerSourceNotification);
 			PowerSourceNotification = 0;
 		}
+		if (KernelErrorDir)
+		{
+			closedir(KernelErrorDir);
+			KernelErrorDir = nullptr;
+		}
 	}
 	
 	static FGuid RunGUID()
@@ -395,6 +403,7 @@ struct FMacApplicationInfo
 	FString XcodePath;
 	NSOperatingSystemVersion XcodeVersion;
 	NSPipe* StdErrPipe;
+	DIR* KernelErrorDir;
 	static PLCrashReporter* CrashReporter;
 };
 static FMacApplicationInfo GMacAppInfo;
@@ -1769,6 +1778,36 @@ void FMacCrashContext::GenerateInfoInFolder(char const* const InfoFolder) const
 		close(ConfigDst);
 		close(ConfigSrc);
 
+		// Copy all the GPU restart logs from the user machine into our log
+		if ( !GMacAppInfo.bIsSandboxed && GIsGPUCrashed && GMacAppInfo.KernelErrorDir )
+		{
+			struct dirent DirEntry;
+			struct dirent* DirResult;
+			while(readdir_r(GMacAppInfo.KernelErrorDir, &DirEntry, &DirResult) == 0 && DirResult == &DirEntry)
+			{
+				if (strstr(DirEntry.d_name, ".gpuRestart"))
+				{
+					FCStringAnsi::Strncpy(FilePath, "/Library/Logs/DiagnosticReports/", PATH_MAX);
+					FCStringAnsi::Strcat(FilePath, PATH_MAX, DirEntry.d_name);
+					if (access(FilePath, R_OK|F_OK) == 0)
+					{
+						char const* SysLogHeader = "\nAppending GPU Restart Log: ";
+						write(LogDst, SysLogHeader, strlen(SysLogHeader));
+						
+						write(LogDst, FilePath, strlen(FilePath));
+						write(LogDst, "\n", strlen("\n"));
+
+						int SysLogSrc = open(FilePath, O_RDONLY);
+						while((Bytes = read(SysLogSrc, Data, PATH_MAX)) > 0)
+						{
+							write(LogDst, Data, Bytes);
+						}
+						close(SysLogSrc);
+					}
+				}
+			}
+		}
+		
 		// Copy the system log to capture GPU restarts and other nasties not reported by our application
 		if ( !GMacAppInfo.bIsSandboxed && GMacAppInfo.SystemLogSize >= 0 && access("/var/log/system.log", R_OK|F_OK) == 0 )
 		{
