Fixing compression of CSVs (used by bad match reporting).

diff --git a/Engine/Source/Runtime/Core/Private/ProfilingDebugging/CsvProfiler.cpp b/Engine/Source/Runtime/Core/Private/ProfilingDebugging/CsvProfiler.cpp
--- a/Engine/Source/Runtime/Core/Private/ProfilingDebugging/CsvProfiler.cpp
+++ b/Engine/Source/Runtime/Core/Private/ProfilingDebugging/CsvProfiler.cpp
@@ -1086,7 +1086,7 @@ private:
 						// Compress the data in Buffer into the GZipBuffer array
 						CompressedSize = GZipBuffer.Num();
 						if (FCompression::CompressMemory(
-							NAME_Zlib,
+							NAME_Gzip,
 							GZipBuffer.GetData(), CompressedSize,
 							Buffer.GetData(), BytesInBuffer,
 							ECompressionFlags::COMPRESS_BiasSpeed))
