Don't attempt to load bundle manager module if it doesn't exist.

diff --git a/Engine/Source/Runtime/Core/Private/GenericPlatform/GenericPlatformMisc.cpp b/Engine/Source/Runtime/Core/Private/GenericPlatform/GenericPlatformMisc.cpp
--- a/Engine/Source/Runtime/Core/Private/GenericPlatform/GenericPlatformMisc.cpp
+++ b/Engine/Source/Runtime/Core/Private/GenericPlatform/GenericPlatformMisc.cpp
@@ -833,12 +833,16 @@ IPlatformInstallBundleManager* FGenericPlatformMisc::GetPlatformInstallBundleMan
 		IPlatformInstallBundleManagerModule* Module = nullptr;
 		GConfig->GetString(TEXT("InstallBundleManager"), TEXT("ModuleName"), ModuleName, GEngineIni);
 
-		FModuleStatus Status;
-		Module = FModuleManager::LoadModulePtr<IPlatformInstallBundleManagerModule>(*ModuleName);
-		if (Module)
+		if (FModuleManager::Get().ModuleExists(*ModuleName))
 		{
-			Manager = Module->GetInstallBundleManager();
+			FModuleStatus Status;
+			Module = FModuleManager::LoadModulePtr<IPlatformInstallBundleManagerModule>(*ModuleName);
+			if (Module)
+			{
+				Manager = Module->GetInstallBundleManager();
+			}
 		}
+
 		bCheckedIni = true;
 	}
 
