We want the callstack where we are spin waiting so crash there versus rasing a signal

diff --git a/Engine/Source/Runtime/Core/Private/Apple/ApplePlatformStackWalk.cpp b/Engine/Source/Runtime/Core/Private/Apple/ApplePlatformStackWalk.cpp
--- a/Engine/Source/Runtime/Core/Private/Apple/ApplePlatformStackWalk.cpp
+++ b/Engine/Source/Runtime/Core/Private/Apple/ApplePlatformStackWalk.cpp
@@ -334,7 +334,8 @@ static void RaiseSIGUSR2ForThreadAndWait(uint32 ThreadId)
 		{
 			if (TotalWaitTime > MaxWaitTime)
 			{
-				raise(SIGQUIT);
+				// We have gone over our wait time and now should crash versus waiting possibly forever
+				*(int*)0x10 = 0x11;
 			}
 
 			// sleep a small amount until MaxWaitTime.
