Implemented Gauntlet's "-dev" feature for Mac. This gives the ability to use a locally build executable against network builds by comparing timestamps and using the -basedir argument to point the client to the network data.

diff --git a/Engine/Source/Programs/AutomationTool/Gauntlet/Platform/Mac/Gauntlet.TargetDeviceMac.cs b/Engine/Source/Programs/AutomationTool/Gauntlet/Platform/Mac/Gauntlet.TargetDeviceMac.cs
--- a/Engine/Source/Programs/AutomationTool/Gauntlet/Platform/Mac/Gauntlet.TargetDeviceMac.cs
+++ b/Engine/Source/Programs/AutomationTool/Gauntlet/Platform/Mac/Gauntlet.TargetDeviceMac.cs
@@ -220,19 +220,20 @@ namespace Gauntlet
 				}
 			}
 
-			if (Path.IsPathRooted(InBuild.ExecutablePath))
-			{
-				MacApp.ExecutablePath = InBuild.ExecutablePath;
-			}
-			else
+			// For Mac turn Foo.app into Foo/Content/MacOS/Foo
+			string AppPath = Path.GetDirectoryName(InBuild.ExecutablePath);
+			string FileName = Path.GetFileNameWithoutExtension(InBuild.ExecutablePath);
+			MacApp.ExecutablePath = Path.Combine(InBuild.ExecutablePath, "Contents", "MacOS", FileName);
+
+			if (!Path.IsPathRooted(MacApp.ExecutablePath))
 			{
 				// TODO - this check should be at a higher level....
-				string BinaryPath = Path.Combine(BuildPath, InBuild.ExecutablePath);
+				string BinaryPath = Path.Combine(BuildPath, MacApp.ExecutablePath);
 
 				// check for a local newer executable
 				if (Globals.Params.ParseParam("dev") && AppConfig.ProcessType.UsesEditor() == false)
 				{
-					string LocalBinary = Path.Combine(Environment.CurrentDirectory, InBuild.ExecutablePath);
+					string LocalBinary = Path.Combine(Environment.CurrentDirectory, MacApp.ExecutablePath);
 
 					bool LocalFileExists = File.Exists(LocalBinary);
 					bool LocalFileNewer = LocalFileExists && File.GetLastWriteTime(LocalBinary) > File.GetLastWriteTime(BinaryPath);
@@ -242,8 +243,10 @@ namespace Gauntlet
 
 					if (LocalFileExists && LocalFileNewer)
 					{
+						string LocalAppPath = Path.Combine(Environment.CurrentDirectory, InBuild.ExecutablePath);
+						
 						// need to -basedir to have our exe load content from the path
-						MacApp.CommandArguments += string.Format(" -basedir={0}", Path.GetDirectoryName(BinaryPath));
+						MacApp.CommandArguments += string.Format(" -basedir={0}", Path.GetDirectoryName(LocalAppPath));
 
 						BinaryPath = LocalBinary;
 					}
@@ -252,11 +255,6 @@ namespace Gauntlet
 				MacApp.ExecutablePath = BinaryPath;
 			}
 
-			// now turn the Foo.app into Foo/Content/MacOS/Foo
-			string AppPath = Path.GetDirectoryName(MacApp.ExecutablePath);
-			string FileName = Path.GetFileNameWithoutExtension(MacApp.ExecutablePath);
-			MacApp.ExecutablePath = Path.Combine(MacApp.ExecutablePath, "Contents", "MacOS", FileName);
-
 			return MacApp;
 		}
 
diff --git a/Engine/Source/Runtime/Core/Private/Mac/MacPlatformProcess.cpp b/Engine/Source/Runtime/Core/Private/Mac/MacPlatformProcess.cpp
--- a/Engine/Source/Runtime/Core/Private/Mac/MacPlatformProcess.cpp
+++ b/Engine/Source/Runtime/Core/Private/Mac/MacPlatformProcess.cpp
@@ -732,31 +732,50 @@ const TCHAR* FMacPlatformProcess::BaseDir()
 	if (!Result[0])
 	{
 		SCOPED_AUTORELEASE_POOL;
-		NSFileManager* FileManager = [NSFileManager defaultManager];
-		NSString *BasePath = [[NSBundle mainBundle] bundlePath];
-		// If it has .app extension, it's a bundle, otherwise BasePath is a full path to Binaries/Mac (in case of command line tools)
-		if ([[BasePath pathExtension] isEqual: @"app"])
+		
+		FString CommandLine = [[[NSProcessInfo processInfo] arguments] componentsJoinedByString:@" "];
+	
+		FString BaseArg;
+		FParse::Value(*CommandLine, TEXT("-basedir="), BaseArg);
+
+		if (BaseArg.Len())
 		{
-			NSString* BundledBinariesPath = NULL;
-			if (!FApp::IsProjectNameEmpty())
-			{
-				BundledBinariesPath = [BasePath stringByAppendingPathComponent : [NSString stringWithFormat : @"Contents/UE4/%s/Binaries/Mac", TCHAR_TO_UTF8(FApp::GetProjectName())]];
-			}
-			if (!BundledBinariesPath || ![FileManager fileExistsAtPath:BundledBinariesPath])
-			{
-				BundledBinariesPath = [BasePath stringByAppendingPathComponent: @"Contents/UE4/Engine/Binaries/Mac"];
-			}
-			if ([FileManager fileExistsAtPath: BundledBinariesPath])
-			{
-				BasePath = BundledBinariesPath;
-			}
-			else
+			BaseArg = BaseArg.Replace(TEXT("\\"), TEXT("/"));
+			BaseArg += TEXT('/');
+			FCString::Strcpy(Result, *BaseArg);
+			
+			FPlatformMisc::LowLevelOutputDebugStringf(TEXT("BaseDir set to %s"), Result);
+		}
+		else
+		{
+			NSFileManager* FileManager = [NSFileManager defaultManager];
+			NSString *BasePath = [[NSBundle mainBundle] bundlePath];
+			// If it has .app extension, it's a bundle, otherwise BasePath is a full path to Binaries/Mac (in case of command line tools)
+			if ([[BasePath pathExtension] isEqual: @"app"])
 			{
-				BasePath = [BasePath stringByDeletingLastPathComponent];
+				NSString* BundledBinariesPath = NULL;
+				if (!FApp::IsProjectNameEmpty())
+				{
+					BundledBinariesPath = [BasePath stringByAppendingPathComponent : [NSString stringWithFormat : @"Contents/UE4/%s/Binaries/Mac", TCHAR_TO_UTF8(FApp::GetProjectName())]];
+				}
+				if (!BundledBinariesPath || ![FileManager fileExistsAtPath:BundledBinariesPath])
+				{
+					BundledBinariesPath = [BasePath stringByAppendingPathComponent: @"Contents/UE4/Engine/Binaries/Mac"];
+				}
+				if ([FileManager fileExistsAtPath: BundledBinariesPath])
+				{
+					BasePath = BundledBinariesPath;
+				}
+				else
+				{
+					BasePath = [BasePath stringByDeletingLastPathComponent];
+				}
 			}
+		
+			FCString::Strcpy(Result, MAC_MAX_PATH, *FString(BasePath));
+			FCString::Strcat(Result, TEXT("/"));
+			
 		}
-		FCString::Strcpy(Result, MAC_MAX_PATH, *FString(BasePath));
-		FCString::Strcat(Result, TEXT("/"));
 	}
 	return Result;
 }
