added `ElementType` typedef to TUniquePtr/TSharedPtr/TSharedRef

diff --git a/Engine/Source/Runtime/Core/Public/Templates/SharedPointer.h b/Engine/Source/Runtime/Core/Public/Templates/SharedPointer.h
--- a/Engine/Source/Runtime/Core/Public/Templates/SharedPointer.h
+++ b/Engine/Source/Runtime/Core/Public/Templates/SharedPointer.h
@@ -147,6 +147,7 @@ template< class ObjectType, ESPMode Mode >
 class TSharedRef
 {
 public:
+	using ElementType = ObjectType;
 
 	// NOTE: TSharedRef has no default constructor as it does not support empty references.  You must
 	//		 initialize your TSharedRef to a valid object at construction time.
@@ -546,6 +547,7 @@ class TSharedPtr
 	static_assert( TSharedPtr::ObjectTypeHasSameModeSharedFromThis || !TSharedPtr::ObjectTypeHasOppositeModeSharedFromThis, "You cannot use a TSharedPtr of one mode with a type which inherits TSharedFromThis of another mode.");
 
 public:
+	using ElementType = ObjectType;
 
 	/**
 	 * Constructs an empty shared pointer
@@ -951,6 +953,7 @@ template< class ObjectType, ESPMode Mode >
 class TWeakPtr
 {
 public:
+	using ElementType = ObjectType;
 
 	/** Constructs an empty TWeakPtr */
 	// NOTE: FNullTag parameter is an Unreal extension to standard shared_ptr behavior
diff --git a/Engine/Source/Runtime/Core/Public/Templates/UniquePtr.h b/Engine/Source/Runtime/Core/Public/Templates/UniquePtr.h
--- a/Engine/Source/Runtime/Core/Public/Templates/UniquePtr.h
+++ b/Engine/Source/Runtime/Core/Public/Templates/UniquePtr.h
@@ -87,6 +87,8 @@ class TUniquePtr : private Deleter
 	friend class TUniquePtr;
 
 public:
+	using ElementType = T;
+
 	/**
 	 * Default constructor - initializes the TUniquePtr to null.
 	 */
@@ -318,6 +320,8 @@ class TUniquePtr<T[], Deleter> : private Deleter
 	friend class TUniquePtr;
 
 public:
+	using ElementType = T;
+
 	/**
 	 * Default constructor - initializes the TUniquePtr to null.
 	 */
