Allow crash report contexts to query the game name if it hasn't already been initialized at the point that a crash occurs. Also query the platform name.

Fixes problem with CrashReportClient not being able to submit a crash unless the game name is set.

diff --git a/Engine/Source/Runtime/Core/Private/GenericPlatform/GenericPlatformCrashContext.cpp b/Engine/Source/Runtime/Core/Private/GenericPlatform/GenericPlatformCrashContext.cpp
--- a/Engine/Source/Runtime/Core/Private/GenericPlatform/GenericPlatformCrashContext.cpp
+++ b/Engine/Source/Runtime/Core/Private/GenericPlatform/GenericPlatformCrashContext.cpp
@@ -77,8 +77,6 @@ namespace NCachedCrashContextProperties
 	static TOptional<bool> bIsVanilla;
 	static FString GameName;
 	static FString ExecutableName;
-	static FString PlatformName;
-	static FString PlatformNameIni;
 	static FString DeploymentName;
 	static FString BaseDir;
 	static FString RootDir;
@@ -117,8 +115,6 @@ void FGenericCrashContext::Initialize()
 
 	NCachedCrashContextProperties::GameName = FString::Printf( TEXT("UE4-%s"), FApp::GetProjectName() );
 	NCachedCrashContextProperties::ExecutableName = FPlatformProcess::ExecutableName();
-	NCachedCrashContextProperties::PlatformName = FPlatformProperties::PlatformName();
-	NCachedCrashContextProperties::PlatformNameIni = FPlatformProperties::IniPlatformName();
 	NCachedCrashContextProperties::BaseDir = FPlatformProcess::BaseDir();
 	NCachedCrashContextProperties::RootDir = FPlatformMisc::RootDir();
 	NCachedCrashContextProperties::EpicAccountId = FPlatformMisc::GetEpicAccountId();
@@ -171,7 +167,8 @@ void FGenericCrashContext::Initialize()
 	}
 
 	const FGuid Guid = FGuid::NewGuid();
-	NCachedCrashContextProperties::CrashGUIDRoot = FString::Printf(TEXT("%s%s-%s"), *CrashGUIDRootPrefix, *NCachedCrashContextProperties::PlatformNameIni, *Guid.ToString(EGuidFormats::Digits));
+	const FString IniPlatformName(FPlatformProperties::IniPlatformName());
+	NCachedCrashContextProperties::CrashGUIDRoot = FString::Printf(TEXT("%s%s-%s"), *CrashGUIDRootPrefix, *IniPlatformName, *Guid.ToString(EGuidFormats::Digits));
 
 	// Initialize delegate for updating SecondsSinceStart, because FPlatformTime::Seconds() is not POSIX safe.
 	const float PollingInterval = 1.0f;
@@ -298,7 +295,22 @@ void FGenericCrashContext::SerializeContentToBuffer() const
 	AddCrashProperty( TEXT( "SecondsSinceStart" ), NCachedCrashContextProperties::SecondsSinceStart );
 
 	// Add common crash properties.
-	AddCrashProperty( TEXT( "GameName" ), *NCachedCrashContextProperties::GameName );
+	if (NCachedCrashContextProperties::GameName.Len() > 0)
+	{
+		AddCrashProperty(TEXT("GameName"), *NCachedCrashContextProperties::GameName);
+	}
+	else
+	{
+		const TCHAR* ProjectName = FApp::GetProjectName();
+		if (ProjectName != nullptr && ProjectName[0] != 0)
+		{
+			AddCrashProperty(TEXT("GameName"), *FString::Printf(TEXT("UE4-%s"), ProjectName));
+		}
+		else
+		{
+			AddCrashProperty(TEXT("GameName"), TEXT(""));
+		}
+	}
 	AddCrashProperty( TEXT( "ExecutableName" ), *NCachedCrashContextProperties::ExecutableName );
 	AddCrashProperty( TEXT( "BuildConfiguration" ), EBuildConfigurations::ToString( FApp::GetBuildConfiguration() ) );
 	AddCrashProperty( TEXT( "GameSessionID" ), *NCachedCrashContextProperties::GameSessionID );
@@ -311,8 +323,8 @@ void FGenericCrashContext::SerializeContentToBuffer() const
 
 	AddCrashProperty( TEXT( "Symbols" ), Symbols);
 
-	AddCrashProperty( TEXT( "PlatformName" ), *NCachedCrashContextProperties::PlatformName );
-	AddCrashProperty( TEXT( "PlatformNameIni" ), *NCachedCrashContextProperties::PlatformNameIni );
+	AddCrashProperty( TEXT( "PlatformName" ), FPlatformProperties::PlatformName() );
+	AddCrashProperty( TEXT( "PlatformNameIni" ), FPlatformProperties::IniPlatformName());
 	AddCrashProperty( TEXT( "EngineMode" ), FPlatformMisc::GetEngineMode() );
 	AddCrashProperty( TEXT( "EngineModeEx" ), EngineModeExString());
 
