Improve Bundle Manager Progress reporting and Simplify ESP progress
- For background download progress, don't include chunkdbs that were cached before we began the step.  This should now match up with the size returned by GetContentState().
- FInstallBundleStatus is now cumulative and will only ever add information/progress in subsequent calls.
- Only return FInstallBundleStatus if we can actually find a request for it.  Returning information after this just led to bugs and backsliding progress.  The last valid time to query progress is in the completion callback for a bundle.
- Remove Hacks from ESP progress reporting.

diff --git a/Engine/Source/Runtime/Core/Public/GenericPlatform/GenericPlatformInstallBundleManager.h b/Engine/Source/Runtime/Core/Public/GenericPlatform/GenericPlatformInstallBundleManager.h
--- a/Engine/Source/Runtime/Core/Public/GenericPlatform/GenericPlatformInstallBundleManager.h
+++ b/Engine/Source/Runtime/Core/Public/GenericPlatform/GenericPlatformInstallBundleManager.h
@@ -107,7 +107,6 @@ ENUM_CLASS_FLAGS(EInstallBundleRequestFlags)
 
 enum class EInstallBundleStatus : int
 {
-	NotRequested,
 	QueuedForDownload,
 	Downloading,
 	QueuedForInstall,
@@ -117,29 +116,40 @@ enum class EInstallBundleStatus : int
 	Installed,
 };
 
-struct FInstallBundleProgress
+struct FInstallBundleDownloadProgress
 {
 	// Num bytes received
-	uint64 ProgressBytes = 0;
-	// Num bytes written to storage (<= ProgressBytes)
-	uint64 ProgressBytesWritten = 0;
+	uint64 BytesDownloaded = 0;
+	// Num bytes written to storage (<= BytesDownloaded)
+	uint64 BytesDownloadedAndWritten = 0;
 	// Num bytes needed
-	uint64 ProgressTotalBytes = 0;
-	float ProgressPercent = 0;
+	uint64 TotalBytesToDownload = 0;
+	float PercentComplete = 0;
 };
 
 struct FInstallBundleStatus
 {
 	FName BundleName;
 
-	EInstallBundleStatus Status = EInstallBundleStatus::NotRequested;
+	EInstallBundleStatus Status = EInstallBundleStatus::QueuedForDownload;
 
 	EInstallBundlePauseFlags PauseFlags = EInstallBundlePauseFlags::None;
 
 	FText StatusText;
 
-	// Progress is only present if Status is Downloading or Installing
-	TOptional<FInstallBundleProgress> Progress;
+	// Download progress of EInstallBundleStatus::Downloading
+	// Will be set if Status >= EInstallBundleStatus::Downloading
+	TOptional<FInstallBundleDownloadProgress> BackgroundDownloadProgress;
+
+	// Download progress of EInstallBundleStatus::Install
+	// Will be set if Status >= EInstallBundleStatus::Install
+	// We may download during install if background downloads are turned off or fail
+	// We may also do small downloads during install as a normal part of installation
+	TOptional<FInstallBundleDownloadProgress> InstallDownloadProgress;
+
+	float Install_Percent = 0;
+
+	float Finishing_Percent = 0;
 };
 
 struct FInstallBundleResultInfo
