Report all thread backtrace when a hang happens

diff --git a/Engine/Source/Runtime/Core/Private/GenericPlatform/GenericPlatformCrashContext.cpp b/Engine/Source/Runtime/Core/Private/GenericPlatform/GenericPlatformCrashContext.cpp
--- a/Engine/Source/Runtime/Core/Private/GenericPlatform/GenericPlatformCrashContext.cpp
+++ b/Engine/Source/Runtime/Core/Private/GenericPlatform/GenericPlatformCrashContext.cpp
@@ -784,6 +784,11 @@ FORCENOINLINE void FGenericCrashContext::CapturePortableCallStack(int32 NumStack
 }
 
 void FGenericCrashContext::SetPortableCallStack(const uint64* StackFrames, int32 NumStackFrames)
+{
+	GetPortableCallStack(StackFrames, NumStackFrames, CallStack);
+}
+
+void FGenericCrashContext::GetPortableCallStack(const uint64* StackFrames, int32 NumStackFrames, TArray<FCrashStackFrame>& OutCallStack)
 {
 	// Get all the modules in the current process
 	uint32 NumModules = (uint32)FPlatformStackWalk::GetProcessModuleCount();
@@ -795,7 +800,7 @@ void FGenericCrashContext::SetPortableCallStack(const uint64* StackFrames, int32
 	Modules.SetNum(NumModules);
 
 	// Update the callstack with offsets from each module
-	CallStack.Reset(NumStackFrames);
+	OutCallStack.Reset(NumStackFrames);
 	for(int32 Idx = 0; Idx < NumStackFrames; Idx++)
 	{
 		const uint64 StackFrame = StackFrames[Idx];
@@ -814,11 +819,11 @@ void FGenericCrashContext::SetPortableCallStack(const uint64* StackFrames, int32
 		// Add the callstack item
 		if(FoundModule == nullptr)
 		{
-			CallStack.Add(FCrashStackFrame(TEXT("Unknown"), 0, StackFrame));
+			OutCallStack.Add(FCrashStackFrame(TEXT("Unknown"), 0, StackFrame));
 		}
 		else
 		{
-			CallStack.Add(FCrashStackFrame(FPaths::GetBaseFilename(FoundModule->ImageName), FoundModule->BaseOfImage, StackFrame - FoundModule->BaseOfImage));
+			OutCallStack.Add(FCrashStackFrame(FPaths::GetBaseFilename(FoundModule->ImageName), FoundModule->BaseOfImage, StackFrame - FoundModule->BaseOfImage));
 		}
 	}
 }
diff --git a/Engine/Source/Runtime/Core/Private/HAL/ThreadingBase.cpp b/Engine/Source/Runtime/Core/Private/HAL/ThreadingBase.cpp
--- a/Engine/Source/Runtime/Core/Private/HAL/ThreadingBase.cpp
+++ b/Engine/Source/Runtime/Core/Private/HAL/ThreadingBase.cpp
@@ -214,7 +214,7 @@ const FString& FThreadManager::GetThreadName(uint32 ThreadId)
 	return NoThreadName;
 }
 
-#if PLATFORM_WINDOWS
+#if PLATFORM_WINDOWS || PLATFORM_MAC
 static void GetAllThreadStackBackTraces_ProcessSingle(
 	uint32 CurThreadId,
 	uint32 ThreadId,
diff --git a/Engine/Source/Runtime/Core/Private/Mac/MacPlatformMisc.cpp b/Engine/Source/Runtime/Core/Private/Mac/MacPlatformMisc.cpp
--- a/Engine/Source/Runtime/Core/Private/Mac/MacPlatformMisc.cpp
+++ b/Engine/Source/Runtime/Core/Private/Mac/MacPlatformMisc.cpp
@@ -20,6 +20,7 @@
 #include "HAL/PlatformOutputDevices.h"
 #include "HAL/IConsoleManager.h"
 #include "HAL/FileManager.h"
+#include "HAL/ThreadManager.h"
 #include "Misc/OutputDeviceError.h"
 #include "Misc/OutputDeviceRedirector.h"
 #include "Misc/FeedbackContext.h"
@@ -71,6 +72,10 @@ static FAutoConsoleVariableRef CVarMacExplicitRendererID(
 	TEXT("Forces the Mac RHI to use the specified rendering device which is a 0-based index into the list of GPUs provided by FMacPlatformMisc::GetGPUDescriptors or -1 to disable & use the default device. (Default: -1, off)"),
 	ECVF_RenderThreadSafe|ECVF_ReadOnly
 	);
+static TAutoConsoleVariable<int32> CVarMacPlatformDumpAllThreadsOnHang(
+	TEXT("Mac.DumpAllThreadsOnHang"),
+	1,
+	TEXT("If > 0, then when reporting a hang generate a backtrace for all threads."));
 
 /*------------------------------------------------------------------------------
  FMacApplicationInfo - class to contain all state for crash reporting that is unsafe to acquire in a signal.
@@ -1989,6 +1994,65 @@ void FMacCrashContext::GenerateEnsureInfoAndLaunchReporter() const
 	}
 }
 
+void FMacCrashContext::AddThreadContext(
+	uint32 ThreadIdEnteredOn,
+	uint32 ThreadId,
+	const FString& ThreadName,
+	const TArray<FCrashStackFrame>& PortableCallStack)
+{
+	AllThreadContexts += TEXT("<Thread>");
+	{
+		AllThreadContexts += TEXT("<CallStack>");
+
+		int32 MaxModuleNameLen = 0;
+		for (const FCrashStackFrame& StFrame : PortableCallStack)
+		{
+			MaxModuleNameLen = FMath::Max(MaxModuleNameLen, StFrame.ModuleName.Len());
+		}
+
+		FString CallstackStr;
+		for (const FCrashStackFrame& StFrame : PortableCallStack)
+		{
+			CallstackStr += FString::Printf(TEXT("%-*s 0x%016x + %-8x"), MaxModuleNameLen + 1, *StFrame.ModuleName, StFrame.BaseAddress, StFrame.Offset);
+			CallstackStr += LINE_TERMINATOR;
+		}
+		AppendEscapedXMLString(AllThreadContexts, *CallstackStr);
+		AllThreadContexts += TEXT("</CallStack>") LINE_TERMINATOR;
+	}
+
+	AllThreadContexts += FString::Printf(TEXT("<IsCrashed>%s</IsCrashed>") LINE_TERMINATOR, (ThreadId == ThreadIdEnteredOn) ? TEXT("true") : TEXT("false"));
+	// TODO: do we need thread register states?
+	AllThreadContexts += TEXT("<Registers></Registers>") LINE_TERMINATOR;
+	AllThreadContexts += FString::Printf(TEXT("<ThreadID>%d</ThreadID>") LINE_TERMINATOR, ThreadId);
+	AllThreadContexts += FString::Printf(TEXT("<ThreadName>%s</ThreadName>") LINE_TERMINATOR, *ThreadName);
+	AllThreadContexts += TEXT("</Thread>");
+	AllThreadContexts += LINE_TERMINATOR;
+}
+
+void FMacCrashContext::CaptureAllThreadContext(uint32 ThreadIdEnteredOn)
+{
+	TArray<typename FThreadManager::FThreadStackBackTrace> StackTraces;
+	FThreadManager::Get().GetAllThreadStackBackTraces(StackTraces);
+
+	for (int32 Idx = 0; Idx < StackTraces.Num(); ++Idx)
+	{
+		const FThreadManager::FThreadStackBackTrace& ThreadStTrace = StackTraces[Idx];
+		const uint32 ThreadId     = ThreadStTrace.ThreadId;
+		const FString& ThreadName = ThreadStTrace.ThreadName;
+
+		TArray<FCrashStackFrame> PortableStack;
+		GetPortableCallStack(ThreadStTrace.ProgramCounters.GetData(), ThreadStTrace.ProgramCounters.Num(), PortableStack);
+
+		AddThreadContext(ThreadIdEnteredOn, ThreadId, ThreadName, PortableStack);
+	}
+}
+
+bool FMacCrashContext::GetPlatformAllThreadContextsString(FString& OutStr) const
+{
+	OutStr = AllThreadContexts;
+	return !OutStr.IsEmpty();
+}
+
 void ReportAssert(const TCHAR* ErrorMessage, int NumStackFramesToIgnore)
 {
 	GCrashErrorMessage = ErrorMessage;
@@ -2044,6 +2108,12 @@ void ReportHang(const TCHAR* ErrorMessage, const uint64* StackFrames, int32 NumS
 
 		FMacCrashContext EnsureContext(ECrashContextType::Hang, ErrorMessage);
 		EnsureContext.SetPortableCallStack(StackFrames, NumStackFrames);
+
+		if (CVarMacPlatformDumpAllThreadsOnHang.AsVariable()->GetInt())
+		{
+			EnsureContext.CaptureAllThreadContext(HungThreadId);
+		}
+
 		EnsureContext.GenerateEnsureInfoAndLaunchReporter();
 
 		bReentranceGuard = false;
diff --git a/Engine/Source/Runtime/Core/Public/GenericPlatform/GenericPlatformCrashContext.h b/Engine/Source/Runtime/Core/Public/GenericPlatform/GenericPlatformCrashContext.h
--- a/Engine/Source/Runtime/Core/Public/GenericPlatform/GenericPlatformCrashContext.h
+++ b/Engine/Source/Runtime/Core/Public/GenericPlatform/GenericPlatformCrashContext.h
@@ -238,6 +238,9 @@ public:
 	/** Sets the portable callstack to a specified stack */
 	virtual void SetPortableCallStack(const uint64* StackFrames, int32 NumStackFrames);
 
+	/** Gets the portable callstack to a specified stack and puts it into OutCallStack */
+	virtual void GetPortableCallStack(const uint64* StackFrames, int32 NumStackFrames, TArray<FCrashStackFrame>& OutCallStack);
+
 	/**
 	 * @return whether this crash is a non-crash event
 	 */
diff --git a/Engine/Source/Runtime/Core/Public/HAL/ThreadManager.h b/Engine/Source/Runtime/Core/Public/HAL/ThreadManager.h
--- a/Engine/Source/Runtime/Core/Public/HAL/ThreadManager.h
+++ b/Engine/Source/Runtime/Core/Public/HAL/ThreadManager.h
@@ -40,7 +40,7 @@ public:
 	/** Returns the name of a thread given its TLS id */
 	const FString& GetThreadName(uint32 ThreadId);
 
-#if PLATFORM_WINDOWS
+#if PLATFORM_WINDOWS || PLATFORM_MAC
 	struct FThreadStackBackTrace
 	{
 		uint32 ThreadId;
diff --git a/Engine/Source/Runtime/Core/Public/Mac/MacPlatformCrashContext.h b/Engine/Source/Runtime/Core/Public/Mac/MacPlatformCrashContext.h
--- a/Engine/Source/Runtime/Core/Public/Mac/MacPlatformCrashContext.h
+++ b/Engine/Source/Runtime/Core/Public/Mac/MacPlatformCrashContext.h
@@ -20,6 +20,32 @@ struct CORE_API FMacCrashContext : public FApplePlatformCrashContext
 	
 	/** Generates information for ensures sent via the CrashReporter */
 	void GenerateEnsureInfoAndLaunchReporter() const;
+
+	/** Captures all information about all threads */
+	void CaptureAllThreadContext(uint32 ThreadIdEnteredOn);
+
+protected:
+	virtual bool GetPlatformAllThreadContextsString(FString& OutStr) const override;
+
+private:
+	void AddThreadContext(
+		uint32 ThreadIdEnteredOn, 
+		uint32 ThreadId,
+		const FString& ThreadName,
+		const TArray<FCrashStackFrame>& PortableCallStack);
+
+	/**
+	* <Thread>
+	*   <CallStack>...</CallStack>
+	*   <IsCrashed>...</IsCrashed>
+	*   <Registers>...</Registers>
+	*   <ThreadID>...</ThreadID>
+	*   <ThreadName>...</ThreadName>
+	* </Thead>
+	* <Thread>...</Thread>
+	* ...
+	*/
+	FString AllThreadContexts;
 };
 
 typedef FMacCrashContext FPlatformCrashContext;
