Dirty fix for some reads occasionally returning EFAULT (UE-69018).

- This is a dirty fix which can incur performance costs if such failures are frequent. Proper fix is tracked as UE-72914.

diff --git a/Engine/Source/Runtime/Core/Private/Unix/UnixPlatformFile.cpp b/Engine/Source/Runtime/Core/Private/Unix/UnixPlatformFile.cpp
--- a/Engine/Source/Runtime/Core/Private/Unix/UnixPlatformFile.cpp
+++ b/Engine/Source/Runtime/Core/Private/Unix/UnixPlatformFile.cpp
@@ -250,6 +250,20 @@ private:
 			int64 ThisSize = FMath::Min<int64>(READWRITE_SIZE, BytesToRead);
 			check(Destination);
 			int64 ThisRead = read(FileHandle, Destination, ThisSize);
+			if (ThisRead == -1 && errno == EFAULT)
+			{
+				// hacky workaround for UE-69018
+				void* TempDest = FMemory::Malloc(ThisSize);
+				if (TempDest)
+				{
+					ThisRead = read(FileHandle, TempDest, ThisSize);
+					if (ThisRead > 0 && ThisRead <= ThisSize)
+					{
+						FMemory::Memcpy(Destination, TempDest, ThisRead);
+					}
+					FMemory::Free(TempDest);
+				}
+			}
 			BytesRead += ThisRead;
 			if (ThisRead != ThisSize)
 			{
