[Switch] Changes neeeded for CSV reporting in live builds. Cleaned up temporary storage support a bit.

diff --git a/Engine/Source/Runtime/Core/Private/Misc/Paths.cpp b/Engine/Source/Runtime/Core/Private/Misc/Paths.cpp
--- a/Engine/Source/Runtime/Core/Private/Misc/Paths.cpp
+++ b/Engine/Source/Runtime/Core/Private/Misc/Paths.cpp
@@ -356,16 +356,18 @@ FString FPaths::VideoCaptureDir()
 FString FPaths::ProjectLogDir()
 {
 #if PLATFORM_PS4
-
 	const FString* OverrideDir = FPS4PlatformFile::GetOverrideLogDirectory();
 	if (OverrideDir != nullptr)
 	{
 		return *OverrideDir;
 	}
-
-#endif
-
-#if PLATFORM_MAC || PLATFORM_XBOXONE
+#elif PLATFORM_SWITCH
+	const FString* OverrideDir = FSwitchPlatformFile::GetOverrideLogDirectory();
+	if (OverrideDir != nullptr)
+	{
+		return *OverrideDir;
+	}
+#elif PLATFORM_MAC || PLATFORM_XBOXONE
 	if (UE4Paths_Private::CustomUserDirArgument().IsEmpty())
 	{
 		return FPlatformProcess::UserLogsDir();
