Implement FPlatformProcess::GetCurrentCoreNumber

diff --git a/Engine/Source/Runtime/Core/Private/Android/AndroidProcess.cpp b/Engine/Source/Runtime/Core/Private/Android/AndroidProcess.cpp
--- a/Engine/Source/Runtime/Core/Private/Android/AndroidProcess.cpp
+++ b/Engine/Source/Runtime/Core/Private/Android/AndroidProcess.cpp
@@ -96,6 +96,13 @@ uint32 FAndroidPlatformProcess::GetCurrentProcessId()
 	return getpid();
 }
 
+uint32 FAndroidPlatformProcess::GetCurrentCoreNumber()
+{
+	unsigned CPU;
+	int Err = syscall(__NR_getcpu, &CPU, nullptr, nullptr);
+	return (!Err) ? CPU : 0;
+}
+
 const TCHAR* FAndroidPlatformProcess::BaseDir()
 {
 	return TEXT("");
diff --git a/Engine/Source/Runtime/Core/Private/GenericPlatform/GenericPlatformProcess.cpp b/Engine/Source/Runtime/Core/Private/GenericPlatform/GenericPlatformProcess.cpp
--- a/Engine/Source/Runtime/Core/Private/GenericPlatform/GenericPlatformProcess.cpp
+++ b/Engine/Source/Runtime/Core/Private/GenericPlatform/GenericPlatformProcess.cpp
@@ -53,6 +53,12 @@ uint32 FGenericPlatformProcess::GetCurrentProcessId()
 	return 0;
 }
 
+uint32 FGenericPlatformProcess::GetCurrentCoreNumber()
+{
+	return 0;
+}
+
+
 void FGenericPlatformProcess::SetThreadAffinityMask( uint64 AffinityMask )
 {
 	// Not implemented cross-platform. Each platform may or may not choose to implement this.
diff --git a/Engine/Source/Runtime/Core/Private/Mac/MacPlatformProcess.cpp b/Engine/Source/Runtime/Core/Private/Mac/MacPlatformProcess.cpp
--- a/Engine/Source/Runtime/Core/Private/Mac/MacPlatformProcess.cpp
+++ b/Engine/Source/Runtime/Core/Private/Mac/MacPlatformProcess.cpp
@@ -16,6 +16,7 @@
 #include <mach/thread_act.h>
 #include <mach/thread_policy.h>
 #include <libproc.h>
+#include <cpuid.h>
 #include "Apple/PostAppleSystemHeaders.h"
 
 void* FMacPlatformProcess::GetDllHandle( const TCHAR* Filename )
@@ -658,6 +659,13 @@ uint32 FMacPlatformProcess::GetCurrentProcessId()
 	return getpid();
 }
 
+uint32 FMacPlatformProcess::GetCurrentCoreNumber()
+{
+	int CPUInfo[4];
+	__cpuid(1, CPUInfo[0], CPUInfo[1], CPUInfo[2], CPUInfo[3]);
+	return (CPUInfo[1] >> 24) & 0xff;
+}
+
 bool FMacPlatformProcess::GetProcReturnCode( FProcHandle& ProcessHandle, int32* ReturnCode )
 {
 	SCOPED_AUTORELEASE_POOL;
diff --git a/Engine/Source/Runtime/Core/Private/Unix/UnixPlatformProcess.cpp b/Engine/Source/Runtime/Core/Private/Unix/UnixPlatformProcess.cpp
--- a/Engine/Source/Runtime/Core/Private/Unix/UnixPlatformProcess.cpp
+++ b/Engine/Source/Runtime/Core/Private/Unix/UnixPlatformProcess.cpp
@@ -1453,6 +1453,11 @@ uint32 FUnixPlatformProcess::GetCurrentProcessId()
 	return getpid();
 }
 
+uint32 FUnixPlatformProcess::GetCurrentCoreNumber()
+{
+	return sched_getcpu();
+}
+
 void FUnixPlatformProcess::SetCurrentWorkingDirectoryToBaseDir()
 {
 	FPlatformMisc::CacheLaunchDir();
diff --git a/Engine/Source/Runtime/Core/Private/Windows/WindowsPlatformProcess.cpp b/Engine/Source/Runtime/Core/Private/Windows/WindowsPlatformProcess.cpp
--- a/Engine/Source/Runtime/Core/Private/Windows/WindowsPlatformProcess.cpp
+++ b/Engine/Source/Runtime/Core/Private/Windows/WindowsPlatformProcess.cpp
@@ -456,6 +456,11 @@ uint32 FWindowsPlatformProcess::GetCurrentProcessId()
 	return ::GetCurrentProcessId();
 }
 
+uint32 FWindowsPlatformProcess::GetCurrentCoreNumber()
+{
+	return ::GetCurrentProcessorNumber();
+}
+
 void FWindowsPlatformProcess::SetThreadAffinityMask( uint64 AffinityMask )
 {
 	if( AffinityMask != FPlatformAffinity::GetNoAffinityMask() )
diff --git a/Engine/Source/Runtime/Core/Public/Android/AndroidProcess.h b/Engine/Source/Runtime/Core/Public/Android/AndroidProcess.h
--- a/Engine/Source/Runtime/Core/Public/Android/AndroidProcess.h
+++ b/Engine/Source/Runtime/Core/Public/Android/AndroidProcess.h
@@ -34,6 +34,7 @@ struct CORE_API FAndroidPlatformProcess : public FGenericPlatformProcess
 	static const TCHAR* ComputerName();
 	static void SetThreadAffinityMask( uint64 AffinityMask );
 	static uint32 GetCurrentProcessId();
+	static uint32 GetCurrentCoreNumber();
 	static const TCHAR* BaseDir();
 	static const TCHAR* ExecutableName(bool bRemoveExtension = true);
 	static class FRunnableThread* CreateRunnableThread();
diff --git a/Engine/Source/Runtime/Core/Public/GenericPlatform/GenericPlatformProcess.h b/Engine/Source/Runtime/Core/Public/GenericPlatform/GenericPlatformProcess.h
--- a/Engine/Source/Runtime/Core/Public/GenericPlatform/GenericPlatformProcess.h
+++ b/Engine/Source/Runtime/Core/Public/GenericPlatform/GenericPlatformProcess.h
@@ -173,7 +173,14 @@ struct CORE_API FGenericPlatformProcess
 	 * @return the ProcessId of this process.
 	 */
 	static uint32 GetCurrentProcessId();
-	
+
+	/**
+	 * Retrieves the current hardware CPU core
+	 *
+	 * @return the current hardware core.
+	 */
+	static uint32 GetCurrentCoreNumber();
+
 	/**	 
 	 * Change the thread processor affinity
 	 *
diff --git a/Engine/Source/Runtime/Core/Public/Mac/MacPlatformProcess.h b/Engine/Source/Runtime/Core/Public/Mac/MacPlatformProcess.h
--- a/Engine/Source/Runtime/Core/Public/Mac/MacPlatformProcess.h
+++ b/Engine/Source/Runtime/Core/Public/Mac/MacPlatformProcess.h
@@ -101,6 +101,7 @@ struct CORE_API FMacPlatformProcess : public FGenericPlatformProcess
 	static void FreeDllHandle( void* DllHandle );
 	static void* GetDllExport( void* DllHandle, const TCHAR* ProcName );
 	static uint32 GetCurrentProcessId();
+	static uint32 GetCurrentCoreNumber();
 	static const TCHAR* BaseDir();
 	static const TCHAR* UserDir();
 	static const TCHAR* UserTempDir();
diff --git a/Engine/Source/Runtime/Core/Public/Unix/UnixPlatformProcess.h b/Engine/Source/Runtime/Core/Public/Unix/UnixPlatformProcess.h
--- a/Engine/Source/Runtime/Core/Public/Unix/UnixPlatformProcess.h
+++ b/Engine/Source/Runtime/Core/Public/Unix/UnixPlatformProcess.h
@@ -270,6 +270,7 @@ struct CORE_API FUnixPlatformProcess : public FGenericPlatformProcess
 	static void TerminateProc( FProcHandle & ProcessHandle, bool KillTree = false );
 	static EWaitAndForkResult WaitAndFork();
 	static uint32 GetCurrentProcessId();
+	static uint32 GetCurrentCoreNumber();
 	static bool GetProcReturnCode( FProcHandle & ProcHandle, int32* ReturnCode );
 	static bool Daemonize();
 	static bool IsApplicationRunning( uint32 ProcessId );
diff --git a/Engine/Source/Runtime/Core/Public/Windows/WindowsPlatformProcess.h b/Engine/Source/Runtime/Core/Public/Windows/WindowsPlatformProcess.h
--- a/Engine/Source/Runtime/Core/Public/Windows/WindowsPlatformProcess.h
+++ b/Engine/Source/Runtime/Core/Public/Windows/WindowsPlatformProcess.h
@@ -129,6 +129,7 @@ public:
 	static void PushDllDirectory(const TCHAR* Directory);
 	static void PopDllDirectory(const TCHAR* Directory);
 	static uint32 GetCurrentProcessId();
+	static uint32 GetCurrentCoreNumber();
 	static void SetThreadAffinityMask( uint64 AffinityMask );
 	static const TCHAR* BaseDir();
 	static const TCHAR* UserDir();
