Partially backing out 5830553 which resolved a major hang but is suspected of causing a new major hitch.

diff --git a/Engine/Source/Runtime/Core/Private/Async/TaskGraph.cpp b/Engine/Source/Runtime/Core/Private/Async/TaskGraph.cpp
--- a/Engine/Source/Runtime/Core/Private/Async/TaskGraph.cpp
+++ b/Engine/Source/Runtime/Core/Private/Async/TaskGraph.cpp
@@ -1931,9 +1931,7 @@ void FTaskGraphInterface::BroadcastSlow_OnlyUseForSpecialPurposes(bool bDoTaskTh
 	{
 		Tasks.Add(TGraphTask<FBroadcastTask>::CreateTask().ConstructAndDispatchWhenReady(Callback, StartTime, TEXT("AudioT"), ENamedThreads::SetTaskPriority(ENamedThreads::AudioThread, ENamedThreads::HighTaskPriority), nullptr, nullptr, nullptr));
 	}
-
-	Callback(ENamedThreads::GameThread_Local);
-
+	Tasks.Add(TGraphTask<FBroadcastTask>::CreateTask().ConstructAndDispatchWhenReady(Callback, StartTime, TEXT("GT"), ENamedThreads::GameThread_Local, nullptr, nullptr, nullptr));
 	if (bDoTaskThreads)
 	{
 		check(MyEvent);
@@ -1961,31 +1959,13 @@ void FTaskGraphInterface::BroadcastSlow_OnlyUseForSpecialPurposes(bool bDoTaskTh
 	{
 		double StartTimeInner = FPlatformTime::Seconds();
 		QUICK_SCOPE_CYCLE_COUNTER(STAT_Broadcast_WaitForNamedThreads);
-		
-		// Wait for all tasks to be complete.  Spin and pump messages to avoid deadlocks when other threads send messages and block until messages are processed
-		while (true)
+		FTaskGraphInterface::Get().WaitUntilTasksComplete(Tasks, ENamedThreads::GameThread_Local);
 		{
-			bool bAnyNotDone = false;
-			for (FGraphEventRef& Item : Tasks)
-			{
-				if (Item.GetReference() && !Item->IsComplete())
-				{
-					bAnyNotDone = true;
-					break;
-				}
-			}
-			if (!bAnyNotDone)
+			float ThisTime = FPlatformTime::Seconds() - StartTimeInner;
+			if (ThisTime > 0.02f)
 			{
-				break;
+				UE_CLOG(GPrintBroadcastWarnings, LogTaskGraph, Warning, TEXT("Task graph took %6.2fms to wait for named thread broadcast."), ThisTime * 1000.0f);
 			}
-
-			FPlatformMisc::PumpMessagesOutsideMainLoop();
-		}
-		
-		float EndTimeInner = FPlatformTime::Seconds() - StartTimeInner;
-		if (EndTimeInner > 0.02f)
-		{
-			UE_CLOG(GPrintBroadcastWarnings, LogTaskGraph, Warning, TEXT("Task graph took %6.2fms to wait for named thread broadcast."), EndTimeInner * 1000.0f);
 		}
 	}
 	for (FEvent* TaskEvent : TaskEvents)
