Fix bug with GT hitch detector getting reset every frame if startSuspended is enabled.
Remove unneeded function
Add CSV event when hitchdetector fires
Add some logging around suspend/resume

diff --git a/Engine/Source/Runtime/Core/Private/HAL/ThreadHeartBeat.cpp b/Engine/Source/Runtime/Core/Private/HAL/ThreadHeartBeat.cpp
--- a/Engine/Source/Runtime/Core/Private/HAL/ThreadHeartBeat.cpp
+++ b/Engine/Source/Runtime/Core/Private/HAL/ThreadHeartBeat.cpp
@@ -15,6 +15,7 @@
 #include "HAL/ExceptionHandling.h"
 #include "Stats/Stats.h"
 #include "Async/TaskGraphInterfaces.h"
+#include "ProfilingDebugging/CsvProfiler.h"
 
 #if PLATFORM_SWITCH
 #include "Switch/SwitchPlatformCrashContext.h"
@@ -616,6 +617,27 @@ void FGameThreadHitchHeartBeatThreaded::InitSettings()
 	{
 		bHasCmdLine = FParse::Value(FCommandLine::Get(), TEXT("hitchdetection="), CmdLine_HangDuration);
 		CmdLine_StackWalk = FParse::Param(FCommandLine::Get(), TEXT("hitchdetectionstackwalk"));
+
+		// Determine whether to start suspended
+		bool bStartSuspended = false;
+		if (GConfig)
+		{
+			GConfig->GetBool(TEXT("Core.System"), TEXT("GameThreadHeartBeatStartSuspended"), bStartSuspended, GEngineIni);
+		}
+		if (FParse::Param(FCommandLine::Get(), TEXT("hitchdetectionstartsuspended")))
+		{
+			bStartSuspended = true;
+		}
+		else if (FParse::Param(FCommandLine::Get(), TEXT("hitchdetectionstartrunning")))
+		{
+			bStartSuspended = false;
+		}
+		if (bStartSuspended)
+		{
+			UE_LOG(LogCore, Display, TEXT("Starting with HitchHeartbeat suspended"));
+			SuspendedCount = 1;
+		}
+		bFirst = false;
 	}
 
 	if (bHasCmdLine)
@@ -650,34 +672,11 @@ void FGameThreadHitchHeartBeatThreaded::InitSettings()
 		}
 	}
 	
-	// Figure out whether to start suspended
-	bool bStartSuspended = false;
-	if (GConfig)
-	{
-		GConfig->GetBool(TEXT("Core.System"), TEXT("GameThreadHeartBeatStartSuspended"), bStartSuspended, GEngineIni);
-	}
-	if (bFirst)
-	{
-		if (FParse::Param(FCommandLine::Get(), TEXT("hitchdetectionstartsuspended")))
-		{
-			bStartSuspended = true;
-		}
-		else if (FParse::Param(FCommandLine::Get(), TEXT("hitchdetectionstartrunning")))
-		{
-			bStartSuspended = false;
-		}
-	}
-	if (bStartSuspended)
-	{
-		SuspendedCount = 1;
-	}
-
 	// Start the heart beat thread if it hasn't already been started.
 	if (Thread == nullptr && FPlatformProcess::SupportsMultithreading() && HangDuration > 0)
 	{
 		Thread = FRunnableThread::Create(this, TEXT("FGameThreadHitchHeartBeatThreaded"), 0, TPri_AboveNormal);
 	}
-	bFirst = false;
 #endif
 }
 
@@ -717,7 +716,7 @@ uint32 FGameThreadHitchHeartBeatThreaded::Run()
 					{
 						GHitchDetected = true;
 						UE_LOG(LogCore, Error, TEXT("Hitch detected on gamethread (frame hasn't finished for %8.2fms):"), float(CurrentTime - LocalFrameStartTime) * 1000.0f);
-						//FPlatformMisc::LowLevelOutputDebugStringf(TEXT("Hitch detected on GameThread (frame hasn't finished for %8.2fms):"), float(CurrentTime - LocalFrameStartTime) * 1000.0f);
+						CSV_EVENT_GLOBAL(TEXT("HitchDetector"));
 
 #if WALK_STACK_ON_HITCH_DETECTED
 						if (bWalkStackOnHitch)
@@ -812,6 +811,7 @@ void FGameThreadHitchHeartBeatThreaded::SuspendHeartBeat()
 		return;
 
 	FPlatformAtomics::InterlockedIncrement(&SuspendedCount);
+	UE_LOG(LogCore, Log, TEXT("HitchHeartBeat Suspend called (count %d) - State: %s"), SuspendedCount, SuspendedCount==0 ? TEXT("Running") : TEXT("Suspended") );
 #endif
 }
 void FGameThreadHitchHeartBeatThreaded::ResumeHeartBeat()
@@ -825,20 +825,10 @@ void FGameThreadHitchHeartBeatThreaded::ResumeHeartBeat()
 	{
 		FrameStart(true);
 	}
+	UE_LOG(LogCore, Log, TEXT("HitchHeartBeat Resume called (count %d) - State: %s"), SuspendedCount, SuspendedCount == 0 ? TEXT("Running") : TEXT("Suspended"));
 #endif
 }
 
-bool FGameThreadHitchHeartBeatThreaded::IsSuspended_Gamethread() const
-{
-#if USE_HITCH_DETECTION
-	check(IsInGameThread());
-	return (SuspendedCount > 0);
-#else
-	return false;
-#endif
-}
-
-
 double FGameThreadHitchHeartBeatThreaded::GetFrameStartTime()
 {
 	return FrameStartTime;
@@ -847,4 +837,4 @@ double FGameThreadHitchHeartBeatThreaded::GetFrameStartTime()
 double FGameThreadHitchHeartBeatThreaded::GetCurrentTime()
 {
 	return Clock.Seconds();
-}
+}
\ No newline at end of file
diff --git a/Engine/Source/Runtime/Core/Public/HAL/ThreadHeartBeat.h b/Engine/Source/Runtime/Core/Public/HAL/ThreadHeartBeat.h
--- a/Engine/Source/Runtime/Core/Public/HAL/ThreadHeartBeat.h
+++ b/Engine/Source/Runtime/Core/Public/HAL/ThreadHeartBeat.h
@@ -267,11 +267,6 @@ public:
 	*/
 	void ResumeHeartBeat();
 
-	/**
-	* Is the heartbeat suspended?
-	*/
-	bool IsSuspended_Gamethread() const;
-
 	// No-op, used in FUnixSignalGameHitchHeartBeat
 	void Restart() {}
 
