Adding data to track failed download bytes in bundle manager implementations.
#tests Android Client

diff --git a/Engine/Source/Runtime/Core/Public/GenericPlatform/GenericPlatformInstallBundleManager.h b/Engine/Source/Runtime/Core/Public/GenericPlatform/GenericPlatformInstallBundleManager.h
--- a/Engine/Source/Runtime/Core/Public/GenericPlatform/GenericPlatformInstallBundleManager.h
+++ b/Engine/Source/Runtime/Core/Public/GenericPlatform/GenericPlatformInstallBundleManager.h
@@ -124,6 +124,8 @@ struct FInstallBundleDownloadProgress
 	uint64 BytesDownloadedAndWritten = 0;
 	// Num bytes needed
 	uint64 TotalBytesToDownload = 0;
+	// Num bytes that failed to download
+	uint64 TotalBytesFailedToDownload = 0;
 	float PercentComplete = 0;
 };
 
