Disable static analysis warnings for testing against constants.

diff --git a/Engine/Source/Runtime/Core/Public/Windows/WindowsPlatformCodeAnalysis.h b/Engine/Source/Runtime/Core/Public/Windows/WindowsPlatformCodeAnalysis.h
--- a/Engine/Source/Runtime/Core/Public/Windows/WindowsPlatformCodeAnalysis.h
+++ b/Engine/Source/Runtime/Core/Public/Windows/WindowsPlatformCodeAnalysis.h
@@ -84,7 +84,7 @@
 
 	// We don't use exceptions or care to gracefully handle _alloca() failure.  Also, we wrap _alloca in an
 	// appAlloca macro (not inline methods) and don't want to suppress at all call sites.
-	#pragma warning( disable : 6255 )	// warning C6255: _alloca indicates failure by raising a stack overflow exception.  Consider using _malloca instead.
+	#pragma warning(disable : 6255) // warning C6255: _alloca indicates failure by raising a stack overflow exception.  Consider using _malloca instead.
 
 	// This a very common false positive warning (but some cases are not false positives!). Disabling for now so that we
 	// can more quickly see the benefits of static analysis on new code.
@@ -94,6 +94,12 @@
 	// so we can get the benefits of analysis on cross-platform code sooner.
 	#pragma warning(disable : 6320) // warning C6320: Exception-filter expression is the constant EXCEPTION_EXECUTE_HANDLER. This might mask exceptions that were not intended to be handled.
 
+	// Branching on constants allows us to ensure code paths compile even if they are inactive (eg. if(PLATFORM_FOO){ ... }), and is good practice that should not be discouraged.
+	#pragma warning(disable : 6326) // warning C6326 : Potential comparison of a constant with another constant.
+
+	// Likewise, expressions involving constants can also be valuable to check code paths compile.
+	#pragma warning(disable : 6240) // warning C6240 : (<expression> && <non-zero constant>) always evaluates to the result of <expression>. Did you intend to use the bitwise-and operator?
+
 	//PVS-Studio settings:
 	//-V::505,542,581,601,668,677,686,688,690,703,704,711,719,720,728,730,735,751,1002,1008
 	//-V:TRYCOMPRESSION:519,547
