Added FHistogram::DumpToJsonString2, which uses a slightly easier to query string format.

diff --git a/Engine/Source/Runtime/Core/Private/ProfilingDebugging/Histogram.cpp b/Engine/Source/Runtime/Core/Private/ProfilingDebugging/Histogram.cpp
--- a/Engine/Source/Runtime/Core/Private/ProfilingDebugging/Histogram.cpp
+++ b/Engine/Source/Runtime/Core/Private/ProfilingDebugging/Histogram.cpp
@@ -212,6 +212,32 @@ FString FHistogram::DumpToJsonString() const
 	return DumpToJsonString(&DefaultConvertBinToLabel);
 }
 
+FString FHistogram::DumpToJsonString2(TFunctionRef<FString (double, double)> ConvertBinToLabel) const
+{
+	// {"BinName":{"Count":Count,"Sum":Sum}}
+	FString Result;
+	if (LIKELY(Bins.Num()))
+	{
+		Result += TEXT('[');
+		for (int BinIdx = 0, LastBinIdx = Bins.Num() - 1; BinIdx < LastBinIdx+1; ++BinIdx)
+		{
+			const FBin& Bin = Bins[BinIdx];
+			if (BinIdx != 0)
+			{
+				Result += TEXT(',');
+			}
+			Result += FString::Printf(TEXT("{\"%s\":{\"Count\":%d,\"Sum\":%.5f}}"), *ConvertBinToLabel(Bin.MinValue, Bin.UpperBound), Bin.Count, Bin.Sum);
+		}
+		Result.AppendChar(TEXT(']'));
+	}
+	return Result;
+}
+
+FString FHistogram::DumpToJsonString2() const
+{
+	return DumpToJsonString2(&DefaultConvertBinToLabel);
+}
+
 FString FHistogram::DefaultConvertBinToLabel(double MinValue, double UpperBound)
 {
 	if (UpperBound >= FLT_MAX)
diff --git a/Engine/Source/Runtime/Core/Public/ProfilingDebugging/Histogram.h b/Engine/Source/Runtime/Core/Public/ProfilingDebugging/Histogram.h
--- a/Engine/Source/Runtime/Core/Public/ProfilingDebugging/Histogram.h
+++ b/Engine/Source/Runtime/Core/Public/ProfilingDebugging/Histogram.h
@@ -68,6 +68,16 @@ struct CORE_API FHistogram
 	/** Same as DumpToJsonString but uses a DefaultConvertBinToLabel. */
 	FString DumpToJsonString() const;
 
+	/** 
+	* Returns a string in a Json format: [{"BinName":{"Count":Count,"Sum":Sum}},...]. 
+	* Bucket name is constructed by calling ConvertBinToLabel on the MinValue and UpperBound for each bin. 
+	* Convert function is used to allow the bin range, which is stored as a double, to be printed prettily.
+	*/
+	FString DumpToJsonString2(TFunctionRef<FString (double, double)> ConvertBinToLabel) const;
+
+	/** Same as DumpToJsonString2 but uses a DefaultConvertBinToLabel. */
+	FString DumpToJsonString2() const;
+
 	/** Default stringifier for bins for use with DumpToJsonString. Truncates to int and uses Plus as the suffix for the last bin. ie. [0.0, 3.75, 9.8] -> 0_3, 3_9, 9_Plus */
 	static FString DefaultConvertBinToLabel(double MinValue, double UpperBound);
 
