Add A/B testing support to the CSV profiler. A specified cvar's value is modifed every N frames (default: 7), and the stat value is written out as a stat to the CSV. Use with CsvSplit to generate multiple CSVs (one for each value)

Params:
-csvabtest="<cvar>=<value0>,<value1>..."
[-csvabtestStatFrameOffset=<N frames>]
[-csvabtestswitchduration=<N frames>]
[-csvabtestfastcvarset] - skips cvar syncs

diff --git a/Engine/Source/Runtime/Core/Private/HAL/ConsoleManager.cpp b/Engine/Source/Runtime/Core/Private/HAL/ConsoleManager.cpp
--- a/Engine/Source/Runtime/Core/Private/HAL/ConsoleManager.cpp
+++ b/Engine/Source/Runtime/Core/Private/HAL/ConsoleManager.cpp
@@ -156,6 +156,9 @@ public:
 
 	void OnChanged(EConsoleVariableFlags SetBy)
 	{
+		// SetBy can include set flags. Discard them here
+		SetBy = EConsoleVariableFlags(SetBy & ~ECVF_SetFlagMask);
+
 		// you have to specify a SetBy e.g. ECVF_SetByCommandline
 		check(((uint32)SetBy & ECVF_SetByMask) || SetBy == ECVF_Default);
 
@@ -164,7 +167,7 @@ public:
 
 		// only change on main thread
 
-		Flags = (EConsoleVariableFlags)(((uint32)Flags & ~ECVF_SetByMask) | SetBy);
+		Flags = (EConsoleVariableFlags)(((uint32)Flags & ECVF_FlagMask) | SetBy);
 
 		OnChangedCallback.ExecuteIfBound(this);
 	}
@@ -261,7 +264,7 @@ private: // -----------------------------------------
 };
 
 template <class T>
-void OnCVarChange(T& Dst, const T& Src, EConsoleVariableFlags Flags)
+void OnCVarChange(T& Dst, const T& Src, EConsoleVariableFlags Flags, EConsoleVariableFlags SetBy)
 {
 	FConsoleManager& ConsoleManager = (FConsoleManager&)IConsoleManager::Get();
 
@@ -284,7 +287,10 @@ void OnCVarChange(T& Dst, const T& Src, EConsoleVariableFlags Flags)
 		check(0);
 	}
 
-	ConsoleManager.OnCVarChanged();
+	if ((SetBy & ECVF_Set_NoSinkCall_Unsafe)== 0)
+	{
+		ConsoleManager.OnCVarChanged();
+	}
 }
 
 // T: int32, float, FString
@@ -334,7 +340,7 @@ private: // ----------------------------------------------------
 	void OnChanged(EConsoleVariableFlags SetBy)
 	{
 		// propagate from main thread to render thread
-		OnCVarChange(Data.ShadowedValue[1], Data.ShadowedValue[0], Flags);
+		OnCVarChange(Data.ShadowedValue[1], Data.ShadowedValue[0], Flags, SetBy);
 		FConsoleVariableBase::OnChanged(SetBy);
 	}
 };
@@ -470,7 +476,7 @@ private: // ----------------------------------------------------
 		if(CanChange(SetBy))
 		{
 			// propagate from main thread to render thread or to reference
-			OnCVarChange(RefValue, MainValue, Flags);
+			OnCVarChange(RefValue, MainValue, Flags, SetBy);
 			FConsoleVariableBase::OnChanged(SetBy);
 		}
 	}
@@ -547,7 +553,7 @@ private: // ----------------------------------------------------
 		if (CanChange(SetBy))
 		{
 			// propagate from main thread to render thread or to reference
-			OnCVarChange(RefValue, MainValue, Flags);
+			OnCVarChange(RefValue, MainValue, Flags, SetBy);
 			FConsoleVariableBase::OnChanged(SetBy);
 		}
 	}
diff --git a/Engine/Source/Runtime/Core/Private/Misc/ConfigCacheIni.cpp b/Engine/Source/Runtime/Core/Private/Misc/ConfigCacheIni.cpp
--- a/Engine/Source/Runtime/Core/Private/Misc/ConfigCacheIni.cpp
+++ b/Engine/Source/Runtime/Core/Private/Misc/ConfigCacheIni.cpp
@@ -3837,7 +3837,7 @@ const TCHAR* ConvertValueFromHumanFriendlyValue( const TCHAR* Value )
 CORE_API void OnSetCVarFromIniEntry(const TCHAR *IniFile, const TCHAR *Key, const TCHAR* Value, uint32 SetBy, bool bAllowCheating)
 {
 	check(IniFile && Key && Value);
-	check((SetBy & ~ECVF_SetByMask) == 0);
+	check((SetBy & ECVF_FlagMask) == 0);
 
 
 	Value = ConvertValueFromHumanFriendlyValue(Value);
diff --git a/Engine/Source/Runtime/Core/Private/ProfilingDebugging/CsvProfiler.cpp b/Engine/Source/Runtime/Core/Private/ProfilingDebugging/CsvProfiler.cpp
--- a/Engine/Source/Runtime/Core/Private/ProfilingDebugging/CsvProfiler.cpp
+++ b/Engine/Source/Runtime/Core/Private/ProfilingDebugging/CsvProfiler.cpp
@@ -152,6 +152,135 @@ bool IsContinuousWriteEnabled(bool bGameThread)
 	return CVarValue > 0;
 }
 
+#if CSV_PROFILER_ALLOW_DEBUG_FEATURES
+class FCsvABTest
+{
+public:
+	FCsvABTest() 
+		: ConsoleVariable(nullptr)
+		, StatFrameOffset(1)
+		, SwitchDuration(7)
+		, PreviousValue(FLT_MAX)
+		, bPrevCapturing(false)
+		, bFastCVarSet(false)
+	{}
+
+	void InitFromCommandline()
+	{
+		FString ABTestString;
+		if (FParse::Value(FCommandLine::Get(), TEXT("csvABTest="), ABTestString))
+		{
+			int32 FindIndex;
+			if (ABTestString.FindChar(TEXT('='), FindIndex))
+			{
+				CVarName = ABTestString.Mid(0, FindIndex);
+				FString ValueStr = ABTestString.Mid(FindIndex+1);
+				while (true)
+				{
+					int32 CommaIndex;
+					if (ValueStr.FindChar(TEXT(','), CommaIndex))
+					{
+						FString Val = ValueStr.Mid(0, CommaIndex);
+						CVarValues.Add(FCString::Atof(*Val));
+						ValueStr = ValueStr.Mid(CommaIndex + 1);
+					}
+					else
+					{
+						CVarValues.Add(FCString::Atof(*ValueStr));
+						break;
+					}
+				} 
+			}
+			else
+			{
+				CVarName = ABTestString;
+				CVarValues.Add(0.0f);
+				CVarValues.Add(1.0f);
+			}
+			ConsoleVariable = IConsoleManager::Get().FindConsoleVariable(*CVarName);
+			if (ConsoleVariable)
+			{
+				CVarStatFName = *CVarName;
+				FParse::Value(FCommandLine::Get(), TEXT("csvABTestStatFrameOffset="), StatFrameOffset);
+				FParse::Value(FCommandLine::Get(), TEXT("csvABTestSwitchDuration="), SwitchDuration);
+				bFastCVarSet = FParse::Param(FCommandLine::Get(), TEXT("csvABTestFastCVarSet"));				
+
+				UE_LOG(LogCsvProfiler, Log, TEXT("Initialized CSV Profiler A/B test with CVar: '%s' with values:"), *CVarName);
+				for (int32 i = 0; i < CVarValues.Num(); i++)
+				{
+					UE_LOG(LogCsvProfiler, Log, TEXT("  [%d] : %.2f"), i, CVarValues[i]);
+				}
+				UE_LOG(LogCsvProfiler, Log, TEXT("Stat Offset: %d frames"), StatFrameOffset);
+				UE_LOG(LogCsvProfiler, Log, TEXT("Switch Duration : %d frames"), SwitchDuration);
+				UE_LOG(LogCsvProfiler, Log, TEXT("Fast cvar set: %s"), bFastCVarSet ? TEXT("Enabled") : TEXT("Disabled"));
+			}
+			else
+			{
+				UE_LOG(LogCsvProfiler, Log, TEXT("CSV Profiler A/B test initialization failed. Can't find CVar %s"), *CVarName);
+			}
+		}
+	}
+
+	void BeginFrameUpdate(int32 FrameNumber, bool bCapturing)
+	{
+		if (ConsoleVariable == nullptr)
+		{
+			return;
+		}
+		if (bCapturing)
+		{
+			if (bPrevCapturing == false)
+			{
+				UE_LOG(LogCsvProfiler, Log, TEXT("CSV Profiler A/B test - %s original value: %s"), *CVarName, *OriginalCVarValue);
+				OriginalCVarValue = ConsoleVariable->GetString();
+			}
+			int32 ValueIndex = (FrameNumber / SwitchDuration) % CVarValues.Num();
+			int32 StatValueIndex = ((FrameNumber - StatFrameOffset) / SwitchDuration) % CVarValues.Num();
+
+			{
+				float Value = CVarValues[ValueIndex];
+				if (Value != PreviousValue)
+				{
+					EConsoleVariableFlags CVarFlags = ECVF_SetByCode;
+					if (bFastCVarSet)
+					{ 
+						CVarFlags = EConsoleVariableFlags( CVarFlags | ECVF_Set_NoSinkCall_Unsafe );
+					} 
+					ConsoleVariable->Set(*FString::Printf(TEXT("%f"), Value), CVarFlags);
+					PreviousValue = Value;
+				}
+			}
+
+			FCsvProfiler::RecordCustomStat(CVarStatFName, CSV_CATEGORY_INDEX_GLOBAL, CVarValues[StatValueIndex], ECsvCustomStatOp::Set);
+		}
+		else
+		{
+			if (bPrevCapturing == true)
+			{
+				// Restore cvar to old value
+				// TODO: Set Setby flag to the original value
+				ConsoleVariable->Set(*OriginalCVarValue);
+				UE_LOG(LogCsvProfiler, Log, TEXT("CSV Profiler A/B test - setting %s=%s"), *CVarName, *OriginalCVarValue);
+			}
+		}
+		bPrevCapturing = bCapturing;
+	}
+
+private:
+	FString CVarName;
+	TArray<float> CVarValues;
+	FName CVarStatFName;
+	IConsoleVariable* ConsoleVariable;
+	int32 StatFrameOffset;
+	int32 SwitchDuration;
+	float PreviousValue;
+	bool bPrevCapturing;
+	bool bFastCVarSet;
+	FString OriginalCVarValue;
+};
+static FCsvABTest GCsvABTest;
+
+#endif // CSV_PROFILER_ALLOW_DEBUG_FEATURES
 
 class FCsvCategoryData
 {
@@ -2247,6 +2376,8 @@ FCsvProfiler::~FCsvProfiler()
 void FCsvProfiler::BeginFrame()
 {
 	LLM_SCOPE(ELLMTag::CsvProfiler);
+	QUICK_SCOPE_CYCLE_COUNTER(STAT_FCsvProfiler_BeginFrame);
+	CSV_SCOPED_TIMING_STAT_EXCLUSIVE(CsvProfiler);
 
 	check(IsInGameThread());
 
@@ -2355,6 +2486,8 @@ void FCsvProfiler::BeginFrame()
 	{
 		CSVTest();
 	}
+
+	GCsvABTest.BeginFrameUpdate(CaptureFrameNumber, GCsvProfilerIsCapturing);
 #endif // CSV_PROFILER_ALLOW_DEBUG_FEATURES
 }
 
@@ -2363,6 +2496,7 @@ void FCsvProfiler::EndFrame()
 	LLM_SCOPE(ELLMTag::CsvProfiler);
 
 	QUICK_SCOPE_CYCLE_COUNTER(STAT_FCsvProfiler_EndFrame);
+	CSV_SCOPED_TIMING_STAT_EXCLUSIVE(CsvProfiler);
 
 	check(IsInGameThread());
 	if (GCsvProfilerIsCapturing)
@@ -2882,7 +3016,7 @@ void FCsvProfiler::Init()
 			break;
 		}
 	}
-
+	GCsvABTest.InitFromCommandline();
 #endif // CSV_PROFILER_ALLOW_DEBUG_FEATURES
 
 	// Always disable the CSV profiling thread if the platform does not support threading.
diff --git a/Engine/Source/Runtime/Core/Public/HAL/IConsoleManager.h b/Engine/Source/Runtime/Core/Public/HAL/IConsoleManager.h
--- a/Engine/Source/Runtime/Core/Public/HAL/IConsoleManager.h
+++ b/Engine/Source/Runtime/Core/Public/HAL/IConsoleManager.h
@@ -48,6 +48,9 @@ template <class T> class TConsoleVariableData;
  */
 enum EConsoleVariableFlags
 {
+	/* Mask for flags. Use this instead of ~ECVF_SetByMask */
+	ECVF_FlagMask = 0x0000ffff,
+
 	/**
 	 * Default, no flags are set, the value is set by the constructor 
 	 */
@@ -88,9 +91,17 @@ enum EConsoleVariableFlags
 
 	// ------------------------------------------------
 
+	/* Set flags */
+	ECVF_SetFlagMask =				0x00ff0000,
+
+	// Use to set a cvar without calling all cvar sinks. Much faster, but potentially unsafe. Use only if you know the particular cvar/setting does not require a sink call
+	ECVF_Set_NoSinkCall_Unsafe =	0x00010000,
+
+	// ------------------------------------------------
+
 	/* to get some history of where the last value was set by ( useful for track down why a cvar is in a specific state */
 	ECVF_SetByMask =				0xff000000,
-	
+
 	// the ECVF_SetBy are sorted in override order (weak to strong), the value is not serialized, it only affects it's override behavior when calling Set()
 
 	// lowest priority (default after console variable creation)
