CsvProfiler A/B test: Set the default stat frame offset to 0 frames. It's better to compute this value offline

diff --git a/Engine/Source/Runtime/Core/Private/ProfilingDebugging/CsvProfiler.cpp b/Engine/Source/Runtime/Core/Private/ProfilingDebugging/CsvProfiler.cpp
--- a/Engine/Source/Runtime/Core/Private/ProfilingDebugging/CsvProfiler.cpp
+++ b/Engine/Source/Runtime/Core/Private/ProfilingDebugging/CsvProfiler.cpp
@@ -158,7 +158,7 @@ class FCsvABTest
 public:
 	FCsvABTest() 
 		: ConsoleVariable(nullptr)
-		, StatFrameOffset(1)
+		, StatFrameOffset(0)
 		, SwitchDuration(7)
 		, PreviousValue(FLT_MAX)
 		, bPrevCapturing(false)
