[MemPro] don't try to pause MemPro on shutdown because MemPro::SetPause will create mempro automatically

diff --git a/Engine/Source/Runtime/Core/Private/MemPro/MemProProfiler.cpp b/Engine/Source/Runtime/Core/Private/MemPro/MemProProfiler.cpp
--- a/Engine/Source/Runtime/Core/Private/MemPro/MemProProfiler.cpp
+++ b/Engine/Source/Runtime/Core/Private/MemPro/MemProProfiler.cpp
@@ -115,7 +115,7 @@ bool FMemProProfiler::IsUsingPort( uint32 Port )
 
 
 /*
- * initialisation for MemPro. IT may have
+ * initialisation for MemPro.
  */
 
 void FMemProProfiler::PostInit()
@@ -123,7 +123,6 @@ void FMemProProfiler::PostInit()
 	//shutdown MemPro when the engine is shutting down so that the send thread terminates cleanly
 	FCoreDelegates::OnPreExit.AddLambda( []()
 	{
-		MemPro::SetPaused(true);
 		GMemProEnabled = 0;
 		MemPro::Disconnect();
 		//MemPro::Shutdown(); ...disabled for now as I was getting hangs on shutdown.
