Strip usage_other CSV stat as it's inaccurate and can be inferred.

diff --git a/Engine/Source/Runtime/Core/Private/Windows/WindowsPlatformProcess.cpp b/Engine/Source/Runtime/Core/Private/Windows/WindowsPlatformProcess.cpp
--- a/Engine/Source/Runtime/Core/Private/Windows/WindowsPlatformProcess.cpp
+++ b/Engine/Source/Runtime/Core/Private/Windows/WindowsPlatformProcess.cpp
@@ -495,7 +495,7 @@ bool FWindowsPlatformProcess::GetApplicationMemoryUsage(uint32 ProcessId, SIZE_T
 	return bSuccess;
 }
 
-bool FWindowsPlatformProcess::GetPerFrameProcessorUsage(uint32 ProcessId, float& ProcessUsageFraction, float& OtherUsageFraction, float& IdleUsageFraction)
+bool FWindowsPlatformProcess::GetPerFrameProcessorUsage(uint32 ProcessId, float& ProcessUsageFraction, float& IdleUsageFraction)
 {
 	bool bSuccess = true;
 
@@ -580,11 +580,10 @@ bool FWindowsPlatformProcess::GetPerFrameProcessorUsage(uint32 ProcessId, float&
 	{
 		ProcessUsageFraction = LastProcessTime;
 		IdleUsageFraction = LastIdleTime;
-		OtherUsageFraction = FMath::Clamp(1.f - (LastProcessTime + LastIdleTime), 0.f, 1.f);
 	}
 	else
 	{
-		ProcessUsageFraction = OtherUsageFraction = IdleUsageFraction = 0.f;
+		ProcessUsageFraction = IdleUsageFraction = 0.f;
 	}
 
 	return bSuccess;
diff --git a/Engine/Source/Runtime/Core/Public/Windows/WindowsPlatformProcess.h b/Engine/Source/Runtime/Core/Public/Windows/WindowsPlatformProcess.h
--- a/Engine/Source/Runtime/Core/Public/Windows/WindowsPlatformProcess.h
+++ b/Engine/Source/Runtime/Core/Public/Windows/WindowsPlatformProcess.h
@@ -157,7 +157,7 @@ public:
 	static void TerminateProc( FProcHandle & ProcessHandle, bool KillTree = false );
 	static bool GetProcReturnCode( FProcHandle & ProcHandle, int32* ReturnCode );
 	static bool GetApplicationMemoryUsage(uint32 ProcessId, SIZE_T* OutMemoryUsage);
-	static bool GetPerFrameProcessorUsage(uint32 ProcessId, float& ProcessUsageFraction, float& OtherUsageFraction, float& IdleUsageFraction);
+	static bool GetPerFrameProcessorUsage(uint32 ProcessId, float& ProcessUsageFraction, float& IdleUsageFraction);
 	static bool IsApplicationRunning( uint32 ProcessId );
 	static bool IsApplicationRunning( const TCHAR* ProcName );
 	static FString GetApplicationName( uint32 ProcessId );	
diff --git a/Engine/Source/Runtime/Launch/Private/LaunchEngineLoop.cpp b/Engine/Source/Runtime/Launch/Private/LaunchEngineLoop.cpp
--- a/Engine/Source/Runtime/Launch/Private/LaunchEngineLoop.cpp
+++ b/Engine/Source/Runtime/Launch/Private/LaunchEngineLoop.cpp
@@ -990,11 +990,10 @@ static void UpdateCoreCsvStats_BeginFrame()
 	if (FCsvProfiler::Get()->IsCapturing())
 	{
 		const uint32 ProcessId = (uint32)GetCurrentProcessId();
-		float ProcessUsageFraction = 0.f, OtherUsageFraction = 0.f, IdleUsageFraction = 0.f;
-		FWindowsPlatformProcess::GetPerFrameProcessorUsage(ProcessId, ProcessUsageFraction, OtherUsageFraction, IdleUsageFraction);
+		float ProcessUsageFraction = 0.f, IdleUsageFraction = 0.f;
+		FWindowsPlatformProcess::GetPerFrameProcessorUsage(ProcessId, ProcessUsageFraction, IdleUsageFraction);
 
 		CSV_CUSTOM_STAT_GLOBAL(CPUUsage_Process, ProcessUsageFraction, ECsvCustomStatOp::Set);
-		CSV_CUSTOM_STAT_GLOBAL(CPUUsage_Other, OtherUsageFraction, ECsvCustomStatOp::Set);
 		CSV_CUSTOM_STAT_GLOBAL(CPUUsage_Idle, IdleUsageFraction, ECsvCustomStatOp::Set);
 	}
 #endif
