make it so my assert functions are not inlined so they continue to be separate crashes

diff --git a/Engine/Source/Runtime/Core/Public/IOS/IOSPlatformMisc.h b/Engine/Source/Runtime/Core/Public/IOS/IOSPlatformMisc.h
--- a/Engine/Source/Runtime/Core/Public/IOS/IOSPlatformMisc.h
+++ b/Engine/Source/Runtime/Core/Public/IOS/IOSPlatformMisc.h
@@ -217,8 +217,8 @@ struct CORE_API FIOSPlatformMisc : public FApplePlatformMisc
 	static bool RequestDeviceCheckToken(TFunction<void(const TArray<uint8>&)> QuerySucceededFunc, TFunction<void(const FString&, const FString&)> QueryFailedFunc);
     
     // added these for now because Crashlytics doesn't properly break up different callstacks all ending in UE_LOG(LogXXX, Fatal, ...)
-    static CA_NO_RETURN void GPUAssert();
-    static CA_NO_RETURN void MetalAssert();
+    static FORCENOINLINE CA_NO_RETURN void GPUAssert();
+    static FORCENOINLINE CA_NO_RETURN void MetalAssert();
 };
 
 typedef FIOSPlatformMisc FPlatformMisc;
