[MemPro] fix static analysis warnings.

Note: no [REDACTED_HANDLE] tags on the modified code because the fixes are going to be rolled back into the canonical version of MemPro

diff --git a/Engine/Source/Runtime/Core/Private/MemPro/MemPro.cpp b/Engine/Source/Runtime/Core/Private/MemPro/MemPro.cpp
--- a/Engine/Source/Runtime/Core/Private/MemPro/MemPro.cpp
+++ b/Engine/Source/Runtime/Core/Private/MemPro/MemPro.cpp
@@ -2002,7 +2002,7 @@ namespace MemPro
 		MEMPRO_ASSERT(len <= max_path_len);
 
 		// round up to 4 bytes
-		static char temp[max_path_len] MEMPRO_ALIGN_SUFFIX(16);
+		static char temp[max_path_len+1] MEMPRO_ALIGN_SUFFIX(16);
 		memset(temp, 0, sizeof(temp));
 		memcpy(temp, p_str, len);
 
@@ -3187,7 +3187,7 @@ namespace MemPro
 
 			platform_thread.m_Handle = ::CreateThread(NULL, 0, PlatformThreadMain, p_os_thread_mem, 0, NULL);
 
-			return GetThreadId(platform_thread.m_Handle);
+			return (platform_thread.m_Handle == NULL) ? 0 : GetThreadId(platform_thread.m_Handle);
 		}
 
 		//------------------------------------------------------------------------
@@ -3287,10 +3287,12 @@ namespace MemPro
 				LLM_SCOPED_PAUSE_TRACKING(ELLMAllocType::System);
 			#endif
 
-			VirtualFree(p, size, MEM_RELEASE);
+			VirtualFree(p, 0, MEM_RELEASE);
 		}
 
 		//------------------------------------------------------------------------
+		#pragma warning(push)
+		#pragma warning(disable:6322) // warning C6322: Empty _except block.
 		void SetThreadName(unsigned int thread_id, const char* p_name)
 		{
 			 // see http://msdn.microsoft.com/en-us/library/xcb2z8hs.aspx
@@ -3321,6 +3323,7 @@ namespace MemPro
 			{
 			}
 		}
+		#pragma warning(pop)
 
 		//------------------------------------------------------------------------
 		void Sleep(int ms)
@@ -3436,6 +3439,7 @@ namespace MemPro
 		BOOL CALLBACK EnumerateLoadedModulesCallback(__in PCSTR ModuleName,__in ULONG ModuleBase,__in ULONG ModuleSize,__in_opt PVOID UserContext)
 		#endif
 		{
+			if(UserContext == nullptr) return false;
 			EnumModulesContext* p_context = (EnumModulesContext*)UserContext;
 
 			p_context->mp_CallbackFunction(ModuleBase, ModuleName, p_context->mp_Context);
