Fix IWYU issue
#rnx

diff --git a/Engine/Source/Runtime/Core/Public/Logging/LogMacros.h b/Engine/Source/Runtime/Core/Public/Logging/LogMacros.h
--- a/Engine/Source/Runtime/Core/Public/Logging/LogMacros.h
+++ b/Engine/Source/Runtime/Core/Public/Logging/LogMacros.h
@@ -7,6 +7,7 @@
 #include "Misc/AssertionMacros.h"
 #include "Containers/UnrealString.h"
 #include "Logging/LogCategory.h"
+#include "Logging/LogScopedCategoryAndVerbosityOverride.h"
 #include "Templates/IsValidVariadicFunctionArg.h"
 #include "Templates/AndOrNot.h"
 #include "Templates/IsArrayOrRefOfType.h"
