Add LexToString for EInstallBundleStatus

diff --git a/Engine/Source/Runtime/Core/Public/GenericPlatform/GenericPlatformInstallBundleManager.h b/Engine/Source/Runtime/Core/Public/GenericPlatform/GenericPlatformInstallBundleManager.h
--- a/Engine/Source/Runtime/Core/Public/GenericPlatform/GenericPlatformInstallBundleManager.h
+++ b/Engine/Source/Runtime/Core/Public/GenericPlatform/GenericPlatformInstallBundleManager.h
@@ -114,8 +114,28 @@ enum class EInstallBundleStatus : int
 	QueuedForFinish,
 	Finishing,
 	Installed,
+	Count,
 };
 
+inline const TCHAR* LexToString(EInstallBundleStatus Status)
+{
+	using UnderType = __underlying_type(EInstallBundleStatus);
+	static const TCHAR* Strings[] =
+	{
+		TEXT("QueuedForDownload"),
+		TEXT("Downloading"),
+		TEXT("QueuedForInstall"),
+		TEXT("Installing"),
+		TEXT("QueuedForFinish"),
+		TEXT("Finishing"),
+		TEXT("Installed"),
+	};
+
+	static_assert(static_cast<UnderType>(EInstallBundleStatus::Count) == ARRAY_COUNT(Strings), "");
+
+	return Strings[static_cast<UnderType>(Status)];
+}
+
 struct FInstallBundleDownloadProgress
 {
 	// Num bytes received
