Make sure unattended ensures are not being send on Mac when editor's Privacy > Bug Reports setting is set to "Don't Send"

diff --git a/Engine/Source/Runtime/Core/Private/Mac/MacPlatformMisc.cpp b/Engine/Source/Runtime/Core/Private/Mac/MacPlatformMisc.cpp
--- a/Engine/Source/Runtime/Core/Private/Mac/MacPlatformMisc.cpp
+++ b/Engine/Source/Runtime/Core/Private/Mac/MacPlatformMisc.cpp
@@ -1875,7 +1875,8 @@ void FMacCrashContext::GenerateCrashInfoAndLaunchReporter() const
 		bSendUnattendedBugReports = false;
 	}
 
-	if (GMacAppInfo.bIsUnattended && !bSendUnattendedBugReports)
+	const bool bUnattended = GMacAppInfo.bIsUnattended || IsRunningDedicatedServer();
+	if (bUnattended && !bSendUnattendedBugReports)
 	{
 		bCanRunCrashReportClient = false;
 	}
@@ -1959,7 +1960,8 @@ void FMacCrashContext::GenerateEnsureInfoAndLaunchReporter() const
 		bSendUnattendedBugReports = false;
 	}
 
-	if(GMacAppInfo.bIsUnattended && !bSendUnattendedBugReports)
+	const bool bUnattended = GMacAppInfo.bIsUnattended || !IsInteractiveEnsureMode() || IsRunningDedicatedServer();
+	if (bUnattended && !bSendUnattendedBugReports)
 	{
 		bCanRunCrashReportClient = false;
 	}
