Disable ApplePlatformThreadStackWalk when running under the debugger

diff --git a/Engine/Source/Runtime/Core/Private/Mac/MacPlatformMisc.cpp b/Engine/Source/Runtime/Core/Private/Mac/MacPlatformMisc.cpp
--- a/Engine/Source/Runtime/Core/Private/Mac/MacPlatformMisc.cpp
+++ b/Engine/Source/Runtime/Core/Private/Mac/MacPlatformMisc.cpp
@@ -423,6 +423,16 @@ void FMacPlatformMisc::PlatformPreInit()
 	// No SIGPIPE crashes please - they are a pain to debug!
 	signal(SIGPIPE, SIG_IGN);
 
+	// Disable ApplePlatformThreadStackWalk when the debugger is attached
+	if (FPlatformMisc::IsDebuggerPresent() && !GAlwaysReportCrash)
+	{
+		IConsoleVariable* CVarApplePlatformThreadStackWalkEnable = IConsoleManager::Get().FindConsoleVariable(TEXT("ApplePlatformThreadStackWalk.Enable"));
+		if (CVarApplePlatformThreadStackWalkEnable)
+		{
+			CVarApplePlatformThreadStackWalkEnable->Set(0);
+		}
+	}
+
 	// Increase the maximum number of simultaneously open files
 	uint32 MaxFilesPerProc = OPEN_MAX;
 	size_t UInt32Size = sizeof(uint32);
