Temporary workaround for suspend/resume issue

diff --git a/Engine/Source/Runtime/Core/Private/HAL/ThreadHeartBeat.cpp b/Engine/Source/Runtime/Core/Private/HAL/ThreadHeartBeat.cpp
--- a/Engine/Source/Runtime/Core/Private/HAL/ThreadHeartBeat.cpp
+++ b/Engine/Source/Runtime/Core/Private/HAL/ThreadHeartBeat.cpp
@@ -820,7 +820,14 @@ void FGameThreadHitchHeartBeatThreaded::ResumeHeartBeat()
 	if (!IsInGameThread())
 		return;
 
-	check(SuspendedCount > 0);
+	// Temporary workaround for suspend/resume issue
+	//check(SuspendedCount > 0);
+	if (SuspendedCount == 0)
+	{
+		UE_LOG(LogCore, Warning, TEXT("HitchHeartBeat Resume called when SuspendedCount was already 0! Ignoring"));
+		return;
+	}
+
 	if (FPlatformAtomics::InterlockedDecrement(&SuspendedCount) == 0)
 	{
 		FrameStart(true);
