Fix UE4 not compiling on Windows with WinSDK versions later than 10.0.18362.0 due to header conflicts.

diff --git a/Engine/Source/Runtime/Core/Private/Windows/WindowsPlatformMisc.cpp b/Engine/Source/Runtime/Core/Private/Windows/WindowsPlatformMisc.cpp
--- a/Engine/Source/Runtime/Core/Private/Windows/WindowsPlatformMisc.cpp
+++ b/Engine/Source/Runtime/Core/Private/Windows/WindowsPlatformMisc.cpp
@@ -44,6 +44,7 @@
 #include "Runtime/Launch/Resources/Windows/Resource.h"
 
 #include "Windows/AllowWindowsPlatformTypes.h"
+THIRD_PARTY_INCLUDES_START
 	#include <time.h>
 	#include <mmsystem.h>
 	#include <rpcsal.h>
@@ -53,6 +54,7 @@
 	#include <shellapi.h>
 	#include <IPHlpApi.h>
 	#include <VersionHelpers.h>
+THIRD_PARTY_INCLUDES_END
 #include "Windows/HideWindowsPlatformTypes.h"
 
 #include "Modules/ModuleManager.h"
diff --git a/Engine/Source/Runtime/Online/ICMP/Private/Windows/IcmpWindows.cpp b/Engine/Source/Runtime/Online/ICMP/Private/Windows/IcmpWindows.cpp
--- a/Engine/Source/Runtime/Online/ICMP/Private/Windows/IcmpWindows.cpp
+++ b/Engine/Source/Runtime/Online/ICMP/Private/Windows/IcmpWindows.cpp
@@ -7,12 +7,14 @@
 
 #include "Windows/WindowsHWrapper.h"
 #include "Windows/AllowWindowsPlatformTypes.h"
+THIRD_PARTY_INCLUDES_START
 #define _WINSOCK_DEPRECATED_NO_WARNINGS
 #include "WinSock2.h"
 #include "Ws2tcpip.h"
 #include "iphlpapi.h"
 #include "IcmpAPI.h"
 #undef _WINSOCK_DEPRECATED_NO_WARNINGS
+THIRD_PARTY_INCLUDES_END
 #include "Windows/HideWindowsPlatformTypes.h"
 
 #if PLATFORM_32BITS
diff --git a/Engine/Source/Runtime/Sockets/Private/Windows/SocketSubsystemWindows.cpp b/Engine/Source/Runtime/Sockets/Private/Windows/SocketSubsystemWindows.cpp
--- a/Engine/Source/Runtime/Sockets/Private/Windows/SocketSubsystemWindows.cpp
+++ b/Engine/Source/Runtime/Sockets/Private/Windows/SocketSubsystemWindows.cpp
@@ -6,7 +6,9 @@
 #include "BSDSockets/IPAddressBSD.h"
 
 #include "Windows/AllowWindowsPlatformTypes.h"
+THIRD_PARTY_INCLUDES_START
 #include "Iphlpapi.h"
+THIRD_PARTY_INCLUDES_END
 #include "Windows/HideWindowsPlatformTypes.h"
 
 
