Make Mac flush logs during shutdown

diff --git a/Engine/Source/Runtime/Core/Private/Mac/MacPlatformMisc.cpp b/Engine/Source/Runtime/Core/Private/Mac/MacPlatformMisc.cpp
--- a/Engine/Source/Runtime/Core/Private/Mac/MacPlatformMisc.cpp
+++ b/Engine/Source/Runtime/Core/Private/Mac/MacPlatformMisc.cpp
@@ -600,6 +600,15 @@ void FMacPlatformMisc::RequestExit( bool Force )
 	
 	if( Force )
 	{
+		// Make sure the log is flushed.
+		if (GLog)
+		{
+			// This may be called from other thread, so set this thread as the master.
+			GLog->SetCurrentThreadAsMasterThread();
+			GLog->TearDown();
+		}
+
+
 		// Exit immediately, by request.
 		_Exit(GIsCriticalError ? 3 : 0);
 	}
