csvabtest supports multiple cvar separated by semicolon

ben.woodhouse
#rnx

diff --git a/Engine/Source/Runtime/Core/Private/ProfilingDebugging/CsvProfiler.cpp b/Engine/Source/Runtime/Core/Private/ProfilingDebugging/CsvProfiler.cpp
--- a/Engine/Source/Runtime/Core/Private/ProfilingDebugging/CsvProfiler.cpp
+++ b/Engine/Source/Runtime/Core/Private/ProfilingDebugging/CsvProfiler.cpp
@@ -156,127 +156,198 @@ bool IsContinuousWriteEnabled(bool bGameThread)
 class FCsvABTest
 {
 public:
-	FCsvABTest() 
-		: ConsoleVariable(nullptr)
-		, StatFrameOffset(0)
+	FCsvABTest()
+		: StatFrameOffset(0)
 		, SwitchDuration(7)
-		, PreviousValue(FLT_MAX)
 		, bPrevCapturing(false)
 		, bFastCVarSet(false)
 	{}
 
-	void InitFromCommandline()
+	void AddCVarABData(const FString& CVarName, int32 Count)
 	{
-		FString ABTestString;
-		if (FParse::Value(FCommandLine::Get(), TEXT("csvABTest="), ABTestString))
+		Count = CVarValues.Num() - Count;
+		IConsoleVariable* ConsoleVariable = IConsoleManager::Get().FindConsoleVariable(*CVarName);
+
+		if (Count > 0 && ConsoleVariable != nullptr)
+		{
+			CVarABDataArray.Add({ CVarName, *CVarName, ConsoleVariable, ConsoleVariable->GetString(), Count, FLT_MAX });
+		}
+		else if (ConsoleVariable == nullptr)
+		{
+			UE_LOG(LogCsvProfiler, Log, TEXT("Skipping CVar %s - Not found"), *CVarName);
+		}
+		else if (Count == 0)
+		{
+			UE_LOG(LogCsvProfiler, Log, TEXT("Skipping CVar %s - No value specified"), *CVarName);
+		}
+	}
+
+	void IterateABTestArguments(const FString& ABTestString)
+	{
+		int32 FindIndex;
+		if (!ABTestString.FindChar(TEXT('='), FindIndex))
+		{
+			return;
+		}
+
+		int32 Count = CVarValues.Num();
+
+		FString CVarName = ABTestString.Mid(0, FindIndex);
+		FString ValueStr = ABTestString.Mid(FindIndex + 1);
+		while (true)
 		{
-			int32 FindIndex;
-			if (ABTestString.FindChar(TEXT('='), FindIndex))
+			int32 CommaIndex;
+			bool bComma = ValueStr.FindChar(TEXT(','), CommaIndex);
+			int32 SemiColonIndex;
+			bool bSemiColon = ValueStr.FindChar(TEXT(';'), SemiColonIndex);
+
+			if (bComma)
 			{
-				CVarName = ABTestString.Mid(0, FindIndex);
-				FString ValueStr = ABTestString.Mid(FindIndex+1);
-				while (true)
+				if (!bSemiColon || (bSemiColon && CommaIndex<SemiColonIndex))
 				{
-					int32 CommaIndex;
-					if (ValueStr.FindChar(TEXT(','), CommaIndex))
-					{
-						FString Val = ValueStr.Mid(0, CommaIndex);
-						CVarValues.Add(FCString::Atof(*Val));
-						ValueStr = ValueStr.Mid(CommaIndex + 1);
-					}
-					else
-					{
-						CVarValues.Add(FCString::Atof(*ValueStr));
-						break;
-					}
-				} 
+					FString Val = ValueStr.Mid(0, CommaIndex);
+					CVarValues.Add(FCString::Atof(*Val));
+					ValueStr = ValueStr.Mid(CommaIndex + 1);
+					continue;
+				}
 			}
-			else
+
+			if (bSemiColon)
 			{
-				CVarName = ABTestString;
-				CVarValues.Add(0.0f);
-				CVarValues.Add(1.0f);
+				if (SemiColonIndex==0)
+				{
+					AddCVarABData(CVarName, Count);
+					IterateABTestArguments(ValueStr.Mid(SemiColonIndex + 1));
+					break;
+				}
+				else
+				{
+					FString Val = ValueStr.Mid(0, SemiColonIndex);
+					CVarValues.Add(FCString::Atof(*Val));
+					ValueStr = ValueStr.Mid(SemiColonIndex);
+					continue;
+				}
 			}
-			ConsoleVariable = IConsoleManager::Get().FindConsoleVariable(*CVarName);
-			if (ConsoleVariable)
-			{
-				CVarStatFName = *CVarName;
-				FParse::Value(FCommandLine::Get(), TEXT("csvABTestStatFrameOffset="), StatFrameOffset);
-				FParse::Value(FCommandLine::Get(), TEXT("csvABTestSwitchDuration="), SwitchDuration);
-				bFastCVarSet = FParse::Param(FCommandLine::Get(), TEXT("csvABTestFastCVarSet"));				
 
-				UE_LOG(LogCsvProfiler, Log, TEXT("Initialized CSV Profiler A/B test with CVar: '%s' with values:"), *CVarName);
-				for (int32 i = 0; i < CVarValues.Num(); i++)
+			CVarValues.Add(FCString::Atof(*ValueStr));
+			AddCVarABData(CVarName, Count);
+			break;
+		}
+
+	}
+
+	void InitFromCommandline()
+	{
+		FString ABTestString;
+		if (FParse::Value(FCommandLine::Get(), TEXT("csvABTest="), ABTestString, false))
+		{
+			IterateABTestArguments(ABTestString);
+
+			if (CVarABDataArray.Num() > 0)
+			{
+				UE_LOG(LogCsvProfiler, Log, TEXT("Initialized CSV Profiler A/B test")); 
+				
+				int32 CVarValuesIndex = 0;
+				for (int32 Index = 0; Index < CVarABDataArray.Num(); ++Index)
 				{
-					UE_LOG(LogCsvProfiler, Log, TEXT("  [%d] : %.2f"), i, CVarValues[i]);
+					const FCVarABData& CVarABData = CVarABDataArray[Index];
+
+					UE_LOG(LogCsvProfiler, Log, TEXT("  CVar %s [Original value: %s] AB Test with values:"), *CVarABData.CVarName, *CVarABData.OriginalValue);
+					for (int32 i = 0; i < CVarABData.Count; ++i)
+					{
+						UE_LOG(LogCsvProfiler, Log, TEXT("    [%d] : %.2f"), i, CVarValues[CVarValuesIndex + i]);
+					}
+
+					CVarValuesIndex += CVarABData.Count;
 				}
+
+				FParse::Value(FCommandLine::Get(), TEXT("csvABTestStatFrameOffset="), StatFrameOffset);
+				FParse::Value(FCommandLine::Get(), TEXT("csvABTestSwitchDuration="), SwitchDuration);
+				bFastCVarSet = FParse::Param(FCommandLine::Get(), TEXT("csvABTestFastCVarSet"));
 				UE_LOG(LogCsvProfiler, Log, TEXT("Stat Offset: %d frames"), StatFrameOffset);
 				UE_LOG(LogCsvProfiler, Log, TEXT("Switch Duration : %d frames"), SwitchDuration);
 				UE_LOG(LogCsvProfiler, Log, TEXT("Fast cvar set: %s"), bFastCVarSet ? TEXT("Enabled") : TEXT("Disabled"));
+
 			}
 			else
 			{
-				UE_LOG(LogCsvProfiler, Log, TEXT("CSV Profiler A/B test initialization failed. Can't find CVar %s"), *CVarName);
+				UE_LOG(LogCsvProfiler, Log, TEXT("CSV Profiler A/B has not initialized"));
 			}
 		}
 	}
 
 	void BeginFrameUpdate(int32 FrameNumber, bool bCapturing)
 	{
-		if (ConsoleVariable == nullptr)
+		if (CVarABDataArray.Num() == 0)
 		{
 			return;
 		}
+		
 		if (bCapturing)
 		{
-			if (bPrevCapturing == false)
+			int32 CVarValuesIndex = 0;
+			for (int32 Index = 0; Index < CVarABDataArray.Num(); ++Index)
 			{
-				UE_LOG(LogCsvProfiler, Log, TEXT("CSV Profiler A/B test - %s original value: %s"), *CVarName, *OriginalCVarValue);
-				OriginalCVarValue = ConsoleVariable->GetString();
-			}
-			int32 ValueIndex = (FrameNumber / SwitchDuration) % CVarValues.Num();
-			int32 StatValueIndex = ((FrameNumber - StatFrameOffset) / SwitchDuration) % CVarValues.Num();
+				FCVarABData& CVarABData = CVarABDataArray[Index];
+
+				int32 ValueIndex = (FrameNumber / SwitchDuration) % CVarABData.Count;
+				int32 StatValueIndex = ((FrameNumber - StatFrameOffset) / SwitchDuration) % CVarABData.Count;
+				
+				ValueIndex += CVarValuesIndex;
+				StatValueIndex += CVarValuesIndex;
+				CVarValuesIndex += CVarABData.Count;
 
-			{
-				float Value = CVarValues[ValueIndex];
-				if (Value != PreviousValue)
 				{
-					EConsoleVariableFlags CVarFlags = ECVF_SetByCode;
-					if (bFastCVarSet)
-					{ 
-						CVarFlags = EConsoleVariableFlags( CVarFlags | ECVF_Set_NoSinkCall_Unsafe );
-					} 
-					ConsoleVariable->Set(*FString::Printf(TEXT("%f"), Value), CVarFlags);
-					PreviousValue = Value;
+					float Value = CVarValues[ValueIndex];
+					if (Value != CVarABData.PreviousValue)
+					{
+						EConsoleVariableFlags CVarFlags = ECVF_SetByCode;
+						if (bFastCVarSet)
+						{
+							CVarFlags = EConsoleVariableFlags(CVarFlags | ECVF_Set_NoSinkCall_Unsafe);
+						}
+						CVarABData.ConsoleVariable->Set(*FString::Printf(TEXT("%f"), Value), CVarFlags);
+
+						CVarABData.PreviousValue = Value;
+					}
 				}
-			}
 
-			FCsvProfiler::RecordCustomStat(CVarStatFName, CSV_CATEGORY_INDEX_GLOBAL, CVarValues[StatValueIndex], ECsvCustomStatOp::Set);
+				FCsvProfiler::RecordCustomStat(CVarABData.CVarStatFName, CSV_CATEGORY_INDEX_GLOBAL, CVarValues[StatValueIndex], ECsvCustomStatOp::Set);
+			}
 		}
-		else
+		else if (bPrevCapturing == true)
 		{
-			if (bPrevCapturing == true)
+			// Restore cvar to old value
+			// TODO: Set Setby flag to the original value
+			for (int32 Index = 0; Index < CVarABDataArray.Num(); ++Index)
 			{
-				// Restore cvar to old value
-				// TODO: Set Setby flag to the original value
-				ConsoleVariable->Set(*OriginalCVarValue);
-				UE_LOG(LogCsvProfiler, Log, TEXT("CSV Profiler A/B test - setting %s=%s"), *CVarName, *OriginalCVarValue);
+				CVarABDataArray[Index].ConsoleVariable->Set(*CVarABDataArray[Index].OriginalValue);
+
+				UE_LOG(LogCsvProfiler, Log, TEXT("CSV Profiler A/B test - setting %s=%s"), *CVarABDataArray[Index].CVarName, *CVarABDataArray[Index].OriginalValue);
 			}
+
 		}
 		bPrevCapturing = bCapturing;
 	}
 
 private:
-	FString CVarName;
+	struct FCVarABData
+	{
+		FString CVarName;
+		FName CVarStatFName;
+		IConsoleVariable* ConsoleVariable;
+		FString OriginalValue;
+		int32 Count;
+		float PreviousValue;
+	};
+	
+	TArray<FCVarABData> CVarABDataArray;
 	TArray<float> CVarValues;
-	FName CVarStatFName;
-	IConsoleVariable* ConsoleVariable;
+
 	int32 StatFrameOffset;
 	int32 SwitchDuration;
-	float PreviousValue;
 	bool bPrevCapturing;
 	bool bFastCVarSet;
-	FString OriginalCVarValue;
 };
 static FCsvABTest GCsvABTest;
 
