#UE4 - update Party / MCP code to use latest EOS SDK features
- added TVariant to core from David Parton
- added accessors to TypeCompatibleBytes from Johan Torp
- added some allocator changes to memory from Steve Robb

diff --git a/Engine/Source/Runtime/Core/Public/Containers/ContainerAllocationPolicies.h b/Engine/Source/Runtime/Core/Public/Containers/ContainerAllocationPolicies.h
--- a/Engine/Source/Runtime/Core/Public/Containers/ContainerAllocationPolicies.h
+++ b/Engine/Source/Runtime/Core/Public/Containers/ContainerAllocationPolicies.h
@@ -863,6 +863,21 @@ public:
 	typedef TInlineAllocator<InlineBitArrayDWORDs,typename SecondaryAllocator::BitArrayAllocator>	BitArrayAllocator;
 };
 
+/** An inline sparse array allocator that doesn't have any secondary storage. */
+template <uint32 NumInlineElements>
+class TFixedSparseArrayAllocator
+{
+private:
+
+	/** The size to allocate inline for the bit array. */
+	enum { InlineBitArrayDWORDs = (NumInlineElements + NumBitsPerDWORD - 1) / NumBitsPerDWORD};
+
+public:
+
+	typedef TFixedAllocator<NumInlineElements>    ElementAllocator;
+	typedef TFixedAllocator<InlineBitArrayDWORDs> BitArrayAllocator;
+};
+
 //
 // Set allocation definitions.
 //
@@ -913,7 +928,7 @@ private:
 
 	enum { NumInlineHashBuckets = (NumInlineElements + AverageNumberOfElementsPerHashBucket - 1) / AverageNumberOfElementsPerHashBucket };
 
-	static_assert(!(NumInlineHashBuckets & (NumInlineHashBuckets - 1)), "Number of inline buckets must be a power of two");
+	static_assert(NumInlineHashBuckets > 0 && !(NumInlineHashBuckets & (NumInlineHashBuckets - 1)), "Number of inline buckets must be a power of two");
 
 public:
 
@@ -938,6 +953,43 @@ public:
 	typedef TInlineAllocator<NumInlineHashBuckets,typename SecondaryAllocator::HashAllocator>                HashAllocator;
 };
 
+/** An inline set allocator that doesn't have any secondary storage. */
+template<
+	uint32 NumInlineElements,
+	uint32 AverageNumberOfElementsPerHashBucket = DEFAULT_NUMBER_OF_ELEMENTS_PER_HASH_BUCKET,
+	uint32 MinNumberOfHashedElements            = DEFAULT_MIN_NUMBER_OF_HASHED_ELEMENTS
+	>
+class TFixedSetAllocator
+{
+private:
+
+	enum { NumInlineHashBuckets = (NumInlineElements + AverageNumberOfElementsPerHashBucket - 1) / AverageNumberOfElementsPerHashBucket };
+
+	static_assert(NumInlineHashBuckets > 0 && !(NumInlineHashBuckets & (NumInlineHashBuckets - 1)), "Number of inline buckets must be a power of two");
+
+public:
+
+	/** Computes the number of hash buckets to use for a given number of elements. */
+	static FORCEINLINE uint32 GetNumberOfHashBuckets(uint32 NumHashedElements)
+	{
+		const uint32 NumDesiredHashBuckets = FPlatformMath::RoundUpToPowerOfTwo(NumHashedElements / AverageNumberOfElementsPerHashBucket);
+		if (NumDesiredHashBuckets < NumInlineHashBuckets)
+		{
+			return NumInlineHashBuckets;
+		}
+
+		if (NumHashedElements < MinNumberOfHashedElements)
+		{
+			return NumInlineHashBuckets;
+		}
+
+		return NumDesiredHashBuckets;
+	}
+
+	typedef TFixedSparseArrayAllocator<NumInlineElements> SparseArrayAllocator;
+	typedef TFixedAllocator<NumInlineHashBuckets>         HashAllocator;
+};
+
 
 /**
  * 'typedefs' for various allocator defaults.
diff --git a/Engine/Source/Runtime/Core/Public/Misc/TVariant.h b/Engine/Source/Runtime/Core/Public/Misc/TVariant.h
new file mode 100644
--- /dev/null
+++ b/Engine/Source/Runtime/Core/Public/Misc/TVariant.h
@@ -0,0 +1,429 @@
+// Copyright 1998-2019 Epic Games, Inc. All Rights Reserved.
+
+#pragma once
+
+#include "Templates/ChooseClass.h"
+#include "Templates/EnableIf.h"
+#include "Templates/IsClass.h"
+#include "Templates/IsConstructible.h"
+#include "Templates/MemoryOps.h"
+#include "Templates/RemoveReference.h"
+#include "Templates/TypeCompatibleBytes.h"
+#include "Templates/UnrealTemplate.h"
+#include "Templates/UnrealTypeTraits.h"
+
+#include "Misc/AssertionMacros.h"
+
+// Due to a bug in Visual Studio, we must use a recursive template to determine max sizeof() and alignof() of types in a template parameter pack.
+// On other compilers, we use a constexpr array and pluck out the largest
+#if defined(_MSC_VER) && !defined(__clang__)
+#define TVARIANT_STORAGE_USE_RECURSIVE_TEMPLATE 1
+#else
+#define TVARIANT_STORAGE_USE_RECURSIVE_TEMPLATE 0
+#endif
+
+namespace UE4Variant_Details
+{
+	/** Determine if all the types in a template parameter pack has duplicate types */
+	template <typename...>
+	struct TTypePackContainsDuplicates;
+
+	/** A template parameter pack containing a single type has no duplicates */
+	template <typename T>
+	struct TTypePackContainsDuplicates<T>
+	{
+		static constexpr bool Value = false;
+	};
+
+	/**
+	 * A template parameter pack containing the same type adjacently contains duplicate types.
+	 * The next structure ensures that we check all pairs of types in a template parameter pack.
+	 */
+	template <typename T, typename... Ts>
+	struct TTypePackContainsDuplicates<T, T, Ts...>
+	{
+		static constexpr bool Value = true;
+	};
+
+	/** Check all pairs of types in a template parameter pack to determine if any type is duplicated */
+	template <typename T, typename U, typename... Rest>
+	struct TTypePackContainsDuplicates<T, U, Rest...>
+	{
+		static constexpr bool Value = TTypePackContainsDuplicates<T, Rest...>::Value || TTypePackContainsDuplicates<U, Rest...>::Value;
+	};
+
+	/** Determine if any of the types in a template parameter pack are references */
+	template <typename... Ts>
+	struct TContainsReferenceType
+	{
+		static constexpr bool Value = TOr<TIsReferenceType<Ts>...>::Value;
+	};
+
+#if TVARIANT_STORAGE_USE_RECURSIVE_TEMPLATE
+	/** Determine the max alignof and sizeof of all types in a template parameter pack */
+	template <typename... Ts>
+	struct TVariantStorageTraits;
+
+	template <typename T, typename... Ts>
+	struct TVariantStorageTraits<T, Ts...>
+	{
+		static constexpr SIZE_T MaxSizeof(SIZE_T CurrentSize)
+		{
+			return TVariantStorageTraits<Ts...>::MaxSizeof(TVariantStorageTraits<T>::MaxSizeof(CurrentSize));
+		}
+
+		static constexpr SIZE_T MaxAlignof(SIZE_T CurrentSize)
+		{
+			return TVariantStorageTraits<Ts...>::MaxAlignof(TVariantStorageTraits<T>::MaxAlignof(CurrentSize));
+		}
+	};
+
+	template <typename T>
+	struct TVariantStorageTraits<T>
+	{
+		static constexpr SIZE_T MaxSizeof(SIZE_T CurrentSize)
+		{
+			return CurrentSize > sizeof(T) ? CurrentSize : sizeof(T);
+		}
+
+		static constexpr SIZE_T MaxAlignof(SIZE_T CurrentSize)
+		{
+			return CurrentSize > alignof(T) ? CurrentSize : alignof(T);
+		}
+	};
+
+	/** Expose a type that is suitable for storing any of the types in a template parameter pack */
+	template <typename T, typename... Ts>
+	struct TVariantStorage
+	{
+		static constexpr SIZE_T SizeofValue = TVariantStorageTraits<T, Ts...>::MaxSizeof(0);
+		static constexpr SIZE_T AlignofValue = TVariantStorageTraits<T, Ts...>::MaxAlignof(0);
+		static_assert(SizeofValue > 0, "MaxSizeof must be greater than 0");
+		static_assert(AlignofValue > 0, "MaxAlignof must be greater than 0");
+		using Type = TAlignedBytes<SizeofValue, AlignofValue>;
+	};
+#else
+	/** Determine the max alignof and sizeof of all types in a template parameter pack and provide a type that is compatible with those sizes */
+	template <typename... Ts>
+	struct TVariantStorage
+	{
+		static constexpr SIZE_T MaxOf(const SIZE_T Sizes[])
+		{
+			SIZE_T MaxSize = Sizes[0];
+			for (int32 Itr = 1; Itr < sizeof...(Ts); ++Itr)
+			{
+				if (Sizes[Itr] > MaxSize)
+				{
+					MaxSize = Sizes[Itr];
+				}
+			}
+			return MaxSize;
+		}
+		static constexpr SIZE_T MaxSizeof()
+		{
+			constexpr SIZE_T Sizes[] = { sizeof(Ts)... };
+			return MaxOf(Sizes);
+		}
+		static constexpr SIZE_T MaxAlignof()
+		{
+			constexpr SIZE_T Sizes[] = { alignof(Ts)... };
+			return MaxOf(Sizes);
+		}
+
+		static constexpr SIZE_T SizeofValue = MaxSizeof();
+		static constexpr SIZE_T AlignofValue = MaxAlignof();
+		static_assert(SizeofValue > 0, "MaxSizeof must be greater than 0");
+		static_assert(AlignofValue > 0, "MaxAlignof must be greater than 0");
+		using Type = TAlignedBytes<SizeofValue, AlignofValue>;
+	};
+#endif
+
+	/** Helper to lookup indices of each type in a template parameter pack */
+	template <SIZE_T N, typename LookupType, typename... Ts>
+	struct TParameterPackTypeIndexHelper
+	{
+		static constexpr SIZE_T Value = (SIZE_T)-1;
+	};
+
+	/** When the type we're looking up bubbles up to the top, we return the current index */
+	template <SIZE_T N, typename T, typename... Ts>
+	struct TParameterPackTypeIndexHelper<N, T, T, Ts...>
+	{
+		static constexpr SIZE_T Value = N;
+	};
+
+	/** When different type than the lookup is at the front of the parameter pack, we increase the index and move to the next type */
+	template <SIZE_T N, typename LookupType, typename T, typename... Ts>
+	struct TParameterPackTypeIndexHelper<N, LookupType, T, Ts...>
+	{
+		static constexpr SIZE_T Value = TParameterPackTypeIndexHelper<N + 1, LookupType, Ts...>::Value;
+	};
+
+	/** Entry-point for looking up the index of a type in a template parameter pack */
+	template <typename LookupType, typename... Ts>
+	struct TParameterPackTypeIndex
+	{
+		static constexpr SIZE_T Value = TParameterPackTypeIndexHelper<0, LookupType, Ts...>::Value;
+	};
+
+	/** An adapter for calling DestructItem */
+	template <typename T>
+	struct TDestructorCaller
+	{
+		static constexpr void Destruct(void* Storage)
+		{
+			DestructItem(static_cast<T*>(Storage));
+		}
+	};
+
+	/** Lookup a type in a template parameter pack by its index and call the destructor */
+	template <typename... Ts>
+	struct TDestructorLookup
+	{
+		/** If the index matches, call the destructor, otherwise call with the next index and type in the parameter pack*/
+		static void Destruct(SIZE_T TypeIndex, void* Value)
+		{
+			static constexpr void(*Destructors[])(void*) = { &TDestructorCaller<Ts>::Destruct... };
+			check(TypeIndex < ARRAY_COUNT(Destructors));
+			Destructors[TypeIndex](Value);
+		}
+	};
+
+	/** An adapter for calling a copy constructor of a type */
+	template <typename T>
+	struct TCopyConstructorCaller
+	{
+		/** Call the copy constructor of a type with the provided memory location and value */
+		static void Construct(void* Storage, const void* Value)
+		{
+			new(Storage) T(*static_cast<const T*>(Value));
+		}
+	};
+
+	/** A utility for calling a type's copy constructor based on an index into a template parameter pack */
+	template <typename... Ts>
+	struct TCopyConstructorLookup
+	{
+		/** Construct the type at the index in the template parameter pack with the provided memory location and value */
+		static void Construct(SIZE_T TypeIndex, void* Storage, const void* Value)
+		{
+			static constexpr void(*CopyConstructors[])(void*, const void*) = { &TCopyConstructorCaller<Ts>::Construct... };
+			check(TypeIndex < ARRAY_COUNT(CopyConstructors));
+			CopyConstructors[TypeIndex](Storage, Value);
+		}
+	};
+
+
+	/** A utility for calling a type's move constructor based on an index into a template parameter pack */
+	template <typename T>
+	struct TMoveConstructorCaller
+	{
+		/** Call the move constructor of a type with the provided memory location and value */
+		static void Construct(void* Storage, void* Value)
+		{
+			new(Storage) T(MoveTemp(*static_cast<T*>(Value)));
+		}
+	};
+
+	/** A utility for calling a type's move constructor based on an index into a template parameter pack */
+	template <typename... Ts>
+	struct TMoveConstructorLookup
+	{
+		/** Construct the type at the index in the template parameter pack with the provided memory location and value */
+		static void Construct(SIZE_T TypeIndex, void* Target, void* Source)
+		{
+			static constexpr void(*MoveConstructors[])(void*, void*) = { &TMoveConstructorCaller<Ts>::Construct... };
+			check(TypeIndex < ARRAY_COUNT(MoveConstructors));
+			MoveConstructors[TypeIndex](Target, Source);
+		}
+	};
+
+	/** Determine if the type with the provided index in the template parameter pack is the same */
+	template <typename LookupType, typename... Ts>
+	struct TIsType
+	{
+		/** Check if the type at the provided index is the lookup type */
+		static bool IsSame(SIZE_T TypeIndex)
+		{
+			static constexpr bool bIsSameType[] = { TIsSame<Ts, LookupType>::Value... };
+			check(TypeIndex < ARRAY_COUNT(bIsSameType));
+			return bIsSameType[TypeIndex];
+		}
+	};
+}
+
+#undef TVARIANT_STORAGE_USE_RECURSIVE_TEMPLATE
+
+/**
+ * A special tag used to indicate that in-place construction of a variant should take place.
+ */
+template <typename T>
+struct TInPlaceType {};
+
+/**
+ * A special tag that can be used as the first type in a TVariant parameter pack if none of the other types can be default-constructed.
+ */
+struct FEmptyVariantState {};
+
+/**
+ * A type-safe union based loosely on std::variant. This flavor of variant requires that all the types in the declaring template parameter pack be unique.
+ * Attempting to use the value of a Get() when the underlying type is different leads to undefined behavior.
+ */
+template <typename T, typename... Ts>
+class TVariant
+{
+	static_assert(!UE4Variant_Details::TTypePackContainsDuplicates<T, Ts...>::Value, "All the types used in TVariant should be unique");
+	static_assert(!UE4Variant_Details::TContainsReferenceType<T, Ts...>::Value, "TVariant cannot hold reference types");
+
+public:
+	/** Default initialize the TVariant to the first type in the parameter pack */
+	TVariant()
+	{
+		static_assert(TIsConstructible<T>::Value, "To default-initialize a TVariant, the first type in the parameter pack must be default constructible. Use FEmptyVariantState as the first type if none of the other types can be listed first.");
+		new(&Storage) T();
+		TypeIndex = 0;
+	}
+
+	/** Perform in-place construction of a type into the variant */
+	template <typename U, typename... TArgs>
+	explicit TVariant(TInPlaceType<U>&&, TArgs&&... Args)
+	{
+		constexpr SIZE_T Index = UE4Variant_Details::TParameterPackTypeIndex<U, T, Ts...>::Value;
+		static_assert(Index != (SIZE_T)-1, "The TVariant is not declared to hold the type being constructed");
+		
+		new(&Storage) U(Forward<TArgs>(Args)...);
+		TypeIndex = Index;
+	}
+
+	/** Copy construct the variant from another variant of the same type */
+	TVariant(const TVariant& Other)
+		: TypeIndex(Other.TypeIndex)
+	{
+		UE4Variant_Details::TCopyConstructorLookup<T, Ts...>::Construct(TypeIndex, &Storage, &Other.Storage);
+	}
+
+	/** Move construct the variant from another variant of the same type */
+	TVariant(TVariant&& Other)
+		: TypeIndex(Other.TypeIndex)
+	{
+		UE4Variant_Details::TMoveConstructorLookup<T, Ts...>::Construct(TypeIndex, &Storage, &Other.Storage);
+	}
+
+	/** Copy assign a variant from another variant of the same type */
+	TVariant& operator=(const TVariant& Other)
+	{
+		if (&Other != this)
+		{
+			TVariant Temp = Other;
+			Swap(Temp, *this);
+		}
+		return *this;
+	}
+
+	/** Move assign a variant from another variant of the same type */
+	TVariant& operator=(TVariant&& Other)
+	{
+		if (&Other != this)
+		{
+			TVariant Temp = MoveTemp(Other);
+			Swap(Temp, *this);
+		}
+		return *this;
+	}
+
+	/** Destruct the underlying type (if appropriate) */
+	~TVariant()
+	{
+		UE4Variant_Details::TDestructorLookup<T, Ts...>::Destruct(TypeIndex, &Storage);
+	}
+
+	/** Determine if the variant holds the specific type */
+	template <typename U>
+	bool IsType() const
+	{
+		return UE4Variant_Details::TIsType<U, T, Ts...>::IsSame(TypeIndex);
+	}
+
+	/** Get a reference to the held value. Bad things can happen if this is called on a variant that does not hold the type asked for */
+	template <typename U>
+	U& Get()
+	{
+		constexpr SIZE_T Index = UE4Variant_Details::TParameterPackTypeIndex<U, T, Ts...>::Value;
+		static_assert(Index != (SIZE_T)-1, "The TVariant is not declared to hold the type passed to Get<>");
+
+		check(Index == TypeIndex);
+		return *reinterpret_cast<U*>(&Storage);
+	}
+
+	/** Get a reference to the held value. Bad things can happen if this is called on a variant that does not hold the type asked for */
+	template <typename U>
+	const U& Get() const
+	{
+		// Temporarily remove the const qualifier so we can implement Get in one location.
+		return const_cast<TVariant*>(this)->template Get<U>();
+	}
+
+	/** Get a pointer to the held value if the held type is the same as the one specified */
+	template <typename U>
+	U* TryGet()
+	{
+		constexpr SIZE_T Index = UE4Variant_Details::TParameterPackTypeIndex<U, T, Ts...>::Value;
+		static_assert(Index != (SIZE_T)-1, "The TVariant is not declared to hold the type passed to TryGet<>");
+		return Index == TypeIndex ? reinterpret_cast<U*>(&Storage) : nullptr;
+	}
+
+	/** Get a pointer to the held value if the held type is the same as the one specified */
+	template <typename U>
+	const U* TryGet() const
+	{
+		// Temporarily add the const qualifier so we can implement TrGet in one location.
+		return const_cast<TVariant*>(this)->template TryGet<U>();
+	}
+
+	/** Set a specifically-typed value into the variant */
+	template <typename U>
+	void Set(typename TIdentity<U>::Type&& Value)
+	{
+		Emplace<U>(MoveTemp(Value));
+	}
+
+	/** Set a specifically-typed value into the variant */
+	template <typename U>
+	void Set(const typename TIdentity<U>::Type& Value)
+	{
+		Emplace<U>(Value);
+	}
+
+	/** Set a specifically-typed value into the variant using in-place construction */
+	template <typename U, typename... TArgs>
+	void Emplace(TArgs&&... Args)
+	{
+		constexpr SIZE_T Index = UE4Variant_Details::TParameterPackTypeIndex<U, T, Ts...>::Value;
+		static_assert(Index != (SIZE_T)-1, "The TVariant is not declared to hold the type passed to Emplace<>");
+		
+		UE4Variant_Details::TDestructorLookup<T, Ts...>::Destruct(TypeIndex, &Storage);
+		new(&Storage) U(Forward<TArgs>(Args)...);
+		TypeIndex = Index;
+	}
+
+	/** Lookup the index of a type in the template parameter pack at compile time. */
+	template <typename U>
+	static constexpr SIZE_T IndexOfType() 
+	{
+		constexpr SIZE_T Index = UE4Variant_Details::TParameterPackTypeIndex<U, T, Ts...>::Value;
+		static_assert(Index != (SIZE_T)-1, "The TVariant is not declared to hold the type passed to IndexOfType<>");
+		return Index;
+	}
+
+	/** Returns the currently held type's index into the template parameter pack */
+	SIZE_T GetIndex() const
+	{
+		return TypeIndex;
+	}
+
+private:
+	/** Memory location for the value */
+	typename UE4Variant_Details::TVariantStorage<T, Ts...>::Type Storage;
+	/** Index into the template parameter pack for the type held. */
+	SIZE_T TypeIndex;
+};
diff --git a/Engine/Source/Runtime/Core/Public/Templates/TypeCompatibleBytes.h b/Engine/Source/Runtime/Core/Public/Templates/TypeCompatibleBytes.h
--- a/Engine/Source/Runtime/Core/Public/Templates/TypeCompatibleBytes.h
+++ b/Engine/Source/Runtime/Core/Public/Templates/TypeCompatibleBytes.h
@@ -59,4 +59,7 @@ struct TTypeCompatibleBytes :
 		sizeof(ElementType),
 		alignof(ElementType)
 		>
-{};
+{
+	ElementType*		GetTypedPtr()		{ return (ElementType*)this;  }
+	const ElementType*	GetTypedPtr() const	{ return (const ElementType*)this; }
+};
