[LLM] adding support for per-project LLM tags + LLM_DUMP_TAG returns the tag, for extra debugging logic as required

diff --git a/Engine/Source/Runtime/Core/Private/HAL/LowLevelMemTracker.cpp b/Engine/Source/Runtime/Core/Private/HAL/LowLevelMemTracker.cpp
--- a/Engine/Source/Runtime/Core/Private/HAL/LowLevelMemTracker.cpp
+++ b/Engine/Source/Runtime/Core/Private/HAL/LowLevelMemTracker.cpp
@@ -25,6 +25,7 @@ TAutoConsoleVariable<int32> CVarLLMWriteInterval(
 DECLARE_LLM_MEMORY_STAT(TEXT("LLM Overhead"), STAT_LLMOverheadTotal, STATGROUP_LLMOverhead);
 
 DEFINE_STAT(STAT_EngineSummaryLLM);
+DEFINE_STAT(STAT_ProjectSummaryLLM);
 
 /*
  * LLM stats referenced by ELLMTagNames
@@ -117,13 +118,6 @@ DECLARE_LLM_MEMORY_STAT(TEXT("Particles"), STAT_ParticlesSummaryLLM, STATGROUP_L
 DECLARE_LLM_MEMORY_STAT(TEXT("UI"), STAT_UISummaryLLM, STATGROUP_LLM);
 DECLARE_LLM_MEMORY_STAT(TEXT("Textures"), STAT_TexturesSummaryLLM, STATGROUP_LLM);
 
-struct FLLMTagInfo
-{
-	const TCHAR* Name;
-	FName StatName;				// shows in the LLMFULL stat group
-	FName SummaryStatName;		// shows in the LLM stat group
-};
-
 extern const TCHAR* LLMGetTagName(ELLMTag Tag)
 {
 #define LLM_TAG_NAME_ARRAY(Enum,Str,Stat,Group,ParentTag) TEXT(Str),
@@ -237,16 +231,16 @@ public:
 	void SetStat(int64 Tag, int64 Value);
 #endif
 
-	void Update(FLLMPlatformTag* PlatformTags, const int32* ParentTags);
+	void Update(FLLMCustomTag* CustomTags, const int32* ParentTags);
 
 	void SetEnabled(bool value) { Enabled = value; }
 
 private:
-	void WriteGraph(FLLMPlatformTag* PlatformTags, const int32* ParentTags);
+	void WriteGraph(FLLMCustomTag* CustomTags, const int32* ParentTags);
 
 	void Write(const FString& Text);
 
-	static FString GetTagName(int64 Tag, FLLMPlatformTag* PlatformTags, const int32* ParentTags);
+	static FString GetTagName(int64 Tag, FLLMCustomTag* CustomTags, const int32* ParentTags);
 
 	static const TCHAR* GetTrackerCsvName(ELLMTracker InTracker);
 
@@ -311,7 +305,7 @@ public:
 
 	void SetCSVEnabled(bool Value);
 
-	void WriteCsv(FLLMPlatformTag* PlatformTags, const int32* ParentTags);
+	void WriteCsv(FLLMCustomTag* CustomTags, const int32* ParentTags);
 
 #define LLM_USE_ALLOC_INFO_STRUCT (LLM_STAT_TAGS_ENABLED || LLM_ALLOW_ASSETS_TAGS)
 
@@ -336,7 +330,7 @@ public:
 	}
 
 	void SetTotalTags(ELLMTag Untagged, ELLMTag Tracked);
-	void Update(FLLMPlatformTag* PlatformTags, const int32* ParentTags);
+	void Update(FLLMCustomTag* CustomTags, const int32* ParentTags);
 	void UpdateTotals();
 
 	int64 GetTagAmount(ELLMTag Tag) const;
@@ -379,11 +373,11 @@ protected:
 			ELLMTag InUntaggedTotalTag,
 			FLLMThreadState& InStateCopy,
 			FLLMCsvWriter& CsvWriter,
-			FLLMPlatformTag* PlatformTags,
+			FLLMCustomTag* CustomTags,
 			int64* EnumTagAmounts,
 			int64* OutAllocTypeAmounts);
 
-		void UpdateFrameStatGroups( FLLMPlatformTag* PlatformTags, const int32* ParentTags );
+		void UpdateFrameStatGroups( FLLMCustomTag* CustomTags, const int32* ParentTags );
 
 		static void IncMemoryStatByFName(FName Name, int64 Amount);
 
@@ -435,7 +429,7 @@ protected:
 
 	double LastTrimTime;
 
-	int64 EnumTagAmounts[(int32)ELLMTag::PlatformTagEnd];
+	int64 EnumTagAmounts[LLM_TAG_COUNT];
 
 	int64 AllocTypeAmounts[(int32)ELLMAllocType::Count];
 };
@@ -451,7 +445,7 @@ static int64 FNameToTag(FName Name)
 	int64 NameIndex = Name.GetComparisonIndex();
 	int64 NameNumber = Name.GetNumber();
 	int64 tag = (NameNumber << 32) | NameIndex;
-	LLMCheckf(tag > (int64)ELLMTag::PlatformTagEnd, TEXT("Passed with a name index [%d - %s] that was less than MemTracker_MaxUserAllocation"), NameIndex, *Name.ToString());
+	LLMCheckf(tag > LLM_TAG_COUNT, TEXT("Passed with a name index [%d - %s] that was less than MemTracker_MaxUserAllocation"), NameIndex, *Name.ToString());
 
 	// convert it to a tag, but you can actually convert this to an FMinimalName in the debugger to view it - *((FMinimalName*)&Tag)
 	return tag;
@@ -504,16 +498,16 @@ FLowLevelMemTracker::FLowLevelMemTracker()
 		ActiveSets[Index] = Index == (int32)ELLMTagSet::None;
 	}
 
-	for (int32 Index = 0; Index < sizeof(PlatformTags)/sizeof(FLLMPlatformTag); Index++ )
+	for (int32 Index = 0; Index < LLM_CUSTOM_TAG_COUNT; Index++ )
 	{
-		PlatformTags[Index].Name = nullptr;
+		CustomTags[Index].Name = nullptr;
 	}
 
-	for (int32 Index = 0; Index < (int32)ELLMTag::PlatformTagEnd; Index++ )
+	for (int32 Index = 0; Index < LLM_TAG_COUNT; Index++ )
 	{
 		ParentTags[Index] = LLMGetTagParent((ELLMTag)Index);
 	}
-	for (int32 Index = 0; Index < (int32)ELLMTag::PlatformTagEnd; Index++ )
+	for (int32 Index = 0; Index < LLM_TAG_COUNT; Index++ )
 	{
 		if (ParentTags[Index] != -1)
 		{
@@ -582,7 +576,7 @@ void FLowLevelMemTracker::UpdateStatsPerFrame(const TCHAR* LogName)
 	// update the trackers
 	for (int32 TrackerIndex = 0; TrackerIndex < (int32)ELLMTracker::Max; TrackerIndex++)
 	{
-		GetTracker((ELLMTracker)TrackerIndex)->Update(PlatformTags,ParentTags);
+		GetTracker((ELLMTracker)TrackerIndex)->Update(CustomTags,ParentTags);
 	}
 
 	// calculate FMalloc unused stat and set it in the Default tracker
@@ -632,8 +626,8 @@ void FLowLevelMemTracker::UpdateStatsPerFrame(const TCHAR* LogName)
 
 	if (bCsvWriterEnabled)
 	{
-		GetTracker(ELLMTracker::Default)->WriteCsv(PlatformTags,ParentTags);
-		GetTracker(ELLMTracker::Platform)->WriteCsv(PlatformTags,ParentTags);
+		GetTracker(ELLMTracker::Default)->WriteCsv(CustomTags,ParentTags);
+		GetTracker(ELLMTracker::Platform)->WriteCsv(CustomTags,ParentTags);
 	}
 
 	if (LogName != nullptr)
@@ -862,10 +856,10 @@ static bool IsAssetTagForAssets(ELLMTagSet Set)
 	return Set == ELLMTagSet::Assets || Set == ELLMTagSet::AssetClasses;
 }
 
-void FLowLevelMemTracker::RegisterPlatformTag(int32 Tag, const TCHAR* Name, FName StatName, FName SummaryStatName, int32 ParentTag)
+void FLowLevelMemTracker::RegisterCustomTagInternal(int32 Tag, const TCHAR* Name, FName StatName, FName SummaryStatName, int32 ParentTag)
 {
-	LLMCheck(Tag >= (int32)ELLMTag::PlatformTagStart && Tag <= (int32)ELLMTag::PlatformTagEnd);
-	FLLMPlatformTag& PlatformTag = PlatformTags[Tag - (int32)ELLMTag::PlatformTagStart];
+	LLMCheck(Tag >= LLM_CUSTOM_TAG_START && Tag <= LLM_CUSTOM_TAG_END);
+	FLLMCustomTag& PlatformTag = CustomTags[Tag - LLM_CUSTOM_TAG_START];
 	PlatformTag.Tag = Tag;
 	PlatformTag.Name = Name;
 	PlatformTag.StatName = StatName;
@@ -878,6 +872,19 @@ void FLowLevelMemTracker::RegisterPlatformTag(int32 Tag, const TCHAR* Name, FNam
 	}
 }
 
+void FLowLevelMemTracker::RegisterPlatformTag(int32 Tag, const TCHAR* Name, FName StatName, FName SummaryStatName, int32 ParentTag)
+{
+	LLMCheck(Tag >= (int32)ELLMTag::PlatformTagStart && Tag <= (int32)ELLMTag::PlatformTagEnd);
+	RegisterCustomTagInternal( Tag, Name, StatName, SummaryStatName, ParentTag );
+}
+
+void FLowLevelMemTracker::RegisterProjectTag(int32 Tag, const TCHAR* Name, FName StatName, FName SummaryStatName, int32 ParentTag)
+{
+	LLMCheck(Tag >= (int32)ELLMTag::ProjectTagStart && Tag <= (int32)ELLMTag::ProjectTagEnd);
+	RegisterCustomTagInternal( Tag, Name, StatName, SummaryStatName, ParentTag );
+}
+
+
 bool FLowLevelMemTracker::FindTagByName( const TCHAR* Name, uint64& OutTag ) const
 {
 	if( Name != nullptr )
@@ -890,11 +897,11 @@ bool FLowLevelMemTracker::FindTagByName( const TCHAR* Name, uint64& OutTag ) con
 				return true;
 			}
 		}
-		for ( int32 PlatformTagIndex = 0; PlatformTagIndex <= (int32)ELLMTag::PlatformTagEnd - (int32)ELLMTag::PlatformTagStart; PlatformTagIndex++ )
+		for ( int32 PlatformTagIndex = 0; PlatformTagIndex < LLM_CUSTOM_TAG_COUNT; PlatformTagIndex++ )
 		{
-			if( PlatformTags[PlatformTagIndex].Name != nullptr && FCString::Stricmp( Name, PlatformTags[PlatformTagIndex].Name ) == 0 )
+			if( CustomTags[PlatformTagIndex].Name != nullptr && FCString::Stricmp( Name, CustomTags[PlatformTagIndex].Name ) == 0 )
 			{
-				OutTag = (int32)ELLMTag::PlatformTagStart + PlatformTagIndex;
+				OutTag = LLM_CUSTOM_TAG_START + PlatformTagIndex;
 				return true;
 			}
 		}
@@ -911,9 +918,9 @@ const TCHAR* FLowLevelMemTracker::FindTagName(uint64 Tag) const
 	{
 		Result = LLMGetTagName((ELLMTag)Tag);
 	}
-	else if( Tag >= (int32)ELLMTag::PlatformTagStart && Tag <= (int32)ELLMTag::PlatformTagEnd )
+	else if( Tag >= LLM_CUSTOM_TAG_START && Tag <= LLM_CUSTOM_TAG_END )
 	{
-		Result = PlatformTags[Tag - (int32)ELLMTag::PlatformTagStart].Name;
+		Result = CustomTags[Tag - LLM_CUSTOM_TAG_START].Name;
 	}
 
 	return Result;
@@ -934,12 +941,14 @@ int64 FLowLevelMemTracker::GetActiveTag(ELLMTracker Tracker)
     return GetTracker(Tracker)->GetActiveTag();
 }
 
-void FLowLevelMemTracker::DumpTag( ELLMTracker Tracker, const char* FileName, int LineNumber )
+int64 FLowLevelMemTracker::DumpTag( ELLMTracker Tracker, const char* FileName, int LineNumber )
 {
 	int64 Tag = GetActiveTag(Tracker);
 	const TCHAR* TagName = FindTagName(Tag);
 
 	FPlatformMisc::LowLevelOutputDebugStringf( TEXT("LLM TAG: %s (%lld) @ %s:%d\n"), TagName ? TagName : TEXT("<unknown>"), Tag, FileName ? ANSI_TO_TCHAR(FileName) : TEXT("?"), LineNumber );
+
+	return Tag;
 }
 
 
@@ -1125,7 +1134,7 @@ FLLMTracker::FLLMTracker()
 {
 	TlsSlot = FPlatformTLS::AllocTlsSlot();
 
-	for (int32 Index = 0; Index < (int32)ELLMTag::PlatformTagEnd; ++Index)
+	for (int32 Index = 0; Index < LLM_TAG_COUNT; ++Index)
 	{
 		EnumTagAmounts[Index] = 0;
 	}
@@ -1250,7 +1259,7 @@ void FLLMTracker::TrackAllocation(const void* Ptr, uint64 Size, ELLMTag DefaultT
 		AllocInfo.AssetTag = State->GetTopAssetTag();
 			#endif
 		#else
-		LLMCheck(tag >= 0 && tag <= (int64)ELLMTag::PlatformTagEnd);
+		LLMCheck(tag >= 0 && tag < (int64)LLM_TAG_COUNT);
 		AllocInfo = (ELLMTag)tag;
 		#endif
 
@@ -1362,7 +1371,7 @@ void FLLMTracker::SetTotalTags(ELLMTag InUntaggedTotalTag, ELLMTag InTrackedTota
 	TrackedTotalTag = InTrackedTotalTag;
 }
 
-void FLLMTracker::Update(FLLMPlatformTag* PlatformTags, const int32* ParentTags)
+void FLLMTracker::Update(FLLMCustomTag* CustomTags, const int32* ParentTags)
 {
 	// protect any accesses to the ThreadStates array
 	FScopeLock Lock(&ThreadArraySection);
@@ -1371,8 +1380,8 @@ void FLLMTracker::Update(FLLMPlatformTag* PlatformTags, const int32* ParentTags)
 	int ThreadStateNum = ThreadStates.Num();
 	for (int32 ThreadIndex = 0; ThreadIndex < ThreadStateNum; ThreadIndex++)
 	{
-		ThreadStates[ThreadIndex]->UpdateFrameStatGroups(PlatformTags,ParentTags);
-		ThreadStates[ThreadIndex]->GetFrameStatTotals(UntaggedTotalTag, StateCopy, CsvWriter, PlatformTags, EnumTagAmounts, AllocTypeAmounts);
+		ThreadStates[ThreadIndex]->UpdateFrameStatGroups(CustomTags,ParentTags);
+		ThreadStates[ThreadIndex]->GetFrameStatTotals(UntaggedTotalTag, StateCopy, CsvWriter, CustomTags, EnumTagAmounts, AllocTypeAmounts);
 	}
 
 	EnumTagAmounts[(int32)TrackedTotalTag] = TrackedMemoryOverFrames;
@@ -1407,9 +1416,9 @@ void FLLMTracker::UpdateTotals()
 	}
 }
 
-void FLLMTracker::WriteCsv(FLLMPlatformTag* PlatformTags, const int32* ParentTags)
+void FLLMTracker::WriteCsv(FLLMCustomTag* CustomTags, const int32* ParentTags)
 {
-	CsvWriter.Update(PlatformTags,ParentTags);
+	CsvWriter.Update(CustomTags,ParentTags);
 }
 
 int64 FLLMTracker::GetActiveTag()
@@ -1658,7 +1667,7 @@ void FLLMTracker::FLLMThreadState::GetFrameStatTotals(
 	ELLMTag InUntaggedTotalTag,
 	FLLMThreadState& InStateCopy,
 	FLLMCsvWriter& InCsvWriter,
-	FLLMPlatformTag* PlatformTags,
+	FLLMCustomTag* CustomTags,
 	int64* OutEnumTagAmounts,
 	int64* OutAllocTypeAmounts)
 {
@@ -1720,14 +1729,14 @@ void FLLMTracker::FLLMThreadState::GetFrameStatTotals(
 
 		//---------------------
 		// update the stats
-		if (Tag > (int64)ELLMTag::PlatformTagEnd)
+		if (Tag >= (int64)LLM_TAG_COUNT)
 		{
 			IncMemoryStatByFName(TagToFName(Tag), Amount);
 		}
-		else if (Tag >= (int64)ELLMTag::PlatformTagStart)
+		else if (Tag >= LLM_CUSTOM_TAG_START)
 		{
-			IncMemoryStatByFName(PlatformTags[Tag - (int32)ELLMTag::PlatformTagStart].StatName, int64(Amount));
-			IncMemoryStatByFName(PlatformTags[Tag - (int32)ELLMTag::PlatformTagStart].SummaryStatName, int64(Amount));
+			IncMemoryStatByFName(CustomTags[Tag - LLM_CUSTOM_TAG_START].StatName, int64(Amount));
+			IncMemoryStatByFName(CustomTags[Tag - LLM_CUSTOM_TAG_START].SummaryStatName, int64(Amount));
 		}
 		else
 		{
@@ -1741,7 +1750,7 @@ void FLLMTracker::FLLMThreadState::GetFrameStatTotals(
 	InStateCopy.Clear();
 }
 
-void FLLMTracker::FLLMThreadState::UpdateFrameStatGroups( FLLMPlatformTag* PlatformTags, const int32* ParentTags )
+void FLLMTracker::FLLMThreadState::UpdateFrameStatGroups( FLLMCustomTag* CustomTags, const int32* ParentTags )
 {
 	FScopeLock Lock(&TagSection);
 
@@ -1752,7 +1761,7 @@ void FLLMTracker::FLLMThreadState::UpdateFrameStatGroups( FLLMPlatformTag* Platf
 		if( Amount != 0 )
 		{
 			int64 Tag = TaggedAllocTags[TagIndex];
-			if( Tag >= 0 && Tag < (int32)ELLMTag::PlatformTagEnd)
+			if( Tag >= 0 && Tag < LLM_TAG_COUNT)
 			{
 				int32 ParentTag = ParentTags[Tag];
 				if( ParentTag != -1 )
@@ -1890,12 +1899,12 @@ void FLLMCsvWriter::SetStat(int64 Tag, int64 Value)
 /*
 * memory can be allocated in this function
 */
-void FLLMCsvWriter::Update(FLLMPlatformTag* PlatformTags, const int32* ParentTags)
+void FLLMCsvWriter::Update(FLLMCustomTag* CustomTags, const int32* ParentTags)
 {
 	double Now = FPlatformTime::Seconds();
 	if (Now - LastWriteTime >= (double)CVarLLMWriteInterval.GetValueOnGameThread())
 	{
-		WriteGraph(PlatformTags, ParentTags);
+		WriteGraph(CustomTags, ParentTags);
 
 		LastWriteTime = Now;
 	}
@@ -1923,7 +1932,7 @@ void FLLMCsvWriter::Write(const FString& Text)
  * create the csv file on the first call. When it finds a new stat name it seeks
  * back to the start of the file and re-writes the column names.
 */
-void FLLMCsvWriter::WriteGraph(FLLMPlatformTag* PlatformTags, const int32* ParentTags)
+void FLLMCsvWriter::WriteGraph(FLLMCustomTag* CustomTags, const int32* ParentTags)
 {
 	// create the csv file
 	if (!Archive)
@@ -1960,7 +1969,7 @@ void FLLMCsvWriter::WriteGraph(FLLMPlatformTag* PlatformTags, const int32* Paren
 
 		for (int32 i = 0; i < StatValueCountLocal; ++i)
 		{
-			FString StatName = GetTagName(StatValuesForWrite[i].Tag, PlatformTags, ParentTags);
+			FString StatName = GetTagName(StatValuesForWrite[i].Tag, CustomTags, ParentTags);
 			FString Text = FString::Printf(TEXT("%s,"), *StatName);
 			Write(Text);
 		}
@@ -1995,11 +2004,11 @@ void FLLMCsvWriter::WriteGraph(FLLMPlatformTag* PlatformTags, const int32* Paren
 /*
  * convert a Tag to a string. If the Tag is actually a Stat then extract the name of the stat.
 */
-FString FLLMCsvWriter::GetTagName(int64 Tag, FLLMPlatformTag* PlatformTags, const int32* ParentTags)
+FString FLLMCsvWriter::GetTagName(int64 Tag, FLLMCustomTag* CustomTags, const int32* ParentTags)
 {
 	FString Result;
 
-	if (Tag > (int64)ELLMTag::PlatformTagEnd)
+	if (Tag >= (int64)LLM_TAG_COUNT)
 	{
 		FString Name = TagToFName(Tag).ToString();
 
@@ -2017,14 +2026,14 @@ FString FLLMCsvWriter::GetTagName(int64 Tag, FLLMPlatformTag* PlatformTags, cons
 
 		Result = Name;
 	}
-	else if (Tag >= (int32)ELLMTag::PlatformTagStart && Tag <= (int32)ELLMTag::PlatformTagEnd)
+	else if (Tag >= LLM_CUSTOM_TAG_START && Tag <= LLM_CUSTOM_TAG_END)
 	{
 		if (ParentTags != nullptr && ParentTags[Tag] != -1)
 		{
-			Result = GetTagName( ParentTags[Tag], PlatformTags, nullptr ) + TEXT("/");
+			Result = GetTagName( ParentTags[Tag], CustomTags, nullptr ) + TEXT("/");
 		}
 
-		Result += PlatformTags[Tag - (int32)ELLMTag::PlatformTagStart].Name;
+		Result += CustomTags[Tag - LLM_CUSTOM_TAG_START].Name;
 	}
 	else
 	{
@@ -2032,7 +2041,7 @@ FString FLLMCsvWriter::GetTagName(int64 Tag, FLLMPlatformTag* PlatformTags, cons
 
 		if (ParentTags != nullptr && ParentTags[Tag] != -1)
 		{
-			Result = GetTagName( ParentTags[Tag], PlatformTags, nullptr ) + TEXT("/");
+			Result = GetTagName( ParentTags[Tag], CustomTags, nullptr ) + TEXT("/");
 		}
 
 		Result += LLMGetTagName((ELLMTag)Tag);
diff --git a/Engine/Source/Runtime/Core/Public/HAL/LowLevelMemStats.h b/Engine/Source/Runtime/Core/Public/HAL/LowLevelMemStats.h
--- a/Engine/Source/Runtime/Core/Public/HAL/LowLevelMemStats.h
+++ b/Engine/Source/Runtime/Core/Public/HAL/LowLevelMemStats.h
@@ -21,6 +21,7 @@
 	DECLARE_STATS_GROUP(TEXT("LLM Assets"), STATGROUP_LLMAssets, STATCAT_Advanced);
 
 	DECLARE_LLM_MEMORY_STAT_EXTERN(TEXT("Engine"), STAT_EngineSummaryLLM, STATGROUP_LLM, CORE_API);
+	DECLARE_LLM_MEMORY_STAT_EXTERN(TEXT("Project"), STAT_ProjectSummaryLLM, STATGROUP_LLM, CORE_API);
 
 #else
 	#define DECLARE_LLM_MEMORY_STAT(CounterName,StatId,GroupId)
diff --git a/Engine/Source/Runtime/Core/Public/HAL/LowLevelMemTracker.h b/Engine/Source/Runtime/Core/Public/HAL/LowLevelMemTracker.h
--- a/Engine/Source/Runtime/Core/Public/HAL/LowLevelMemTracker.h
+++ b/Engine/Source/Runtime/Core/Public/HAL/LowLevelMemTracker.h
@@ -233,12 +233,22 @@ enum class ELLMTag : LLM_TAG_TYPE
 	//------------------------------
 	// Platform tags
 	PlatformTagStart = 100,
-	PlatformTagEnd = 0xff,
+	PlatformTagEnd = 149,
+
+	//------------------------------
+	// Project tags
+	ProjectTagStart = 150,
+	ProjectTagEnd = 255,
 
 	// anything above this value is treated as an FName for a stat section
 };
+static_assert( ELLMTag::GenericTagCount <= ELLMTag::PlatformTagStart, "too many LLM tags defined"); 
 
 static const uint32 LLM_TAG_COUNT = 256;
+static const uint32 LLM_CUSTOM_TAG_START = (int32)ELLMTag::PlatformTagStart;
+static const uint32 LLM_CUSTOM_TAG_END = (int32)ELLMTag::ProjectTagEnd;
+static const uint32 LLM_CUSTOM_TAG_COUNT = LLM_CUSTOM_TAG_END + 1 - LLM_CUSTOM_TAG_START;
+
 
 /**
  * Passed in to OnLowLevelAlloc to specify the type of allocation. Used to track FMalloc total
@@ -322,8 +332,8 @@ extern FName LLMGetTagStat(ELLMTag Tag);
 /**
  * LLM tag dumping, to help with identifying mis-tagged items. Probably don't want to check in with these in use!
  */
-#define LLM_DUMP_TAG()  FLowLevelMemTracker::Get().DumpTag(ELLMTracker::Default,__FILE__,__LINE__);
-#define LLM_DUMP_PLATFORM_TAG()  FLowLevelMemTracker::Get().DumpTag(ELLMTracker::Platform,__FILE__,__LINE__);
+#define LLM_DUMP_TAG()  FLowLevelMemTracker::Get().DumpTag(ELLMTracker::Default,__FILE__,__LINE__)
+#define LLM_DUMP_PLATFORM_TAG()  FLowLevelMemTracker::Get().DumpTag(ELLMTracker::Platform,__FILE__,__LINE__)
 
 
 typedef void*(*LLMAllocFunction)(size_t);
@@ -382,7 +392,7 @@ private:
 	int32 Alignment;
 };
 
-struct FLLMPlatformTag
+struct FLLMCustomTag
 {
 	int32 Tag;
 	const TCHAR* Name;
@@ -390,6 +400,14 @@ struct FLLMPlatformTag
 	FName SummaryStatName;
 };
 
+struct FLLMTagInfo
+{
+	const TCHAR* Name;
+	FName StatName;				// shows in the LLMFULL stat group
+	FName SummaryStatName;		// shows in the LLM summary stat group
+	int32 ParentTag = -1;
+};
+
 /*
  * The main LLM tracker class
  */
@@ -444,6 +462,7 @@ public:
     int64 GetActiveTag(ELLMTracker Tracker);
 
 	void RegisterPlatformTag(int32 Tag, const TCHAR* Name, FName StatName, FName SummaryStatName, int32 ParentTag = -1);
+	void RegisterProjectTag(int32 Tag, const TCHAR* Name, FName StatName, FName SummaryStatName, int32 ParentTag = -1);
     
 	// look up the tag associated with the given name
 	bool FindTagByName( const TCHAR* Name, uint64& OutTag ) const;
@@ -458,7 +477,7 @@ public:
 	void SetTagAmountForTracker(ELLMTracker Tracker, ELLMTag Tag, int64 Amount, bool bAddToTotal );
 
 	// Dump the current tag for the given tracker to the output
-	void DumpTag( ELLMTracker Tracker, const char* FileName, int LineNumber );
+	int64 DumpTag( ELLMTracker Tracker, const char* FileName, int LineNumber );
 
 private:
 	FLowLevelMemTracker();
@@ -469,6 +488,8 @@ private:
 
 	class FLLMTracker* GetTracker(ELLMTracker Tracker);
 
+	void RegisterCustomTagInternal(int32 Tag, const TCHAR* Name, FName StatName, FName SummaryStatName, int32 ParentTag = -1);
+
 	friend class FLLMScope;
 	friend class FLLMPauseScope;
 	friend class FLLMScopeFromPtr;
@@ -487,7 +508,7 @@ private:
 
 	bool bInitialisedTrackers;
 
-	FLLMPlatformTag PlatformTags[(int32)ELLMTag::PlatformTagEnd + 1 - (int32)ELLMTag::PlatformTagStart];
+	FLLMCustomTag CustomTags[LLM_CUSTOM_TAG_COUNT];
 
 	FLLMTracker* Trackers[(int32)ELLMTracker::Max];
 
