Fix for TUnion missing operator=().

C++:
1) Added TUnion::operator=(), and modified the copy constructor to use it.

diff --git a/Engine/Source/Runtime/Core/Public/Containers/Union.h b/Engine/Source/Runtime/Core/Public/Containers/Union.h
--- a/Engine/Source/Runtime/Core/Public/Containers/Union.h
+++ b/Engine/Source/Runtime/Core/Public/Containers/Union.h
@@ -110,18 +110,7 @@ public:
 	TUnion(const TUnion& Other)
 	:	CurrentSubtypeIndex(-1)
 	{
-		// Copy the value of the appropriate subtype from the other union
-		switch(Other.CurrentSubtypeIndex)
-		{
-		case (uint8)-1: break;
-		case 0: SetSubtype<TypeA>(Other.GetSubtype<TypeA>()); break;
-		case 1: SetSubtype<TypeB>(Other.GetSubtype<TypeB>()); break;
-		case 2: SetSubtype<TypeC>(Other.GetSubtype<TypeC>()); break;
-		case 3: SetSubtype<TypeD>(Other.GetSubtype<TypeD>()); break;
-		case 4: SetSubtype<TypeE>(Other.GetSubtype<TypeE>()); break;
-		case 5: SetSubtype<TypeF>(Other.GetSubtype<TypeF>()); break;
-		default: FatalErrorUndefinedSubtype(); break;
-		};
+		*this = Other;
 	}
 
 	/** Destructor. */
@@ -241,6 +230,24 @@ public:
 		return Result;
 	}
 
+	TUnion& operator=(const TUnion& Other)
+	{
+		// Copy the value of the appropriate subtype from the other union
+		switch (Other.CurrentSubtypeIndex)
+		{
+		case (uint8)-1: break;
+		case 0: SetSubtype<TypeA>(Other.GetSubtype<TypeA>()); break;
+		case 1: SetSubtype<TypeB>(Other.GetSubtype<TypeB>()); break;
+		case 2: SetSubtype<TypeC>(Other.GetSubtype<TypeC>()); break;
+		case 3: SetSubtype<TypeD>(Other.GetSubtype<TypeD>()); break;
+		case 4: SetSubtype<TypeE>(Other.GetSubtype<TypeE>()); break;
+		case 5: SetSubtype<TypeF>(Other.GetSubtype<TypeF>()); break;
+		default: FatalErrorUndefinedSubtype(); break;
+		};
+
+		return *this;
+	}
+
 	/** Equality comparison. */
 	bool operator==(const TUnion& Other) const
 	{
