[LLM] add missing CsvProfiler LLM scope to wait tracking functions

diff --git a/Engine/Source/Runtime/Core/Private/ProfilingDebugging/CsvProfiler.cpp b/Engine/Source/Runtime/Core/Private/ProfilingDebugging/CsvProfiler.cpp
--- a/Engine/Source/Runtime/Core/Private/ProfilingDebugging/CsvProfiler.cpp
+++ b/Engine/Source/Runtime/Core/Private/ProfilingDebugging/CsvProfiler.cpp
@@ -2901,6 +2901,7 @@ void FCsvProfiler::BeginWait()
 #if RECORD_TIMESTAMPS
 	if (GCsvProfilerIsCapturing && GCsvCategoriesEnabled[CSV_CATEGORY_INDEX(Exclusive)])
 	{
+		LLM_SCOPE(ELLMTag::CsvProfiler);
 		const char* WaitStatName = FCsvProfilerThreadData::Get().GetWaitStatName();
 		if (WaitStatName != GIgnoreWaitStatName)
 		{
@@ -2915,6 +2916,7 @@ void FCsvProfiler::EndWait()
 #if RECORD_TIMESTAMPS
 	if (GCsvProfilerIsCapturing && GCsvCategoriesEnabled[CSV_CATEGORY_INDEX(Exclusive)])
 	{
+		LLM_SCOPE(ELLMTag::CsvProfiler);
 		const char* WaitStatName = FCsvProfilerThreadData::Get().GetWaitStatName();
 		if (WaitStatName != GIgnoreWaitStatName)
 		{
