Add DumpEarlyConfigReads command to dump all config reads that occur before it's possible to install patched content to CSV/Log.
This is disabled in shipping.

diff --git a/Engine/Source/Runtime/Core/Private/Misc/ConfigCacheIni.cpp b/Engine/Source/Runtime/Core/Private/Misc/ConfigCacheIni.cpp
--- a/Engine/Source/Runtime/Core/Private/Misc/ConfigCacheIni.cpp
+++ b/Engine/Source/Runtime/Core/Private/Misc/ConfigCacheIni.cpp
@@ -2032,6 +2032,9 @@ bool FConfigCacheIni::GetString( const TCHAR* Section, const TCHAR* Key, FString
 		return false;
 	}
 	Value = ConfigValue->GetValue();
+
+	FCoreDelegates::OnConfigValueRead.Broadcast(*Filename, Section, Key);
+
 	return true;
 }
 
@@ -2056,7 +2059,14 @@ bool FConfigCacheIni::GetText( const TCHAR* Section, const TCHAR* Key, FText& Va
 	{
 		return false;
 	}
-	return FTextStringHelper::ReadFromBuffer( *ConfigValue->GetValue(), Value, Section ) != nullptr;
+	if (FTextStringHelper::ReadFromBuffer(*ConfigValue->GetValue(), Value, Section) == nullptr)
+	{
+		return false;
+	}
+
+	FCoreDelegates::OnConfigValueRead.Broadcast(*Filename, Section, Key);
+
+	return true;
 }
 
 bool FConfigCacheIni::GetSection( const TCHAR* Section, TArray<FString>& Result, const FString& Filename )
@@ -2078,6 +2088,9 @@ bool FConfigCacheIni::GetSection( const TCHAR* Section, TArray<FString>& Result,
 	{
 		Result.Add(FString::Printf(TEXT("%s=%s"), *It.Key().ToString(), *It.Value().GetValue()));
 	}
+
+	FCoreDelegates::OnConfigSectionRead.Broadcast(*Filename, Section);
+
 	return true;
 }
 
@@ -2098,6 +2111,12 @@ FConfigSection* FConfigCacheIni::GetSectionPrivate( const TCHAR* Section, bool F
 	{
 		File->Dirty = true;
 	}
+
+	if (Sec)
+	{
+		FCoreDelegates::OnConfigSectionRead.Broadcast(*Filename, Section);
+	}
+
 	return Sec;
 }
 
@@ -2110,6 +2129,11 @@ bool FConfigCacheIni::DoesSectionExist(const TCHAR* Section, const FString& File
 
 	bReturnVal = (File != nullptr && File->Find(Section) != nullptr);
 
+	if (bReturnVal)
+	{
+		FCoreDelegates::OnConfigSectionNameRead.Broadcast(*Filename, Section);
+	}
+
 	return bReturnVal;
 }
 
@@ -2268,6 +2292,8 @@ bool FConfigCacheIni::GetSectionNames( const FString& Filename, TArray<FString>&
 		{
 			// insert each item at the beginning of the array because TIterators return results in reverse order from which they were added
 			out_SectionNames.Insert(It.Key(),0);
+
+			FCoreDelegates::OnConfigSectionNameRead.Broadcast(*Filename, *It.Key());
 		}
 		bResult = true;
 	}
@@ -2309,6 +2335,8 @@ bool FConfigCacheIni::GetPerObjectConfigSections( const FString& Filename, const
 					// found a PerObjectConfig section for the class specified - add it to the list
 					out_SectionNames.Insert(SectionName,0);
 					bResult = true;
+
+					FCoreDelegates::OnConfigSectionNameRead.Broadcast(*Filename, *SectionName);
 				}
 			}
 		}
@@ -2458,6 +2486,11 @@ int32 FConfigCacheIni::GetArray
 #endif
 	}
 
+	if (out_Arr.Num())
+	{
+		FCoreDelegates::OnConfigValueRead.Broadcast(*Filename, Section, Key);
+	}
+
 	return out_Arr.Num();
 }
 /** Loads a "delimited" list of string
@@ -3955,7 +3988,6 @@ private:
 public:
 	FCVarIniHistoryHelper() : bRecurseCheck( false )
 	{
-		
 		FCoreDelegates::OnApplyCVarFromIni.AddRaw(this, &FCVarIniHistoryHelper::OnApplyCVarFromIniCallback);
 	}
 
@@ -3964,8 +3996,6 @@ public:
 		FCoreDelegates::OnApplyCVarFromIni.RemoveAll(this);
 	}
 
-
-
 	void ReapplyIniHistory()
 	{
 		for (const FCVarIniHistory& IniHistory : CVarIniHistory)
@@ -4029,10 +4059,127 @@ public:
 		bRecurseCheck=false;
 	}
 };
-
 FCVarIniHistoryHelper *IniHistoryHelper = nullptr;
 
+#if !UE_BUILD_SHIPPING
+class FConfigHistoryHelper
+{
+private:
+	enum class HistoryType
+	{
+		Value,
+		Section,
+		SectionName,
+		Count
+	};
+	friend const TCHAR* LexToString(HistoryType Type)
+	{
+		static const TCHAR* Strings[] = {
+			TEXT("Value"),
+			TEXT("Section"),
+			TEXT("SectionName")
+		};
+		using UnderType = __underlying_type(HistoryType);
+		static_assert(static_cast<UnderType>(HistoryType::Count) == ARRAY_COUNT(Strings), "");
+		return Strings[static_cast<UnderType>(Type)];
+	}
 
+	struct FConfigHistory
+	{
+		HistoryType Type = HistoryType::Value;
+
+		FString FileName;
+		FString SectionName;
+		FString Key;
+	};
+	friend uint32 GetTypeHash(const FConfigHistory& CH)
+	{
+		uint32 Hash = ::GetTypeHash(CH.Type);
+		Hash = HashCombine(Hash, ::GetTypeHash(CH.FileName));
+		Hash = HashCombine(Hash, ::GetTypeHash(CH.SectionName));
+		Hash = HashCombine(Hash, ::GetTypeHash(CH.Key));
+		return Hash;
+	}
+	friend bool operator==(const FConfigHistory& A, const FConfigHistory& B)
+	{
+		return
+			A.Type == B.Type &&
+			A.FileName == B.FileName &&
+			A.SectionName == B.SectionName &&
+			A.Key == B.Key;
+	}
+
+	TSet<FConfigHistory> History;
+
+	void OnConfigValueRead(const TCHAR* FileName, const TCHAR* SectionName, const TCHAR* Key)
+	{
+		History.Emplace(FConfigHistory{ HistoryType::Value, FileName, SectionName, Key });
+	}
+
+	void OnConfigSectionRead(const TCHAR* FileName, const TCHAR* SectionName)
+	{
+		History.Emplace(FConfigHistory{ HistoryType::Section, FileName, SectionName });
+	}
+
+	void OnConfigSectionNameRead(const TCHAR* FileName, const TCHAR* SectionName)
+	{
+		History.Emplace(FConfigHistory{ HistoryType::SectionName, FileName, SectionName });
+	}
+
+public:
+	FConfigHistoryHelper()
+	{
+		FCoreDelegates::OnConfigValueRead.AddRaw(this, &FConfigHistoryHelper::OnConfigValueRead);
+		FCoreDelegates::OnConfigSectionRead.AddRaw(this, &FConfigHistoryHelper::OnConfigSectionRead);
+		FCoreDelegates::OnConfigSectionNameRead.AddRaw(this, &FConfigHistoryHelper::OnConfigSectionNameRead);
+	}
+
+	~FConfigHistoryHelper()
+	{
+		FCoreDelegates::OnConfigValueRead.RemoveAll(this);
+		FCoreDelegates::OnConfigSectionRead.RemoveAll(this);
+		FCoreDelegates::OnConfigSectionNameRead.RemoveAll(this);
+	}
+
+	void DumpHistory()
+	{
+		const FString SavePath = FPaths::ProjectLogDir() / TEXT("ConfigHistory.csv");
+
+		FArchive* Writer = IFileManager::Get().CreateFileWriter(*SavePath, FILEWRITE_NoFail);
+
+		auto WriteLine = [Writer](FString&& Line)
+		{
+			UE_LOG(LogConfig, Display, TEXT("%s"), *Line);
+			FTCHARToUTF8 UTF8String(*(MoveTemp(Line) + LINE_TERMINATOR));
+			Writer->Serialize((UTF8CHAR*)UTF8String.Get(), UTF8String.Length());
+		};
+
+		UE_LOG(LogConfig, Display, TEXT("Dumping History of Config Reads to %s"), *SavePath);
+		UE_LOG(LogConfig, Display, TEXT("Begin History of Config Reads"));
+		UE_LOG(LogConfig, Display, TEXT("------------------------------------------------------"));
+		WriteLine(FString::Printf(TEXT("Type, File, Section, Key")));
+		for (const FConfigHistory& CH : History)
+		{
+			switch (CH.Type)
+			{
+			case HistoryType::Value:
+				WriteLine(FString::Printf(TEXT("%s, %s, %s, %s"), LexToString(CH.Type), *CH.FileName, *CH.SectionName, *CH.Key));
+				break;
+			case HistoryType::Section:
+			case HistoryType::SectionName:
+				WriteLine(FString::Printf(TEXT("%s, %s, %s, None"), LexToString(CH.Type), *CH.FileName, *CH.SectionName));
+				break;
+			}
+		}
+		UE_LOG(LogConfig, Display, TEXT("------------------------------------------------------"));
+		UE_LOG(LogConfig, Display, TEXT("End History of Config Reads"));
+
+		delete Writer;
+		Writer = nullptr;
+	}
+};
+FConfigHistoryHelper* ConfigHistoryHelper = nullptr;
+#endif // !UE_BUILD_SHIPPING
 
 
 void RecordApplyCVarSettingsFromIni()
@@ -4074,4 +4221,29 @@ void DeleteRecordedCVarSettingsFromIni()
 {
 	check(IniHistoryHelper != nullptr);
 	delete IniHistoryHelper;
+	IniHistoryHelper = nullptr;
+}
+
+void RecordConfigReadsFromIni()
+{
+#if !UE_BUILD_SHIPPING
+	ConfigHistoryHelper = new FConfigHistoryHelper();
+#endif
+}
+
+void DumpRecordedConfigReadsFromIni()
+{
+#if !UE_BUILD_SHIPPING
+	check(ConfigHistoryHelper);
+	ConfigHistoryHelper->DumpHistory();
+#endif
+}
+
+void DeleteRecordConfigReadsFromIni()
+{
+#if !UE_BUILD_SHIPPING
+	check(ConfigHistoryHelper);
+	delete ConfigHistoryHelper;
+	ConfigHistoryHelper = nullptr;
+#endif
 }
diff --git a/Engine/Source/Runtime/Core/Private/Misc/CoreDelegates.cpp b/Engine/Source/Runtime/Core/Private/Misc/CoreDelegates.cpp
--- a/Engine/Source/Runtime/Core/Private/Misc/CoreDelegates.cpp
+++ b/Engine/Source/Runtime/Core/Private/Misc/CoreDelegates.cpp
@@ -114,6 +114,9 @@ FCoreDelegates::FPreLoadConfigFileDelegate FCoreDelegates::PreLoadConfigFileDele
 FCoreDelegates::FPreSaveConfigFileDelegate FCoreDelegates::PreSaveConfigFileDelegate;
 FCoreDelegates::FOnFConfigFileCreated FCoreDelegates::OnFConfigCreated;
 FCoreDelegates::FOnFConfigFileCreated FCoreDelegates::OnFConfigDeleted;
+FCoreDelegates::FOnConfigValueRead FCoreDelegates::OnConfigValueRead;
+FCoreDelegates::FOnConfigSectionRead FCoreDelegates::OnConfigSectionRead;
+FCoreDelegates::FOnConfigSectionRead FCoreDelegates::OnConfigSectionNameRead;
 FCoreDelegates::FOnApplyCVarFromIni FCoreDelegates::OnApplyCVarFromIni;
 FCoreDelegates::FOnSystemResolutionChanged FCoreDelegates::OnSystemResolutionChanged;
 
diff --git a/Engine/Source/Runtime/Core/Public/Misc/ConfigCacheIni.h b/Engine/Source/Runtime/Core/Public/Misc/ConfigCacheIni.h
--- a/Engine/Source/Runtime/Core/Public/Misc/ConfigCacheIni.h
+++ b/Engine/Source/Runtime/Core/Public/Misc/ConfigCacheIni.h
@@ -954,4 +954,8 @@ CORE_API void ReapplyRecordedCVarSettingsFromIni();
 /**
  * Helper function to clean up ini history
  */
-CORE_API void DeleteRecordedCVarSettingsFromIni();
\ No newline at end of file
+CORE_API void DeleteRecordedCVarSettingsFromIni();
+
+CORE_API void RecordConfigReadsFromIni();
+CORE_API void DumpRecordedConfigReadsFromIni();
+CORE_API void DeleteRecordConfigReadsFromIni();
\ No newline at end of file
diff --git a/Engine/Source/Runtime/Core/Public/Misc/CoreDelegates.h b/Engine/Source/Runtime/Core/Public/Misc/CoreDelegates.h
--- a/Engine/Source/Runtime/Core/Public/Misc/CoreDelegates.h
+++ b/Engine/Source/Runtime/Core/Public/Misc/CoreDelegates.h
@@ -281,6 +281,13 @@ public:
 	static FOnFConfigFileCreated OnFConfigCreated;
 	static FOnFConfigFileDeleted OnFConfigDeleted;
 
+	DECLARE_MULTICAST_DELEGATE_ThreeParams(FOnConfigValueRead, const TCHAR* /*IniFilename*/, const TCHAR* /*SectionName*/, const TCHAR* /*Key*/);
+	static FOnConfigValueRead OnConfigValueRead;
+
+	DECLARE_MULTICAST_DELEGATE_TwoParams(FOnConfigSectionRead, const TCHAR* /*IniFilename*/, const TCHAR* /*SectionName*/);
+	static FOnConfigSectionRead OnConfigSectionRead;
+	static FOnConfigSectionRead OnConfigSectionNameRead;
+
 	DECLARE_MULTICAST_DELEGATE_FourParams(FOnApplyCVarFromIni, const TCHAR* /*SectionName*/, const TCHAR* /*IniFilename*/, uint32 /*SetBy*/, bool /*bAllowCheating*/);
 	static FOnApplyCVarFromIni OnApplyCVarFromIni;
 
diff --git a/Engine/Source/Runtime/Launch/Private/LaunchEngineLoop.cpp b/Engine/Source/Runtime/Launch/Private/LaunchEngineLoop.cpp
--- a/Engine/Source/Runtime/Launch/Private/LaunchEngineLoop.cpp
+++ b/Engine/Source/Runtime/Launch/Private/LaunchEngineLoop.cpp
@@ -1598,11 +1598,18 @@ int32 FEngineLoop::PreInit(const TCHAR* CmdLine)
 		}
 	}
 
+	const bool bDumpEarlyConfigReads = FParse::Param(FCommandLine::Get(), TEXT("DumpEarlyConfigReads"));
+
 	constexpr bool bWithConfigPatching = (WITH_CONFIG_PATCHING != 0);
 	if (bWithConfigPatching)
 	{
 		UE_LOG(LogInit, Verbose, TEXT("Begin recording CVar changes for config patching."));
 
+		if (bDumpEarlyConfigReads)
+		{
+			RecordConfigReadsFromIni();
+		}
+
 		RecordApplyCVarSettingsFromIni();
 	}
 
@@ -2234,7 +2241,7 @@ int32 FEngineLoop::PreInit(const TCHAR* CmdLine)
 				IPlatformInstallBundleManager* BundleManager = FPlatformMisc::GetPlatformInstallBundleManager();
 				if (BundleManager != nullptr && !BundleManager->IsNullInterface())
 				{
-					IPlatformInstallBundleManager::InstallBundleCompleteDelegate.AddRaw(this, &FEngineLoop::OnStartupContentMounted);
+					IPlatformInstallBundleManager::InstallBundleCompleteDelegate.AddRaw(this, &FEngineLoop::OnStartupContentMounted, bDumpEarlyConfigReads);
 				}
 				// If not using the bundle manager, config will be reloaded after ESP, see below
 			}
@@ -2367,6 +2374,12 @@ int32 FEngineLoop::PreInit(const TCHAR* CmdLine)
 			{
 				SCOPED_BOOT_TIMING("ReapplyCVarsFromIniAfterEarlyStartupScreen");
 
+				if (bDumpEarlyConfigReads)
+				{
+					DumpRecordedConfigReadsFromIni();
+					DeleteRecordConfigReadsFromIni();
+				}
+
 				ReapplyRecordedCVarSettingsFromIni();
 				DeleteRecordedCVarSettingsFromIni();
 			}
@@ -3622,10 +3635,16 @@ void FEngineLoop::ProcessLocalPlayerSlateOperations() const
 	}
 }
 
-void FEngineLoop::OnStartupContentMounted(FInstallBundleResultInfo Result)
+void FEngineLoop::OnStartupContentMounted(FInstallBundleResultInfo Result, bool bDumpEarlyConfigReads)
 {
 	if (Result.bIsStartup && Result.Result == EInstallBundleResult::OK)
 	{
+		if (bDumpEarlyConfigReads)
+		{
+			DumpRecordedConfigReadsFromIni();
+			DeleteRecordConfigReadsFromIni();
+		}
+
 		ReapplyRecordedCVarSettingsFromIni();
 		DeleteRecordedCVarSettingsFromIni();
 
diff --git a/Engine/Source/Runtime/Launch/Public/LaunchEngineLoop.h b/Engine/Source/Runtime/Launch/Public/LaunchEngineLoop.h
--- a/Engine/Source/Runtime/Launch/Public/LaunchEngineLoop.h
+++ b/Engine/Source/Runtime/Launch/Public/LaunchEngineLoop.h
@@ -118,7 +118,7 @@ private:
 	/** Utility function that processes Slate operations. */
 	void ProcessLocalPlayerSlateOperations() const;
 
-	void OnStartupContentMounted(FInstallBundleResultInfo Result);
+	void OnStartupContentMounted(FInstallBundleResultInfo Result, bool bDumpEarlyConfigReads);
 
 protected:
 
