Add comments.

diff --git a/Engine/Source/Runtime/Core/Public/Misc/ConfigCacheIni.h b/Engine/Source/Runtime/Core/Public/Misc/ConfigCacheIni.h
--- a/Engine/Source/Runtime/Core/Public/Misc/ConfigCacheIni.h
+++ b/Engine/Source/Runtime/Core/Public/Misc/ConfigCacheIni.h
@@ -956,6 +956,17 @@ CORE_API void ReapplyRecordedCVarSettingsFromIni();
  */
 CORE_API void DeleteRecordedCVarSettingsFromIni();
 
+/**
+ * Helper function to start recording config reads
+ */
 CORE_API void RecordConfigReadsFromIni();
+
+/**
+ * Helper function to dump config reads to csv after RecordConfigReadsFromIni was called
+ */
 CORE_API void DumpRecordedConfigReadsFromIni();
-CORE_API void DeleteRecordConfigReadsFromIni();
\ No newline at end of file
+
+/**
+ * Helper function to clean up config read history
+ */
+CORE_API void DeleteRecordConfigReadsFromIni();
