- Fixing up ALL platform header files to match new required naming conventions for HAL files (XXXPlatformFile.h, not XXXFile.h, and all platforms must have a header file, even if it just #includes another one, like IOSPlatformTime.h just #includes ApplePlatformTime.h)

diff --git a/Engine/Plugins/Media/AndroidCamera/Source/AndroidCamera/Private/Player/AndroidCameraPlayer.cpp b/Engine/Plugins/Media/AndroidCamera/Source/AndroidCamera/Private/Player/AndroidCameraPlayer.cpp
--- a/Engine/Plugins/Media/AndroidCamera/Source/AndroidCamera/Private/Player/AndroidCameraPlayer.cpp
+++ b/Engine/Plugins/Media/AndroidCamera/Source/AndroidCamera/Private/Player/AndroidCameraPlayer.cpp
@@ -3,7 +3,7 @@
 #include "AndroidCameraPlayer.h"
 #include "AndroidCameraPrivate.h"
 
-#include "Android/AndroidFile.h"
+#include "Android/AndroidPlatformFile.h"
 #include "AndroidCameraSettings.h"
 #include "IMediaEventSink.h"
 #include "MediaSamples.h"
diff --git a/Engine/Plugins/Media/AndroidMedia/Source/AndroidMedia/Private/Player/AndroidMediaPlayer.cpp b/Engine/Plugins/Media/AndroidMedia/Source/AndroidMedia/Private/Player/AndroidMediaPlayer.cpp
--- a/Engine/Plugins/Media/AndroidMedia/Source/AndroidMedia/Private/Player/AndroidMediaPlayer.cpp
+++ b/Engine/Plugins/Media/AndroidMedia/Source/AndroidMedia/Private/Player/AndroidMediaPlayer.cpp
@@ -3,7 +3,7 @@
 #include "AndroidMediaPlayer.h"
 #include "AndroidMediaPrivate.h"
 
-#include "Android/AndroidFile.h"
+#include "Android/AndroidPlatformFile.h"
 #include "AndroidMediaSettings.h"
 #include "IMediaEventSink.h"
 #include "MediaSamples.h"
diff --git a/Engine/Plugins/Runtime/AndroidMoviePlayer/Source/AndroidMoviePlayer/Private/AndroidMovieStreamer.cpp b/Engine/Plugins/Runtime/AndroidMoviePlayer/Source/AndroidMoviePlayer/Private/AndroidMovieStreamer.cpp
--- a/Engine/Plugins/Runtime/AndroidMoviePlayer/Source/AndroidMoviePlayer/Private/AndroidMovieStreamer.cpp
+++ b/Engine/Plugins/Runtime/AndroidMoviePlayer/Source/AndroidMoviePlayer/Private/AndroidMovieStreamer.cpp
@@ -4,7 +4,7 @@
 
 #include "Android/AndroidApplication.h"
 #include "Android/AndroidJava.h"
-#include "Android/AndroidFile.h"
+#include "Android/AndroidPlatformFile.h"
 
 #include "Rendering/RenderingCommon.h"
 #include "RenderUtils.h"
diff --git a/Engine/Source/Developer/Android/AndroidTargetPlatform/Private/AndroidTargetPlatform.h b/Engine/Source/Developer/Android/AndroidTargetPlatform/Private/AndroidTargetPlatform.h
--- a/Engine/Source/Developer/Android/AndroidTargetPlatform/Private/AndroidTargetPlatform.h
+++ b/Engine/Source/Developer/Android/AndroidTargetPlatform/Private/AndroidTargetPlatform.h
@@ -14,7 +14,7 @@
 #include "Delegates/Delegate.h"
 #include "Containers/Ticker.h"
 #include "Misc/ScopeLock.h"
-#include "Android/AndroidProperties.h"
+#include "Android/AndroidPlatformProperties.h"
 #include "Interfaces/ITargetPlatformModule.h"
 #include "Common/TargetPlatformBase.h"
 #include "Interfaces/IAndroidDeviceDetection.h"
diff --git a/Engine/Source/Developer/Android/AndroidTargetPlatform/Private/AndroidTargetPlatformModule.cpp b/Engine/Source/Developer/Android/AndroidTargetPlatform/Private/AndroidTargetPlatformModule.cpp
--- a/Engine/Source/Developer/Android/AndroidTargetPlatform/Private/AndroidTargetPlatformModule.cpp
+++ b/Engine/Source/Developer/Android/AndroidTargetPlatform/Private/AndroidTargetPlatformModule.cpp
@@ -2,7 +2,7 @@
 
 #include "CoreMinimal.h"
 #include "Modules/ModuleManager.h"
-#include "Android/AndroidProperties.h"
+#include "Android/AndroidPlatformProperties.h"
 #include "Interfaces/ITargetPlatformModule.h"
 #include "Common/TargetPlatformBase.h"
 #include "Interfaces/IAndroidDeviceDetection.h"
diff --git a/Engine/Source/Developer/Android/Android_ETC1aTargetPlatform/Private/Android_ETC1aTargetPlatformModule.cpp b/Engine/Source/Developer/Android/Android_ETC1aTargetPlatform/Private/Android_ETC1aTargetPlatformModule.cpp
--- a/Engine/Source/Developer/Android/Android_ETC1aTargetPlatform/Private/Android_ETC1aTargetPlatformModule.cpp
+++ b/Engine/Source/Developer/Android/Android_ETC1aTargetPlatform/Private/Android_ETC1aTargetPlatformModule.cpp
@@ -6,7 +6,7 @@
 
 #include "CoreMinimal.h"
 #include "Modules/ModuleManager.h"
-#include "Android/AndroidProperties.h"
+#include "Android/AndroidPlatformProperties.h"
 #include "Interfaces/ITargetPlatformModule.h"
 #include "Common/TargetPlatformBase.h"
 #include "Interfaces/IAndroidDeviceDetection.h"
diff --git a/Engine/Source/Developer/Lumin/LuminTargetPlatform/Private/LuminTargetPlatform.h b/Engine/Source/Developer/Lumin/LuminTargetPlatform/Private/LuminTargetPlatform.h
--- a/Engine/Source/Developer/Lumin/LuminTargetPlatform/Private/LuminTargetPlatform.h
+++ b/Engine/Source/Developer/Lumin/LuminTargetPlatform/Private/LuminTargetPlatform.h
@@ -15,7 +15,7 @@
 #include "Containers/Ticker.h"
 #include "Misc/ScopeLock.h"
 #include "Common/TargetPlatformBase.h"
-#include "Android/AndroidProperties.h"
+#include "Android/AndroidPlatformProperties.h"
 #include "AndroidTargetPlatform.h"
 
 #if WITH_ENGINE
diff --git a/Engine/Source/Runtime/ApplicationCore/Private/Android/AndroidApplication.cpp b/Engine/Source/Runtime/ApplicationCore/Private/Android/AndroidApplication.cpp
--- a/Engine/Source/Runtime/ApplicationCore/Private/Android/AndroidApplication.cpp
+++ b/Engine/Source/Runtime/ApplicationCore/Private/Android/AndroidApplication.cpp
@@ -7,7 +7,7 @@
 #include "IInputDeviceModule.h"
 #include "HAL/OutputDevices.h"
 #include "Misc/AssertionMacros.h"
-#include "Android/AndroidMisc.h"
+#include "Android/AndroidPlatformMisc.h"
 #include "Misc/CoreDelegates.h"
 
 DEFINE_LOG_CATEGORY_STATIC(LogAndroidApplication, Log, All);
diff --git a/Engine/Source/Runtime/ApplicationCore/Private/Android/AndroidApplicationMisc.cpp b/Engine/Source/Runtime/ApplicationCore/Private/Android/AndroidApplicationMisc.cpp
deleted file mode 100644
diff --git a/Engine/Source/Runtime/ApplicationCore/Private/Android/AndroidPlatformApplicationMisc.cpp b/Engine/Source/Runtime/ApplicationCore/Private/Android/AndroidPlatformApplicationMisc.cpp
new file mode 100644
--- /dev/null
+++ b/Engine/Source/Runtime/ApplicationCore/Private/Android/AndroidPlatformApplicationMisc.cpp
@@ -0,0 +1,233 @@
+// Copyright 1998-2019 Epic Games, Inc. All Rights Reserved.
+
+#include "Android/AndroidPlatformApplicationMisc.h"
+
+#include "Android/AndroidApplication.h"
+#include "Android/AndroidFeedbackContext.h"
+#include "Android/AndroidErrorOutputDevice.h"
+#include "Android/AndroidInputInterface.h"
+#include "HAL/PlatformMisc.h"
+#include "Misc/ConfigCacheIni.h"
+#include "Internationalization/Regex.h"
+#include "Modules/ModuleManager.h"
+
+void FAndroidApplicationMisc::LoadPreInitModules()
+{
+	FModuleManager::Get().LoadModule(TEXT("OpenGLDrv"));
+#if USE_ANDROID_AUDIO
+	FModuleManager::Get().LoadModule(TEXT("AndroidAudio"));
+	FModuleManager::Get().LoadModule(TEXT("AudioMixerAndroid"));
+#endif
+}
+
+class FFeedbackContext* FAndroidApplicationMisc::GetFeedbackContext()
+{
+	static FAndroidFeedbackContext Singleton;
+	return &Singleton;
+}
+
+class FOutputDeviceError* FAndroidApplicationMisc::GetErrorOutputDevice()
+{
+	static FAndroidErrorOutputDevice Singleton;
+	return &Singleton;
+}
+
+
+GenericApplication* FAndroidApplicationMisc::CreateApplication()
+{
+	return FAndroidApplication::CreateAndroidApplication();
+}
+
+void FAndroidApplicationMisc::RequestMinimize()
+{
+#if USE_ANDROID_JNI
+	extern void AndroidThunkCpp_Minimize();
+	AndroidThunkCpp_Minimize();
+#endif
+}
+
+bool FAndroidApplicationMisc::IsScreensaverEnabled()
+{
+#if USE_ANDROID_JNI
+	extern bool AndroidThunkCpp_IsScreensaverEnabled();
+	return AndroidThunkCpp_IsScreensaverEnabled();
+#else
+	return false;
+#endif
+}
+
+bool FAndroidApplicationMisc::ControlScreensaver(EScreenSaverAction Action)
+{
+#if USE_ANDROID_JNI
+	extern void AndroidThunkCpp_KeepScreenOn(bool Enable);
+	switch (Action)
+	{
+		case EScreenSaverAction::Disable:
+			// Prevent display sleep.
+			AndroidThunkCpp_KeepScreenOn(true);
+			break;
+
+		case EScreenSaverAction::Enable:
+			// Stop preventing display sleep now that we are done.
+			AndroidThunkCpp_KeepScreenOn(false);
+			break;
+	}
+	return true;
+#else
+	return false;
+#endif
+}
+
+void FAndroidApplicationMisc::SetGamepadsAllowed(bool bAllowed)
+{
+	if (FAndroidInputInterface* InputInterface = (FAndroidInputInterface*)FAndroidApplication::Get()->GetInputInterface())
+	{
+		InputInterface->SetGamepadsAllowed(bAllowed);
+	}
+}
+
+void FAndroidApplicationMisc::ResetGamepadAssignments()
+{
+	FAndroidInputInterface::ResetGamepadAssignments();
+}
+
+void FAndroidApplicationMisc::ResetGamepadAssignmentToController(int32 ControllerId)
+{
+	FAndroidInputInterface::ResetGamepadAssignmentToController(ControllerId);
+}
+
+bool FAndroidApplicationMisc::IsControllerAssignedToGamepad(int32 ControllerId)
+{
+	return FAndroidInputInterface::IsControllerAssignedToGamepad(ControllerId);
+}
+
+FString FAndroidApplicationMisc::GetGamepadControllerName(int32 ControllerId)
+{
+	return FAndroidInputInterface::GetGamepadControllerName(ControllerId);
+}
+
+extern void AndroidThunkCpp_ClipboardCopy(const FString& Str);
+void FAndroidApplicationMisc::ClipboardCopy(const TCHAR* Str)
+{
+#if USE_ANDROID_JNI
+	AndroidThunkCpp_ClipboardCopy(Str);
+#endif
+}
+
+extern FString AndroidThunkCpp_ClipboardPaste();
+void FAndroidApplicationMisc::ClipboardPaste(class FString& Result)
+{
+#if USE_ANDROID_JNI
+	Result = AndroidThunkCpp_ClipboardPaste();
+#endif
+}
+
+struct FScreenDensity
+{
+	FString Model;
+	bool IsRegex;
+	int32 Density;
+	
+	FScreenDensity()
+		: Model()
+		, IsRegex(false)
+		, Density(0)
+	{
+	}
+
+	bool InitFromString(const FString& InSourceString)
+	{
+		Model = TEXT("");
+		Density = 0;
+		IsRegex = false;
+
+		// The initialization is only successful if the Model and Density values can all be parsed from the string
+		const bool bSuccessful = FParse::Value(*InSourceString, TEXT("Model="), Model) && FParse::Value(*InSourceString, TEXT("Density="), Density);
+
+		// Regex= is optional, it lets us know if this model requires regular expression matching, which is much more expensive.
+		FParse::Bool(*InSourceString, TEXT("IsRegex="), IsRegex);
+
+		return bSuccessful;
+	}
+
+	bool IsMatch(const FString& InDeviceModel) const
+	{
+		if ( IsRegex )
+		{
+			const FRegexPattern RegexPattern(Model);
+			FRegexMatcher RegexMatcher(RegexPattern, InDeviceModel);
+
+			return RegexMatcher.FindNext();
+		}
+		else
+		{
+			return Model == InDeviceModel;
+		}
+	}
+};
+
+static float GetWindowUpscaleFactor()
+{
+	// Determine the difference between the native resolution of the device, and the size of our window,
+	// and return that scalar.
+
+	int32_t SurfaceWidth, SurfaceHeight;
+	FAndroidWindow::CalculateSurfaceSize(FAndroidWindow::GetHardwareWindow(), SurfaceWidth, SurfaceHeight);
+
+	FPlatformRect ScreenRect = FAndroidWindow::GetScreenRect();
+
+	const float CalculatedScaleFactor = FVector2D(ScreenRect.Right - ScreenRect.Left, ScreenRect.Bottom - ScreenRect.Top).Size() / FVector2D(SurfaceWidth, SurfaceHeight).Size();
+
+	return CalculatedScaleFactor;
+}
+
+EScreenPhysicalAccuracy FAndroidApplicationMisc::ComputePhysicalScreenDensity(int32& OutScreenDensity)
+{
+	FString MyDeviceModel = FPlatformMisc::GetDeviceModel();
+
+	TArray<FString> DeviceStrings;
+	GConfig->GetArray(TEXT("DeviceScreenDensity"), TEXT("Devices"), DeviceStrings, GEngineIni);
+
+	TArray<FScreenDensity> Devices;
+	for ( const FString& DeviceString : DeviceStrings )
+	{
+		FScreenDensity DensityEntry;
+		if ( DensityEntry.InitFromString(DeviceString) )
+		{
+			Devices.Add(DensityEntry);
+		}
+	}
+
+	for ( const FScreenDensity& Device : Devices )
+	{
+		if ( Device.IsMatch(MyDeviceModel) )
+		{
+			OutScreenDensity = Device.Density * GetWindowUpscaleFactor();
+			return EScreenPhysicalAccuracy::Truth;
+		}
+	}
+
+#if USE_ANDROID_JNI
+	extern FString AndroidThunkCpp_GetMetaDataString(const FString& Key);
+	FString DPIStrings = AndroidThunkCpp_GetMetaDataString(TEXT("ue4.displaymetrics.dpi"));
+	TArray<FString> DPIValues;
+	DPIStrings.ParseIntoArray(DPIValues, TEXT(","));
+
+	float xdpi, ydpi;
+	LexFromString(xdpi, *DPIValues[0]);
+	LexFromString(ydpi, *DPIValues[1]);
+
+	OutScreenDensity = ( xdpi + ydpi ) / 2.0f;
+
+	if ( OutScreenDensity <= 0 || OutScreenDensity > 2000 )
+	{
+		return EScreenPhysicalAccuracy::Unknown;
+	}
+
+	OutScreenDensity *= GetWindowUpscaleFactor();
+	return EScreenPhysicalAccuracy::Approximation;
+#else
+	// @todo Lumin: implement this on Lumin probably
+	return EScreenPhysicalAccuracy::Unknown;
+#endif
+}
diff --git a/Engine/Source/Runtime/ApplicationCore/Private/Android/AndroidPlatformSurvey.cpp b/Engine/Source/Runtime/ApplicationCore/Private/Android/AndroidPlatformSurvey.cpp
new file mode 100644
--- /dev/null
+++ b/Engine/Source/Runtime/ApplicationCore/Private/Android/AndroidPlatformSurvey.cpp
@@ -0,0 +1,40 @@
+// Copyright 1998-2019 Epic Games, Inc. All Rights Reserved.
+
+/*=============================================================================
+	AndroidSurvey.cpp: HardwareSurvey implementation
+=============================================================================*/
+
+#include "Android/AndroidPlatformSurvey.h"
+#include "Android/AndroidWindow.h"
+
+bool FAndroidPlatformSurvey::GetSurveyResults(FHardwareSurveyResults& OutResults, bool bWait)
+{
+#if PLATFORM_64BITS
+	OutResults.OSBits = 64;
+#else
+	OutResults.OSBits = 32;
+#endif
+
+	FCString::Strcpy(OutResults.Platform, *(FAndroidMisc::GetDeviceMake() + TEXT("-") + FAndroidMisc::GetDeviceModel()));
+	FCString::Strcpy(OutResults.OSVersion, *FAndroidMisc::GetAndroidVersion());
+	FCString::Strcpy(OutResults.OSLanguage, *FPlatformMisc::GetDefaultLocale());
+	// @todo vulkan: Get Vulkan version somehow
+	FCString::Strcpy(OutResults.RenderingAPI, FAndroidMisc::ShouldUseVulkan() ? TEXT("Vulkan") : *FAndroidMisc::GetGLVersion());
+	FCString::Strcpy(OutResults.RHIAdapter.AdapterName, *FAndroidMisc::GetGPUFamily());
+
+	// display 0 is max size
+	int32 ScreenWidth, ScreenHeight;
+	FAndroidWindow::CalculateSurfaceSize(FAndroidWindow::GetHardwareWindow(), ScreenWidth, ScreenHeight);
+
+	OutResults.Displays[0].CurrentModeWidth = ScreenWidth;
+	OutResults.Displays[0].CurrentModeHeight = ScreenHeight;
+
+	OutResults.CPUCount = FPlatformMisc::NumberOfCores();
+
+
+	// display 1 is current size
+	FPlatformRect ViewSize = FAndroidWindow::GetScreenRect();
+	OutResults.Displays[1].CurrentModeWidth = ViewSize.Right - ViewSize.Left;
+	OutResults.Displays[1].CurrentModeHeight = ViewSize.Bottom - ViewSize.Top;
+	return true;
+}
diff --git a/Engine/Source/Runtime/ApplicationCore/Private/Android/AndroidSurvey.cpp b/Engine/Source/Runtime/ApplicationCore/Private/Android/AndroidSurvey.cpp
deleted file mode 100644
diff --git a/Engine/Source/Runtime/ApplicationCore/Public/Android/AndroidApplicationMisc.h b/Engine/Source/Runtime/ApplicationCore/Public/Android/AndroidApplicationMisc.h
deleted file mode 100644
diff --git a/Engine/Source/Runtime/ApplicationCore/Public/Android/AndroidPlatformApplicationMisc.h b/Engine/Source/Runtime/ApplicationCore/Public/Android/AndroidPlatformApplicationMisc.h
new file mode 100644
--- /dev/null
+++ b/Engine/Source/Runtime/ApplicationCore/Public/Android/AndroidPlatformApplicationMisc.h
@@ -0,0 +1,28 @@
+// Copyright 1998-2019 Epic Games, Inc. All Rights Reserved.
+
+#pragma once
+
+#include "GenericPlatform/GenericPlatformApplicationMisc.h"
+
+struct APPLICATIONCORE_API FAndroidApplicationMisc : public FGenericPlatformApplicationMisc
+{
+	static void LoadPreInitModules();
+	static class FFeedbackContext* GetFeedbackContext();
+	static class FOutputDeviceError* GetErrorOutputDevice();
+	static class GenericApplication* CreateApplication();
+	static void RequestMinimize();
+	static bool IsScreensaverEnabled();
+	static bool ControlScreensaver(EScreenSaverAction Action);
+	static void SetGamepadsAllowed(bool bAllowed);
+	static void ResetGamepadAssignments();
+	static void ResetGamepadAssignmentToController(int32 ControllerId);
+	static bool IsControllerAssignedToGamepad(int32 ControllerId);
+	static FString GetGamepadControllerName(int32 ControllerId);
+	static void ClipboardCopy(const TCHAR* Str);
+	static void ClipboardPaste(class FString& Dest);
+	static EScreenPhysicalAccuracy ComputePhysicalScreenDensity(int32& OutScreenDensity);
+};
+
+#if !PLATFORM_LUMIN
+typedef FAndroidApplicationMisc FPlatformApplicationMisc;
+#endif
diff --git a/Engine/Source/Runtime/ApplicationCore/Public/Android/AndroidPlatformSplash.h b/Engine/Source/Runtime/ApplicationCore/Public/Android/AndroidPlatformSplash.h
new file mode 100644
--- /dev/null
+++ b/Engine/Source/Runtime/ApplicationCore/Public/Android/AndroidPlatformSplash.h
@@ -0,0 +1,19 @@
+// Copyright 1998-2019 Epic Games, Inc. All Rights Reserved.
+
+/*=============================================================================================
+	AndroidSplash.h: Android platform splash screen...
+==============================================================================================*/
+
+#pragma once
+#include "GenericPlatform/GenericPlatformSplash.h"
+
+/**
+ * Android splash implementation
+ */
+struct APPLICATIONCORE_API FAndroidSplash : public FGenericPlatformSplash
+{
+	// default implementation for now
+};
+
+
+typedef FAndroidSplash FPlatformSplash;
diff --git a/Engine/Source/Runtime/ApplicationCore/Public/Android/AndroidPlatformSurvey.h b/Engine/Source/Runtime/ApplicationCore/Public/Android/AndroidPlatformSurvey.h
new file mode 100644
--- /dev/null
+++ b/Engine/Source/Runtime/ApplicationCore/Public/Android/AndroidPlatformSurvey.h
@@ -0,0 +1,19 @@
+// Copyright 1998-2019 Epic Games, Inc. All Rights Reserved.
+
+
+/*=============================================================================================
+	AndroidSurvey.h: Google Android platform hardware-survey classes
+==============================================================================================*/
+
+#pragma once
+#include "GenericPlatform/GenericPlatformSurvey.h"
+
+/**
+ * Android implementation of FGenericPlatformSurvey
+ **/
+struct FAndroidPlatformSurvey : public FGenericPlatformSurvey
+{
+	static bool GetSurveyResults(FHardwareSurveyResults& OutResults, bool bWait=false);
+};
+
+typedef FAndroidPlatformSurvey FPlatformSurvey;
\ No newline at end of file
diff --git a/Engine/Source/Runtime/ApplicationCore/Public/Android/AndroidSplash.h b/Engine/Source/Runtime/ApplicationCore/Public/Android/AndroidSplash.h
deleted file mode 100644
diff --git a/Engine/Source/Runtime/ApplicationCore/Public/Android/AndroidSurvey.h b/Engine/Source/Runtime/ApplicationCore/Public/Android/AndroidSurvey.h
deleted file mode 100644
diff --git a/Engine/Source/Runtime/ApplicationCore/Public/Android/AndroidWindowUtils.h b/Engine/Source/Runtime/ApplicationCore/Public/Android/AndroidWindowUtils.h
--- a/Engine/Source/Runtime/ApplicationCore/Public/Android/AndroidWindowUtils.h
+++ b/Engine/Source/Runtime/ApplicationCore/Public/Android/AndroidWindowUtils.h
@@ -6,7 +6,7 @@
 #include "RHI.h"
 
 #if PLATFORM_ANDROID
-#include "Android/AndroidMisc.h"
+#include "Android/AndroidPlatformMisc.h"
 
 #if USE_ANDROID_JNI
 extern bool AndroidThunkCpp_IsGearVRApplication();
diff --git a/Engine/Source/Runtime/ApplicationCore/Public/HAL/PlatformApplicationMisc.h b/Engine/Source/Runtime/ApplicationCore/Public/HAL/PlatformApplicationMisc.h
--- a/Engine/Source/Runtime/ApplicationCore/Public/HAL/PlatformApplicationMisc.h
+++ b/Engine/Source/Runtime/ApplicationCore/Public/HAL/PlatformApplicationMisc.h
@@ -4,12 +4,4 @@
 
 #include "CoreTypes.h"
 
-#if PLATFORM_PS4
-#include "PS4/PS4ApplicationMisc.h"
-#elif PLATFORM_XBOXONE
-#include "XboxOne/XboxOneApplicationMisc.h"
-#elif PLATFORM_ANDROID
-#include "Android/AndroidApplicationMisc.h"
-#else
 #include COMPILED_PLATFORM_HEADER(PlatformApplicationMisc.h)
-#endif
diff --git a/Engine/Source/Runtime/ApplicationCore/Public/HAL/PlatformSplash.h b/Engine/Source/Runtime/ApplicationCore/Public/HAL/PlatformSplash.h
--- a/Engine/Source/Runtime/ApplicationCore/Public/HAL/PlatformSplash.h
+++ b/Engine/Source/Runtime/ApplicationCore/Public/HAL/PlatformSplash.h
@@ -3,12 +3,4 @@
 
 #include "CoreTypes.h"
 
-#if PLATFORM_PS4
-#include "PS4/PS4Splash.h"
-#elif PLATFORM_XBOXONE
-#include "XboxOne/XboxOneSplash.h"
-#elif PLATFORM_ANDROID
-#include "Android/AndroidSplash.h"
-#else
 #include COMPILED_PLATFORM_HEADER(PlatformSplash.h)
-#endif
diff --git a/Engine/Source/Runtime/ApplicationCore/Public/HAL/PlatformSurvey.h b/Engine/Source/Runtime/ApplicationCore/Public/HAL/PlatformSurvey.h
--- a/Engine/Source/Runtime/ApplicationCore/Public/HAL/PlatformSurvey.h
+++ b/Engine/Source/Runtime/ApplicationCore/Public/HAL/PlatformSurvey.h
@@ -3,15 +3,5 @@
 
 #include "CoreTypes.h"
 
-#if PLATFORM_PS4
-#include "PS4/PS4Survey.h"
-#elif PLATFORM_XBOXONE
-#include "XboxOne/XboxOneSurvey.h"
-#elif PLATFORM_ANDROID
-#include "Android/AndroidSurvey.h"
-#elif PLATFORM_UNIX
-#include "Unix/UnixPlatformSurvey.h"
-#else
 #include COMPILED_PLATFORM_HEADER(PlatformSurvey.h)
-#endif
 
diff --git a/Engine/Source/Runtime/ApplicationCore/Public/Linux/LinuxPlatformSurvey.h b/Engine/Source/Runtime/ApplicationCore/Public/Linux/LinuxPlatformSurvey.h
new file mode 100644
--- /dev/null
+++ b/Engine/Source/Runtime/ApplicationCore/Public/Linux/LinuxPlatformSurvey.h
@@ -0,0 +1,5 @@
+// Copyright 1998-2019 Epic Games, Inc. All Rights Reserved.
+
+#pragma once
+
+#include "Unix/UnixPlatformSurvey.h"
diff --git a/Engine/Source/Runtime/ApplicationCore/Public/Lumin/LuminPlatformApplicationMisc.h b/Engine/Source/Runtime/ApplicationCore/Public/Lumin/LuminPlatformApplicationMisc.h
--- a/Engine/Source/Runtime/ApplicationCore/Public/Lumin/LuminPlatformApplicationMisc.h
+++ b/Engine/Source/Runtime/ApplicationCore/Public/Lumin/LuminPlatformApplicationMisc.h
@@ -2,7 +2,7 @@
 
 #pragma once
 
-#include "Android/AndroidApplicationMisc.h"
+#include "Android/AndroidPlatformApplicationMisc.h"
 
 struct APPLICATIONCORE_API FLuminPlatformApplicationMisc : public FAndroidApplicationMisc
 {
diff --git a/Engine/Source/Runtime/ApplicationCore/Public/Lumin/LuminPlatformSplash.h b/Engine/Source/Runtime/ApplicationCore/Public/Lumin/LuminPlatformSplash.h
new file mode 100644
--- /dev/null
+++ b/Engine/Source/Runtime/ApplicationCore/Public/Lumin/LuminPlatformSplash.h
@@ -0,0 +1,5 @@
+// Copyright 1998-2019 Epic Games, Inc. All Rights Reserved.
+
+#pragma once
+
+#include "Android/AndroidPlatformSplash.h"
diff --git a/Engine/Source/Runtime/ApplicationCore/Public/Lumin/LuminPlatformSurvey.h b/Engine/Source/Runtime/ApplicationCore/Public/Lumin/LuminPlatformSurvey.h
new file mode 100644
--- /dev/null
+++ b/Engine/Source/Runtime/ApplicationCore/Public/Lumin/LuminPlatformSurvey.h
@@ -0,0 +1,5 @@
+// Copyright 1998-2019 Epic Games, Inc. All Rights Reserved.
+
+#pragma once
+
+#include "Android/AndroidPlatformSurvey.h"
diff --git a/Engine/Source/Runtime/Core/Private/Android/AndroidFile.cpp b/Engine/Source/Runtime/Core/Private/Android/AndroidFile.cpp
deleted file mode 100644
diff --git a/Engine/Source/Runtime/Core/Private/Android/AndroidMemory.cpp b/Engine/Source/Runtime/Core/Private/Android/AndroidMemory.cpp
deleted file mode 100644
diff --git a/Engine/Source/Runtime/Core/Private/Android/AndroidMisc.cpp b/Engine/Source/Runtime/Core/Private/Android/AndroidMisc.cpp
deleted file mode 100644
diff --git a/Engine/Source/Runtime/Core/Private/Android/AndroidOutputDevices.cpp b/Engine/Source/Runtime/Core/Private/Android/AndroidOutputDevices.cpp
deleted file mode 100644
diff --git a/Engine/Source/Runtime/Core/Private/Android/AndroidPlatformFile.cpp b/Engine/Source/Runtime/Core/Private/Android/AndroidPlatformFile.cpp
new file mode 100644
--- /dev/null
+++ b/Engine/Source/Runtime/Core/Private/Android/AndroidPlatformFile.cpp
@@ -0,0 +1,2020 @@
+// Copyright 1998-2019 Epic Games, Inc. All Rights Reserved.
+
+/*=============================================================================
+	AndroidFile.cpp: Android platform implementations of File functions
+=============================================================================*/
+
+#include "Android/AndroidPlatformFile.h"
+
+#if USE_ANDROID_FILE
+#include "Misc/App.h"
+#include "Misc/Paths.h"
+#include "Android/AndroidJavaEnv.h"
+
+#include <dirent.h>
+#include <jni.h>
+#include <unistd.h>
+#include <android/asset_manager.h>
+#include <android/asset_manager_jni.h>
+#include <android/storage_manager.h>
+#include "Android/AndroidJava.h"
+#include "Containers/Map.h"
+#include <limits>
+
+DEFINE_LOG_CATEGORY_STATIC(LogAndroidFile, Log, All);
+
+#define LOG_ANDROID_FILE 0
+
+#define LOG_ANDROID_FILE_MANIFEST 0
+
+// Support 64 bit file access.
+//#define UE_ANDROID_FILE_64 0
+#define UE_ANDROID_FILE_64 1
+
+#if UE_ANDROID_FILE_64
+	#define __lseek(_fd, _offset, _whence)			lseek64(_fd, _offset, _whence)
+	#define __pread(_fd, _buf, _count, _offset)		pread64(_fd, _buf, _count, _offset)
+	#define __pwrite(_fd, _buf, _count, _offset)	pwrite64(_fd, _buf, _count, _offset)
+#else
+	#define __lseek(_fd, _offset, _whence)			lseek(_fd, _offset, _whence)
+	#define __pread(_fd, _buf, _count, _offset)		pread(_fd, _buf, _count, _offset)
+	#define __pwrite(_fd, _buf, _count, _offset)	pwrite(_fd, _buf, _count, _offset)
+#endif
+
+// make an FTimeSpan object that represents the "epoch" for time_t (from a stat struct)
+const FDateTime AndroidEpoch(1970, 1, 1);
+
+namespace
+{
+	FFileStatData AndroidStatToUEFileData(struct stat& FileInfo)
+	{
+		const bool bIsDirectory = S_ISDIR(FileInfo.st_mode);
+
+		int64 FileSize = -1;
+		if (!bIsDirectory)
+		{
+			FileSize = FileInfo.st_size;
+		}
+
+		return FFileStatData(
+			AndroidEpoch + FTimespan::FromSeconds(FileInfo.st_ctime), 
+			AndroidEpoch + FTimespan::FromSeconds(FileInfo.st_atime), 
+			AndroidEpoch + FTimespan::FromSeconds(FileInfo.st_mtime), 
+			FileSize,
+			bIsDirectory,
+			!!(FileInfo.st_mode & S_IWUSR)
+		);
+	}
+}
+
+#define USE_UTIME 0
+
+
+// AndroidProcess uses this for executable name
+FString GAndroidProjectName;
+FString GPackageName;
+int32 GAndroidPackageVersion = 0;
+int32 GAndroidPackagePatchVersion = 0;
+FString GAndroidAppType;
+
+// External File Path base - setup during load
+FString GFilePathBase;
+// Obb File Path base - setup during load
+FString GOBBFilePathBase;
+// Internal File Direcory Path (for application) - setup during load
+FString GInternalFilePath;
+// External File Direcory Path (for application) - setup during load
+FString GExternalFilePath;
+// External font path base - setup during load
+FString GFontPathBase;
+
+// Is the OBB in an APK file or not
+bool GOBBinAPK;
+FString GAPKFilename;
+
+#define FILEBASE_DIRECTORY "/UE4Game/"
+
+extern jobject AndroidJNI_GetJavaAssetManager();
+extern AAssetManager * AndroidThunkCpp_GetAssetManager();
+
+//This function is declared in the Java-defined class, GameActivity.java: "public native void nativeSetObbInfo(String PackageName, int Version, int PatchVersion);"
+JNI_METHOD void Java_com_epicgames_ue4_GameActivity_nativeSetObbInfo(JNIEnv* jenv, jobject thiz, jstring ProjectName, jstring PackageName, jint Version, jint PatchVersion, jstring AppType)
+{
+	GAndroidProjectName = FJavaHelper::FStringFromParam(jenv, ProjectName);
+	GPackageName = FJavaHelper::FStringFromParam(jenv, PackageName);
+	GAndroidAppType = FJavaHelper::FStringFromParam(jenv, AppType);
+	
+	GAndroidPackageVersion = Version;
+	GAndroidPackagePatchVersion = PatchVersion;
+}
+
+// Constructs the base path for any files which are not in OBB/pak data
+const FString &GetFileBasePath()
+{
+	static FString BasePath = GFilePathBase + FString(FILEBASE_DIRECTORY) + FApp::GetProjectName() + FString("/");
+	return BasePath;
+}
+
+/**
+	Android file handle implementation for partial (i.e. parcels) files.
+	This can handle all the variations of accessing data for assets and files.
+**/
+class CORE_API FFileHandleAndroid : public IFileHandle
+{
+	enum { READWRITE_SIZE = 1024 * 1024 };
+
+	void LogInfo()
+	{
+#if LOG_ANDROID_FILE
+		FPlatformMisc::LowLevelOutputDebugStringf(
+			TEXT("FFileHandleAndroid => Asset = %p, Handle = %d, Bounds = [%d,%d)"),
+			File->Asset, File->Handle, int32(Start), int32(Start+Length));
+#endif
+	}
+
+public:
+
+	struct FileReference
+	{
+		FString Path;
+		AAsset * Asset;
+		int32 Handle;
+
+		FileReference()
+			: Asset(nullptr), Handle(-1)
+		{
+		}
+
+		FileReference(const FString & path, AAsset * asset)
+			: Path(path), Asset(asset), Handle(0)
+		{
+		}
+
+		FileReference(const FString & path, int32 handle)
+			: Path(path), Asset(nullptr), Handle(handle)
+		{
+		}
+
+		~FileReference()
+		{
+			if (Handle != -1)
+			{
+				close(Handle);
+			}
+			if (nullptr != Asset)
+			{
+				AAsset_close(Asset);
+			}
+		}
+	};
+
+	TSharedPtr<FileReference> File;
+	int64 Start;
+	int64 Length;
+	int64 CurrentOffset;
+
+	FORCEINLINE void CheckValid()
+	{
+		check(File.IsValid() && File->Handle != -1);
+	}
+
+	// Invalid handle.
+	FFileHandleAndroid()
+		: File(MakeShareable(new FileReference()))
+		, Start(0), Length(0), CurrentOffset(0)
+	{
+	}
+
+	// Handle that covers a subsegment of another handle.
+	FFileHandleAndroid(const FFileHandleAndroid & base,
+		int64 start, int64 length)
+		: File(base.File)
+		, Start(base.Start + start), Length(length)
+		, CurrentOffset(base.Start + start)
+	{
+		CheckValid();
+		LogInfo();
+	}
+
+	// Handle that covers a subsegment of provided file.
+	FFileHandleAndroid(const FString & path, int32 filehandle, int64 filestart, int64 filelength)
+		: File(MakeShareable(new FileReference(path, filehandle)))
+		, Start(filestart), Length(filelength), CurrentOffset(0)
+	{
+		CheckValid();
+#if UE_ANDROID_FILE_64
+		lseek64(File->Handle, filestart, SEEK_SET);
+#else
+		lseek(File->Handle, filestart, SEEK_SET);
+#endif
+		LogInfo();
+	}
+
+	// Handle that covers the entire file content.
+	FFileHandleAndroid(const FString & path, int32 filehandle)
+		: File(MakeShareable(new FileReference(path, filehandle)))
+		, Start(0), Length(0), CurrentOffset(0)
+	{
+		CheckValid();
+		Length = __lseek(File->Handle, 0, SEEK_END);
+		__lseek(File->Handle, 0, SEEK_SET);
+		LogInfo();
+	}
+
+	// Handle that covers the entire content of an asset.
+	FFileHandleAndroid(const FString & path, AAsset * asset)
+		: File(MakeShareable(new FileReference(path, asset)))
+		, Start(0), Length(0), CurrentOffset(0)
+	{
+#if UE_ANDROID_FILE_64 && PLATFORM_32BITS
+		File->Handle = AAsset_openFileDescriptor64(File->Asset, &Start, &Length);
+#else
+		off_t OutStart = Start;
+		off_t OutLength = Length;
+		File->Handle = AAsset_openFileDescriptor(File->Asset, &OutStart, &OutLength);
+		Start = OutStart;
+		Length = OutLength;
+#endif
+		CheckValid();
+		LogInfo();
+	}
+
+	virtual ~FFileHandleAndroid()
+	{
+	}
+
+	virtual int64 Tell() override
+	{
+		CheckValid();
+		int64 pos = CurrentOffset;
+		check(pos != -1);
+		return pos - Start; // We are treating 'tell' as a virtual location from file Start
+	}
+
+	virtual bool Seek(int64 NewPosition) override
+	{
+		CheckValid();
+		// we need to offset all positions by the Start offset
+		CurrentOffset = NewPosition += Start;
+		check(NewPosition >= 0);
+		
+		return true;
+	}
+
+    virtual bool SeekFromEnd(int64 NewPositionRelativeToEnd = 0) override
+	{
+		CheckValid();
+		check(NewPositionRelativeToEnd <= 0);
+		// We need to convert this to a virtual offset inside the file we are interested in
+		CurrentOffset = Start + (Length - NewPositionRelativeToEnd);
+		return true;
+	}
+
+	
+
+	virtual bool Read(uint8* Destination, int64 BytesToRead) override
+	{
+		CheckValid();
+#if LOG_ANDROID_FILE
+		FPlatformMisc::LowLevelOutputDebugStringf(
+			TEXT("(%d/%d) FFileHandleAndroid:Read => Path = %s, BytesToRead = %d"),
+			FAndroidTLS::GetCurrentThreadId(), File->Handle,
+			*(File->Path), int32(BytesToRead));
+#endif
+		check(BytesToRead >= 0);
+		check(Destination);
+		
+		while (BytesToRead > 0)
+		{
+
+			int64 ThisSize = FMath::Min<int64>(READWRITE_SIZE, BytesToRead);
+			
+			ThisSize = __pread(File->Handle, Destination, ThisSize, CurrentOffset);
+#if LOG_ANDROID_FILE
+			FPlatformMisc::LowLevelOutputDebugStringf(
+				TEXT("(%d/%d) FFileHandleAndroid:Read => Path = %s, ThisSize = %d, destination = %X"),
+				FAndroidTLS::GetCurrentThreadId(), File->Handle,
+				*(File->Path), int32(ThisSize), Destination);
+#endif
+			if (ThisSize < 0)
+			{
+				return false;
+			}
+			else if (ThisSize == 0)
+			{
+				break;
+			}
+			CurrentOffset += ThisSize;
+			Destination += ThisSize;
+			BytesToRead -= ThisSize;
+		}
+
+		return BytesToRead == 0;
+	}
+
+	virtual bool Write(const uint8* Source, int64 BytesToWrite) override
+	{
+		CheckValid();
+		if (nullptr != File->Asset)
+		{
+			// Can't write to assets.
+			return false;
+		}
+
+		bool bSuccess = true;
+		while (BytesToWrite)
+		{
+			check(BytesToWrite >= 0);
+			int64 ThisSize = FMath::Min<int64>(READWRITE_SIZE, BytesToWrite);
+			check(Source);
+			if (__pwrite(File->Handle, Source, ThisSize, CurrentOffset) != ThisSize)
+			{
+				bSuccess = false;
+				break;
+			}
+			CurrentOffset += ThisSize;
+			Source += ThisSize;
+			BytesToWrite -= ThisSize;
+		}
+		
+		// Update the cached file length
+		Length = FMath::Max(Length, CurrentOffset);
+
+		return bSuccess;
+	}
+
+	virtual bool Flush(const bool bFullFlush = false) override
+	{
+		CheckValid();
+		if (nullptr != File->Asset)
+		{
+			// Can't write to assets.
+			return false;
+		}
+
+		return bFullFlush
+			? fsync(File->Handle) == 0
+			: fdatasync(File->Handle) == 0;
+	}
+
+	virtual bool Truncate(int64 NewSize) override
+	{
+		CheckValid();
+		if (nullptr != File->Asset)
+		{
+			// Can't write to assets.
+			return false;
+		}
+
+		return ftruncate(File->Handle, NewSize) == 0;
+	}
+
+	virtual int64 Size() override
+	{
+		return Length;
+	}
+};
+
+
+class FAndroidFileManifestReader
+{
+private:
+	bool bInitialized;
+	FString ManifestFileName;
+	TMap<FString, FDateTime> ManifestEntries;
+public:
+
+	FAndroidFileManifestReader( const FString& InManifestFileName ) : ManifestFileName(InManifestFileName), bInitialized(false)
+	{
+	}
+
+	bool GetFileTimeStamp( const FString& FileName, FDateTime& DateTime ) 
+	{
+		if ( bInitialized == false )
+		{
+			Read();
+			bInitialized = true;
+		}
+
+		const FDateTime* Result = ManifestEntries.Find( FileName );
+		if ( Result )
+		{
+			DateTime = *Result;
+#if LOG_ANDROID_FILE_MANIFEST
+			FPlatformMisc::LowLevelOutputDebugStringf(TEXT("Found time stamp for '%s', %s"), *FileName, *DateTime.ToString());
+#endif
+			return true;
+		}
+#if LOG_ANDROID_FILE_MANIFEST
+		FPlatformMisc::LowLevelOutputDebugStringf(TEXT("Didn't find time stamp for '%s'"), *FileName);
+#endif
+		return false;
+	}
+
+
+	bool SetFileTimeStamp( const FString& FileName, const FDateTime& DateTime )
+	{
+		if (bInitialized == false)
+		{
+			Read();
+			bInitialized = true;
+		}
+
+		FDateTime* Result = ManifestEntries.Find( FileName );
+		if ( Result == NULL )
+		{
+			ManifestEntries.Add(FileName, DateTime );
+		}
+		else
+		{
+			*Result = DateTime;
+		}
+#if LOG_ANDROID_FILE_MANIFEST
+		FPlatformMisc::LowLevelOutputDebugStringf(TEXT("SetFileTimeStamp '%s', %s"), *FileName, *DateTime.ToString());
+#endif
+		return true;
+	}
+
+	// read manifest from disk
+	void Read()
+	{
+		// Local filepaths are directly in the deployment directory.
+		static const FString &BasePath = GetFileBasePath();
+		const FString ManifestPath = BasePath + ManifestFileName;
+
+		ManifestEntries.Empty();
+
+		// int Handle = open( TCHAR_TO_UTF8(ManifestFileName), O_RDWR );
+		int Handle = open( TCHAR_TO_UTF8(*ManifestPath), O_RDONLY );
+
+		if ( Handle == -1 )
+		{
+#if LOG_ANDROID_FILE_MANIFEST
+			FPlatformMisc::LowLevelOutputDebugStringf(TEXT("Failed to open file for read'%s'"), *ManifestFileName);
+#endif
+			return;
+		}
+
+		FString EntireFile;
+		char Buffer[1024];
+		Buffer[1023] = '\0';
+		int BytesRead = 1023;
+		while ( BytesRead == 1023 )
+		{
+			BytesRead = read(Handle, Buffer, 1023);
+			check( Buffer[1023] == '\0');
+			EntireFile.Append(FString(UTF8_TO_TCHAR(Buffer)));
+		}
+
+		close( Handle );
+
+		TArray<FString> Lines;
+		EntireFile.ParseIntoArrayLines(Lines);
+
+#if LOG_ANDROID_FILE_MANIFEST
+		FPlatformMisc::LowLevelOutputDebugStringf(TEXT("Loaded manifest file %s"), *ManifestFileName);
+		for( const auto& Line : Lines )
+		{
+			FPlatformMisc::LowLevelOutputDebugStringf(TEXT("Manifest Line %s"), *Line );
+		}
+#endif
+		for ( const auto& Line : Lines) 
+		{
+#if LOG_ANDROID_FILE_MANIFEST
+			FPlatformMisc::LowLevelOutputDebugStringf(TEXT("Processing line '%s' "), *Line);
+#endif
+			FString Filename;
+			FString DateTimeString;
+			if ( Line.Split(TEXT("\t"), &Filename, &DateTimeString) )
+			{
+				FDateTime ModifiedDate;
+				if ( FDateTime::ParseIso8601(*DateTimeString, ModifiedDate) )
+				{
+#if LOG_ANDROID_FILE_MANIFEST
+					FPlatformMisc::LowLevelOutputDebugStringf(TEXT("Read time stamp '%s' %s"), *Filename, *ModifiedDate.ToString());
+#endif
+					Filename.ReplaceInline( TEXT("\\"), TEXT("/") );
+					ManifestEntries.Emplace( MoveTemp(Filename), ModifiedDate );
+				}
+				else
+				{
+#if LOG_ANDROID_FILE_MANIFEST
+					FPlatformMisc::LowLevelOutputDebugStringf(TEXT("Failed to parse date for file '%s' %s"), *Filename, *DateTimeString);
+#endif					
+				}
+			}
+#if LOG_ANDROID_FILE_MANIFEST
+			else
+			{
+				FPlatformMisc::LowLevelOutputDebugStringf(TEXT("Unable to split line '%s'"), *Line);
+			}
+#endif
+		}
+	}
+
+	void Write()
+	{
+		
+		// Local filepaths are directly in the deployment directory.
+		static const FString &BasePath = GetFileBasePath();
+		const FString ManifestPath = BasePath + ManifestFileName;
+
+
+		int Handle = open(TCHAR_TO_UTF8(*ManifestPath), O_WRONLY | O_CREAT | O_TRUNC);
+
+		if ( Handle == -1 )
+		{
+#if LOG_ANDROID_FILE_MANIFEST
+			FPlatformMisc::LowLevelOutputDebugStringf(TEXT("Failed to open file for write '%s'"), *ManifestFileName);
+#endif
+			return;
+		}
+
+
+		for ( const auto& Line : ManifestEntries )
+		{
+			const int BufferSize = 4096;
+			char Buffer[BufferSize] = {"\0"}; 
+			const FString RawDateTimeString = Line.Value.ToIso8601();
+			const FString DateTimeString = FString::Printf(TEXT("%s\r\n"), *RawDateTimeString);
+			strncpy(Buffer, TCHAR_TO_UTF8(*Line.Key), BufferSize-1);
+			strncat(Buffer, "\t", BufferSize-1 );
+			strncat(Buffer, TCHAR_TO_UTF8(*DateTimeString), BufferSize-1 );
+			write( Handle, Buffer, strlen( Buffer ) );
+		}
+
+		close( Handle );
+	}
+};
+
+FAndroidFileManifestReader NonUFSManifest(TEXT("Manifest_NonUFSFiles_Android.txt"));
+FAndroidFileManifestReader UFSManifest(TEXT("Manifest_UFSFiles_Android.txt"));
+
+/*
+	Access to files in multiple ZIP archives.
+*/
+class FZipUnionFile
+{
+public:
+	struct Entry
+	{
+		TSharedPtr<FFileHandleAndroid> File;
+		FString FileName;
+		int32 ModTime = 0;
+		bool IsDirectory = false;
+
+		Entry(TSharedPtr<FFileHandleAndroid> file, const FString & filename, int32 modtime, bool isdir)
+			: File(file), FileName(filename), ModTime(modtime), IsDirectory(isdir)
+		{
+		}
+	};
+
+	typedef TMap<FString, TSharedPtr<Entry>> FEntryMap;
+
+	struct Directory
+	{
+		FEntryMap::TIterator Current;
+		FString Path;
+
+		Directory(FEntryMap & entries, const FString & dirpath)
+			: Current(entries.CreateIterator()), Path(dirpath)
+		{
+			if (!Path.IsEmpty())
+			{
+				Path /= "";
+			}
+			// This would be much easier, and efficient, if TMap
+			// supported getting iterators to found entries in
+			// the map. Instead we need to iterate the entire
+			// map to find the initial directory entry.
+			while (Current && Current.Key() != Path)
+			{
+				++Current;
+			}
+		}
+
+		bool Next()
+		{
+			for (++Current; Current; ++Current)
+			{
+				if (Current.Key().StartsWith(Path))
+				{
+					int32 i = Current.Key().Find(TEXT("/"), ESearchCase::CaseSensitive, ESearchDir::FromStart, Path.Len());
+					if (i == INDEX_NONE || i == Current.Key().Len() - 1)
+					{
+						break;
+					}
+				}
+			}
+			return !!Current;
+		}
+	};
+
+	FZipUnionFile()
+	{
+	}
+
+	~FZipUnionFile()
+	{
+	}
+
+	void AddPatchFile(TSharedPtr<FFileHandleAndroid> file)
+	{
+		int64 FileLength = file->Size();
+
+		// Is it big enough to be a ZIP?
+		check( FileLength >= kEOCDLen );
+
+		int64 ReadAmount = kMaxEOCDSearch;
+		if (ReadAmount > FileLength)
+		{
+			ReadAmount = FileLength;
+		}
+
+		// Check magic signature for ZIP.
+		file->Seek(0);
+		uint32 Header;
+		verify( file->Read((uint8*)&Header, sizeof(Header)) );
+		check( Header != kEOCDSignature );
+		check( Header == kLFHSignature );
+
+		/*
+		* Perform the traditional EOCD snipe hunt. We're searching for the End
+		* of Central Directory magic number, which appears at the start of the
+		* EOCD block. It's followed by 18 bytes of EOCD stuff and up to 64KB of
+		* archive comment. We need to read the last part of the file into a
+		* buffer, dig through it to find the magic number, parse some values
+		* out, and use those to determine the extent of the CD. We start by
+		* pulling in the last part of the file.
+		*/
+		int64 SearchStart = FileLength - ReadAmount;
+		ByteBuffer Buffer(ReadAmount);
+		verify( file->Seek(SearchStart) );
+		verify( file->Read(Buffer.Data, Buffer.Size) );
+
+		/*
+		* Scan backward for the EOCD magic. In an archive without a trailing
+		* comment, we'll find it on the first try. (We may want to consider
+		* doing an initial minimal read; if we don't find it, retry with a
+		* second read as above.)
+		*/
+		int64 EOCDIndex = Buffer.Size - kEOCDLen;
+		for (; EOCDIndex >= 0; --EOCDIndex)
+		{
+			if (Buffer.GetValue<uint32>(EOCDIndex) == kEOCDSignature)
+			{
+				break;
+			}
+		}
+		check( EOCDIndex >= 0 );
+
+		/*
+		* Grab the CD offset and size, and the number of entries in the
+		* archive.
+		*/
+		uint16 NumEntries = (Buffer.GetValue<uint16>(EOCDIndex + kEOCDNumEntries));
+		int64 DirSize = (Buffer.GetValue<uint32>(EOCDIndex + kEOCDSize));
+		int64 DirOffset = (Buffer.GetValue<uint32>(EOCDIndex + kEOCDFileOffset));
+		check( DirOffset + DirSize <= FileLength );
+		check( NumEntries > 0 );
+
+		/*
+		* Walk through the central directory, adding entries to the hash table.
+		*/
+		FFileHandleAndroid DirectoryMap(*file, DirOffset, DirSize);
+		int64 Offset = 0;
+		for (uint16 EntryIndex = 0; EntryIndex < NumEntries; ++EntryIndex)
+		{
+			uint32 Signature;
+			verify( DirectoryMap.Seek(Offset) );
+			verify( DirectoryMap.Read((uint8*)&Signature, sizeof(Signature)) );
+
+			// NumEntries may be 65535 so also stop if signature invalid.
+			if (Signature != kCDESignature)
+			{
+				// Hit the end of the central directory, stop.
+				break;
+			}
+
+			// Entry information. Note, we try and read in incremental
+			// order to avoid missing read-aheads.
+
+			uint16 Method;
+			verify( DirectoryMap.Seek(Offset + kCDEMethod) );
+			verify( DirectoryMap.Read((uint8*)&Method, sizeof(Method)) );
+
+			int32 WhenModified;
+			verify( DirectoryMap.Seek(Offset + kCDEModWhen) );
+			verify( DirectoryMap.Read((uint8*)&WhenModified, sizeof(WhenModified)) );
+
+			uint32 UncompressedLength;
+			verify( DirectoryMap.Seek(Offset + kCDEUncompLen) );
+			verify( DirectoryMap.Read((uint8*)&UncompressedLength, sizeof(UncompressedLength)) );
+
+			uint16 FileNameLen;
+			verify( DirectoryMap.Seek(Offset + kCDENameLen) );
+			verify( DirectoryMap.Read((uint8*)&FileNameLen, sizeof(FileNameLen)) );
+
+			uint16 ExtraLen;
+			verify( DirectoryMap.Seek(Offset + kCDEExtraLen) );
+			verify( DirectoryMap.Read((uint8*)&ExtraLen, sizeof(ExtraLen)) );
+
+			uint16 CommentLen;
+			verify( DirectoryMap.Seek(Offset + kCDECommentLen) );
+			verify( DirectoryMap.Read((uint8*)&CommentLen, sizeof(CommentLen)) );
+
+			// We only add uncompressed entries as we don't support decompression.
+			if (Method == kCompressStored)
+			{
+				uint32 LocalOffset;
+				verify( DirectoryMap.Seek(Offset + kCDELocalOffset) );
+				verify( DirectoryMap.Read((uint8*)&LocalOffset, sizeof(LocalOffset)) );
+
+				ByteBuffer FileNameBuffer(FileNameLen+1);
+				verify( DirectoryMap.Seek(Offset + kCDELen) );
+				verify( DirectoryMap.Read(FileNameBuffer.Data, FileNameBuffer.Size) );
+				FileNameBuffer.Data[FileNameBuffer.Size-1] = 0;
+				FString FileName(UTF8_TO_TCHAR(FileNameBuffer.Data));
+
+				verify( file->Seek(LocalOffset) );
+
+				uint32 LocalSignature;
+				verify( file->Read((uint8*)&LocalSignature, sizeof(LocalSignature)) );
+
+				uint16 LocalFileNameLen;
+				verify( file->Seek(LocalOffset + kLFHNameLen) );
+				verify( file->Read((uint8*)&LocalFileNameLen, sizeof(LocalFileNameLen)) );
+
+				uint16 LocalExtraLen;
+				verify( file->Seek(LocalOffset + kLFHExtraLen) );
+				verify( file->Read((uint8*)&LocalExtraLen, sizeof(LocalExtraLen)) );
+
+				int64 EntryOffset = LocalOffset + kLFHLen + LocalFileNameLen + LocalExtraLen;
+
+				// Add the entry.
+#if LOG_ANDROID_FILE
+				FPlatformMisc::LowLevelOutputDebugStringf(TEXT("FUnionZipFile::AddPatchFile.. FILE: '%s'"), *FileName);
+#endif
+				if (FileName.EndsWith(TEXT("/"))) // We only care about actual files in the zip
+				{
+					check(UncompressedLength == 0);
+				}
+				else
+				{
+					Entries.Add(
+						FileName,
+						MakeShareable(new Entry(
+							MakeShareable(new FFileHandleAndroid(
+								*file, EntryOffset, UncompressedLength)),
+							FileName,
+							WhenModified,
+							false)));
+
+					// Add extra directory entries to contain the file
+					// that we can use to later discover directory contents.
+					FileName = FPaths::GetPath(FileName);
+					while (!FileName.IsEmpty())
+					{
+						FString DirName = FileName + "/";
+						if (!Entries.Contains(DirName))
+						{
+#if LOG_ANDROID_FILE
+							FPlatformMisc::LowLevelOutputDebugStringf(TEXT("FUnionZipFile::AddPatchFile.. DIR: '%s'"), *DirName);
+#endif
+							Entries.Add(
+								DirName,
+								MakeShareable(new Entry(
+									nullptr, DirName, 0, true))
+							);
+						}
+						FileName = FPaths::GetPath(FileName);
+					}
+				}
+			}
+
+			// Skip to next entry.
+			Offset += kCDELen + FileNameLen + ExtraLen + CommentLen;
+		}
+
+		// Keep the entries sorted so that we can do iteration to discover
+		// directory contents, and other stuff.
+		Entries.KeySort(TLess<FString>());
+	}
+
+	bool HasEntry(const FString & Path)
+	{
+		return Entries.Contains(Path);
+	}
+
+	Entry & GetEntry(const FString & Path)
+	{
+		return *Entries[Path];
+	}
+
+	int64 GetEntryLength(const FString & Path)
+	{
+		return Entries[Path]->File->Size();
+	}
+
+	int64 GetEntryModTime(const FString & Path)
+	{
+		return Entries[Path]->ModTime;
+	}
+
+	Directory OpenDirectory(const FString & Path)
+	{
+		return Directory(Entries, Path);
+	}
+
+	AAsset * GetEntryAsset(const FString & Path)
+	{
+		return Entries[Path]->File->File->Asset;
+	}
+
+	FString GetEntryRootPath(const FString & Path)
+	{
+		return Entries[Path]->File->File->Path;
+	}
+
+private:
+	FEntryMap Entries;
+
+	// Zip file constants.
+
+	const uint32 kEOCDSignature = 0x06054b50;
+	const uint32 kEOCDLen = 22;
+	const uint32 kEOCDNumEntries = 8; // offset to #of entries in file
+	const uint32 kEOCDSize = 12; // size of the central directory
+	const uint32 kEOCDFileOffset = 16; // offset to central directory
+
+	const uint32 kMaxCommentLen = 65535; // longest possible in ushort
+	const uint32 kMaxEOCDSearch = (kMaxCommentLen + kEOCDLen);
+
+	const uint32 kLFHSignature = 0x04034b50;
+	const uint32 kLFHLen = 30; // excluding variable-len fields
+	const uint32 kLFHNameLen = 26; // offset to filename length
+	const uint32 kLFHExtraLen = 28; // offset to extra length
+
+	const uint32 kCDESignature = 0x02014b50;
+	const uint32 kCDELen = 46; // excluding variable-len fields
+	const uint32 kCDEMethod = 10; // offset to compression method
+	const uint32 kCDEModWhen = 12; // offset to modification timestamp
+	const uint32 kCDECRC = 16; // offset to entry CRC
+	const uint32 kCDECompLen = 20; // offset to compressed length
+	const uint32 kCDEUncompLen = 24; // offset to uncompressed length
+	const uint32 kCDENameLen = 28; // offset to filename length
+	const uint32 kCDEExtraLen = 30; // offset to extra length
+	const uint32 kCDECommentLen = 32; // offset to comment length
+	const uint32 kCDELocalOffset = 42; // offset to local hdr
+
+	const uint32 kCompressStored = 0; // no compression
+	const uint32 kCompressDeflated = 8; // standard deflate
+
+	struct ByteBuffer
+	{
+		uint8* Data;
+		int64 Size;
+
+		ByteBuffer(int64 size)
+			: Data(new uint8[size]), Size(size)
+		{
+		}
+
+		~ByteBuffer()
+		{
+			delete[] Data;
+		}
+
+		template <typename T>
+		T GetValue(int64 offset)
+		{
+			return *reinterpret_cast<T*>(this->Data + offset);
+		}
+	};
+};
+
+// NOTE: Files are stored either loosely in the deployment directory
+// or packed in an OBB archive. We don't know which one unless we try
+// and get the files. We always first check if the files are local,
+// i.e. loosely stored in deployment dir, if they aren't we assume
+// they are archived (and can use the asset system to get them).
+
+/**
+	Implementation for Android file I/O. These handles access to these
+	kinds of files:
+
+	1. Plain-old-files in the file system (i.e. sdcard).
+	2. Resources packed in OBBs (aka ZIPs) placed in download locations.
+	3. Resources packed in OBBs embedded in the APK.
+	4. Direct assets packaged in the APK.
+
+	The base filenames are checked in the above order to allow for
+	overriding content from the most "frozen" to the most "fluid"
+	state. Hence creating a virtual single union file-system.
+**/
+class CORE_API FAndroidPlatformFile : public IAndroidPlatformFile
+{
+public:
+	// Singleton implementation.
+	static FAndroidPlatformFile & GetPlatformPhysical()
+	{
+		static FAndroidPlatformFile AndroidPlatformSingleton;
+		return AndroidPlatformSingleton;
+	}
+
+	FAndroidPlatformFile()
+		: AssetMgr(nullptr)
+	{
+		AssetMgr = AndroidThunkCpp_GetAssetManager();
+	}
+
+	//~ For visibility of overloads we don't override
+	using IAndroidPlatformFile::IterateDirectory;
+	using IAndroidPlatformFile::IterateDirectoryStat;
+
+	// On initialization we search for OBBs that we need to
+	// open to find resources.
+	virtual bool Initialize(IPlatformFile* Inner, const TCHAR* CmdLine) override
+	{
+#if LOG_ANDROID_FILE
+		FPlatformMisc::LowLevelOutputDebugStringf(TEXT("FAndroidPlatformFile::Initialize(..)"));
+#endif
+		if (!IPhysicalPlatformFile::Initialize(Inner, CmdLine))
+		{
+			FPlatformMisc::LowLevelOutputDebugStringf(TEXT("FAndroidPlatformFile::Initialize failed"));
+			return false;
+		}
+		if (GOBBinAPK)
+		{
+			// Open the APK as a ZIP
+			FZipUnionFile APKZip;
+			int32 Handle = open(TCHAR_TO_UTF8(*GAPKFilename), O_RDONLY);
+			if (Handle == -1)
+			{
+				FPlatformMisc::LowLevelOutputDebugStringf(TEXT("FAndroidPlatformFile::Initialize unable to open APK: %s"), *GAPKFilename);
+				return false;
+			}
+			FFileHandleAndroid* APKFile = new FFileHandleAndroid(GAPKFilename, Handle);
+			APKZip.AddPatchFile(MakeShareable(APKFile));
+
+			// Now open the OBB in the APK and mount it
+			if (APKZip.HasEntry("assets/main.obb.png"))
+			{
+				auto OBBEntry = APKZip.GetEntry("assets/main.obb.png");
+				FFileHandleAndroid* OBBFile = static_cast<FFileHandleAndroid*>(new FFileHandleAndroid(*OBBEntry.File, 0, OBBEntry.File->Size()));
+				check(nullptr != OBBFile);
+				ZipResource.AddPatchFile(MakeShareable(OBBFile));
+				FPlatformMisc::LowLevelOutputDebugStringf(TEXT("Mounted OBB in APK: %s"), *GAPKFilename);
+			}
+			else
+			{
+				FPlatformMisc::LowLevelOutputDebugStringf(TEXT("OBB not found in APK: %s"), *GAPKFilename);
+				return false;
+			}
+		}
+		else
+		{
+			// For external OBBs we mount the specific OBB files,
+			// main and patch, only. As required by Android specs.
+			// See <http://developer.android.com/google/play/expansion-files.html>
+			FString OBBDir1 = GOBBFilePathBase + FString(TEXT("/Android/obb/") + GPackageName);
+			FString OBBDir2 = GOBBFilePathBase + FString(TEXT("/obb/") + GPackageName);
+			FString MainOBBName = FString::Printf(TEXT("main.%d.%s.obb"), GAndroidPackageVersion, *GPackageName);
+			FString PatchOBBName = FString::Printf(TEXT("patch.%d.%s.obb"), GAndroidPackageVersion, *GPackageName);
+			if (FileExists(*(OBBDir1 / MainOBBName), true))
+			{
+				MountOBB(*(OBBDir1 / MainOBBName));
+			}
+			else if (FileExists(*(OBBDir2 / MainOBBName), true))
+			{
+				MountOBB(*(OBBDir2 / MainOBBName));
+			}
+			if (FileExists(*(OBBDir1 / PatchOBBName), true))
+			{
+				MountOBB(*(OBBDir1 / PatchOBBName));
+			}
+			else if (FileExists(*(OBBDir2 / PatchOBBName), true))
+			{
+				MountOBB(*(OBBDir2 / PatchOBBName));
+			}
+		}
+
+		// make sure the base path directory exists (UE4Game and UE4Game/ProjectName)
+		FString FileBaseDir = GFilePathBase + FString(FILEBASE_DIRECTORY);
+		mkdir(TCHAR_TO_UTF8(*FileBaseDir), 0766);
+		mkdir(TCHAR_TO_UTF8(*(FileBaseDir + GAndroidProjectName)), 0766);
+
+		return true;
+	}
+
+	virtual bool FileExists(const TCHAR* Filename) override
+	{
+		return FileExists(Filename, false);
+	}
+
+	bool FileExists(const TCHAR* Filename, bool AllowLocal)
+	{
+#if LOG_ANDROID_FILE
+		FPlatformMisc::LowLevelOutputDebugStringf(TEXT("FAndroidPlatformFile::FileExists('%s')"), Filename);
+#endif
+		FString LocalPath;
+		FString AssetPath;
+		PathToAndroidPaths(LocalPath, AssetPath, Filename, AllowLocal);
+
+		bool result = false;
+		struct stat FileInfo;
+		if (!LocalPath.IsEmpty() && (stat(TCHAR_TO_UTF8(*LocalPath), &FileInfo) == 0))
+		{
+			// For local files we need to check if it's a plain
+			// file, as opposed to directories.
+			result = S_ISREG(FileInfo.st_mode);
+		}
+		else
+		{
+			// For everything else we only check existence.
+			result = IsResource(AssetPath) || IsAsset(AssetPath);
+		}
+#if LOG_ANDROID_FILE
+		FPlatformMisc::LowLevelOutputDebugStringf(TEXT("FAndroidPlatformFile::FileExists('%s') => %s\nResolved as %s"), Filename, result ? TEXT("TRUE") : TEXT("FALSE"), *LocalPath);
+#endif
+		return result;
+	}
+
+	virtual int64 FileSize(const TCHAR* Filename) override
+	{
+		return FileSize(Filename, false);
+	}
+
+	int64 FileSize(const TCHAR* Filename, bool AllowLocal)
+	{
+#if LOG_ANDROID_FILE
+		FPlatformMisc::LowLevelOutputDebugStringf(TEXT("FAndroidPlatformFile::FileSize('%s')"), Filename);
+#endif
+		FString LocalPath;
+		FString AssetPath;
+		PathToAndroidPaths(LocalPath, AssetPath, Filename, AllowLocal);
+
+		struct stat FileInfo;
+		FileInfo.st_size = -1;
+		if (!LocalPath.IsEmpty() && (stat(TCHAR_TO_UTF8(*LocalPath), &FileInfo) == 0))
+		{
+			// make sure to return -1 for directories
+			if (S_ISDIR(FileInfo.st_mode))
+			{
+				FileInfo.st_size = -1;
+			}
+			return FileInfo.st_size;
+		}
+		else if (IsResource(AssetPath))
+		{
+			FileInfo.st_size = ZipResource.GetEntryLength(AssetPath);
+		}
+		else
+		{
+			AAsset * file = AAssetManager_open(AssetMgr, TCHAR_TO_UTF8(*AssetPath), AASSET_MODE_UNKNOWN);
+			if (nullptr != file)
+			{
+				FileInfo.st_size = AAsset_getLength(file);
+				AAsset_close(file);
+			}
+		}
+		return FileInfo.st_size;
+	}
+
+	virtual bool DeleteFile(const TCHAR* Filename) override
+	{
+		return DeleteFile(Filename, false);
+	}
+
+	bool DeleteFile(const TCHAR* Filename, bool AllowLocal)
+	{
+#if LOG_ANDROID_FILE
+		FPlatformMisc::LowLevelOutputDebugStringf(TEXT("FAndroidPlatformFile::DeleteFile('%s')"), Filename);
+#endif
+		FString LocalPath;
+		FString AssetPath;
+		PathToAndroidPaths(LocalPath, AssetPath, Filename, AllowLocal);
+
+		// Only delete if we have a local file.
+		if (IsLocal(LocalPath))
+		{
+			return unlink(TCHAR_TO_UTF8(*LocalPath)) == 0;
+		}
+		return false;
+	}
+
+	// NOTE: Returns false if the file is not found.
+	virtual bool IsReadOnly(const TCHAR* Filename) override
+	{
+		return IsReadOnly(Filename, false);
+	}
+
+	bool IsReadOnly(const TCHAR* Filename, bool AllowLocal)
+	{
+#if LOG_ANDROID_FILE
+		FPlatformMisc::LowLevelOutputDebugStringf(TEXT("FAndroidPlatformFile::IsReadOnly('%s')"), Filename);
+#endif
+		FString LocalPath;
+		FString AssetPath;
+		PathToAndroidPaths(LocalPath, AssetPath, Filename, AllowLocal);
+
+		if (IsLocal(LocalPath))
+		{
+			if (access(TCHAR_TO_UTF8(*LocalPath), W_OK) == -1)
+			{
+				return errno == EACCES;
+			}
+		}
+		else
+		{
+			// Anything other than local files are from read-only sources.
+			return IsResource(AssetPath) || IsAsset(AssetPath);
+		}
+		return false;
+	}
+
+	virtual bool MoveFile(const TCHAR* To, const TCHAR* From) override
+	{
+		return MoveFile(To, From, false);
+	}
+
+	bool MoveFile(const TCHAR* To, const TCHAR* From, bool AllowLocal)
+	{
+#if LOG_ANDROID_FILE
+		FPlatformMisc::LowLevelOutputDebugStringf(TEXT("FAndroidPlatformFile::MoveFile('%s', '%s')"), To, From);
+#endif
+		// Can only move local files.
+		FString ToLocalPath;
+		FString ToAssetPath;
+		PathToAndroidPaths(ToLocalPath, ToAssetPath, To, AllowLocal);
+		FString FromLocalPath;
+		FString FromAssetPath;
+		PathToAndroidPaths(FromLocalPath, FromAssetPath, From, AllowLocal);
+
+		if (IsLocal(FromLocalPath))
+		{
+			return rename(TCHAR_TO_UTF8(*FromLocalPath), TCHAR_TO_UTF8(*ToLocalPath)) != -1;
+		}
+		return false;
+	}
+
+	virtual bool SetReadOnly(const TCHAR* Filename, bool bNewReadOnlyValue) override
+	{
+		return SetReadOnly(Filename, bNewReadOnlyValue, false);
+	}
+
+	bool SetReadOnly(const TCHAR* Filename, bool bNewReadOnlyValue, bool AllowLocal)
+	{
+#if LOG_ANDROID_FILE
+		FPlatformMisc::LowLevelOutputDebugStringf(TEXT("FAndroidPlatformFile::SetReadOnly('%s')"), Filename);
+#endif
+		FString LocalPath;
+		FString AssetPath;
+		PathToAndroidPaths(LocalPath, AssetPath, Filename, AllowLocal);
+
+		if (IsLocal(LocalPath))
+		{
+			struct stat FileInfo;
+			if (stat(TCHAR_TO_UTF8(*LocalPath), &FileInfo) != -1)
+			{
+				if (bNewReadOnlyValue)
+				{
+					FileInfo.st_mode &= ~S_IWUSR;
+				}
+				else
+				{
+					FileInfo.st_mode |= S_IWUSR;
+				}
+				return chmod(TCHAR_TO_UTF8(*LocalPath), FileInfo.st_mode) == 0;
+			}
+		}
+		return false;
+	}
+
+	virtual FDateTime GetTimeStamp(const TCHAR* Filename) override
+	{
+		return GetTimeStamp(Filename, false);
+	}
+
+	FDateTime GetTimeStamp(const TCHAR* Filename, bool AllowLocal)
+	{
+#if LOG_ANDROID_FILE_MANIFEST
+		FPlatformMisc::LowLevelOutputDebugStringf(TEXT("FAndroidPlatformFile::GetTimeStamp('%s')"), Filename);
+#endif
+		FString LocalPath;
+		FString AssetPath;
+		PathToAndroidPaths(LocalPath, AssetPath, Filename, AllowLocal);
+
+		if (IsLocal(LocalPath))
+		{
+
+#if USE_UTIME
+			struct stat FileInfo;
+			if (stat(TCHAR_TO_UTF8(*LocalPath), &FileInfo) == -1)
+			{
+				return FDateTime::MinValue();
+			}
+			// convert _stat time to FDateTime
+			FTimespan TimeSinceEpoch(0, 0, FileInfo.st_mtime);
+			return AndroidEpoch + TimeSinceEpoch;
+#else
+			FDateTime Result;
+			if ( NonUFSManifest.GetFileTimeStamp(AssetPath,Result) )
+			{
+				return Result;
+			}
+
+			if ( UFSManifest.GetFileTimeStamp( AssetPath, Result ) )
+			{
+				return Result;
+			}
+
+#if LOG_ANDROID_FILE_MANIFEST
+			FPlatformMisc::LowLevelOutputDebugStringf(TEXT("Failed to find time stamp in NonUFSManifest for file '%s'"), Filename);
+#endif
+
+			// pak file outside of obb may not be in manifest so check if it exists
+			if (AssetPath.EndsWith(".pak"))
+			{
+				// return local file access timestamp (if exists)
+				return GetAccessTimeStamp(Filename, true);
+			}
+
+			return FDateTime::MinValue();
+
+#endif
+		}
+		else if (IsResource(AssetPath))
+		{
+			FTimespan TimeSinceEpoch(0, 0, ZipResource.GetEntryModTime(AssetPath));
+			return AndroidEpoch + TimeSinceEpoch;
+		}
+		else
+		{
+			// No TimeStamp for assets, so just return a default timespan for now.
+			return FDateTime::MinValue();
+		}
+	}
+
+	virtual void SetTimeStamp(const TCHAR* Filename, const FDateTime DateTime) override
+	{
+		return SetTimeStamp(Filename, DateTime, false);
+	}
+	
+	void SetTimeStamp(const TCHAR* Filename, const FDateTime DateTime, bool AllowLocal)
+	{
+#if LOG_ANDROID_FILE
+		FPlatformMisc::LowLevelOutputDebugStringf(TEXT("FAndroidPlatformFile::SetTimeStamp('%s')"), Filename);
+#endif
+		FString LocalPath;
+		FString AssetPath;
+		PathToAndroidPaths(LocalPath, AssetPath, Filename, AllowLocal);
+
+		// Can only set time stamp on local files
+		if (IsLocal(LocalPath))
+		{
+#if USE_UTIME
+			// get file times
+			struct stat FileInfo;
+			if (stat(TCHAR_TO_UTF8(*LocalPath), &FileInfo) == -1)
+			{
+				return;
+			}
+			// change the modification time only
+			struct utimbuf Times;
+			Times.actime = FileInfo.st_atime;
+			Times.modtime = (DateTime - AndroidEpoch).GetTotalSeconds();
+			utime(TCHAR_TO_UTF8(*LocalPath), &Times);
+#else
+			// do something better as utime isn't supported on android very well...
+			FDateTime TempDateTime;
+			if ( NonUFSManifest.GetFileTimeStamp( AssetPath, TempDateTime ) )
+			{
+				NonUFSManifest.SetFileTimeStamp( AssetPath, DateTime );
+				NonUFSManifest.Write();
+			}
+			else
+			{
+				UFSManifest.SetFileTimeStamp( AssetPath, DateTime );
+				UFSManifest.Write();
+			}
+#endif
+		}
+	}
+
+	virtual FDateTime GetAccessTimeStamp(const TCHAR* Filename) override
+	{
+		return GetAccessTimeStamp(Filename, false);
+	}
+
+	FDateTime GetAccessTimeStamp(const TCHAR* Filename, bool AllowLocal)
+	{
+#if LOG_ANDROID_FILE
+		FPlatformMisc::LowLevelOutputDebugStringf(TEXT("FAndroidPlatformFile::GetAccessTimeStamp('%s')"), Filename);
+#endif
+		FString LocalPath;
+		FString AssetPath;
+		PathToAndroidPaths(LocalPath, AssetPath, Filename, AllowLocal);
+
+		if (IsLocal(LocalPath))
+		{
+			struct stat FileInfo;
+			if (stat(TCHAR_TO_UTF8(*LocalPath), &FileInfo) == -1)
+			{
+				return FDateTime::MinValue();
+			}
+			// convert _stat time to FDateTime
+			FTimespan TimeSinceEpoch(0, 0, FileInfo.st_atime);
+			return AndroidEpoch + TimeSinceEpoch;
+		}
+		else
+		{
+			// No TimeStamp for resources nor assets, so just return a default timespan for now.
+			return FDateTime::MinValue();
+		}
+	}
+
+	virtual FFileStatData GetStatData(const TCHAR* FilenameOrDirectory) override
+	{
+		return GetStatData(FilenameOrDirectory, false);
+	}
+
+	FFileStatData GetStatData(const TCHAR* FilenameOrDirectory, bool AllowLocal)
+	{
+#if LOG_ANDROID_FILE
+		FPlatformMisc::LowLevelOutputDebugStringf(TEXT("FAndroidPlatformFile::GetStatData('%s')"), FilenameOrDirectory);
+#endif
+		FString LocalPath;
+		FString AssetPath;
+		PathToAndroidPaths(LocalPath, AssetPath, FilenameOrDirectory, AllowLocal);
+
+		if (IsLocal(LocalPath))
+		{
+			struct stat FileInfo;
+			if (stat(TCHAR_TO_UTF8(*LocalPath), &FileInfo) != -1)
+			{
+				return AndroidStatToUEFileData(FileInfo);
+			}
+		}
+		else if (IsResource(AssetPath))
+		{
+			return FFileStatData(
+				FDateTime::MinValue(),					// CreationTime
+				FDateTime::MinValue(),					// AccessTime
+				FDateTime::MinValue(),					// ModificationTime
+				ZipResource.GetEntryLength(AssetPath),	// FileSize
+				false,									// bIsDirectory
+				true									// bIsReadOnly
+				);
+		}
+		else
+		{
+			AAsset * file = AAssetManager_open(AssetMgr, TCHAR_TO_UTF8(*AssetPath), AASSET_MODE_UNKNOWN);
+			if (nullptr != file)
+			{
+				bool isDirectory = false;
+				AAssetDir * subdir = AAssetManager_openDir(AssetMgr, TCHAR_TO_UTF8(*AssetPath));
+				if (nullptr != subdir)
+				{
+					isDirectory = true;
+					AAssetDir_close(subdir);
+				}
+
+				int64 FileSize = -1;
+				if (!isDirectory)
+				{
+					FileSize = AAsset_getLength(file);
+				}
+
+				FFileStatData StatData(
+					FDateTime::MinValue(),				// CreationTime
+					FDateTime::MinValue(),				// AccessTime
+					FDateTime::MinValue(),				// ModificationTime
+					FileSize,							// FileSize
+					isDirectory,						// bIsDirectory
+					true								// bIsReadOnly
+					);
+
+				AAsset_close(file);
+
+				return StatData;
+			}
+		}
+		
+		return FFileStatData();
+	}
+
+	virtual FString GetFilenameOnDisk(const TCHAR* Filename) override
+	{
+		return Filename;
+	}
+
+	virtual IFileHandle* OpenRead(const TCHAR* Filename, bool bAllowWrite = false) override
+	{
+		return OpenRead(Filename, false, bAllowWrite);
+	}
+
+	IFileHandle* OpenRead(const TCHAR* Filename, bool AllowLocal, bool bAllowWrite)
+	{
+#if LOG_ANDROID_FILE
+		FPlatformMisc::LowLevelOutputDebugStringf(TEXT("FAndroidPlatformFile::OpenRead('%s')"), Filename);
+#endif
+		FString LocalPath;
+		FString AssetPath;
+		PathToAndroidPaths(LocalPath, AssetPath, Filename, AllowLocal);
+
+		if (IsLocal(LocalPath))
+		{
+			int32 Handle = open(TCHAR_TO_UTF8(*LocalPath), O_RDONLY);
+			if (Handle != -1)
+			{
+				return new FFileHandleAndroid(LocalPath, Handle);
+			}
+		}
+		else if (IsResource(AssetPath))
+		{
+			return new FFileHandleAndroid(
+				*ZipResource.GetEntry(AssetPath).File,
+				0, ZipResource.GetEntry(AssetPath).File->Size());
+		}
+		else
+		{
+			AAsset * asset = AAssetManager_open(AssetMgr, TCHAR_TO_UTF8(*AssetPath), AASSET_MODE_RANDOM);
+			if (nullptr != asset)
+			{
+				return new FFileHandleAndroid(AssetPath, asset);
+			}
+		}
+		return nullptr;
+	}
+
+	// Regardless of the file being local, asset, or resource, we
+	// assert that opening a file for write will open a local file.
+	// The intent is to allow creating fresh files that override
+	// packaged content.
+	virtual IFileHandle* OpenWrite(const TCHAR* Filename, bool bAppend, bool bAllowRead) override
+	{
+		return OpenWrite(Filename, bAppend, bAllowRead, false);
+	}
+
+	IFileHandle* OpenWrite(const TCHAR* Filename, bool bAppend, bool bAllowRead, bool AllowLocal)
+	{
+#if LOG_ANDROID_FILE
+		FPlatformMisc::LowLevelOutputDebugStringf(TEXT("FAndroidPlatformFile::OpenWrite('%s')"), Filename);
+#endif
+		FString LocalPath;
+		FString AssetPath;
+		PathToAndroidPaths(LocalPath, AssetPath, Filename, AllowLocal);
+
+		int Flags = O_CREAT;
+		if (!bAppend)
+		{
+			Flags |= O_TRUNC;
+		}
+		if (bAllowRead)
+		{
+			Flags |= O_RDWR;
+		}
+		else
+		{
+			Flags |= O_WRONLY;
+		}
+
+		int32 Handle = open(TCHAR_TO_UTF8(*LocalPath), Flags, S_IRUSR | S_IWUSR);
+		if (Handle != -1)
+		{
+			FFileHandleAndroid* FileHandleAndroid = new FFileHandleAndroid(LocalPath, Handle);
+			if (bAppend)
+			{
+				FileHandleAndroid->SeekFromEnd(0);
+			}
+			return FileHandleAndroid;
+		}
+		return nullptr;
+	}
+
+	virtual bool DirectoryExists(const TCHAR* Directory) override
+	{
+		return DirectoryExists(Directory, false, false);
+	}
+	
+	bool DirectoryExists(const TCHAR* Directory, bool AllowLocal, bool AllowAsset)
+	{
+#if LOG_ANDROID_FILE
+		FPlatformMisc::LowLevelOutputDebugStringf(TEXT("FAndroidPlatformFile::DirectoryExists('%s')"), Directory);
+#endif
+		FString LocalPath;
+		FString AssetPath;
+		PathToAndroidPaths(LocalPath, AssetPath, Directory, AllowLocal);
+
+		bool Found = false;
+		if (IsLocal(LocalPath))
+		{
+#if LOG_ANDROID_FILE
+			FPlatformMisc::LowLevelOutputDebugStringf(TEXT("FAndroidPlatformFile::DirectoryExists('%s') => Check IsLocal: '%s'"), Directory, *(LocalPath + "/"));
+#endif
+			struct stat FileInfo;
+			if (stat(TCHAR_TO_UTF8(*LocalPath), &FileInfo) != -1)
+			{
+				Found = S_ISDIR(FileInfo.st_mode);
+			}
+		}
+		else if (IsResource(AssetPath + "/"))
+		{
+			Found = true;
+#if LOG_ANDROID_FILE
+			FPlatformMisc::LowLevelOutputDebugStringf(TEXT("FAndroidPlatformFile::DirectoryExists('%s') => Found as resource: '%s'"), Directory, *(AssetPath + "/"));
+#endif
+		}
+		else if (AllowAsset)
+		{
+			// This can be very slow on some devices so only do it if we enabled
+			AAssetDir * dir = AAssetManager_openDir(AssetMgr, TCHAR_TO_UTF8(*AssetPath));
+			Found = AAssetDir_getNextFileName(dir) != nullptr;
+			AAssetDir_close(dir);
+#if LOG_ANDROID_FILE
+			if (Found)
+			{
+				FPlatformMisc::LowLevelOutputDebugStringf(TEXT("FAndroidPlatformFile::DirectoryExists('%s') => Found as asset: '%s'"), Directory, *(AssetPath));
+			}
+#endif
+		}
+#if LOG_ANDROID_FILE
+		if (Found)
+		{
+			FPlatformMisc::LowLevelOutputDebugStringf(TEXT("FAndroidPlatformFile::DirectoryExists('%s') => FOUND"), Directory);
+		}
+		else
+		{
+			FPlatformMisc::LowLevelOutputDebugStringf(TEXT("FAndroidPlatformFile::DirectoryExists('%s') => NOT"), Directory);
+		}
+#endif
+		return Found;
+	}
+
+	// We assert that created dirs are in the local file-system.
+	virtual bool CreateDirectory(const TCHAR* Directory) override
+	{
+		return CreateDirectory(Directory, false);
+	}
+
+	bool CreateDirectory(const TCHAR* Directory, bool AllowLocal)
+	{
+#if LOG_ANDROID_FILE
+		FPlatformMisc::LowLevelOutputDebugStringf(TEXT("FAndroidPlatformFile::CreateDirectory('%s')"), Directory);
+#endif
+		FString LocalPath;
+		FString AssetPath;
+		PathToAndroidPaths(LocalPath, AssetPath, Directory, AllowLocal);
+
+		return (mkdir(TCHAR_TO_UTF8(*LocalPath), 0755) == 0) || (errno == EEXIST);
+	}
+
+	// We assert that modifying dirs are in the local file-system.
+	virtual bool DeleteDirectory(const TCHAR* Directory) override
+	{
+		return DeleteDirectory(Directory, false);
+	}
+
+	bool DeleteDirectory(const TCHAR* Directory, bool AllowLocal)
+	{
+#if LOG_ANDROID_FILE
+		FPlatformMisc::LowLevelOutputDebugStringf(TEXT("FAndroidPlatformFile::DeleteDirectory('%s')"), Directory);
+#endif
+		FString LocalPath;
+		FString AssetPath;
+		PathToAndroidPaths(LocalPath, AssetPath, Directory, AllowLocal);
+
+		return rmdir(TCHAR_TO_UTF8(*LocalPath));
+	}
+
+	virtual bool IterateDirectory(const TCHAR* Directory, FDirectoryVisitor& Visitor) override
+	{
+		return IterateDirectory(Directory, Visitor, false, false);
+	}
+
+	bool IterateDirectory(const TCHAR* Directory, FDirectoryVisitor& Visitor, bool AllowLocal, bool AllowAsset)
+	{
+		const FString DirectoryStr = Directory;
+
+		auto InternalVisitor = [&](const FString& InLocalPath, struct dirent* InEntry) -> bool
+		{
+			const FString DirPath = DirectoryStr / UTF8_TO_TCHAR(InEntry->d_name);
+			return Visitor.Visit(*DirPath, InEntry->d_type == DT_DIR);
+		};
+		
+		auto InternalResourceVisitor = [&](const FString& InResourceName, bool IsDirectory) -> bool
+		{
+			return Visitor.Visit(*InResourceName, IsDirectory);
+		};
+		
+		auto InternalAssetVisitor = [&](const char* InAssetPath) -> bool
+		{
+			bool isDirectory = false;
+			AAssetDir* subdir = AAssetManager_openDir(AssetMgr, InAssetPath);
+			if (nullptr != subdir)
+			{
+				isDirectory = true;
+				AAssetDir_close(subdir);
+			}
+
+			return Visitor.Visit(UTF8_TO_TCHAR(InAssetPath), isDirectory);
+		};
+
+		return IterateDirectoryCommon(Directory, InternalVisitor, InternalResourceVisitor, InternalAssetVisitor, AllowLocal, AllowAsset);
+	}
+
+	virtual bool IterateDirectoryStat(const TCHAR* Directory, FDirectoryStatVisitor& Visitor) override
+	{
+		return IterateDirectoryStat(Directory, Visitor, false, false);
+	}
+
+	bool IterateDirectoryStat(const TCHAR* Directory, FDirectoryStatVisitor& Visitor, bool AllowLocal, bool AllowAsset)
+	{
+		const FString DirectoryStr = Directory;
+
+		auto InternalVisitor = [&](const FString& InLocalPath, struct dirent* InEntry) -> bool
+		{
+			const FString DirPath = DirectoryStr / UTF8_TO_TCHAR(InEntry->d_name);
+
+			struct stat FileInfo;
+			if (stat(TCHAR_TO_UTF8(*(InLocalPath / UTF8_TO_TCHAR(InEntry->d_name))), &FileInfo) != -1)
+			{
+				return Visitor.Visit(*DirPath, AndroidStatToUEFileData(FileInfo));
+			}
+
+			return true;
+		};
+		
+		auto InternalResourceVisitor = [&](const FString& InResourceName, bool IsDir) -> bool
+		{
+			return Visitor.Visit(
+				*InResourceName, 
+				FFileStatData(
+					FDateTime::MinValue(),						// CreationTime
+					FDateTime::MinValue(),						// AccessTime
+					FDateTime::MinValue(),						// ModificationTime
+					ZipResource.GetEntryLength(InResourceName),	// FileSize
+					IsDir,										// bIsDirectory
+					true										// bIsReadOnly
+					)
+				);
+		};
+		
+		auto InternalAssetVisitor = [&](const char* InAssetPath) -> bool
+		{
+			bool isDirectory = false;
+			AAssetDir* subdir = AAssetManager_openDir(AssetMgr, InAssetPath);
+			if (nullptr != subdir)
+			{
+				isDirectory = true;
+				AAssetDir_close(subdir);
+			}
+
+			int64 FileSize = -1;
+			if (!isDirectory)
+			{
+				AAsset* file = AAssetManager_open(AssetMgr, InAssetPath, AASSET_MODE_UNKNOWN);
+				FileSize = AAsset_getLength(file);
+				AAssetDir_close(subdir);
+			}
+
+			return Visitor.Visit(
+				UTF8_TO_TCHAR(InAssetPath), 
+				FFileStatData(
+					FDateTime::MinValue(),	// CreationTime
+					FDateTime::MinValue(),	// AccessTime
+					FDateTime::MinValue(),	// ModificationTime
+					FileSize,				// FileSize
+					isDirectory,			// bIsDirectory
+					true					// bIsReadOnly
+					)
+				);
+		};
+
+		return IterateDirectoryCommon(Directory, InternalVisitor, InternalResourceVisitor, InternalAssetVisitor, AllowLocal, AllowAsset);
+	}
+
+	bool IterateDirectoryCommon(const TCHAR* Directory, const TFunctionRef<bool(const FString&, struct dirent*)>& Visitor, const TFunctionRef<bool(const FString&, bool)>& ResourceVisitor, const TFunctionRef<bool(const char*)>& AssetVisitor, bool AllowLocal, bool AllowAsset)
+	{
+#if LOG_ANDROID_FILE
+		FPlatformMisc::LowLevelOutputDebugStringf(TEXT("FAndroidPlatformFile::IterateDirectory('%s')"), Directory);
+#endif
+		FString LocalPath;
+		FString AssetPath;
+		PathToAndroidPaths(LocalPath, AssetPath, Directory, AllowLocal);
+
+		if (IsLocal(LocalPath))
+		{
+			DIR* Handle = opendir(TCHAR_TO_UTF8(*LocalPath));
+			if (Handle)
+			{
+				bool Result = true;
+				struct dirent *Entry;
+				while ((Entry = readdir(Handle)) != nullptr && Result == true)
+				{
+					if (FCString::Strcmp(UTF8_TO_TCHAR(Entry->d_name), TEXT(".")) && FCString::Strcmp(UTF8_TO_TCHAR(Entry->d_name), TEXT("..")))
+					{
+#if LOG_ANDROID_FILE
+						FPlatformMisc::LowLevelOutputDebugStringf(TEXT("FAndroidPlatformFile::IterateDirectory('%s').. LOCAL Visit: '%s'"), Directory, *(FString(Directory) / UTF8_TO_TCHAR(Entry->d_name)));
+#endif
+						Result = Visitor(LocalPath, Entry);
+					}
+				}
+				closedir(Handle);
+				return Result;
+			}
+		}
+		else if (IsResource(AssetPath))
+		{
+			FZipUnionFile::Directory ResourceDir = ZipResource.OpenDirectory(AssetPath);
+			bool Result = true;
+			while (Result && ResourceDir.Next())
+			{
+#if LOG_ANDROID_FILE
+				FPlatformMisc::LowLevelOutputDebugStringf(TEXT("FAndroidPlatformFile::IterateDirectory('%s').. RESOURCE Visit: '%s'"), Directory, *ResourceDir.Current.Key());
+#endif
+				FString ResourcePath = ResourceDir.Current.Key();
+				bool IsDirectory = ResourceDir.Current.Value()->IsDirectory;
+				if (IsDirectory && ResourcePath.EndsWith(TEXT("/"), ESearchCase::CaseSensitive))
+				{
+					ResourcePath.GetCharArray()[ResourcePath.Len() - 1] = 0;
+					ResourcePath.TrimToNullTerminator();
+				}
+				Result = ResourceVisitor(ResourcePath, IsDirectory);
+			}
+			return Result;
+		}
+		else if (IsResource(AssetPath + "/"))
+		{
+			FZipUnionFile::Directory ResourceDir = ZipResource.OpenDirectory(AssetPath + "/");
+			bool Result = true;
+			while (Result && ResourceDir.Next())
+			{
+#if LOG_ANDROID_FILE
+				FPlatformMisc::LowLevelOutputDebugStringf(TEXT("FAndroidPlatformFile::IterateDirectory('%s').. RESOURCE/ Visit: '%s'"), Directory, *ResourceDir.Current.Key());
+#endif
+				FString ResourcePath = ResourceDir.Current.Key();
+				bool IsDirectory = ResourceDir.Current.Value()->IsDirectory;
+				if (IsDirectory && ResourcePath.EndsWith(TEXT("/"), ESearchCase::CaseSensitive))
+				{
+					ResourcePath.GetCharArray()[ResourcePath.Len() - 1] = 0;
+					ResourcePath.TrimToNullTerminator();
+				}
+				Result = ResourceVisitor(ResourcePath, IsDirectory);
+			}
+			return Result;
+		}
+		else if (AllowAsset)
+		{
+			AAssetDir * dir = AAssetManager_openDir(AssetMgr, TCHAR_TO_UTF8(*AssetPath));
+			if (nullptr != dir)
+			{
+				bool Result = true;
+				const char * fileName = nullptr;
+				while ((fileName = AAssetDir_getNextFileName(dir)) != nullptr && Result == true)
+				{
+#if LOG_ANDROID_FILE
+					FPlatformMisc::LowLevelOutputDebugStringf(TEXT("FAndroidPlatformFile::IterateDirectory('%s').. ASSET Visit: '%s'"), Directory, UTF8_TO_TCHAR(fileName));
+#endif
+					Result = AssetVisitor(fileName);
+				}
+				AAssetDir_close(dir);
+				return Result;
+			}
+		}
+		return false;
+	}
+
+	virtual jobject GetAssetManager() override
+	{
+		return AndroidJNI_GetJavaAssetManager();
+	}
+
+	virtual bool IsAsset(const TCHAR* Filename) override
+	{
+#if LOG_ANDROID_FILE
+		FPlatformMisc::LowLevelOutputDebugStringf(TEXT("FAndroidPlatformFile::FileIsAsset('%s')"), Filename);
+#endif
+		FString LocalPath;
+		FString AssetPath;
+		PathToAndroidPaths(LocalPath, AssetPath, Filename, true);
+
+		if (IsLocal(LocalPath))
+		{
+			return false;
+		}
+		else if (IsResource(AssetPath))
+		{
+			return ZipResource.GetEntryAsset(AssetPath) != nullptr;
+		}
+		else if (IsAsset(AssetPath))
+		{
+			return true;
+		}
+		return false;
+	}
+
+	virtual int64 FileStartOffset(const TCHAR* Filename) override
+	{
+#if LOG_ANDROID_FILE
+		FPlatformMisc::LowLevelOutputDebugStringf(TEXT("FAndroidPlatformFile::FileStartOffset('%s')"), Filename);
+#endif
+		FString LocalPath;
+		FString AssetPath;
+		PathToAndroidPaths(LocalPath, AssetPath, Filename, true);
+
+		if (IsLocal(LocalPath))
+		{
+			return 0;
+		}
+		else if (IsResource(AssetPath))
+		{
+			return ZipResource.GetEntry(AssetPath).File->Start;
+		}
+		else if (IsAsset(AssetPath))
+		{
+			AAsset * file = AAssetManager_open(AssetMgr, TCHAR_TO_UTF8(*AssetPath), AASSET_MODE_UNKNOWN);
+			if (nullptr != file)
+			{
+				off_t start = -1;
+				off_t length = -1;
+				int handle = AAsset_openFileDescriptor(file, &start, &length);
+				if (handle != -1)
+				{
+					close(handle);
+				}
+				AAsset_close(file);
+				return start;
+			}
+		}
+		return -1;
+	}
+
+	virtual FString FileRootPath(const TCHAR* Filename) override
+	{
+#if LOG_ANDROID_FILE
+		FPlatformMisc::LowLevelOutputDebugStringf(TEXT("FAndroidPlatformFile::FileRootPath('%s')"), Filename);
+#endif
+		FString LocalPath;
+		FString AssetPath;
+		PathToAndroidPaths(LocalPath, AssetPath, Filename, true);
+
+		if (IsLocal(LocalPath))
+		{
+			return LocalPath;
+		}
+		else if (IsResource(AssetPath))
+		{
+			return ZipResource.GetEntryRootPath(AssetPath);
+		}
+		else if (IsAsset(AssetPath))
+		{
+			return AssetPath;
+		}
+		return "";
+	}
+
+
+private:
+	FString NormalizePath(const TCHAR* Path)
+	{
+		FString Result(Path);
+		Result.ReplaceInline(TEXT("\\"), TEXT("/"));
+		// This replacement addresses a "bug" where some callers
+		// pass in paths that are badly composed with multiple
+		// subdir separators.
+		Result.ReplaceInline(TEXT("//"), TEXT("/"));
+		if (!Result.IsEmpty() && Result[Result.Len() - 1] == TEXT('/'))
+		{
+			Result.LeftChop(1);
+		}
+		// Remove redundant current-dir references.
+		Result.ReplaceInline(TEXT("/./"), TEXT("/"));
+		return Result;
+	}
+
+	void PathToAndroidPaths(FString & LocalPath, FString & AssetPath, const TCHAR* Path, bool AllowLocal)
+	{
+		LocalPath.Empty();
+		AssetPath.Empty();
+
+		FString AndroidPath = NormalizePath(Path);
+#if LOG_ANDROID_FILE
+		FPlatformMisc::LowLevelOutputDebugStringf(TEXT("FAndroidPlatformFile::PathToAndroidPaths('%s') => AndroidPath = '%s'"), Path, *AndroidPath);
+#endif
+		if (!AndroidPath.IsEmpty())
+		{
+			if (AndroidPath.StartsWith(TEXT("/")) ||	// (AllowLocal && AndroidPath.StartsWith(TEXT("/"))) ||		// requiring AllowLocal isn't really useful; system will do permission checks anyway
+				AndroidPath.StartsWith(GFontPathBase) ||
+				AndroidPath.StartsWith(TEXT("/system/etc/")) ||
+				AndroidPath.StartsWith(GInternalFilePath.Left(AndroidPath.Len())) ||
+				AndroidPath.StartsWith(GExternalFilePath.Left(AndroidPath.Len())))
+			{
+				// Absolute paths are only local.
+				LocalPath = AndroidPath;
+				AssetPath = AndroidPath;
+			}
+			else
+			{
+				while (AndroidPath.StartsWith(TEXT("../")))
+				{
+					AndroidPath = AndroidPath.RightChop(3);
+				}
+				AndroidPath.ReplaceInline(FPlatformProcess::BaseDir(), TEXT(""));
+				if (AndroidPath.Equals(TEXT("..")))
+				{
+					AndroidPath = TEXT("");
+				}
+
+				// Local filepaths are directly in the deployment directory.
+				static FString BasePath = GetFileBasePath();
+				LocalPath = BasePath + AndroidPath;
+
+				// Asset paths are relative to the base directory.
+				AssetPath = AndroidPath;
+			}
+		}
+
+#if LOG_ANDROID_FILE
+		FPlatformMisc::LowLevelOutputDebugStringf(TEXT("FAndroidPlatformFile::PathToAndroidPaths('%s') => LocalPath = '%s'"), Path, *LocalPath);
+		FPlatformMisc::LowLevelOutputDebugStringf(TEXT("FAndroidPlatformFile::PathToAndroidPaths('%s') => AssetPath = '%s'"), Path, *AssetPath);
+#endif
+	}
+
+	// There is a limited set of paths we allow local file access to.
+	// We filter out non-permitted local paths here unless explicitly
+	// allowed. The current two cases are direct font access and with
+	// the AllowLocal flag (as used to mount local OBBs). Eventually
+	// we may need to also allow access to downloads somehow!
+	bool IsLocal(const FString & LocalPath)
+	{
+		return !LocalPath.IsEmpty() && access(TCHAR_TO_UTF8(*LocalPath), F_OK) == 0;
+	}
+
+	bool IsAsset(const FString & AssetPath)
+	{
+		AAsset * file = AAssetManager_open(AssetMgr, TCHAR_TO_UTF8(*AssetPath), AASSET_MODE_UNKNOWN);
+		if (nullptr != file)
+		{
+			AAsset_close(file);
+			return true;
+		}
+		return false;
+	}
+
+	bool IsResource(const FString & ResourcePath)
+	{
+		return ZipResource.HasEntry(ResourcePath);
+	}
+
+	class FMountOBBVisitor : public IPlatformFile::FDirectoryVisitor
+	{
+		FAndroidPlatformFile & AndroidPlatformFile;
+
+	public:
+		FMountOBBVisitor(FAndroidPlatformFile & APF)
+			: AndroidPlatformFile(APF)
+		{
+		}
+
+		virtual ~FMountOBBVisitor()
+		{
+		}
+
+		virtual bool Visit(const TCHAR* FilenameOrDirectory, bool bIsDirectory) override
+		{
+			if (FString(FilenameOrDirectory).EndsWith(TEXT(".obb")) ||
+				FString(FilenameOrDirectory).EndsWith(TEXT(".obb.png")))
+			{
+				// It's and OBB (actually a ZIP) so we fake mount it.
+				AndroidPlatformFile.MountOBB(FilenameOrDirectory);
+			}
+			return true;
+		}
+	};
+
+	void MountOBB(const TCHAR* Filename)
+	{
+		FFileHandleAndroid* File
+			= static_cast<FFileHandleAndroid*>(OpenRead(Filename, true, false));
+		check(nullptr != File);
+		ZipResource.AddPatchFile(MakeShareable(File));
+		FPlatformMisc::LowLevelOutputDebugStringf(
+			TEXT("Mounted OBB '%s'"), Filename);
+	}
+
+	AAssetManager* AssetMgr;
+	FZipUnionFile ZipResource;
+};
+
+IPlatformFile& IPlatformFile::GetPlatformPhysical()
+{
+	return FAndroidPlatformFile::GetPlatformPhysical();
+}
+
+IAndroidPlatformFile & IAndroidPlatformFile::GetPlatformPhysical()
+{
+	return FAndroidPlatformFile::GetPlatformPhysical();
+}
+#endif
diff --git a/Engine/Source/Runtime/Core/Private/Android/AndroidPlatformMemory.cpp b/Engine/Source/Runtime/Core/Private/Android/AndroidPlatformMemory.cpp
new file mode 100644
--- /dev/null
+++ b/Engine/Source/Runtime/Core/Private/Android/AndroidPlatformMemory.cpp
@@ -0,0 +1,449 @@
+// Copyright 1998-2019 Epic Games, Inc. All Rights Reserved.
+
+#include "Android/AndroidPlatformMemory.h"
+#include "HAL/LowLevelMemTracker.h"
+#include "HAL/MallocBinned.h"
+#include "HAL/MallocBinned2.h"
+#include "HAL/MallocBinned3.h"
+#include "HAL/MallocAnsi.h"
+#include "unistd.h"
+#include <jni.h>
+#include <sys/sysinfo.h>
+#include <sys/mman.h>
+
+#define JNI_CURRENT_VERSION JNI_VERSION_1_6
+extern JavaVM* GJavaVM;
+
+static int64 GetNativeHeapAllocatedSize()
+{
+	int64 AllocatedSize = 0;
+#if 0 // TODO: this works but sometimes crashes?
+	JNIEnv* Env = NULL;
+	GJavaVM->GetEnv((void **)&Env, JNI_CURRENT_VERSION);
+	jint AttachThreadResult = GJavaVM->AttachCurrentThread(&Env, NULL);
+
+	if(AttachThreadResult != JNI_ERR)
+	{
+		jclass Class = Env->FindClass("android/os/Debug");
+		if (Class)
+		{
+			jmethodID MethodID = Env->GetStaticMethodID(Class, "getNativeHeapAllocatedSize", "()J");
+			if (MethodID)
+			{
+				AllocatedSize = Env->CallStaticLongMethod(Class, MethodID);
+			}
+		}
+	}
+#endif
+	return AllocatedSize;
+}
+
+void FAndroidPlatformMemory::Init()
+{
+	FGenericPlatformMemory::Init();
+
+	const FPlatformMemoryConstants& MemoryConstants = FPlatformMemory::GetConstants();
+	FPlatformMemoryStats MemoryStats = GetStats();
+	UE_LOG(LogInit, Log, TEXT("Memory total: Physical=%.2fMB (%dGB approx) Available=%.2fMB PageSize=%.1fKB"), 
+		float(MemoryConstants.TotalPhysical/1024.0/1024.0),
+		MemoryConstants.TotalPhysicalGB, 
+		float(MemoryStats.AvailablePhysical/1024.0/1024.0),
+		float(MemoryConstants.PageSize/1024.0)
+		);
+}
+
+namespace AndroidPlatformMemory
+{
+	/**
+	* @brief Returns value in bytes from a status line
+	* @param Line in format "Blah:  10000 kB" - needs to be writable as it will modify it
+	* @return value in bytes (10240000, i.e. 10000 * 1024 for the above example)
+	*/
+	uint64 GetBytesFromStatusLine(char * Line)
+	{
+		check(Line);
+		int Len = strlen(Line);
+
+		// Len should be long enough to hold at least " kB\n"
+		const int kSuffixLength = 4;	// " kB\n"
+		if (Len <= kSuffixLength)
+		{
+			return 0;
+		}
+
+		// let's check that this is indeed "kB"
+		char * Suffix = &Line[Len - kSuffixLength];
+		if (strcmp(Suffix, " kB\n") != 0)
+		{
+			// Linux the kernel changed the format, huh?
+			return 0;
+		}
+
+		// kill the kB
+		*Suffix = 0;
+
+		// find the beginning of the number
+		for (const char * NumberBegin = Suffix; NumberBegin >= Line; --NumberBegin)
+		{
+			if (*NumberBegin == ' ')
+			{
+				return static_cast< uint64 >(atol(NumberBegin + 1)) * 1024ULL;
+			}
+		}
+
+		// we were unable to find whitespace in front of the number
+		return 0;
+	}
+}
+
+extern int32 AndroidThunkCpp_GetMetaDataInt(const FString& Key);
+
+FPlatformMemoryStats FAndroidPlatformMemory::GetStats()
+{
+	FPlatformMemoryStats MemoryStats;	// will init from constants
+
+										// open to all kind of overflows, thanks to Linux approach of exposing system stats via /proc and lack of proper C API
+										// And no, sysinfo() isn't useful for this (cannot get the same value for MemAvailable through it for example).
+
+	if (FILE* FileGlobalMemStats = fopen("/proc/meminfo", "r"))
+	{
+		int FieldsSetSuccessfully = 0;
+		uint64 MemFree = 0, Cached = 0;
+		do
+		{
+			char LineBuffer[256] = { 0 };
+			char *Line = fgets(LineBuffer, ARRAY_COUNT(LineBuffer), FileGlobalMemStats);
+			if (Line == nullptr)
+			{
+				break;	// eof or an error
+			}
+
+			// if we have MemAvailable, favor that (see http://git.kernel.org/cgit/linux/kernel/git/torvalds/linux.git/commit/?id=34e431b0ae398fc54ea69ff85ec700722c9da773)
+			if (strstr(Line, "MemAvailable:") == Line)
+			{
+				MemoryStats.AvailablePhysical = AndroidPlatformMemory::GetBytesFromStatusLine(Line);
+				++FieldsSetSuccessfully;
+			}
+			else if (strstr(Line, "SwapFree:") == Line)
+			{
+				MemoryStats.AvailableVirtual = AndroidPlatformMemory::GetBytesFromStatusLine(Line);
+				++FieldsSetSuccessfully;
+			}
+			else if (strstr(Line, "MemFree:") == Line)
+			{
+				MemFree = AndroidPlatformMemory::GetBytesFromStatusLine(Line);
+				++FieldsSetSuccessfully;
+			}
+			else if (strstr(Line, "Cached:") == Line)
+			{
+				Cached = AndroidPlatformMemory::GetBytesFromStatusLine(Line);
+				++FieldsSetSuccessfully;
+			}
+		} while (FieldsSetSuccessfully < 4);
+
+		// if we didn't have MemAvailable (kernels < 3.14 or CentOS 6.x), use free + cached as a (bad) approximation
+		if (MemoryStats.AvailablePhysical == 0)
+		{
+			MemoryStats.AvailablePhysical = FMath::Min(MemFree + Cached, MemoryStats.TotalPhysical);
+		}
+
+		fclose(FileGlobalMemStats);
+	}
+
+	// again /proc "API" :/
+	if (FILE* ProcMemStats = fopen("/proc/self/status", "r"))
+	{
+		int FieldsSetSuccessfully = 0;
+		do
+		{
+			char LineBuffer[256] = { 0 };
+			char *Line = fgets(LineBuffer, ARRAY_COUNT(LineBuffer), ProcMemStats);
+			if (Line == nullptr)
+			{
+				break;	// eof or an error
+			}
+
+			if (strstr(Line, "VmPeak:") == Line)
+			{
+				MemoryStats.PeakUsedVirtual = AndroidPlatformMemory::GetBytesFromStatusLine(Line);
+				++FieldsSetSuccessfully;
+			}
+			else if (strstr(Line, "VmSize:") == Line)
+			{
+				MemoryStats.UsedVirtual = AndroidPlatformMemory::GetBytesFromStatusLine(Line);
+				++FieldsSetSuccessfully;
+			}
+			else if (strstr(Line, "VmHWM:") == Line)
+			{
+				MemoryStats.PeakUsedPhysical = AndroidPlatformMemory::GetBytesFromStatusLine(Line);
+				++FieldsSetSuccessfully;
+			}
+			else if (strstr(Line, "VmRSS:") == Line)
+			{
+				MemoryStats.UsedPhysical = AndroidPlatformMemory::GetBytesFromStatusLine(Line);
+				++FieldsSetSuccessfully;
+			}
+		} while (FieldsSetSuccessfully < 4);
+
+		fclose(ProcMemStats);
+	}
+
+
+	// sanitize stats as sometimes peak < used for some reason
+	MemoryStats.PeakUsedVirtual = FMath::Max(MemoryStats.PeakUsedVirtual, MemoryStats.UsedVirtual);
+	MemoryStats.PeakUsedPhysical = FMath::Max(MemoryStats.PeakUsedPhysical, MemoryStats.UsedPhysical);
+
+	// get this value from Java instead (DO NOT INTEGRATE at this time) - skip this if JavaVM not set up yet!
+#if USE_ANDROID_JNI
+	if (GJavaVM)
+	{
+		MemoryStats.UsedPhysical = static_cast<uint64>(AndroidThunkCpp_GetMetaDataInt(TEXT("ue4.getUsedMemory"))) * 1024ULL;
+	}
+#endif
+
+	return MemoryStats;
+}
+
+uint64 FAndroidPlatformMemory::GetMemoryUsedFast()
+{
+	// get this value from Java instead (DO NOT INTEGRATE at this time) - skip this if JavaVM not set up yet!
+#if USE_ANDROID_JNI
+	if (GJavaVM)
+	{
+		return static_cast<uint64>(AndroidThunkCpp_GetMetaDataInt(TEXT("ue4.getUsedMemory"))) * 1024ULL;
+	}
+#endif
+
+	// minimal code to get Used memory
+	if (FILE* ProcMemStats = fopen("/proc/self/status", "r"))
+	{
+		while (1)
+		{
+			char LineBuffer[256] = { 0 };
+			char *Line = fgets(LineBuffer, ARRAY_COUNT(LineBuffer), ProcMemStats);
+			if (Line == nullptr)
+			{
+				break;	// eof or an error
+			}
+			else if (strstr(Line, "VmRSS:") == Line)
+			{
+				fclose(ProcMemStats);
+				return AndroidPlatformMemory::GetBytesFromStatusLine(Line);
+			}
+		} 
+		fclose(ProcMemStats);
+	}
+
+	return 0;
+}
+
+
+const FPlatformMemoryConstants& FAndroidPlatformMemory::GetConstants()
+{
+	static FPlatformMemoryConstants MemoryConstants;
+
+	if (MemoryConstants.TotalPhysical == 0)
+	{
+		// Gather platform memory stats.
+		struct sysinfo SysInfo;
+		unsigned long long MaxPhysicalRAMBytes = 0;
+		unsigned long long MaxVirtualRAMBytes = 0;
+
+		if (0 == sysinfo(&SysInfo))
+		{
+			MaxPhysicalRAMBytes = static_cast< unsigned long long >(SysInfo.mem_unit) * static_cast< unsigned long long >(SysInfo.totalram);
+			MaxVirtualRAMBytes = static_cast< unsigned long long >(SysInfo.mem_unit) * static_cast< unsigned long long >(SysInfo.totalswap);
+		}
+
+		MemoryConstants.TotalPhysical = MaxPhysicalRAMBytes;
+		MemoryConstants.TotalVirtual = MaxVirtualRAMBytes;
+		MemoryConstants.TotalPhysicalGB = (MemoryConstants.TotalPhysical + 1024 * 1024 * 1024 - 1) / 1024 / 1024 / 1024;
+		MemoryConstants.PageSize = sysconf(_SC_PAGESIZE);
+		MemoryConstants.BinnedPageSize = FMath::Max((SIZE_T)65536, MemoryConstants.PageSize);
+		MemoryConstants.BinnedAllocationGranularity = MemoryConstants.PageSize;
+		MemoryConstants.OsAllocationGranularity = MemoryConstants.PageSize;
+#if PLATFORM_32BITS
+		MemoryConstants.AddressLimit = DECLARE_UINT64(4) * 1024 * 1024 * 1024;
+#else
+		MemoryConstants.AddressLimit = FPlatformMath::RoundUpToPowerOfTwo64(MemoryConstants.TotalPhysical);
+#endif
+	}
+
+	return MemoryConstants;
+}
+
+EPlatformMemorySizeBucket FAndroidPlatformMemory::GetMemorySizeBucket()
+{
+	// @todo android - if running in ES2 mode, or at least without the extensions for texture streaming, we will load all of the textures
+	// so we better look like a low memory device
+	return FGenericPlatformMemory::GetMemorySizeBucket();
+}
+
+// Set rather to use BinnedMalloc2 for binned malloc, can be overridden below
+#define USE_MALLOC_BINNED2 PLATFORM_ANDROID_ARM64
+#if !defined(USE_MALLOC_BINNED3)
+	#define USE_MALLOC_BINNED3 (0)
+#endif
+
+FMalloc* FAndroidPlatformMemory::BaseAllocator()
+{
+#if ENABLE_LOW_LEVEL_MEM_TRACKER
+	// make sure LLM is using UsedPhysical for program size, instead of Availble-Free
+	FPlatformMemoryStats Stats = FAndroidPlatformMemory::GetStats();
+	FLowLevelMemTracker::Get().SetProgramSize(Stats.UsedPhysical);
+#endif
+
+#if USE_MALLOC_BINNED3 && PLATFORM_ANDROID_ARM64
+	return new FMallocBinned3();
+#elif USE_MALLOC_BINNED2
+	return new FMallocBinned2();
+#else
+	const FPlatformMemoryConstants& MemoryConstants = FPlatformMemory::GetConstants();
+	// 1 << FMath::CeilLogTwo(MemoryConstants.TotalPhysical) should really be FMath::RoundUpToPowerOfTwo,
+	// but that overflows to 0 when MemoryConstants.TotalPhysical is close to 4GB, since CeilLogTwo returns 32
+	// this then causes the MemoryLimit to be 0 and crashing the app
+	uint64 MemoryLimit = FMath::Min<uint64>(uint64(1) << FMath::CeilLogTwo(MemoryConstants.TotalPhysical), 0x100000000);
+
+	// todo: Verify MallocBinned2 on 32bit
+	// [RCL] 2017-03-06 FIXME: perhaps BinnedPageSize should be used here, but leaving this change to the Android platform owner.
+	return new FMallocBinned(MemoryConstants.PageSize, MemoryLimit);
+#endif
+}
+
+
+void* FAndroidPlatformMemory::BinnedAllocFromOS(SIZE_T Size)
+{
+	void* Ptr;
+	if (USE_MALLOC_BINNED2)
+	{
+		static FCriticalSection CriticalSection;
+		FScopeLock Lock(&CriticalSection);
+
+		const FPlatformMemoryConstants& MemoryConstants = FPlatformMemory::GetConstants();
+
+		static uint8* Base = nullptr;
+		static uint8* End = nullptr;
+
+		static const SIZE_T MinAllocSize = 4 * 1024 * 1024; // we will allocate chunks of 4MB, that means the amount we will need to unmap, assuming a lot of 64k blocks, will be small.
+
+		if (End - Base < Size)
+		{
+			if (Base)
+			{
+				if (Base < End)
+				{
+					if (munmap(Base, End - Base) != 0)
+					{
+						const int ErrNo = errno;
+						UE_LOG(LogHAL, Fatal, TEXT("munmap (for trim) (addr=%p, len=%llu) failed with errno = %d (%s)"), Base, End - Base,
+							ErrNo, StringCast< TCHAR >(strerror(ErrNo)).Get());
+					}
+				}
+				Base = nullptr;
+				End = nullptr;
+			}
+
+			SIZE_T SizeToAlloc = FMath::Max<SIZE_T>(MinAllocSize, Align(Size, MemoryConstants.PageSize) + MemoryConstants.BinnedPageSize);
+			uint8* UnalignedBase = (uint8*)mmap(nullptr, SizeToAlloc, PROT_READ | PROT_WRITE, MAP_PRIVATE | MAP_ANON, -1, 0);
+			End = UnalignedBase + SizeToAlloc;
+			Base = Align(UnalignedBase, MemoryConstants.BinnedPageSize);
+
+			if (Base > UnalignedBase)
+			{
+				if (munmap(UnalignedBase, Base - UnalignedBase) != 0)
+				{
+					const int ErrNo = errno;
+					UE_LOG(LogHAL, Fatal, TEXT("munmap (for align) (addr=%p, len=%llu) failed with errno = %d (%s)"), UnalignedBase, Base - UnalignedBase,
+						ErrNo, StringCast< TCHAR >(strerror(ErrNo)).Get());
+				}
+			}
+
+
+		}
+		Ptr = Base;
+		uint8* UnalignedBase = Align(Base + Size, MemoryConstants.PageSize);
+		Base = Align(UnalignedBase, MemoryConstants.BinnedPageSize);
+
+		if (Base > End)
+		{
+			Base = End;
+		}
+
+		if (Base > UnalignedBase)
+		{
+			if (munmap(UnalignedBase, Base - UnalignedBase) != 0)
+			{
+				const int ErrNo = errno;
+				UE_LOG(LogHAL, Fatal, TEXT("munmap (for tail align) (addr=%p, len=%llu) failed with errno = %d (%s)"), UnalignedBase, Base - UnalignedBase,
+					ErrNo, StringCast< TCHAR >(strerror(ErrNo)).Get());
+			}
+		}
+	}
+	else
+	{
+		Ptr = mmap(nullptr, Size, PROT_READ | PROT_WRITE, MAP_PRIVATE | MAP_ANON, -1, 0);
+	}
+
+	LLM(FLowLevelMemTracker::Get().OnLowLevelAlloc(ELLMTracker::Platform, Ptr, Size));
+	return Ptr;
+}
+
+void FAndroidPlatformMemory::BinnedFreeToOS(void* Ptr, SIZE_T Size)
+{
+	LLM(FLowLevelMemTracker::Get().OnLowLevelFree(ELLMTracker::Platform, Ptr));
+	if (munmap(Ptr, Size) != 0)
+	{
+		const int ErrNo = errno;
+		UE_LOG(LogHAL, Fatal, TEXT("munmap(addr=%p, len=%llu) failed with errno = %d (%s)"), Ptr, Size,
+			ErrNo, StringCast< TCHAR >(strerror(ErrNo)).Get());
+	}
+}
+
+bool FAndroidPlatformMemory::MemoryRangeDecommit(void* Ptr, SIZE_T Size)
+{
+	return madvise(Ptr, Size, MADV_DONTNEED) == 0;
+}
+
+/**
+* LLM uses these low level functions (LLMAlloc and LLMFree) to allocate memory. It grabs
+* the function pointers by calling FPlatformMemory::GetLLMAllocFunctions. If these functions
+* are not implemented GetLLMAllocFunctions should return false and LLM will be disabled.
+*/
+
+#if ENABLE_LOW_LEVEL_MEM_TRACKER
+
+int64 LLMMallocTotal = 0;
+
+void* LLMAlloc(size_t Size)
+{
+	void* Ptr = mmap(nullptr, Size, PROT_READ | PROT_WRITE, MAP_PRIVATE | MAP_ANON, -1, 0);
+
+	LLMMallocTotal += Size;
+
+	return Ptr;
+}
+
+void LLMFree(void* Addr, size_t Size)
+{
+	LLMMallocTotal -= Size;
+	if (Addr != nullptr && munmap(Addr, Size) != 0)
+	{
+		const int ErrNo = errno;
+		UE_LOG(LogHAL, Fatal, TEXT("munmap(addr=%p, len=%llu) failed with errno = %d (%s)"), Addr, Size,
+			ErrNo, StringCast< TCHAR >(strerror(ErrNo)).Get());
+	}
+}
+
+#endif
+
+bool FAndroidPlatformMemory::GetLLMAllocFunctions(void*(*&OutAllocFunction)(size_t), void(*&OutFreeFunction)(void*, size_t), int32& OutAlignment)
+{
+#if ENABLE_LOW_LEVEL_MEM_TRACKER
+	OutAllocFunction = LLMAlloc;
+	OutFreeFunction = LLMFree;
+	OutAlignment = sysconf(_SC_PAGESIZE);
+	return true;
+#else
+	return false;
+#endif
+}
diff --git a/Engine/Source/Runtime/Core/Private/Android/AndroidPlatformMisc.cpp b/Engine/Source/Runtime/Core/Private/Android/AndroidPlatformMisc.cpp
new file mode 100644
--- /dev/null
+++ b/Engine/Source/Runtime/Core/Private/Android/AndroidPlatformMisc.cpp
@@ -0,0 +1,2284 @@
+// Copyright 1998-2019 Epic Games, Inc. All Rights Reserved.
+
+#include "Android/AndroidPlatformMisc.h"
+#include "Android/AndroidJavaEnv.h"
+#include "HAL/PlatformStackWalk.h"
+#include "Misc/FileHelper.h"
+#include "Misc/App.h"
+#include "Misc/ConfigCacheIni.h"
+#include "Misc/FeedbackContext.h"
+#include "Misc/OutputDeviceRedirector.h"
+#include "HAL/IConsoleManager.h"
+#include "Modules/ModuleManager.h"
+#include <string.h>
+#include <dlfcn.h>
+#include <sys/statfs.h>
+#include <sys/syscall.h>
+#include <sched.h>
+#include <jni.h>
+
+#include "Android/AndroidPlatformCrashContext.h"
+#include "HAL/PlatformMallocCrash.h"
+#include "Android/AndroidJavaMessageBox.h"
+#include "GenericPlatform/GenericPlatformChunkInstall.h"
+
+#include "Misc/Parse.h"
+#include "Internationalization/Regex.h"
+
+#include <android/log.h>
+#if USE_ANDROID_INPUT
+#include <android/keycodes.h>
+#endif
+#if USE_ANDROID_JNI
+#include <android/asset_manager.h>
+#include <cpu-features.h>
+#include <android_native_app_glue.h>
+#include "Templates/Function.h"
+#include "Android/AndroidStats.h"
+#endif
+
+#include "Misc/CoreDelegates.h"
+#include "Async/TaskGraphInterfaces.h"
+
+#include "FramePro/FrameProProfiler.h"
+
+#if USE_ANDROID_JNI
+extern AAssetManager * AndroidThunkCpp_GetAssetManager();
+extern int32 GAndroidPackageVersion;
+#else
+#define GAndroidPackageVersion	1
+#endif
+
+static int32 GAndroidTraceMarkersEnabled = 0;
+static FAutoConsoleVariableRef CAndroidTraceMarkersEnabled(
+	TEXT("android.tracemarkers"),
+	GAndroidTraceMarkersEnabled,
+	TEXT("Enable outputting named events to Android trace marker file.\n"),
+	ECVF_Default
+);
+
+static int32 GAndroidLowPowerBatteryThreshold = 15;
+static FAutoConsoleVariableRef CAndroidLowPowerBatteryThreshold(
+	TEXT("android.LowPowerBatteryThreshold"),
+	GAndroidLowPowerBatteryThreshold,
+	TEXT("The battery level below which the device is considered in a low power state."),
+	ECVF_Default
+);
+
+#if STATS || ENABLE_STATNAMEDEVENTS
+int32 FAndroidMisc::TraceMarkerFileDescriptor = -1;
+
+typedef void (*ATrace_beginSection_Type) (const char* sectionName);
+typedef void (*ATrace_endSection_Type) (void);
+typedef bool (*ATrace_isEnabled_Type) (void);
+
+static ATrace_beginSection_Type ATrace_beginSection = NULL;
+static ATrace_endSection_Type ATrace_endSection = NULL;
+static ATrace_isEnabled_Type ATrace_isEnabled = NULL;
+
+static bool bUseNativeSystrace = false;
+
+#endif
+
+// run time compatibility information
+FString FAndroidMisc::AndroidVersion; // version of android we are running eg "4.0.4"
+FString FAndroidMisc::DeviceMake; // make of the device we are running on eg. "samsung"
+FString FAndroidMisc::DeviceModel; // model of the device we are running on eg "SAMSUNG-SGH-I437"
+FString FAndroidMisc::DeviceBuildNumber; // platform image build number of device "R16NW.G960NKSU1ARD6"
+FString FAndroidMisc::OSLanguage; // language code the device is set to eg "deu"
+
+// Build/API level we are running.
+int32 FAndroidMisc::AndroidBuildVersion = 0;
+
+// Whether or not the system handles the volume buttons (event will still be generated either way)
+bool FAndroidMisc::VolumeButtonsHandledBySystem = true;
+
+// Whether an app restart is needed to free driver allocated memory after precompiling PSOs
+bool FAndroidMisc::bNeedsRestartAfterPSOPrecompile = false;
+
+// Key/Value pair variables from the optional configuration.txt
+TMap<FString, FString> FAndroidMisc::ConfigRulesVariables;
+
+EDeviceScreenOrientation FAndroidMisc::DeviceOrientation = EDeviceScreenOrientation::Unknown;
+
+extern void AndroidThunkCpp_ForceQuit();
+
+// From AndroidFile.cpp
+extern FString GFontPathBase;
+
+void FAndroidMisc::RequestExit( bool Force )
+{
+	UE_LOG(LogAndroid, Log, TEXT("FAndroidMisc::RequestExit(%i)"), Force);
+	if (Force)
+	{
+#if USE_ANDROID_JNI
+		AndroidThunkCpp_ForceQuit();
+#else
+		exit(1);
+#endif
+	}
+	else
+	{
+		GIsRequestingExit = 1;
+	}
+}
+
+void FAndroidMisc::LocalPrint(const TCHAR *Message)
+{
+	// Builds for distribution should not have logging in them:
+	// http://developer.android.com/tools/publishing/preparing.html#publishing-configure
+#if !UE_BUILD_SHIPPING
+	const int MAX_LOG_LENGTH = 4096;
+	// not static since may be called by different threads
+	wchar_t MessageBuffer[MAX_LOG_LENGTH];
+
+	const TCHAR* SourcePtr = Message;
+	while (*SourcePtr)
+	{
+		wchar_t* WritePtr = MessageBuffer;
+		int32 RemainingSpace = MAX_LOG_LENGTH;
+		while (*SourcePtr && --RemainingSpace > 0)
+		{
+			if (*SourcePtr == TEXT('\r'))
+			{
+				// If next character is newline, skip it
+				if (*(++SourcePtr) == TEXT('\n'))
+					++SourcePtr;
+				break;
+			}
+			else if (*SourcePtr == TEXT('\n'))
+			{
+				++SourcePtr;
+				break;
+			}
+			else {
+				*WritePtr++ = static_cast<wchar_t>(*SourcePtr++);
+			}
+		}
+		*WritePtr = '\0';
+		__android_log_print(ANDROID_LOG_DEBUG, "UE4", "%ls", MessageBuffer);
+	}
+#endif
+}
+
+// Test for device vulkan support.
+static void EstablishVulkanDeviceSupport();
+
+namespace FAndroidAppEntry
+{
+	extern void PlatformInit();
+}
+
+void FAndroidMisc::PlatformPreInit()
+{
+	FGenericPlatformMisc::PlatformPreInit();
+	EstablishVulkanDeviceSupport();
+	FAndroidAppEntry::PlatformInit();
+}
+
+static volatile bool HeadPhonesArePluggedIn = false;
+
+static FAndroidMisc::FBatteryState CurrentBatteryState;
+
+static FCriticalSection ReceiversLock;
+static struct
+{
+	int		Volume;
+	double	TimeOfChange;
+} CurrentVolume;
+
+#if USE_ANDROID_JNI
+extern "C"
+{
+
+	JNIEXPORT void Java_com_epicgames_ue4_HeadsetReceiver_stateChanged(JNIEnv * jni, jclass clazz, jint state)
+	{
+		FPlatformMisc::LowLevelOutputDebugStringf(TEXT("nativeHeadsetEvent(%i)"), state);
+		HeadPhonesArePluggedIn = (state == 1);
+	}
+
+	JNIEXPORT void Java_com_epicgames_ue4_VolumeReceiver_volumeChanged(JNIEnv * jni, jclass clazz, jint volume)
+	{
+		FPlatformMisc::LowLevelOutputDebugStringf(TEXT("nativeVolumeEvent(%i)"), volume);
+		ReceiversLock.Lock();
+		CurrentVolume.Volume = volume;
+		CurrentVolume.TimeOfChange = FApp::GetCurrentTime();
+		ReceiversLock.Unlock();
+	}
+
+	JNIEXPORT void Java_com_epicgames_ue4_BatteryReceiver_dispatchEvent(JNIEnv * jni, jclass clazz, jint status, jint level, jint temperature)
+	{
+		FPlatformMisc::LowLevelOutputDebugStringf(TEXT("nativeBatteryEvent(stat = %i, lvl = %i %, temp = %3.2f \u00B0C)"), status, level, float(temperature)/10.f);
+
+		ReceiversLock.Lock();
+		const bool bWasInLowPowerMode = CurrentBatteryState.Level <= GAndroidLowPowerBatteryThreshold;
+
+		FAndroidMisc::FBatteryState state;
+		state.State = (FAndroidMisc::EBatteryState)status;
+		state.Level = level;
+		state.Temperature = float(temperature)/10.f;
+		CurrentBatteryState = state;
+
+		const bool bIsInLowPowerMode = CurrentBatteryState.Level <= GAndroidLowPowerBatteryThreshold;
+		ReceiversLock.Unlock();
+
+		// When we cross the low power battery level threshold, inform the active application
+		if (bIsInLowPowerMode != bWasInLowPowerMode)
+		{
+			FGraphEventRef LowPowerTask = FFunctionGraphTask::CreateAndDispatchWhenReady([=]()
+			{
+				UE_LOG(LogAndroid, Display, TEXT("Low Power Mode Changed: %d"), bIsInLowPowerMode);
+				FCoreDelegates::OnLowPowerMode.Broadcast(bIsInLowPowerMode);
+			}, TStatId(), NULL, ENamedThreads::GameThread);	
+		}
+	}
+}
+#endif
+
+#if USE_ANDROID_JNI
+
+// Manage Java side OS event receivers.
+static struct
+{
+	const char*		ClazzName;
+	JNINativeMethod	Jnim;
+	jclass			Clazz;
+	jmethodID		StartReceiver;
+	jmethodID		StopReceiver;
+} JavaEventReceivers[] =
+{
+	{ "com/epicgames/ue4/VolumeReceiver",{ "volumeChanged", "(I)V",  (void *)Java_com_epicgames_ue4_VolumeReceiver_volumeChanged } },
+	{ "com/epicgames/ue4/BatteryReceiver",{ "dispatchEvent", "(III)V",(void *)Java_com_epicgames_ue4_BatteryReceiver_dispatchEvent } },
+	{ "com/epicgames/ue4/HeadsetReceiver",{ "stateChanged",  "(I)V",  (void *)Java_com_epicgames_ue4_HeadsetReceiver_stateChanged } },
+};
+
+void InitializeJavaEventReceivers()
+{
+	// Register natives to receive Volume, Battery, Headphones events
+	JNIEnv* JEnv = AndroidJavaEnv::GetJavaEnv();
+	if (nullptr != JEnv)
+	{
+		auto CheckJNIExceptions = [&JEnv]()
+		{
+			if (JEnv->ExceptionCheck())
+			{
+				JEnv->ExceptionDescribe();
+				JEnv->ExceptionClear();
+			}
+		};
+		auto GetStaticMethod = [&JEnv, &CheckJNIExceptions](const char* MethodName, jclass Clazz, const char* ClazzName)
+		{
+			jmethodID Method = JEnv->GetStaticMethodID(Clazz, MethodName, "(Landroid/app/Activity;)V");
+			if (Method == 0)
+			{
+				UE_LOG(LogAndroid, Error, TEXT("Can't find method %s of class %s"), ANSI_TO_TCHAR(MethodName), ANSI_TO_TCHAR(ClazzName));
+			}
+			CheckJNIExceptions();
+			return Method;
+		};
+
+		for (auto& JavaEventReceiver : JavaEventReceivers)
+		{
+			JavaEventReceiver.Clazz = AndroidJavaEnv::FindJavaClassGlobalRef(JavaEventReceiver.ClazzName);
+			if (JavaEventReceiver.Clazz == nullptr)
+			{
+				UE_LOG(LogAndroid, Error, TEXT("Can't find class for %s"), ANSI_TO_TCHAR(JavaEventReceiver.ClazzName));
+				continue;
+			}
+			if (JNI_OK != JEnv->RegisterNatives(JavaEventReceiver.Clazz, &JavaEventReceiver.Jnim, 1))
+			{
+				UE_LOG(LogAndroid, Error, TEXT("RegisterNatives failed for %s on %s"), ANSI_TO_TCHAR(JavaEventReceiver.ClazzName), ANSI_TO_TCHAR(JavaEventReceiver.Jnim.name));
+				CheckJNIExceptions();
+			}
+			JavaEventReceiver.StartReceiver = GetStaticMethod("startReceiver", JavaEventReceiver.Clazz, JavaEventReceiver.ClazzName);
+			JavaEventReceiver.StopReceiver = GetStaticMethod("stopReceiver", JavaEventReceiver.Clazz, JavaEventReceiver.ClazzName);
+		}
+	}
+	else
+	{
+		UE_LOG(LogAndroid, Warning, TEXT("Failed to initialize java event receivers. JNIEnv is not valid."));
+	}
+}
+
+void EnableJavaEventReceivers(bool bEnableReceivers)
+{
+	JNIEnv* JEnv = AndroidJavaEnv::GetJavaEnv();
+	if (nullptr != JEnv)
+	{
+		extern struct android_app* GNativeAndroidApp;
+		for (auto& JavaEventReceiver : JavaEventReceivers)
+		{
+			jmethodID methodId = bEnableReceivers ? JavaEventReceiver.StartReceiver : JavaEventReceiver.StopReceiver;
+			if (methodId != 0)
+			{
+				JEnv->CallStaticVoidMethod(JavaEventReceiver.Clazz, methodId, GNativeAndroidApp->activity->clazz);
+			}
+		}
+	}
+}
+
+#endif	//USE_ANDROID_JNI
+
+
+static FDelegateHandle AndroidOnBackgroundBinding;
+static FDelegateHandle AndroidOnForegroundBinding;
+
+#if (STATS || ENABLE_STATNAMEDEVENTS)
+
+static void StartTraceMarkers()
+{
+	if (FAndroidMisc::TraceMarkerFileDescriptor != -1)
+	{
+		UE_LOG(LogAndroid, Warning, TEXT("Systrace event logging already open."));
+		return;
+	}
+
+	// Setup trace file descriptor
+	FAndroidMisc::TraceMarkerFileDescriptor = open("/sys/kernel/debug/tracing/trace_marker", O_WRONLY);
+	if (FAndroidMisc::TraceMarkerFileDescriptor == -1)
+	{
+		UE_LOG(LogAndroid, Warning, TEXT("Trace Marker failed to open; systrace support disabled"));
+	}
+	else
+	{
+		UE_LOG(LogAndroid, Display, TEXT("Started systrace events logging."));
+	}
+}
+
+static void StopTraceMarkers()
+{
+	// Tear down trace file descriptor
+	if (FAndroidMisc::TraceMarkerFileDescriptor != -1)
+	{
+		close(FAndroidMisc::TraceMarkerFileDescriptor);
+		FAndroidMisc::TraceMarkerFileDescriptor = -1;
+		UE_LOG(LogAndroid, Display, TEXT("Stopped systrace events logging."));
+	}
+}
+
+static void UpdateTraceMarkersEnable(IConsoleVariable* Var)
+{
+	if (!GAndroidTraceMarkersEnabled)
+	{
+		StopTraceMarkers();
+	}
+	else
+	{
+		StartTraceMarkers();
+	}
+}
+#endif
+
+void FAndroidMisc::PlatformInit()
+{
+	// Increase the maximum number of simultaneously open files
+	// Display Timer resolution.
+	// Get swap file info
+	// Display memory info
+	// Setup user specified thread affinity if any
+	extern void AndroidSetupDefaultThreadAffinity();
+	AndroidSetupDefaultThreadAffinity();
+
+#if (STATS || ENABLE_STATNAMEDEVENTS)
+	//Loading NDK libandroid.so atrace functions, available in the android libraries way before NDK headers.
+	void* const LibAndroid = dlopen("libandroid.so", RTLD_NOW | RTLD_LOCAL);
+	if (LibAndroid != nullptr)
+	{
+		// Retrieve function pointers from shared object.
+		ATrace_beginSection = reinterpret_cast<ATrace_beginSection_Type>(dlsym(LibAndroid, "ATrace_beginSection"));
+		ATrace_endSection = reinterpret_cast<ATrace_endSection_Type>(dlsym(LibAndroid, "ATrace_endSection"));
+		ATrace_isEnabled = 	reinterpret_cast<ATrace_isEnabled_Type>(dlsym(LibAndroid, "ATrace_isEnabled"));
+	}
+
+	if (!ATrace_beginSection || !ATrace_endSection || !ATrace_isEnabled)
+	{
+		UE_LOG(LogAndroid, Warning, TEXT("Failed to use native systrace functionality."));
+		ATrace_beginSection = nullptr;
+		ATrace_endSection = nullptr;
+		ATrace_isEnabled = nullptr;
+
+		if (FParse::Param(FCommandLine::Get(), TEXT("enablesystrace")))
+		{
+			GAndroidTraceMarkersEnabled = 1;
+		}
+
+		if (GAndroidTraceMarkersEnabled)
+		{
+			StartTraceMarkers();
+		}
+
+		// Watch for CVar update
+		CAndroidTraceMarkersEnabled->SetOnChangedCallback(FConsoleVariableDelegate::CreateStatic(&UpdateTraceMarkersEnable));
+	}
+	else
+	{
+		bUseNativeSystrace = true;
+	}
+#endif
+
+#if USE_ANDROID_JNI
+	InitializeJavaEventReceivers();
+	AndroidOnBackgroundBinding = FCoreDelegates::ApplicationWillEnterBackgroundDelegate.AddStatic(EnableJavaEventReceivers, false);
+	AndroidOnForegroundBinding = FCoreDelegates::ApplicationHasEnteredForegroundDelegate.AddStatic(EnableJavaEventReceivers, true);
+#endif
+}
+
+extern void AndroidThunkCpp_DismissSplashScreen();
+
+void FAndroidMisc::PlatformTearDown()
+{
+#if (STATS || ENABLE_STATNAMEDEVENTS)
+	StopTraceMarkers();
+#endif
+
+	auto RemoveBinding = [](FCoreDelegates::FApplicationLifetimeDelegate& ApplicationLifetimeDelegate, FDelegateHandle& DelegateBinding)
+	{
+		if (DelegateBinding.IsValid())
+		{
+			ApplicationLifetimeDelegate.Remove(DelegateBinding);
+			DelegateBinding.Reset();
+		}
+	};
+
+	RemoveBinding(FCoreDelegates::ApplicationWillEnterBackgroundDelegate, AndroidOnBackgroundBinding);
+	RemoveBinding(FCoreDelegates::ApplicationHasEnteredForegroundDelegate, AndroidOnForegroundBinding);
+}
+
+void FAndroidMisc::PlatformHandleSplashScreen(bool ShowSplashScreen)
+{
+#if USE_ANDROID_JNI
+	if (!ShowSplashScreen)
+	{
+		AndroidThunkCpp_DismissSplashScreen();
+	}
+#endif
+}
+
+void FAndroidMisc::GetEnvironmentVariable(const TCHAR* VariableName, TCHAR* Result, int32 ResultLength)
+{
+	*Result = 0;
+	// @todo Android : get environment variable.
+}
+
+FString FAndroidMisc::GetEnvironmentVariable(const TCHAR* VariableName)
+{
+	// @todo Android : get environment variable.
+	return FString();
+}
+
+const TCHAR* FAndroidMisc::GetSystemErrorMessage(TCHAR* OutBuffer, int32 BufferCount, int32 Error)
+{
+	check(OutBuffer && BufferCount);
+	*OutBuffer = TEXT('\0');
+	if (Error == 0)
+	{
+		Error = errno;
+	}
+	char ErrorBuffer[1024];
+	strerror_r(Error, ErrorBuffer, 1024);
+	FCString::Strcpy(OutBuffer, BufferCount, UTF8_TO_TCHAR((const ANSICHAR*)ErrorBuffer));
+	return OutBuffer;
+}
+
+EAppReturnType::Type FAndroidMisc::MessageBoxExt( EAppMsgType::Type MsgType, const TCHAR* Text, const TCHAR* Caption )
+{
+#if USE_ANDROID_JNI
+	FJavaAndroidMessageBox MessageBox;
+	MessageBox.SetText(Text);
+	MessageBox.SetCaption(Caption);
+	EAppReturnType::Type * ResultValues = nullptr;
+	static EAppReturnType::Type ResultsOk[] = {
+		EAppReturnType::Ok };
+	static EAppReturnType::Type ResultsYesNo[] = {
+		EAppReturnType::Yes, EAppReturnType::No };
+	static EAppReturnType::Type ResultsOkCancel[] = {
+		EAppReturnType::Ok, EAppReturnType::Cancel };
+	static EAppReturnType::Type ResultsYesNoCancel[] = {
+		EAppReturnType::Yes, EAppReturnType::No, EAppReturnType::Cancel };
+	static EAppReturnType::Type ResultsCancelRetryContinue[] = {
+		EAppReturnType::Cancel, EAppReturnType::Retry, EAppReturnType::Continue };
+	static EAppReturnType::Type ResultsYesNoYesAllNoAll[] = {
+		EAppReturnType::Yes, EAppReturnType::No, EAppReturnType::YesAll,
+		EAppReturnType::NoAll };
+	static EAppReturnType::Type ResultsYesNoYesAllNoAllCancel[] = {
+		EAppReturnType::Yes, EAppReturnType::No, EAppReturnType::YesAll,
+		EAppReturnType::NoAll, EAppReturnType::Cancel };
+	static EAppReturnType::Type ResultsYesNoYesAll[] = {
+		EAppReturnType::Yes, EAppReturnType::No, EAppReturnType::YesAll };
+
+	// TODO: Should we localize button text?
+
+	switch (MsgType)
+	{
+	case EAppMsgType::Ok:
+		MessageBox.AddButton(TEXT("Ok"));
+		ResultValues = ResultsOk;
+		break;
+	case EAppMsgType::YesNo:
+		MessageBox.AddButton(TEXT("Yes"));
+		MessageBox.AddButton(TEXT("No"));
+		ResultValues = ResultsYesNo;
+		break;
+	case EAppMsgType::OkCancel:
+		MessageBox.AddButton(TEXT("Ok"));
+		MessageBox.AddButton(TEXT("Cancel"));
+		ResultValues = ResultsOkCancel;
+		break;
+	case EAppMsgType::YesNoCancel:
+		MessageBox.AddButton(TEXT("Yes"));
+		MessageBox.AddButton(TEXT("No"));
+		MessageBox.AddButton(TEXT("Cancel"));
+		ResultValues = ResultsYesNoCancel;
+		break;
+	case EAppMsgType::CancelRetryContinue:
+		MessageBox.AddButton(TEXT("Cancel"));
+		MessageBox.AddButton(TEXT("Retry"));
+		MessageBox.AddButton(TEXT("Continue"));
+		ResultValues = ResultsCancelRetryContinue;
+		break;
+	case EAppMsgType::YesNoYesAllNoAll:
+		MessageBox.AddButton(TEXT("Yes"));
+		MessageBox.AddButton(TEXT("No"));
+		MessageBox.AddButton(TEXT("Yes To All"));
+		MessageBox.AddButton(TEXT("No To All"));
+		ResultValues = ResultsYesNoYesAllNoAll;
+		break;
+	case EAppMsgType::YesNoYesAllNoAllCancel:
+		MessageBox.AddButton(TEXT("Yes"));
+		MessageBox.AddButton(TEXT("No"));
+		MessageBox.AddButton(TEXT("Yes To All"));
+		MessageBox.AddButton(TEXT("No To All"));
+		MessageBox.AddButton(TEXT("Cancel"));
+		ResultValues = ResultsYesNoYesAllNoAllCancel;
+		break;
+	case EAppMsgType::YesNoYesAll:
+		MessageBox.AddButton(TEXT("Yes"));
+		MessageBox.AddButton(TEXT("No"));
+		MessageBox.AddButton(TEXT("Yes To All"));
+		ResultValues = ResultsYesNoYesAll;
+		break;
+	default:
+		check(0);
+	}
+	int32 Choice = MessageBox.Show();
+	if (Choice >= 0 && nullptr != ResultValues)
+	{
+		return ResultValues[Choice];
+	}
+#endif
+
+	// Failed to show dialog, or failed to get a response,
+	// return default cancel response instead.
+	return FGenericPlatformMisc::MessageBoxExt(MsgType, Text, Caption);
+}
+
+bool FAndroidMisc::HasPlatformFeature(const TCHAR* FeatureName)
+{
+	if (FCString::Stricmp(FeatureName, TEXT("Vulkan")) == 0)
+	{
+		return FAndroidMisc::ShouldUseVulkan();
+	}
+
+	return FGenericPlatformMisc::HasPlatformFeature(FeatureName);
+}
+
+bool FAndroidMisc::UseRenderThread()
+{
+	// if we in general don't want to use the render thread due to commandline, etc, then don't
+	if (!FGenericPlatformMisc::UseRenderThread())
+	{
+		return false;
+	}
+
+	// Check for DisableThreadedRendering CVar from DeviceProfiles config
+	// Any devices in the future that need to disable threaded rendering should be given a device profile and use this CVar
+	const IConsoleVariable *const CVar = IConsoleManager::Get().FindConsoleVariable(TEXT("r.AndroidDisableThreadedRendering"));
+	if (CVar && CVar->GetInt() != 0)
+	{
+		return false;
+	}
+
+	// there is a crash with the nvidia tegra dual core processors namely the optimus 2x and xoom
+	// when running multithreaded it can't handle multiple threads using opengl (bug)
+	// tested with lg optimus 2x and motorola xoom
+	// come back and revisit this later
+	// https://code.google.com/p/android/issues/detail?id=32636
+	if (FAndroidMisc::GetGPUFamily() == FString(TEXT("NVIDIA Tegra")) && FPlatformMisc::NumberOfCores() <= 2 && FAndroidMisc::GetGLVersion().StartsWith(TEXT("OpenGL ES 2.")))
+	{
+		return false;
+	}
+
+	// Vivante GC1000 with 2.x driver has issues with render thread
+	if (FAndroidMisc::GetGPUFamily().StartsWith(TEXT("Vivante GC1000")) && FAndroidMisc::GetGLVersion().StartsWith(TEXT("OpenGL ES 2.")))
+	{
+		return false;
+	}
+
+	// there is an issue with presenting the buffer on kindle fire (1st gen) with multiple threads using opengl
+	if (FAndroidMisc::GetDeviceModel() == FString(TEXT("Kindle Fire")))
+	{
+		return false;
+	}
+
+	// there is an issue with swapbuffer ordering on startup on samsung s3 mini with multiple threads using opengl
+	if (FAndroidMisc::GetDeviceModel() == FString(TEXT("GT-I8190L")))
+	{
+		return false;
+	}
+
+	return true;
+}
+
+#if PLATFORM_LUMIN
+
+int32 FAndroidMisc::NumberOfCores()
+{
+//#if USE_ANDROID_JNI
+//	static int32 NumberOfCores = android_getCpuCount();
+//	return NumberOfCores;
+//#else
+	// WARNING: this function ignores edge cases like affinity mask changes (and even more fringe cases like CPUs going offline)
+	// in the name of performance (higher level code calls NumberOfCores() way too often...)
+	static int32 NumberOfCores = 0;
+	if (NumberOfCores == 0)
+	{
+		if (FParse::Param(FCommandLine::Get(), TEXT("usehyperthreading")))
+		{
+			NumberOfCores = NumberOfCoresIncludingHyperthreads();
+		}
+		else
+		{
+			cpu_set_t AvailableCpusMask;
+			CPU_ZERO(&AvailableCpusMask);
+
+			if (0 != sched_getaffinity(0, sizeof(AvailableCpusMask), &AvailableCpusMask))
+			{
+				NumberOfCores = 1;	// we are running on something, right?
+			}
+			else
+			{
+				// read the proc core counts and the proc max frequencies from cpuinfo because of 
+				// potential security restrictions on the sys mount
+				if (FILE* FileGlobalCpuStats = fopen("/proc/cpuinfo", "r"))
+				{
+					char LineBuffer[256] = { 0 };
+					do
+					{
+						char *Line = fgets(LineBuffer, ARRAY_COUNT(LineBuffer), FileGlobalCpuStats);
+						if (Line == nullptr)
+						{
+							break;	// eof or an error
+						}
+						// count the number of processor entries in loop
+						// for Lumin one processor translates to one core
+						if (strstr(Line, "processor") == Line)
+						{
+							NumberOfCores += 1;
+						}
+					} while (1);
+					fclose(FileGlobalCpuStats);
+				}
+			}
+		}
+	}
+	return NumberOfCores;
+//#endif
+}
+
+
+int32 FAndroidMisc::NumberOfCoresIncludingHyperthreads()
+{
+#if USE_ANDROID_JNI
+	return FPlatformMisc::NumberOfCores();
+#else
+	// WARNING: this function ignores edge cases like affinity mask changes (and even more fringe cases like CPUs going offline)
+	// in the name of performance (higher level code calls NumberOfCores() way too often...)
+	static int32 NumCoreIds = 0;
+	if (NumCoreIds == 0)
+	{
+		cpu_set_t AvailableCpusMask;
+		CPU_ZERO(&AvailableCpusMask);
+
+		if (0 != sched_getaffinity(0, sizeof(AvailableCpusMask), &AvailableCpusMask))
+		{
+			NumCoreIds = 1;	// we are running on something, right?
+		}
+		else
+		{
+			return CPU_COUNT(&AvailableCpusMask);
+		}
+	}
+	return NumCoreIds;
+#endif
+}
+
+
+#else
+
+int32 FAndroidMisc::NumberOfCores()
+{
+	int32 NumberOfCores = android_getCpuCount();
+
+	static int CalculatedNumberOfCores = 0;
+
+#ifndef CPU_SETSIZE
+#if PLATFORM_64BITS
+	#define CPU_SETSIZE 1024
+#else
+	#define CPU_SETSIZE 32
+#endif 
+#endif
+
+	char cpuset[CPU_SETSIZE / 8];
+
+	if (CalculatedNumberOfCores == 0)
+	{
+		pid_t ThreadId = gettid();
+		syscall(__NR_sched_getaffinity, ThreadId, sizeof(cpuset), &cpuset);
+
+		char *coreptr = cpuset;
+		int32 CoreSets = CPU_SETSIZE / 8;;
+		while (CoreSets--)
+		{
+			char coremask = *coreptr++;
+			for (int i = 0; i < 8; i++)
+			{
+				CalculatedNumberOfCores += ((coremask & (1 << i)) != 0);
+			}
+		}
+
+		UE_LOG(LogTemp, Log, TEXT("%d cores and %d assignable cores"), NumberOfCores, CalculatedNumberOfCores);
+	}
+
+	return !CalculatedNumberOfCores ? NumberOfCores : CalculatedNumberOfCores;
+}
+
+int32 FAndroidMisc::NumberOfCoresIncludingHyperthreads()
+{
+	return NumberOfCores();
+}
+
+#endif
+
+static FAndroidMisc::FCPUState CurrentCPUState;
+
+FAndroidMisc::FCPUState& FAndroidMisc::GetCPUState(){
+	uint64_t UserTime, NiceTime, SystemTime, SoftIRQTime, IRQTime, IdleTime, IOWaitTime;
+	int32		Index = 0;
+	ANSICHAR	Buffer[500];
+
+	CurrentCPUState.CoreCount = FMath::Min( FAndroidMisc::NumberOfCores(), FAndroidMisc::FCPUState::MaxSupportedCores);
+	FILE* FileHandle = fopen("/proc/stat", "r");
+	if (FileHandle){
+		CurrentCPUState.ActivatedCoreCount = 0;
+		for (size_t n = 0; n < CurrentCPUState.CoreCount; n++) {
+			CurrentCPUState.Status[n] = 0;
+			CurrentCPUState.PreviousUsage[n] = CurrentCPUState.CurrentUsage[n];
+		}
+
+		while (fgets(Buffer, 100, FileHandle)) {
+#if PLATFORM_64BITS
+			sscanf(Buffer, "%5s %8lu %8lu %8lu %8lu %8lu %8lu %8lu", CurrentCPUState.Name,
+				&UserTime, &NiceTime, &SystemTime, &IdleTime, &IOWaitTime, &IRQTime,
+				&SoftIRQTime);
+#else
+			sscanf(Buffer, "%5s %8llu %8llu %8llu %8llu %8llu %8llu %8llu", CurrentCPUState.Name,
+				&UserTime, &NiceTime, &SystemTime, &IdleTime, &IOWaitTime, &IRQTime,
+				&SoftIRQTime);
+#endif
+
+			if (0 == strncmp(CurrentCPUState.Name, "cpu", 3)) {
+				Index = CurrentCPUState.Name[3] - '0';
+				if (Index >= 0 && Index < CurrentCPUState.CoreCount)
+				{
+					if(CurrentCPUState.Name[5] != '\0')
+					{
+						Index = atol(&CurrentCPUState.Name[3]);
+					}
+					CurrentCPUState.CurrentUsage[Index].IdleTime = IdleTime;
+					CurrentCPUState.CurrentUsage[Index].NiceTime = NiceTime;
+					CurrentCPUState.CurrentUsage[Index].SystemTime = SystemTime;
+					CurrentCPUState.CurrentUsage[Index].SoftIRQTime = SoftIRQTime;
+					CurrentCPUState.CurrentUsage[Index].IRQTime = IRQTime;
+					CurrentCPUState.CurrentUsage[Index].IOWaitTime = IOWaitTime;
+					CurrentCPUState.CurrentUsage[Index].UserTime = UserTime;
+					CurrentCPUState.CurrentUsage[Index].TotalTime = UserTime + NiceTime + SystemTime + SoftIRQTime + IRQTime + IdleTime + IOWaitTime;
+					CurrentCPUState.Status[Index] = 1;
+					CurrentCPUState.ActivatedCoreCount++;
+				}
+				if (Index == CurrentCPUState.CoreCount-1)
+					break;
+			}
+		}
+		fclose(FileHandle);
+
+		double WallTime;
+		double CPULoad[CurrentCPUState.CoreCount];
+		CurrentCPUState.AverageUtilization = 0.0;
+		for (size_t n = 0; n < CurrentCPUState.CoreCount; n++) {
+			if (CurrentCPUState.CurrentUsage[n].TotalTime <= CurrentCPUState.PreviousUsage[n].TotalTime) {
+				CPULoad[n] = 0;
+				continue;
+			}
+
+			WallTime = CurrentCPUState.CurrentUsage[n].TotalTime - CurrentCPUState.PreviousUsage[n].TotalTime;
+			IdleTime = CurrentCPUState.CurrentUsage[n].IdleTime - CurrentCPUState.PreviousUsage[n].IdleTime;
+
+			if (!WallTime || WallTime <= IdleTime) {
+				CPULoad[n] = 0;
+				continue;
+			}
+			CPULoad[n] = (WallTime - (double)IdleTime) * 100.0 / WallTime;
+			CurrentCPUState.Utilization[n] = CPULoad[n];
+			CurrentCPUState.AverageUtilization += CPULoad[n];
+		}
+		CurrentCPUState.AverageUtilization /= (double)CurrentCPUState.CoreCount;
+	}else{
+		FMemory::Memzero(CurrentCPUState);
+	}
+	return CurrentCPUState;
+}
+
+
+
+bool FAndroidMisc::SupportsLocalCaching()
+{
+	return true;
+
+	/*if ( SupportsUTime() )
+	{
+		return true;
+	}*/
+
+
+}
+
+/**
+ * Good enough default crash reporter.
+ */
+void DefaultCrashHandler(const FAndroidCrashContext& Context)
+{
+	static int32 bHasEntered = 0;
+	if (FPlatformAtomics::InterlockedCompareExchange(&bHasEntered, 1, 0) == 0)
+	{
+		const SIZE_T StackTraceSize = 65535;
+		ANSICHAR StackTrace[StackTraceSize];
+		StackTrace[0] = 0;
+
+		FPlatformMisc::LowLevelOutputDebugStringf(TEXT("Starting StackWalk..."));
+
+		// Walk the stack and dump it to the allocated memory.
+		FPlatformStackWalk::StackWalkAndDump(StackTrace, StackTraceSize, 0, Context.Context);
+		UE_LOG(LogAndroid, Error, TEXT("\n%s\n"), ANSI_TO_TCHAR(StackTrace));
+
+		if (GLog)
+		{
+			GLog->SetCurrentThreadAsMasterThread();
+			GLog->Flush();
+		}
+
+		if (GWarn)
+		{
+			GWarn->Flush();
+		}
+	}
+}
+
+/** Global pointer to crash handler */
+void (* GCrashHandlerPointer)(const FGenericCrashContext& Context) = NULL;
+
+const int32 TargetSignals[] =
+{
+	SIGQUIT, // SIGQUIT is a user-initiated "crash".
+	SIGILL,
+	SIGFPE,
+	SIGBUS,
+	SIGSEGV,
+	SIGSYS
+};
+
+const int32 NumTargetSignals = ARRAY_COUNT(TargetSignals);
+
+struct sigaction PrevActions[NumTargetSignals];
+static bool PreviousSignalHandlersValid = false;
+
+static void RestorePreviousSignalHandlers()
+{
+	if (PreviousSignalHandlersValid)
+	{
+		for (int32 i = 0; i < NumTargetSignals; ++i)
+		{
+			sigaction(TargetSignals[i], &PrevActions[i], NULL);
+		}
+		PreviousSignalHandlersValid = false;
+	}
+}
+
+/** True system-specific crash handler that gets called first */
+void PlatformCrashHandler(int32 Signal, siginfo* Info, void* Context)
+{
+	// Switch to malloc crash.
+	//FGenericPlatformMallocCrash::Get().SetAsGMalloc(); @todo uncomment after verification
+
+	FPlatformMisc::LowLevelOutputDebugStringf(TEXT("Signal %d caught!"), Signal);
+
+	// Restore system handlers so Android could catch this signal after we are done with crashreport
+	RestorePreviousSignalHandlers();
+
+	FAndroidCrashContext CrashContext(ECrashContextType::Crash, TEXT("Caught signal"));
+	CrashContext.InitFromSignal(Signal, Info, Context);
+
+	if (GCrashHandlerPointer)
+	{
+		GCrashHandlerPointer(CrashContext);
+	}
+	else
+	{
+		// call default one
+		DefaultCrashHandler(CrashContext);
+	}
+}
+
+void FAndroidMisc::SetCrashHandler(void (* CrashHandler)(const FGenericCrashContext& Context))
+{
+	GCrashHandlerPointer = CrashHandler;
+
+	RestorePreviousSignalHandlers();
+	FMemory::Memzero(&PrevActions, sizeof(PrevActions));
+
+	// Passing -1 will leave these restored and won't trap them
+	if ((PTRINT)CrashHandler == -1)
+	{
+		return;
+	}
+
+	struct sigaction Action;
+	FMemory::Memzero(&Action, sizeof(struct sigaction));
+	Action.sa_sigaction = PlatformCrashHandler;
+	sigemptyset(&Action.sa_mask);
+	Action.sa_flags = SA_SIGINFO | SA_RESTART | SA_ONSTACK;
+
+	for (int32 i = 0; i < NumTargetSignals; ++i)
+	{
+		sigaction(TargetSignals[i],	&Action, &PrevActions[i]);
+	}
+	PreviousSignalHandlersValid = true;
+}
+
+bool FAndroidMisc::GetUseVirtualJoysticks()
+{
+	// joystick on commandline means don't require virtual joysticks
+	if (FParse::Param(FCommandLine::Get(), TEXT("joystick")))
+	{
+		return false;
+	}
+
+	// Amazon Fire TV doesn't require virtual joysticks
+	if (FAndroidMisc::GetDeviceMake() == FString("Amazon"))
+	{
+		if (FAndroidMisc::GetDeviceModel().StartsWith(TEXT("AFT")))
+		{
+			return false;
+		}
+	}
+
+	// Stereo-only HMDs don't require virtual joysticks
+	if (IsStandaloneStereoOnlyDevice())
+	{
+		return false;
+	}
+
+	return true;
+}
+
+bool FAndroidMisc::SupportsTouchInput()
+{
+	// Amazon Fire TV doesn't support touch input
+	if (FAndroidMisc::GetDeviceMake() == FString("Amazon"))
+	{
+		if (FAndroidMisc::GetDeviceModel().StartsWith(TEXT("AFT")))
+		{
+			return false;
+		}
+	}
+
+	// Stereo-only HMDs don't support touch input
+	if (IsStandaloneStereoOnlyDevice())
+	{
+		return false;
+	}
+
+	return true;
+}
+
+bool FAndroidMisc::IsStandaloneStereoOnlyDevice()
+{
+	// Oculus HMDs are always in stereo mode
+	if (FAndroidMisc::GetDeviceMake() == FString("Oculus"))
+	{
+		return true;
+	}
+
+	return false;
+}
+
+extern void AndroidThunkCpp_RegisterForRemoteNotifications();
+extern void AndroidThunkCpp_UnregisterForRemoteNotifications();
+
+void FAndroidMisc::RegisterForRemoteNotifications()
+{
+#if USE_ANDROID_JNI
+	AndroidThunkCpp_RegisterForRemoteNotifications();
+#endif
+}
+
+void FAndroidMisc::UnregisterForRemoteNotifications()
+{
+#if USE_ANDROID_JNI
+	AndroidThunkCpp_UnregisterForRemoteNotifications();
+#endif
+}
+
+TArray<uint8> FAndroidMisc::GetSystemFontBytes()
+{
+#if USE_ANDROID_FILE
+	TArray<uint8> FontBytes;
+	static FString FullFontPath = GFontPathBase + FString(TEXT("DroidSans.ttf"));
+	FFileHelper::LoadFileToArray(FontBytes, *FullFontPath);
+	return FontBytes;
+#else 
+	return FGenericPlatformMisc::GetSystemFontBytes();
+#endif
+}
+
+class IPlatformChunkInstall* FAndroidMisc::GetPlatformChunkInstall()
+{
+	static IPlatformChunkInstall* ChunkInstall = nullptr;
+	static bool bIniChecked = false;
+	if (!ChunkInstall || !bIniChecked)
+	{
+		FString ProviderName;
+		IPlatformChunkInstallModule* PlatformChunkInstallModule = nullptr;
+		if (!GEngineIni.IsEmpty())
+		{
+			FString InstallModule;
+			GConfig->GetString(TEXT("StreamingInstall"), TEXT("DefaultProviderName"), InstallModule, GEngineIni);
+			FModuleStatus Status;
+			if (FModuleManager::Get().QueryModule(*InstallModule, Status))
+			{
+				PlatformChunkInstallModule = FModuleManager::LoadModulePtr<IPlatformChunkInstallModule>(*InstallModule);
+				if (PlatformChunkInstallModule != nullptr)
+				{
+					// Attempt to grab the platform installer
+					ChunkInstall = PlatformChunkInstallModule->GetPlatformChunkInstall();
+				}
+			}
+			bIniChecked = true;
+		}
+		if (!ChunkInstall)
+		{
+			// Placeholder instance
+			ChunkInstall = FGenericPlatformMisc::GetPlatformChunkInstall();
+		}
+	}
+
+	return ChunkInstall;
+}
+
+void FAndroidMisc::PrepareMobileHaptics(EMobileHapticsType Type)
+{
+}
+
+void FAndroidMisc::TriggerMobileHaptics()
+{
+#if USE_ANDROID_JNI
+	extern void AndroidThunkCpp_Vibrate(int32 Duration);
+	// tiny little vibration
+	AndroidThunkCpp_Vibrate(10);
+#endif
+}
+
+void FAndroidMisc::ReleaseMobileHaptics()
+{
+
+}
+
+void FAndroidMisc::ShareURL(const FString& URL, const FText& Description, int32 LocationHintX, int32 LocationHintY)
+{
+#if USE_ANDROID_JNI
+	extern void AndroidThunkCpp_ShareURL(const FString& URL, const FText& Description, const FText& SharePrompt, int32 LocationHintX, int32 LocationHintY);
+	AndroidThunkCpp_ShareURL(URL, Description, NSLOCTEXT("AndroidMisc", "ShareURL", "Share URL"), LocationHintX, LocationHintY);
+#endif
+}
+
+FString FAndroidMisc::LoadTextFileFromPlatformPackage(const FString& RelativePath)
+{
+#if USE_ANDROID_JNI
+	AAssetManager* AssetMgr = AndroidThunkCpp_GetAssetManager();
+	AAsset* asset = AAssetManager_open(AssetMgr, TCHAR_TO_UTF8(*RelativePath), AASSET_MODE_BUFFER);
+
+	if (asset)
+	{
+		const void* FileContents = (const ANSICHAR*)AAsset_getBuffer(asset);
+		int32 FileLength = AAsset_getLength(asset);
+
+		TArray<ANSICHAR> TextContents;
+		TextContents.AddUninitialized(FileLength + 1);
+		FMemory::Memcpy(TextContents.GetData(), FileContents, FileLength);
+		TextContents[FileLength] = 0;
+
+		AAsset_close(asset);
+
+		return FString(ANSI_TO_TCHAR(TextContents.GetData()));
+	}
+#endif
+	return FString();
+}
+
+void FAndroidMisc::SetVersionInfo( FString InAndroidVersion, FString InDeviceMake, FString InDeviceModel, FString InDeviceBuildNumber, FString InOSLanguage )
+{
+	AndroidVersion = InAndroidVersion;
+	DeviceMake = InDeviceMake;
+	DeviceModel = InDeviceModel;
+	DeviceBuildNumber = InDeviceBuildNumber;
+	OSLanguage = InOSLanguage;
+
+	UE_LOG(LogAndroid, Display, TEXT("Android Version Make Model BuildNumber Language: %s %s %s %s %s"), *AndroidVersion, *DeviceMake, *DeviceModel, *DeviceBuildNumber, *OSLanguage);
+}
+
+const FString FAndroidMisc::GetAndroidVersion()
+{
+	return AndroidVersion;
+}
+
+const FString FAndroidMisc::GetDeviceMake()
+{
+	return DeviceMake;
+}
+
+const FString FAndroidMisc::GetDeviceModel()
+{
+	return DeviceModel;
+}
+
+const FString FAndroidMisc::GetDeviceBuildNumber()
+{
+	return DeviceBuildNumber;
+}
+
+const FString FAndroidMisc::GetOSLanguage()
+{
+	return OSLanguage;
+}
+
+const FString FAndroidMisc::GetProjectVersion() {
+	return FString::FromInt(GAndroidPackageVersion);
+}
+
+FString FAndroidMisc::GetDefaultLocale()
+{
+	return OSLanguage;
+}
+
+bool FAndroidMisc::GetVolumeButtonsHandledBySystem()
+{
+	return VolumeButtonsHandledBySystem;
+}
+
+void FAndroidMisc::SetVolumeButtonsHandledBySystem(bool enabled)
+{
+	VolumeButtonsHandledBySystem = enabled;
+}
+
+#if USE_ANDROID_JNI
+int32 FAndroidMisc::GetAndroidBuildVersion()
+{
+	if (AndroidBuildVersion > 0)
+	{
+		return AndroidBuildVersion;
+	}
+	if (AndroidBuildVersion <= 0)
+	{
+		JNIEnv* JEnv = AndroidJavaEnv::GetJavaEnv();
+		if (nullptr != JEnv)
+		{
+			jclass Class = AndroidJavaEnv::FindJavaClassGlobalRef("com/epicgames/ue4/GameActivity");
+			if (nullptr != Class)
+			{
+				jfieldID Field = JEnv->GetStaticFieldID(Class, "ANDROID_BUILD_VERSION", "I");
+				if (nullptr != Field)
+				{
+					AndroidBuildVersion = JEnv->GetStaticIntField(Class, Field);
+				}
+				JEnv->DeleteGlobalRef(Class);
+			}
+		}
+	}
+	return AndroidBuildVersion;
+}
+#endif
+
+bool FAndroidMisc::ShouldDisablePluginAtRuntime(const FString& PluginName)
+{
+#if PLATFORM_ANDROID_ARM64 || PLATFORM_ANDROID_X64
+	// disable OnlineSubsystemGooglePlay for unsupported Android architectures
+	if (PluginName.Equals(TEXT("OnlineSubsystemGooglePlay")))
+	{
+		return true;
+	}
+#endif
+	return false;
+}
+
+void FAndroidMisc::SetThreadName(const char* name)
+{
+#if USE_ANDROID_JNI
+	extern void AndroidThunkCpp_SetThreadName(const char * name);
+	AndroidThunkCpp_SetThreadName(name);
+#endif
+}
+
+///////////////////////////////////////////////////////////////////////////////
+//
+// Extracted from vk_platform.h and vulkan.h with modifications just to allow
+// vkCreateInstance/vkDestroyInstance to be called to check if a driver is actually
+// available (presence of libvulkan.so only means it may be available, not that
+// there is an actual usable one). Cannot wait for VulkanRHI init to do this (too
+// late) and vulkan.h header not guaranteed to be available. This part of the header
+// is unlikely to change in future so safe enough to use this truncated version.
+//
+
+#if PLATFORM_ANDROID_ARM
+// On Android/ARMv7a, Vulkan functions use the armeabi-v7a-hard calling
+#define VKAPI_ATTR __attribute__((pcs("aapcs-vfp")))
+#define VKAPI_CALL
+#define VKAPI_PTR  VKAPI_ATTR
+#else
+// On other platforms, use the default calling convention
+#define VKAPI_ATTR
+#define VKAPI_CALL
+#define VKAPI_PTR
+#endif
+
+#define VK_MAKE_VERSION(major, minor, patch) \
+	(((major) << 22) | ((minor) << 12) | (patch))
+
+#define VK_VERSION_MAJOR(version) ((uint32)(version) >> 22)
+#define VK_VERSION_MINOR(version) (((uint32)(version) >> 12) & 0x3ff)
+#define VK_VERSION_PATCH(version) ((uint32)(version) & 0xfff)
+
+typedef uint32 VkFlags;
+typedef uint32 VkBool32;
+
+#define VK_DEFINE_HANDLE(object) typedef struct object##_T* object;
+
+VK_DEFINE_HANDLE(VkInstance)
+VK_DEFINE_HANDLE(VkPhysicalDevice)
+
+typedef enum VkResult {
+	VK_SUCCESS = 0,
+	VK_NOT_READY = 1,
+	VK_TIMEOUT = 2,
+	VK_EVENT_SET = 3,
+	VK_EVENT_RESET = 4,
+	VK_INCOMPLETE = 5,
+	VK_ERROR_OUT_OF_HOST_MEMORY = -1,
+	VK_ERROR_OUT_OF_DEVICE_MEMORY = -2,
+	VK_ERROR_INITIALIZATION_FAILED = -3,
+	VK_ERROR_DEVICE_LOST = -4,
+	VK_ERROR_MEMORY_MAP_FAILED = -5,
+	VK_ERROR_LAYER_NOT_PRESENT = -6,
+	VK_ERROR_EXTENSION_NOT_PRESENT = -7,
+	VK_ERROR_FEATURE_NOT_PRESENT = -8,
+	VK_ERROR_INCOMPATIBLE_DRIVER = -9,
+	VK_ERROR_TOO_MANY_OBJECTS = -10,
+	VK_ERROR_FORMAT_NOT_SUPPORTED = -11,
+	VK_ERROR_SURFACE_LOST_KHR = -1000000000,
+	VK_ERROR_NATIVE_WINDOW_IN_USE_KHR = -1000000001,
+	VK_SUBOPTIMAL_KHR = 1000001003,
+	VK_ERROR_OUT_OF_DATE_KHR = -1000001004,
+	VK_ERROR_INCOMPATIBLE_DISPLAY_KHR = -1000003001,
+	VK_ERROR_VALIDATION_FAILED_EXT = -1000011001,
+	VK_ERROR_INVALID_SHADER_NV = -1000012000,
+	VK_RESULT_BEGIN_RANGE = VK_ERROR_FORMAT_NOT_SUPPORTED,
+	VK_RESULT_END_RANGE = VK_INCOMPLETE,
+	VK_RESULT_RANGE_SIZE = (VK_INCOMPLETE - VK_ERROR_FORMAT_NOT_SUPPORTED + 1),
+	VK_RESULT_MAX_ENUM = 0x7FFFFFFF
+} VkResult;
+
+typedef enum VkStructureType {
+	VK_STRUCTURE_TYPE_APPLICATION_INFO = 0,
+	VK_STRUCTURE_TYPE_INSTANCE_CREATE_INFO = 1,
+	VK_STRUCTURE_TYPE_MAX_ENUM = 0x7FFFFFFF
+} VkStructureType;
+
+typedef VkFlags VkInstanceCreateFlags;
+
+typedef struct VkApplicationInfo {
+	VkStructureType	sType;
+	const void*		pNext;
+	const char*		pApplicationName;
+	uint32			applicationVersion;
+	const char*		pEngineName;
+	uint			engineVersion;
+	uint			apiVersion;
+} VkApplicationInfo;
+
+typedef struct VkInstanceCreateInfo {
+	VkStructureType				sType;
+	const void*					pNext;
+	VkInstanceCreateFlags		flags;
+	const VkApplicationInfo*	pApplicationInfo;
+	uint32						enabledLayerCount;
+	const char* const*			ppEnabledLayerNames;
+	uint32						enabledExtensionCount;
+	const char* const*			ppEnabledExtensionNames;
+} VkInstanceCreateInfo;
+
+typedef struct VkAllocationCallbacks {
+	void*	pUserData;
+	void*	pfnAllocation;
+	void*	pfnReallocation;
+	void*	pfnFree;
+	void*	pfnInternalAllocation;
+	void*	pfnInternalFree;
+} VkAllocationCallbacks;
+
+typedef uint64 VkDeviceSize;
+typedef VkFlags VkSampleCountFlags;
+
+typedef struct VkPhysicalDeviceLimits {
+	uint32				maxImageDimension1D;
+	uint32				maxImageDimension2D;
+	uint32				maxImageDimension3D;
+	uint32				maxImageDimensionCube;
+	uint32				maxImageArrayLayers;
+	uint32				maxTexelBufferElements;
+	uint32				maxUniformBufferRange;
+	uint32				maxStorageBufferRange;
+	uint32				maxPushConstantsSize;
+	uint32				maxMemoryAllocationCount;
+	uint32				maxSamplerAllocationCount;
+	VkDeviceSize		bufferImageGranularity;
+	VkDeviceSize		sparseAddressSpaceSize;
+	uint32				maxBoundDescriptorSets;
+	uint32				maxPerStageDescriptorSamplers;
+	uint32				maxPerStageDescriptorUniformBuffers;
+	uint32				maxPerStageDescriptorStorageBuffers;
+	uint32				maxPerStageDescriptorSampledImages;
+	uint32				maxPerStageDescriptorStorageImages;
+	uint32				maxPerStageDescriptorInputAttachments;
+	uint32				maxPerStageResources;
+	uint32				maxDescriptorSetSamplers;
+	uint32				maxDescriptorSetUniformBuffers;
+	uint32				maxDescriptorSetUniformBuffersDynamic;
+	uint32				maxDescriptorSetStorageBuffers;
+	uint32				maxDescriptorSetStorageBuffersDynamic;
+	uint32				maxDescriptorSetSampledImages;
+	uint32				maxDescriptorSetStorageImages;
+	uint32				maxDescriptorSetInputAttachments;
+	uint32				maxVertexInputAttributes;
+	uint32				maxVertexInputBindings;
+	uint32				maxVertexInputAttributeOffset;
+	uint32				maxVertexInputBindingStride;
+	uint32				maxVertexOutputComponents;
+	uint32				maxTessellationGenerationLevel;
+	uint32				maxTessellationPatchSize;
+	uint32				maxTessellationControlPerVertexInputComponents;
+	uint32				maxTessellationControlPerVertexOutputComponents;
+	uint32				maxTessellationControlPerPatchOutputComponents;
+	uint32				maxTessellationControlTotalOutputComponents;
+	uint32				maxTessellationEvaluationInputComponents;
+	uint32				maxTessellationEvaluationOutputComponents;
+	uint32				maxGeometryShaderInvocations;
+	uint32				maxGeometryInputComponents;
+	uint32				maxGeometryOutputComponents;
+	uint32				maxGeometryOutputVertices;
+	uint32				maxGeometryTotalOutputComponents;
+	uint32				maxFragmentInputComponents;
+	uint32				maxFragmentOutputAttachments;
+	uint32				maxFragmentDualSrcAttachments;
+	uint32				maxFragmentCombinedOutputResources;
+	uint32				maxComputeSharedMemorySize;
+	uint32				maxComputeWorkGroupCount[3];
+	uint32				maxComputeWorkGroupInvocations;
+	uint32				maxComputeWorkGroupSize[3];
+	uint32				subPixelPrecisionBits;
+	uint32				subTexelPrecisionBits;
+	uint32				mipmapPrecisionBits;
+	uint32				maxDrawIndexedIndexValue;
+	uint32				maxDrawIndirectCount;
+	float				maxSamplerLodBias;
+	float				maxSamplerAnisotropy;
+	uint32				maxViewports;
+	uint32				maxViewportDimensions[2];
+	float				viewportBoundsRange[2];
+	uint32				viewportSubPixelBits;
+	size_t				minMemoryMapAlignment;
+	VkDeviceSize		minTexelBufferOffsetAlignment;
+	VkDeviceSize		minUniformBufferOffsetAlignment;
+	VkDeviceSize		minStorageBufferOffsetAlignment;
+	int32				minTexelOffset;
+	uint32				maxTexelOffset;
+	int32				minTexelGatherOffset;
+	uint32				maxTexelGatherOffset;
+	float				minInterpolationOffset;
+	float				maxInterpolationOffset;
+	uint32				subPixelInterpolationOffsetBits;
+	uint32				maxFramebufferWidth;
+	uint32				maxFramebufferHeight;
+	uint32				maxFramebufferLayers;
+	VkSampleCountFlags	framebufferColorSampleCounts;
+	VkSampleCountFlags	framebufferDepthSampleCounts;
+	VkSampleCountFlags	framebufferStencilSampleCounts;
+	VkSampleCountFlags	framebufferNoAttachmentsSampleCounts;
+	uint32				maxColorAttachments;
+	VkSampleCountFlags	sampledImageColorSampleCounts;
+	VkSampleCountFlags	sampledImageIntegerSampleCounts;
+	VkSampleCountFlags	sampledImageDepthSampleCounts;
+	VkSampleCountFlags	sampledImageStencilSampleCounts;
+	VkSampleCountFlags	storageImageSampleCounts;
+	uint32				maxSampleMaskWords;
+	VkBool32			timestampComputeAndGraphics;
+	float				timestampPeriod;
+	uint32				maxClipDistances;
+	uint32				maxCullDistances;
+	uint32				maxCombinedClipAndCullDistances;
+	uint32				discreteQueuePriorities;
+	float				pointSizeRange[2];
+	float				lineWidthRange[2];
+	float				pointSizeGranularity;
+	float				lineWidthGranularity;
+	VkBool32			strictLines;
+	VkBool32			standardSampleLocations;
+	VkDeviceSize		optimalBufferCopyOffsetAlignment;
+	VkDeviceSize		optimalBufferCopyRowPitchAlignment;
+	VkDeviceSize		nonCoherentAtomSize;
+} VkPhysicalDeviceLimits;
+
+typedef struct VkPhysicalDeviceSparseProperties {
+	VkBool32	residencyStandard2DBlockShape;
+	VkBool32	residencyStandard2DMultisampleBlockShape;
+	VkBool32	residencyStandard3DBlockShape;
+	VkBool32	residencyAlignedMipSize;
+	VkBool32	residencyNonResidentStrict;
+} VkPhysicalDeviceSparseProperties;
+
+typedef enum VkPhysicalDeviceType {
+	VK_PHYSICAL_DEVICE_TYPE_OTHER = 0,
+	VK_PHYSICAL_DEVICE_TYPE_INTEGRATED_GPU = 1,
+	VK_PHYSICAL_DEVICE_TYPE_DISCRETE_GPU = 2,
+	VK_PHYSICAL_DEVICE_TYPE_VIRTUAL_GPU = 3,
+	VK_PHYSICAL_DEVICE_TYPE_CPU = 4,
+	VK_PHYSICAL_DEVICE_TYPE_BEGIN_RANGE = VK_PHYSICAL_DEVICE_TYPE_OTHER,
+	VK_PHYSICAL_DEVICE_TYPE_END_RANGE = VK_PHYSICAL_DEVICE_TYPE_CPU,
+	VK_PHYSICAL_DEVICE_TYPE_RANGE_SIZE = (VK_PHYSICAL_DEVICE_TYPE_CPU - VK_PHYSICAL_DEVICE_TYPE_OTHER + 1),
+	VK_PHYSICAL_DEVICE_TYPE_MAX_ENUM = 0x7FFFFFFF
+} VkPhysicalDeviceType;
+
+#define VK_MAX_PHYSICAL_DEVICE_NAME_SIZE 256
+#define VK_UUID_SIZE 16
+
+typedef struct VkPhysicalDeviceProperties {
+	uint32								apiVersion;
+	uint32								driverVersion;
+	uint32								vendorID;
+	uint32								deviceID;
+	VkPhysicalDeviceType				deviceType;
+	char								deviceName[VK_MAX_PHYSICAL_DEVICE_NAME_SIZE];
+	uint8_t								pipelineCacheUUID[VK_UUID_SIZE];
+	VkPhysicalDeviceLimits				limits;
+	VkPhysicalDeviceSparseProperties	sparseProperties;
+} VkPhysicalDeviceProperties;
+
+#define VK_MAX_EXTENSION_NAME_SIZE        256
+#define VK_GOOGLE_DISPLAY_TIMING_EXTENSION_NAME "VK_GOOGLE_display_timing"
+
+typedef struct VkExtensionProperties {
+	char        extensionName[VK_MAX_EXTENSION_NAME_SIZE];
+	uint32_t    specVersion;
+} VkExtensionProperties;
+
+typedef VkResult(VKAPI_PTR *PFN_vkCreateInstance)(const VkInstanceCreateInfo* pCreateInfo, const VkAllocationCallbacks* pAllocator, VkInstance* pInstance);
+typedef void (VKAPI_PTR *PFN_vkDestroyInstance)(VkInstance instance, const VkAllocationCallbacks* pAllocator);
+typedef VkResult(VKAPI_PTR *PFN_vkEnumeratePhysicalDevices)(VkInstance instance, uint32* pPhysicalDeviceCount, VkPhysicalDevice* pPhysicalDevices);
+typedef void (VKAPI_PTR *PFN_vkGetPhysicalDeviceProperties)(VkPhysicalDevice physicalDevice, VkPhysicalDeviceProperties* pProperties);
+typedef VkResult(VKAPI_PTR *PFN_vkEnumerateDeviceExtensionProperties)(VkPhysicalDevice physicalDevice, const char* pLayerName, uint32_t* pPropertyCount, VkExtensionProperties* pProperties);
+
+///////////////////////////////////////////////////////////////////////////////
+
+#define UE_VK_API_VERSION	VK_MAKE_VERSION(1, 0, 1)
+
+enum class EDeviceVulkanSupportStatus
+{
+	Uninitialized,
+	NotSupported,
+	Supported
+};
+
+static FString VulkanVersionString;
+static EDeviceVulkanSupportStatus VulkanSupport = EDeviceVulkanSupportStatus::Uninitialized;
+
+static EDeviceVulkanSupportStatus AttemptVulkanInit(void* VulkanLib)
+{
+	if (VulkanLib == nullptr)
+	{
+		return EDeviceVulkanSupportStatus::NotSupported;
+	}
+
+	// Try to get required functions to check for driver
+	PFN_vkCreateInstance vkCreateInstance = (PFN_vkCreateInstance)dlsym(VulkanLib, "vkCreateInstance");
+	PFN_vkDestroyInstance vkDestroyInstance = (PFN_vkDestroyInstance)dlsym(VulkanLib, "vkDestroyInstance");
+	PFN_vkEnumeratePhysicalDevices vkEnumeratePhysicalDevices = (PFN_vkEnumeratePhysicalDevices)dlsym(VulkanLib, "vkEnumeratePhysicalDevices");
+	PFN_vkGetPhysicalDeviceProperties vkGetPhysicalDeviceProperties = (PFN_vkGetPhysicalDeviceProperties)dlsym(VulkanLib, "vkGetPhysicalDeviceProperties");
+	PFN_vkEnumerateDeviceExtensionProperties vkEnumerateDeviceExtensionProperties = (PFN_vkEnumerateDeviceExtensionProperties)dlsym(VulkanLib, "vkEnumerateDeviceExtensionProperties");
+
+	if (!vkCreateInstance || !vkDestroyInstance || !vkEnumeratePhysicalDevices || !vkGetPhysicalDeviceProperties || !vkEnumerateDeviceExtensionProperties)
+	{
+		UE_LOG(LogAndroid, Log, TEXT("Vulkan not supported: vkCreateInstance: 0x%p, vkDestroyInstance: 0x%p, vkEnumeratePhysicalDevices: 0x%p, vkGetPhysicalDeviceProperties: 0x%p, vkEnumerateDeviceExtensionProperties: 0x%p"), vkCreateInstance, vkDestroyInstance, vkEnumeratePhysicalDevices, vkGetPhysicalDeviceProperties, vkEnumerateDeviceExtensionProperties);
+		return EDeviceVulkanSupportStatus::NotSupported;
+	}
+
+	// try to create instance to verify driver available
+	VkApplicationInfo App;
+	FMemory::Memzero(App);
+	App.sType = VK_STRUCTURE_TYPE_APPLICATION_INFO;
+	App.pApplicationName = "UE4";
+	App.applicationVersion = 0;
+	App.pEngineName = "UE4";
+	App.engineVersion = 0;
+	App.apiVersion = UE_VK_API_VERSION;
+
+	VkInstanceCreateInfo InstInfo;
+	FMemory::Memzero(InstInfo);
+	InstInfo.sType = VK_STRUCTURE_TYPE_INSTANCE_CREATE_INFO;
+	InstInfo.pNext = nullptr;
+	InstInfo.pApplicationInfo = &App;
+	InstInfo.enabledExtensionCount = 0;
+	InstInfo.ppEnabledExtensionNames = nullptr;
+
+	VkInstance Instance;
+	VkResult Result = vkCreateInstance(&InstInfo, nullptr, &Instance);
+	if (Result != VK_SUCCESS)
+	{
+		return EDeviceVulkanSupportStatus::NotSupported;
+	}
+
+	// Determine Vulkan device's API level.
+	uint32 GpuCount = 0;
+	Result = vkEnumeratePhysicalDevices(Instance, &GpuCount, nullptr);
+	if (Result != VK_SUCCESS || GpuCount == 0)
+	{
+		vkDestroyInstance(Instance, nullptr);
+		return EDeviceVulkanSupportStatus::NotSupported;
+	}
+
+	TArray<VkPhysicalDevice> PhysicalDevices;
+	PhysicalDevices.AddZeroed(GpuCount);
+	Result = vkEnumeratePhysicalDevices(Instance, &GpuCount, PhysicalDevices.GetData());
+	if (Result != VK_SUCCESS)
+	{
+		vkDestroyInstance(Instance, nullptr);
+		return EDeviceVulkanSupportStatus::NotSupported;
+	}
+
+	// Don't care which device - This code is making the assumption that all devices will have same api version.
+	VkPhysicalDeviceProperties DeviceProperties;
+	vkGetPhysicalDeviceProperties(PhysicalDevices[0], &DeviceProperties);
+
+	//for now we are allowing devices without the timing extension to run with a basic CPU frame pacer.
+#if 0
+	bool bHasVKGoogleDisplayTiming = false;
+	{
+		TArray<VkExtensionProperties> ExtensionProps;
+		do
+		{
+			uint32 Count = 0;
+			Result = vkEnumerateDeviceExtensionProperties(PhysicalDevices[0], nullptr, &Count, nullptr);
+			check(Result >= VK_SUCCESS);
+
+			if (Count > 0)
+			{
+				ExtensionProps.Empty(Count);
+				ExtensionProps.AddUninitialized(Count);
+				Result = vkEnumerateDeviceExtensionProperties(PhysicalDevices[0], nullptr, &Count, ExtensionProps.GetData());
+				check(Result >= VK_SUCCESS);
+			}
+		} while (Result == VK_INCOMPLETE);		
+		check(Result >= VK_SUCCESS);
+
+		
+		for (int32 i = 0; i < ExtensionProps.Num(); ++i)
+		{
+			UE_LOG(LogAndroid, Log, TEXT("Checking extension: %s."), ANSI_TO_TCHAR(ExtensionProps[i].extensionName));
+			if (!FCStringAnsi::Strcmp(VK_GOOGLE_DISPLAY_TIMING_EXTENSION_NAME, ExtensionProps[i].extensionName))
+			{
+				bHasVKGoogleDisplayTiming = true;
+				break;
+			}
+		}
+	}
+	if (!bHasVKGoogleDisplayTiming)
+	{
+		vkDestroyInstance(Instance, nullptr);
+
+		UE_LOG(LogAndroid, Log, TEXT("Vulkan not supported, cannot find VK_GOOGLE_display_timing extension."));
+		return EDeviceVulkanSupportStatus::NotSupported;
+	}
+#endif
+
+	VulkanVersionString = FString::Printf(TEXT("%d.%d.%d"), VK_VERSION_MAJOR(DeviceProperties.apiVersion), VK_VERSION_MINOR(DeviceProperties.apiVersion), VK_VERSION_PATCH(DeviceProperties.apiVersion));
+	vkDestroyInstance(Instance, nullptr);
+
+	return EDeviceVulkanSupportStatus::Supported;
+}
+
+bool FAndroidMisc::HasVulkanDriverSupport()
+{
+// @todo Lumin: this isn't really the best #define to check here - but basically, without JNI and other version checking, we can't safely do it - we'll need
+// non-JNI platforms that support Vulkan to figure out a way to force it (if they want GL + Vulkan support)
+#if !USE_ANDROID_JNI
+	VulkanSupport = EDeviceVulkanSupportStatus::NotSupported;
+	VulkanVersionString = TEXT("0.0.0");
+#else
+	// this version does not check for VulkanRHI or disabled by cvars!
+	if (VulkanSupport == EDeviceVulkanSupportStatus::Uninitialized)
+
+	{
+		// assume no
+		VulkanSupport = EDeviceVulkanSupportStatus::NotSupported;
+		VulkanVersionString = TEXT("0.0.0");
+
+		// check for libvulkan.so
+		void* VulkanLib = dlopen("libvulkan.so", RTLD_NOW | RTLD_LOCAL);
+		if (VulkanLib != nullptr)
+		{
+			UE_LOG(LogAndroid, Log, TEXT("Vulkan library detected, checking for available driver"));
+
+			// if Nougat, we can check the Vulkan version
+			if (FAndroidMisc::GetAndroidBuildVersion() >= 24)
+			{
+				extern int32 AndroidThunkCpp_GetMetaDataInt(const FString& Key);
+				int32 VulkanVersion = AndroidThunkCpp_GetMetaDataInt(TEXT("android.hardware.vulkan.version"));
+				if (VulkanVersion >= UE_VK_API_VERSION)
+				{
+					// final check, try initializing the instance
+					VulkanSupport = AttemptVulkanInit(VulkanLib);
+				}
+			}
+			else
+			{
+				// otherwise, we need to try initializing the instance
+				VulkanSupport = AttemptVulkanInit(VulkanLib);
+			}
+
+			dlclose(VulkanLib);
+
+			if (VulkanSupport == EDeviceVulkanSupportStatus::Supported)
+			{
+				UE_LOG(LogAndroid, Log, TEXT("VulkanRHI is available, Vulkan capable device detected."));
+				return true;
+			}
+			else
+			{
+				UE_LOG(LogAndroid, Log, TEXT("Vulkan driver NOT available."));
+			}
+		}
+		else
+		{
+			UE_LOG(LogAndroid, Log, TEXT("Vulkan library NOT detected."));
+		}
+	}
+#endif
+	return VulkanSupport == EDeviceVulkanSupportStatus::Supported;
+}
+
+// Test for device vulkan support.
+static void EstablishVulkanDeviceSupport()
+{
+	// just do this check once
+	if (VulkanSupport == EDeviceVulkanSupportStatus::Uninitialized)
+	{
+		// this call will initialize VulkanSupport
+		FAndroidMisc::HasVulkanDriverSupport();
+	}
+}
+
+bool FAndroidMisc::IsVulkanAvailable()
+{
+	check(VulkanSupport != EDeviceVulkanSupportStatus::Uninitialized);
+
+	static int CachedVulkanAvailable = -1;
+	if (CachedVulkanAvailable == -1)
+	{
+		CachedVulkanAvailable = 0;
+		if (VulkanSupport == EDeviceVulkanSupportStatus::Supported)
+		{
+			bool bSupportsVulkan = false;
+			GConfig->GetBool(TEXT("/Script/AndroidRuntimeSettings.AndroidRuntimeSettings"), TEXT("bSupportsVulkan"), bSupportsVulkan, GEngineIni);
+
+			// whether to detect Vulkan by default or require the -detectvulkan command line parameter
+			bool bDetectVulkanByDefault = true;
+			GConfig->GetBool(TEXT("/Script/AndroidRuntimeSettings.AndroidRuntimeSettings"), TEXT("bDetectVulkanByDefault"), bDetectVulkanByDefault, GEngineIni);
+			const bool bDetectVulkanCmdLine = FParse::Param(FCommandLine::Get(), TEXT("detectvulkan"));
+
+			// @todo Lumin: Double check all this stuff after merging general android Vulkan SM5 from main
+			const bool bSupportsVulkanSM5 = ShouldUseDesktopVulkan();
+
+			const bool bVulkanDisabledCmdLine = FParse::Param(FCommandLine::Get(), TEXT("GL")) || FParse::Param(FCommandLine::Get(), TEXT("OpenGL")) || FParse::Param(FCommandLine::Get(), TEXT("ES2"));
+
+			if (!FModuleManager::Get().ModuleExists(TEXT("VulkanRHI")))
+			{
+				UE_LOG(LogAndroid, Log, TEXT("Vulkan not available as VulkanRHI not present."));
+			}
+			else if (!(bSupportsVulkan || bSupportsVulkanSM5))
+			{
+				UE_LOG(LogAndroid, Log, TEXT("Vulkan not available as project packaged without bSupportsVulkan or bSupportsVulkanSM5."));
+			}
+			else if (bVulkanDisabledCmdLine)
+			{
+				UE_LOG(LogAndroid, Log, TEXT("Vulkan API detection is disabled by a command line option."));
+			}
+            else if (!bDetectVulkanByDefault && !bDetectVulkanCmdLine)
+			{
+				UE_LOG(LogAndroid, Log, TEXT("Vulkan available but detection disabled by bDetectVulkanByDefault=False in AndroidRuntimeSettings. Use -detectvulkan to override."));
+			}
+			else
+			{
+				CachedVulkanAvailable = 1;
+			}
+		}
+	}
+
+	return CachedVulkanAvailable == 1;
+}
+
+bool FAndroidMisc::ShouldUseVulkan()
+{
+	check(VulkanSupport != EDeviceVulkanSupportStatus::Uninitialized);
+	static int CachedShouldUseVulkan = -1;
+
+	if (CachedShouldUseVulkan == -1)
+	{
+		CachedShouldUseVulkan = 0;
+
+		static const auto CVarDisableVulkan = IConsoleManager::Get().FindTConsoleVariableDataInt(TEXT("r.Android.DisableVulkanSupport"));
+
+		const bool bVulkanAvailable = IsVulkanAvailable();
+		const bool bVulkanDisabledCVar = CVarDisableVulkan->GetValueOnAnyThread() == 1;
+
+		if (bVulkanAvailable && !bVulkanDisabledCVar)
+		{
+			CachedShouldUseVulkan = 1;
+			UE_LOG(LogAndroid, Log, TEXT("VulkanRHI will be used!"));
+		}
+		else
+		{
+			UE_LOG(LogAndroid, Log, TEXT("VulkanRHI will NOT be used:"));
+			if (!bVulkanAvailable)
+			{
+				UE_LOG(LogAndroid, Log, TEXT(" ** Vulkan support is not available (Driver, RHI or shaders are missing, or disabled by cmdline, see above logging for details)"));
+			}
+			if (bVulkanDisabledCVar)
+			{
+				UE_LOG(LogAndroid, Log, TEXT(" ** Vulkan is disabled via console variable."));
+			}
+			UE_LOG(LogAndroid, Log, TEXT("OpenGL ES will be used."));
+		}
+	}
+
+	return CachedShouldUseVulkan == 1;
+}
+
+bool FAndroidMisc::ShouldUseDesktopVulkan()
+{
+	// @todo Lumin: Double check all this stuff after merging general android Vulkan SM5 from main
+	bool bSupportsVulkanSM5 = false;
+	GConfig->GetBool(TEXT("/Script/AndroidRuntimeSettings.AndroidRuntimeSettings"), TEXT("bSupportsVulkanSM5"), bSupportsVulkanSM5, GEngineIni);
+
+	return bSupportsVulkanSM5;
+}
+
+FString FAndroidMisc::GetVulkanVersion()
+{
+	check(VulkanSupport != EDeviceVulkanSupportStatus::Uninitialized);
+	return VulkanVersionString;
+}
+
+TMap<FString, FString> FAndroidMisc::GetConfigRulesTMap()
+{
+	return ConfigRulesVariables;
+}
+
+FString* FAndroidMisc::GetConfigRulesVariable(const FString& Key)
+{
+	return ConfigRulesVariables.Find(Key);
+}
+
+JNI_METHOD void Java_com_epicgames_ue4_GameActivity_nativeSetConfigRulesVariables(JNIEnv* jenv, jobject thiz, jobjectArray KeyValuePairs)
+{
+	int32 Count = jenv->GetArrayLength(KeyValuePairs);
+	int32 Index = 0;
+	while (Index < Count)
+	{
+		auto javaKey = FJavaHelper::FStringFromLocalRef(jenv, (jstring)(jenv->GetObjectArrayElement(KeyValuePairs, Index++)));
+		auto javaValue = FJavaHelper::FStringFromLocalRef(jenv, (jstring)(jenv->GetObjectArrayElement(KeyValuePairs, Index++)));
+		
+		FAndroidMisc::ConfigRulesVariables.Add(javaKey, javaValue);
+	}
+}
+
+extern bool AndroidThunkCpp_HasMetaDataKey(const FString& Key);
+
+bool FAndroidMisc::IsDaydreamApplication()
+{
+#if USE_ANDROID_JNI
+	static const bool bIsDaydreamApplication = AndroidThunkCpp_HasMetaDataKey(TEXT("com.epicgames.ue4.GameActivity.bDaydream"));
+	return bIsDaydreamApplication;
+#else
+	return false;
+#endif
+}
+
+static bool bDetectedDebugger = false;
+
+JNI_METHOD void Java_com_epicgames_ue4_GameActivity_nativeSetAndroidStartupState(JNIEnv* jenv, jobject thiz, jboolean bDebuggerAttached)
+{
+	// if Java debugger attached, mark detected (but don't lose previous trigger state)
+	if (bDebuggerAttached)
+	{
+		bDetectedDebugger = true;
+	}
+}
+
+#if !UE_BUILD_SHIPPING
+bool FAndroidMisc::IsDebuggerPresent()
+{
+	extern CORE_API bool GIgnoreDebugger;
+	if (GIgnoreDebugger)
+	{
+		return false;
+	}
+
+	if (bDetectedDebugger)
+	{
+		return true;
+	}
+
+	// If a process is tracing this one then TracerPid in /proc/self/status will
+	// be the id of the tracing process. Use SignalHandler safe functions
+
+	int StatusFile = open("/proc/self/status", O_RDONLY);
+	if (StatusFile == -1)
+	{
+		// Failed - unknown debugger status.
+		return false;
+	}
+
+	char Buffer[256];
+	ssize_t Length = read(StatusFile, Buffer, sizeof(Buffer));
+
+	bool bDebugging = false;
+	const char* TracerString = "TracerPid:\t";
+	const ssize_t LenTracerString = strlen(TracerString);
+	int i = 0;
+
+	while ((Length - i) > LenTracerString)
+	{
+		// TracerPid is found
+		if (strncmp(&Buffer[i], TracerString, LenTracerString) == 0)
+		{
+			// 0 if no process is tracing.
+			bDebugging = Buffer[i + LenTracerString] != '0';
+			break;
+		}
+		++i;
+	}
+
+	close(StatusFile);
+
+	// remember if we detected debugger so we can skip check next time
+	if (bDebugging)
+	{
+		bDetectedDebugger = true;
+	}
+
+	return bDebugging;
+}
+#endif
+
+#if STATS || ENABLE_STATNAMEDEVENTS
+
+void FAndroidMisc::BeginNamedEventFrame()
+{
+#if FRAMEPRO_ENABLED
+	FFrameProProfiler::FrameStart();
+#endif // FRAMEPRO_ENABLED
+}
+
+static void WriteTraceMarkerEvent(const ANSICHAR* Text, int32 TraceMarkerFileDescriptor)
+{
+	if (bUseNativeSystrace)
+	{
+		ATrace_beginSection(Text);
+	}
+	else
+	{
+		const int MAX_TRACE_EVENT_LENGTH = 256;
+		ANSICHAR EventBuffer[MAX_TRACE_EVENT_LENGTH];
+		int EventLength = snprintf(EventBuffer, MAX_TRACE_EVENT_LENGTH, "B|%d|%s", getpid(), Text);
+
+		write(TraceMarkerFileDescriptor, EventBuffer, EventLength);
+	}
+}
+
+void FAndroidMisc::BeginNamedEvent(const struct FColor& Color, const TCHAR* Text)
+{
+#if FRAMEPRO_ENABLED
+	FFrameProProfiler::PushEvent(Text);
+#endif // FRAMEPRO_ENABLED
+	if (bUseNativeSystrace ? !ATrace_isEnabled() : TraceMarkerFileDescriptor == -1)
+	{
+		return;
+	}
+
+	const int MAX_TRACE_MESSAGE_LENGTH = 256;
+
+	// not static since may be called by different threads
+	ANSICHAR TextBuffer[MAX_TRACE_MESSAGE_LENGTH];
+
+	const TCHAR* SourcePtr = Text;
+	ANSICHAR* WritePtr = TextBuffer;
+	int32 RemainingSpace = MAX_TRACE_MESSAGE_LENGTH;
+	while (*SourcePtr && --RemainingSpace > 0)
+	{
+		*WritePtr++ = static_cast<ANSICHAR>(*SourcePtr++);
+	}
+	*WritePtr = '\0';
+
+	WriteTraceMarkerEvent(TextBuffer, TraceMarkerFileDescriptor);
+}
+
+void FAndroidMisc::BeginNamedEvent(const struct FColor& Color, const ANSICHAR* Text)
+{
+#if FRAMEPRO_ENABLED
+	FFrameProProfiler::PushEvent(Text);
+#endif // FRAMEPRO_ENABLED
+	if (bUseNativeSystrace ? !ATrace_isEnabled() : TraceMarkerFileDescriptor == -1)
+	{
+		return;
+	}
+
+	WriteTraceMarkerEvent(Text, TraceMarkerFileDescriptor);
+}
+
+void FAndroidMisc::EndNamedEvent()
+{
+#if FRAMEPRO_ENABLED
+	FFrameProProfiler::PopEvent();
+#endif // FRAMEPRO_ENABLED
+	if (bUseNativeSystrace ? !ATrace_isEnabled() : TraceMarkerFileDescriptor == -1)
+	{
+		return;
+	}
+
+	if (bUseNativeSystrace)
+	{
+		ATrace_endSection();
+	}
+	else
+	{
+		const ANSICHAR EventTerminatorChar = 'E';
+		write(TraceMarkerFileDescriptor, &EventTerminatorChar, 1);
+	}
+}
+
+void FAndroidMisc::CustomNamedStat(const TCHAR* Text, float Value, const TCHAR* Graph, const TCHAR* Unit)
+{
+	FRAMEPRO_DYNAMIC_CUSTOM_STAT(TCHAR_TO_WCHAR(Text), Value, TCHAR_TO_WCHAR(Graph), TCHAR_TO_WCHAR(Unit));
+}
+
+void FAndroidMisc::CustomNamedStat(const ANSICHAR* Text, float Value, const ANSICHAR* Graph, const ANSICHAR* Unit)
+{
+	FRAMEPRO_DYNAMIC_CUSTOM_STAT(Text, Value, Graph, Unit);
+}
+
+#endif // STATS || ENABLE_STATNAMEDEVENTS
+
+int FAndroidMisc::GetVolumeState(double* OutTimeOfChangeInSec)
+{
+	int v;
+	ReceiversLock.Lock();
+	v = CurrentVolume.Volume;
+	if (OutTimeOfChangeInSec)
+	{
+		*OutTimeOfChangeInSec = CurrentVolume.TimeOfChange;
+	}
+	ReceiversLock.Unlock();
+	return v;
+}
+
+int32 FAndroidMisc::GetDeviceVolume()
+{
+	//FAndroidMisc::GetVolumeState returns 0-15, scale to 0-100
+	int32 BaseVolume = FAndroidMisc::GetVolumeState();
+	int32 ScaledVolume = (BaseVolume * 100) / 15;
+	return ScaledVolume;
+}
+
+#if USE_ANDROID_FILE
+const TCHAR* FAndroidMisc::GamePersistentDownloadDir()
+{
+	extern FString GExternalFilePath;
+	return *GExternalFilePath;
+}
+
+FString FAndroidMisc::GetLoginId()
+{
+	static FString LoginId = TEXT("");
+
+	// Return already loaded or generated Id
+	if (!LoginId.IsEmpty())
+	{
+		return LoginId;
+	}
+
+	// Check for existing identifier file
+	extern FString GInternalFilePath;
+	extern FString GExternalFilePath;
+	FString InternalLoginIdFilename = GInternalFilePath / TEXT("login-identifier.txt");
+	if (FPaths::FileExists(InternalLoginIdFilename))
+	{
+		if (FFileHelper::LoadFileToString(LoginId, *InternalLoginIdFilename))
+		{
+			return LoginId;
+		}
+	}
+	FString LoginIdFilename = GExternalFilePath / TEXT("login-identifier.txt");
+	if (FPaths::FileExists(LoginIdFilename))
+	{
+		if (FFileHelper::LoadFileToString(LoginId, *LoginIdFilename))
+		{
+			FFileHelper::SaveStringToFile(LoginId, *InternalLoginIdFilename);
+			return LoginId;
+		}
+	}
+
+	// Generate a new one and write to file
+	FGuid DeviceGuid;
+	FPlatformMisc::CreateGuid(DeviceGuid);
+	LoginId = DeviceGuid.ToString();
+	FFileHelper::SaveStringToFile(LoginId, *InternalLoginIdFilename);
+
+	return LoginId;
+}
+#endif
+
+#if USE_ANDROID_JNI
+FString FAndroidMisc::GetDeviceId()
+{
+	extern FString AndroidThunkCpp_GetAndroidId();
+	static FString DeviceId = AndroidThunkCpp_GetAndroidId();
+
+	// note: this can be empty or NOT unique depending on the OEM implementation!
+	return DeviceId;
+}
+
+FString FAndroidMisc::GetUniqueAdvertisingId()
+{
+	extern FString AndroidThunkCpp_GetAdvertisingId();
+	static FString AdvertisingId = AndroidThunkCpp_GetAdvertisingId();
+
+	// note: this can be empty if Google Play not installed, or user is blocking it!
+	return AdvertisingId;
+}
+#endif
+
+FAndroidMisc::FBatteryState FAndroidMisc::GetBatteryState()
+{
+	FBatteryState CurState;
+	ReceiversLock.Lock();
+	CurState = CurrentBatteryState;
+	ReceiversLock.Unlock();
+	return CurState;
+}
+
+int FAndroidMisc::GetBatteryLevel()
+{
+	FBatteryState BatteryState = GetBatteryState();
+	return BatteryState.Level;
+}
+
+bool FAndroidMisc::IsRunningOnBattery()
+{
+	FBatteryState BatteryState = GetBatteryState();
+	return BatteryState.State == BATTERY_STATE_DISCHARGING;
+}
+
+bool FAndroidMisc::IsInLowPowerMode()
+{
+	FBatteryState BatteryState = GetBatteryState();
+	return BatteryState.Level <= GAndroidLowPowerBatteryThreshold;
+}
+
+float FAndroidMisc::GetDeviceTemperatureLevel()
+{
+	return GetBatteryState().Temperature;
+}
+
+bool FAndroidMisc::AreHeadPhonesPluggedIn()
+{
+	return HeadPhonesArePluggedIn;
+}
+
+#define ANDROIDTHUNK_CONNECTION_TYPE_NONE 0
+#define ANDROIDTHUNK_CONNECTION_TYPE_AIRPLANEMODE 1
+#define ANDROIDTHUNK_CONNECTION_TYPE_ETHERNET 2
+#define ANDROIDTHUNK_CONNECTION_TYPE_CELL 3
+#define ANDROIDTHUNK_CONNECTION_TYPE_WIFI 4
+#define ANDROIDTHUNK_CONNECTION_TYPE_WIMAX 5
+#define ANDROIDTHUNK_CONNECTION_TYPE_BLUETOOTH 6
+
+ENetworkConnectionType FAndroidMisc::GetNetworkConnectionType()
+{
+#if USE_ANDROID_JNI
+	extern int32 AndroidThunkCpp_GetNetworkConnectionType();
+
+	switch (AndroidThunkCpp_GetNetworkConnectionType())
+	{
+		case ANDROIDTHUNK_CONNECTION_TYPE_NONE:				return ENetworkConnectionType::None;
+		case ANDROIDTHUNK_CONNECTION_TYPE_AIRPLANEMODE:		return ENetworkConnectionType::AirplaneMode;
+		case ANDROIDTHUNK_CONNECTION_TYPE_ETHERNET:			return ENetworkConnectionType::Ethernet;
+		case ANDROIDTHUNK_CONNECTION_TYPE_CELL:				return ENetworkConnectionType::Cell;
+		case ANDROIDTHUNK_CONNECTION_TYPE_WIFI:				return ENetworkConnectionType::WiFi;
+		case ANDROIDTHUNK_CONNECTION_TYPE_WIMAX:			return ENetworkConnectionType::WiMAX;
+		case ANDROIDTHUNK_CONNECTION_TYPE_BLUETOOTH:		return ENetworkConnectionType::Bluetooth;
+	}
+#endif
+	return ENetworkConnectionType::Unknown;
+}
+
+#if USE_ANDROID_JNI
+bool FAndroidMisc::HasActiveWiFiConnection()
+{
+	ENetworkConnectionType ConnectionType = GetNetworkConnectionType();
+	return (ConnectionType == ENetworkConnectionType::WiFi ||
+			ConnectionType == ENetworkConnectionType::WiMAX);
+}
+#endif
+
+static FAndroidMisc::ReInitWindowCallbackType OnReInitWindowCallback;
+
+FAndroidMisc::ReInitWindowCallbackType FAndroidMisc::GetOnReInitWindowCallback()
+{
+	return OnReInitWindowCallback;
+}
+
+void FAndroidMisc::SetOnReInitWindowCallback(FAndroidMisc::ReInitWindowCallbackType InOnReInitWindowCallback)
+{
+	OnReInitWindowCallback = InOnReInitWindowCallback;
+}
+
+static FAndroidMisc::OnPauseCallBackType OnPauseCallback;
+
+FAndroidMisc::OnPauseCallBackType FAndroidMisc::GetOnPauseCallback()
+{
+	return OnPauseCallback;
+}
+
+void FAndroidMisc::SetOnPauseCallback(FAndroidMisc::OnPauseCallBackType InOnPauseCallback)
+{
+	OnPauseCallback = InOnPauseCallback;
+}
+
+FString FAndroidMisc::GetCPUVendor()
+{
+	return DeviceMake;
+}
+
+FString FAndroidMisc::GetCPUBrand()
+{
+	return DeviceModel;
+}
+
+FString FAndroidMisc::GetCPUChipset()
+{
+	static FString *Chipset = FAndroidMisc::GetConfigRulesVariable(TEXT("hardware"));
+	return (Chipset == NULL) ? FGenericPlatformMisc::GetCPUChipset() : *Chipset;
+}
+
+FString FAndroidMisc::GetPrimaryGPUBrand()
+{
+	return FAndroidMisc::GetGPUFamily();
+}
+
+void FAndroidMisc::GetOSVersions(FString& out_OSVersionLabel, FString& out_OSSubVersionLabel)
+{
+	out_OSVersionLabel = TEXT("Android");
+	out_OSSubVersionLabel = AndroidVersion;
+}
+
+FString FAndroidMisc::GetOSVersion()
+{
+	return AndroidVersion;
+}
+
+bool FAndroidMisc::GetDiskTotalAndFreeSpace(const FString& InPath, uint64& TotalNumberOfBytes, uint64& NumberOfFreeBytes)
+{
+#if USE_ANDROID_FILE
+	extern FString GExternalFilePath;
+	struct statfs FSStat = { 0 };
+	FTCHARToUTF8 Converter(*GExternalFilePath);
+	int Err = statfs((ANSICHAR*)Converter.Get(), &FSStat);
+
+	if (Err == 0)
+	{
+		TotalNumberOfBytes = FSStat.f_blocks * FSStat.f_bsize;
+		NumberOfFreeBytes = FSStat.f_bavail * FSStat.f_bsize;
+	}
+	else
+	{
+		int ErrNo = errno;
+		UE_LOG(LogAndroid, Warning, TEXT("Unable to statfs('%s'): errno=%d (%s)"), *GExternalFilePath, ErrNo, UTF8_TO_TCHAR(strerror(ErrNo)));
+	}
+
+	return (Err == 0);
+#else
+	return false;
+#endif
+}
+
+uint32 FAndroidMisc::GetCoreFrequency(int32 CoreIndex, ECoreFrequencyProperty CoreFrequencyProperty)
+{
+	uint32 ReturnFrequency = 0;
+	char QueryFile[256];
+	const char* FreqProperty = nullptr;
+	static const char* CurrentFrequencyString = "scaling_cur_freq";
+	static const char* MaxFrequencyString = "cpuinfo_max_freq";
+	static const char* MinFrequencyString = "cpuinfo_min_freq";
+	switch (CoreFrequencyProperty)
+	{
+		case ECoreFrequencyProperty::MaxFrequency:
+			FreqProperty = MaxFrequencyString;
+			break;
+		case ECoreFrequencyProperty::MinFrequency:
+			FreqProperty = MinFrequencyString;
+			break;
+		default:
+		case ECoreFrequencyProperty::CurrentFrequency:
+			FreqProperty = CurrentFrequencyString;
+			break;
+	}
+	sprintf(QueryFile, "/sys/devices/system/cpu/cpu%d/cpufreq/%s", CoreIndex, FreqProperty);
+
+	if (FILE* CoreFreqStateFile = fopen(QueryFile, "r"))
+	{
+		char curr_corefreq[32] = { 0 };
+		if( fgets(curr_corefreq, ARRAY_COUNT(curr_corefreq), CoreFreqStateFile) != nullptr)
+		{
+			ReturnFrequency = atol(curr_corefreq);
+		}
+		fclose(CoreFreqStateFile);
+	}
+	return ReturnFrequency;
+}
diff --git a/Engine/Source/Runtime/Core/Private/Android/AndroidPlatformOutputDevices.cpp b/Engine/Source/Runtime/Core/Private/Android/AndroidPlatformOutputDevices.cpp
new file mode 100644
--- /dev/null
+++ b/Engine/Source/Runtime/Core/Private/Android/AndroidPlatformOutputDevices.cpp
@@ -0,0 +1,19 @@
+// Copyright 1998-2019 Epic Games, Inc. All Rights Reserved.
+
+/*=============================================================================
+	AndroidOutputDevices.cpp: Android implementations of OutputDevices functions
+=============================================================================*/
+
+#include "Android/AndroidPlatformOutputDevices.h"
+
+/*class FOutputDevice* FAndroidOutputDevices::GetLog()
+{
+#if UE_BUILD_SHIPPING
+	return NULL;
+#else
+	// Always enable logging via ADB
+	static FOutputDeviceAndroidDebug Singleton;
+	return &Singleton;
+	return NULL;
+#endif
+}*/
diff --git a/Engine/Source/Runtime/Core/Private/Android/AndroidPlatformProcess.cpp b/Engine/Source/Runtime/Core/Private/Android/AndroidPlatformProcess.cpp
new file mode 100644
--- /dev/null
+++ b/Engine/Source/Runtime/Core/Private/Android/AndroidPlatformProcess.cpp
@@ -0,0 +1,276 @@
+// Copyright 1998-2019 Epic Games, Inc. All Rights Reserved.
+
+/*=============================================================================
+	AndroidProcess.cpp: Android implementations of Process functions
+=============================================================================*/
+
+#include "Android/AndroidPlatformProcess.h"
+#include "Android/AndroidPlatformRunnableThread.h"
+#if !PLATFORM_LUMIN
+#include "Android/AndroidPlatformAffinity.h"
+#endif
+#include "Async/TaskGraphInterfaces.h"
+
+#include <sys/syscall.h>
+#include <pthread.h>
+#include <dlfcn.h>
+
+#include "Android/AndroidJavaEnv.h"
+
+#include "Misc/CoreDelegates.h"
+
+// RTLD_NOLOAD not defined for all platforms before NDK15
+#if !PLATFORM_LUMIN
+#if PLATFORM_ANDROID_NDK_VERSION < 150000
+	// not defined for NDK platform before 21
+	#if PLATFORM_USED_NDK_VERSION_INTEGER < 21
+		#define RTLD_NOLOAD   0x00004
+	#endif
+#endif
+#endif
+
+#if !PLATFORM_LUMIN
+int64 FAndroidAffinity::GameThreadMask = FPlatformAffinity::GetNoAffinityMask();
+int64 FAndroidAffinity::RenderingThreadMask = FPlatformAffinity::GetNoAffinityMask();
+#endif
+
+void* FAndroidPlatformProcess::GetDllHandle(const TCHAR* Filename)
+{
+	check(Filename);
+
+	// Check if dylib is already loaded
+	void* Handle = dlopen(TCHAR_TO_ANSI(Filename), RTLD_NOLOAD | RTLD_LAZY | RTLD_LOCAL);
+	
+	if (!Handle)
+	{
+		// Not loaded yet, so try to open it
+		Handle = dlopen(TCHAR_TO_ANSI(Filename), RTLD_LAZY | RTLD_LOCAL);
+	}
+	if (!Handle)
+	{
+		UE_LOG(LogAndroid, Warning, TEXT("dlopen failed: %s"), ANSI_TO_TCHAR(dlerror()));
+	}
+	return Handle;
+}
+
+void FAndroidPlatformProcess::FreeDllHandle(void* DllHandle)
+{
+	check(DllHandle);
+	dlclose(DllHandle);
+}
+
+void* FAndroidPlatformProcess::GetDllExport(void* DllHandle, const TCHAR* ProcName)
+{
+	check(DllHandle);
+	check(ProcName);
+	return dlsym(DllHandle, TCHAR_TO_ANSI(ProcName));
+}
+
+const TCHAR* FAndroidPlatformProcess::ComputerName()
+{
+	static FString ComputerName;
+	if (ComputerName.Len() == 0)
+	{
+		ComputerName = FAndroidMisc::GetDeviceModel();
+	}
+
+	return *ComputerName; 
+}
+
+void FAndroidPlatformProcess::SetThreadAffinityMask( uint64 InAffinityMask )
+{
+	/* 
+	 * On Android we prefer not to touch the thread affinity at all unless the user has specifically requested to change it.
+	 * The only way to override the default mask is to use the android.DefaultThreadAffinity console command set by ini file or device profile.
+	 */
+	if (FPlatformAffinity::GetNoAffinityMask() != InAffinityMask)
+	{
+		pid_t ThreadId = gettid();
+		int AffinityMask = (int)InAffinityMask;
+		syscall(__NR_sched_setaffinity, ThreadId, sizeof(AffinityMask), &AffinityMask);
+	}
+}
+
+uint32 FAndroidPlatformProcess::GetCurrentProcessId()
+{
+	return getpid();
+}
+
+uint32 FAndroidPlatformProcess::GetCurrentCoreNumber()
+{
+	unsigned CPU;
+	int Err = syscall(__NR_getcpu, &CPU, nullptr, nullptr);
+	return (!Err) ? CPU : 0;
+}
+
+const TCHAR* FAndroidPlatformProcess::BaseDir()
+{
+	return TEXT("");
+}
+
+const TCHAR* FAndroidPlatformProcess::ExecutableName(bool bRemoveExtension)
+{
+#if USE_ANDROID_FILE
+	extern FString GAndroidProjectName;
+	return *GAndroidProjectName;
+#else
+	UE_LOG(LogAndroid, Fatal, TEXT("A sub-platform that doesn't use USE_ANDROID_FILE must implement PlatformProcess::ExecutableName"));
+	return TEXT("");
+#endif
+}
+
+FRunnableThread* FAndroidPlatformProcess::CreateRunnableThread()
+{
+	return new FRunnableThreadAndroid();
+}
+
+bool FAndroidPlatformProcess::CanLaunchURL(const TCHAR* URL)
+{
+	return URL != nullptr;
+}
+
+DECLARE_DELEGATE_OneParam(FAndroidLaunchURLDelegate, const FString&);
+
+CORE_API FAndroidLaunchURLDelegate OnAndroidLaunchURL;
+
+void FAndroidPlatformProcess::LaunchURL(const TCHAR* URL, const TCHAR* Parms, FString* Error)
+{
+	if (FCoreDelegates::ShouldLaunchUrl.IsBound() && !FCoreDelegates::ShouldLaunchUrl.Execute(URL))
+	{
+		if (Error)
+		{
+			*Error = TEXT("LaunchURL cancelled by delegate");
+		}
+		return;
+	}
+
+	check(URL);
+	const FString URLWithParams = FString::Printf(TEXT("%s %s"), URL, Parms ? Parms : TEXT("")).TrimEnd();
+
+	OnAndroidLaunchURL.ExecuteIfBound(URLWithParams);
+
+	if (Error != nullptr)
+	{
+		*Error = TEXT("");
+	}
+}
+
+FString FAndroidPlatformProcess::GetGameBundleId()
+{
+#if USE_ANDROID_JNI
+	JNIEnv* JEnv = AndroidJavaEnv::GetJavaEnv();
+	if (nullptr != JEnv)
+	{
+		jclass Class = AndroidJavaEnv::FindJavaClassGlobalRef("com/epicgames/ue4/GameActivity");
+		if (nullptr != Class)
+		{
+			jmethodID getAppPackageNameMethodId = JEnv->GetStaticMethodID(Class, "getAppPackageName", "()Ljava/lang/String;");
+			FString PackageName = FJavaHelper::FStringFromLocalRef(JEnv, (jstring)JEnv->CallStaticObjectMethod(Class, getAppPackageNameMethodId, nullptr));
+			JEnv->DeleteGlobalRef(Class);
+			return PackageName;
+		}
+	}
+#endif
+	return TEXT("");
+}
+
+// Can be specified per device profile
+// android.DefaultThreadAffinity GT 0x01 RT 0x02
+TAutoConsoleVariable<FString> CVarAndroidDefaultThreadAffinity(
+	TEXT("android.DefaultThreadAffinity"), 
+	TEXT(""), 
+	TEXT("Sets the thread affinity for Android platform. Pairs of args [GT|RT] [Hex affinity], ex: android.DefaultThreadAffinity GT 0x01 RT 0x02"));
+
+static void AndroidSetAffinityOnThread()
+{
+	if (IsInActualRenderingThread()) // If RenderingThread is not started yet, affinity will applied at RT creation time 
+	{
+		FPlatformProcess::SetThreadAffinityMask(FPlatformAffinity::GetRenderingThreadMask());
+	}
+	else if (IsInGameThread())
+	{
+		FPlatformProcess::SetThreadAffinityMask(FPlatformAffinity::GetMainGameMask());
+	}
+}
+
+static void ApplyDefaultThreadAffinity(IConsoleVariable* Var)
+{
+	FString AffinityCmd = CVarAndroidDefaultThreadAffinity.GetValueOnAnyThread();
+
+	TArray<FString> Args;
+	if (AffinityCmd.ParseIntoArrayWS(Args) > 0)
+	{
+		for (int32 Index = 0; Index + 1 < Args.Num(); Index += 2)
+		{
+			uint64 Aff = FParse::HexNumber(*Args[Index + 1]);
+			if (!Aff)
+			{
+				UE_LOG(LogAndroid, Display, TEXT("Parsed 0 for affinity, using 0xFFFFFFFFFFFFFFFF instead"));
+				Aff = 0xFFFFFFFFFFFFFFFF;
+			}
+
+#if !PLATFORM_LUMIN
+			if (Args[Index] == TEXT("GT"))
+			{
+				FAndroidAffinity::GameThreadMask = Aff;
+			}
+			else if (Args[Index] == TEXT("RT"))
+			{
+				FAndroidAffinity::RenderingThreadMask = Aff;
+			}
+#endif
+		}
+
+		if (FTaskGraphInterface::IsRunning())
+		{
+			FSimpleDelegateGraphTask::CreateAndDispatchWhenReady(
+				FSimpleDelegateGraphTask::FDelegate::CreateStatic(&AndroidSetAffinityOnThread),
+				TStatId(), NULL, ENamedThreads::GetRenderThread());
+
+			FSimpleDelegateGraphTask::CreateAndDispatchWhenReady(
+				FSimpleDelegateGraphTask::FDelegate::CreateStatic(&AndroidSetAffinityOnThread),
+				TStatId(), NULL, ENamedThreads::GameThread);
+		}
+		else
+		{
+			AndroidSetAffinityOnThread();
+		}
+	}
+}
+
+void AndroidSetupDefaultThreadAffinity()
+{
+	ApplyDefaultThreadAffinity(nullptr);
+
+	// Watch for CVar update
+	CVarAndroidDefaultThreadAffinity->SetOnChangedCallback(FConsoleVariableDelegate::CreateStatic(&ApplyDefaultThreadAffinity));
+}
+
+#if !PLATFORM_LUMIN
+static bool EnableLittleCoreAffinity = false;
+static int32 BigCoreMask = 0;
+static int32 LittleCoreMask = 0;
+
+//This function is declared in the Java-defined class, GameActivity.java: "public native void nativeSetAffinityInfo(boolean bEnableAffinity, int bigCoreMask, int littleCoreMask);
+JNI_METHOD void Java_com_epicgames_ue4_GameActivity_nativeSetAffinityInfo(JNIEnv* jenv, jobject thiz, jboolean bEnableAffinity, jint bigCoreMask, jint littleCoreMask)
+{
+	EnableLittleCoreAffinity = bEnableAffinity;
+	BigCoreMask = bigCoreMask;
+	LittleCoreMask = littleCoreMask;
+}
+
+uint64 FAndroidAffinity::GetLittleCoreMask()
+{
+	static int Mask = 0;
+	if (Mask == 0)
+	{
+		Mask = FGenericPlatformAffinity::GetNoAffinityMask();
+		if (EnableLittleCoreAffinity)
+		{
+			Mask = LittleCoreMask;
+			FPlatformMisc::LowLevelOutputDebugStringf(TEXT("LittleCore Affinity applying mask: 0x%0x"), Mask);
+		}
+	}
+	return Mask;
+}
+#endif
diff --git a/Engine/Source/Runtime/Core/Private/Android/AndroidPlatformString.cpp b/Engine/Source/Runtime/Core/Private/Android/AndroidPlatformString.cpp
new file mode 100644
--- /dev/null
+++ b/Engine/Source/Runtime/Core/Private/Android/AndroidPlatformString.cpp
@@ -0,0 +1,319 @@
+// Copyright 1998-2019 Epic Games, Inc. All Rights Reserved.
+
+/*=============================================================================
+	AndroidString.cpp: Android implementations of string functions
+=============================================================================*/
+
+#include "Android/AndroidPlatformString.h"
+
+#if !PLATFORM_TCHAR_IS_CHAR16
+
+#include "Containers/StringConv.h"
+#include <stdlib.h>
+#include <cwchar>
+
+// This is a full copy of iswspace function from Android sources
+// For some reason function from libc does not work correctly for some korean characters like: 0xBE0C
+int iswspace(wint_t wc)
+{
+	static const wchar_t spaces[] = {
+		' ', '\t', '\n', '\r', 11, 12,  0x0085,
+		0x2000, 0x2001, 0x2002, 0x2003, 0x2004, 0x2005,
+		0x2006, 0x2008, 0x2009, 0x200a,
+		0x2028, 0x2029, 0x205f, 0x3000, 0
+	};
+	return wc && wcschr(spaces, wc);
+}
+
+int vswprintf( TCHAR *buf, int max, const TCHAR *fmt, va_list args )
+{
+	if (fmt == NULL)
+	{
+		if ((max > 0) && (buf != NULL))
+			*buf = 0;
+		return(0);
+	}
+
+	int fmtlen = FAndroidPlatformString::Strlen(fmt);
+	TCHAR *src = (TCHAR *)FMemory_Alloca((fmtlen + 1) * sizeof (TCHAR));
+	FAndroidPlatformString::Strcpy(src, fmtlen + 1, fmt);
+
+	TCHAR *dst = buf;
+	TCHAR *enddst = dst + (max - 1);
+
+	//printf("printf by-hand formatting "); unicode_str_to_stdout(src);
+
+	while ((*src) && (dst < enddst))
+	{
+		if (*src != '%')
+		{
+			*dst = *src;
+			dst++;
+			src++;
+			continue;
+		}
+
+		TCHAR *percent_ptr = src;
+		int fieldlen = 0;
+		int precisionlen = -1;
+
+		src++; // skip the '%' char...
+
+		while (*src == ' ')
+		{
+			*dst = ' ';
+			dst++;
+			src++;
+		}
+
+		// Skip modifier flags that don't need additional processing;
+		// they still get passed to snprintf() below based on the conversion.
+		if (*src == '+')
+		{
+			src++;
+		}
+
+		// check for field width requests...
+		if ((*src == '-') || ((*src >= '0') && (*src <= '9')))
+		{
+			TCHAR *ptr = src + 1;
+			while ((*ptr >= '0') && (*ptr <= '9'))
+				ptr++;
+
+			TCHAR ch = *ptr;
+			*ptr = '\0';
+			fieldlen = atoi(TCHAR_TO_ANSI(src));
+			*ptr = ch;
+
+			src = ptr;
+		}
+
+		if (*src == '.')
+		{
+			TCHAR *ptr = src + 1;
+			while ((*ptr >= '0') && (*ptr <= '9'))
+				ptr++;
+
+			TCHAR ch = *ptr;
+			*ptr = '\0';
+			precisionlen = atoi(TCHAR_TO_ANSI(src + 1));
+			*ptr = ch;
+			src = ptr;
+		}
+
+		// Check for 'ls' field, change to 's'
+		if ((src[0] == 'l' && src[1] == 's'))
+		{
+			src++;
+		}
+
+		switch (*src)
+		{
+		case '%':
+			{
+				src++;
+				*dst = '%';
+				dst++;
+			}
+			break;
+
+		case 'c':
+			{
+				TCHAR val = (TCHAR) va_arg(args, int);
+				src++;
+				*dst = val;
+				dst++;
+			}
+			break;
+
+		case 'd':
+		case 'i':
+		case 'X':
+		case 'x':
+		case 'u':
+			{
+				src++;
+				int val = va_arg(args, int);
+				ANSICHAR ansinum[30];
+				ANSICHAR fmtbuf[30];
+
+				// Yes, this is lame.
+				int cpyidx = 0;
+				while (percent_ptr < src)
+				{
+					fmtbuf[cpyidx] = (ANSICHAR)*percent_ptr;
+					percent_ptr++;
+					cpyidx++;
+				}
+				fmtbuf[cpyidx] = 0;
+
+				int rc = snprintf(ansinum, sizeof(ansinum), fmtbuf, val);
+				if ((dst + rc) > enddst)
+					return -1;	// Fail - the app needs to create a larger buffer and try again
+				for (int i = 0; i < rc; i++)
+				{
+					*dst = (TCHAR)ansinum[i];
+					dst++;
+				}
+			}
+			break;
+
+		case 'p':
+			{
+				src++;
+				void* val = va_arg(args, void*);
+				ANSICHAR ansinum[30];
+				ANSICHAR fmtbuf[30];
+
+				// Yes, this is lame.
+				int cpyidx = 0;
+				while (percent_ptr < src)
+				{
+					fmtbuf[cpyidx] = (ANSICHAR) *percent_ptr;
+					percent_ptr++;
+					cpyidx++;
+				}
+				fmtbuf[cpyidx] = 0;
+
+				int rc = snprintf(ansinum, sizeof (ansinum), fmtbuf, val);
+				if ((dst + rc) > enddst)
+					return -1;	// Fail - the app needs to create a larger buffer and try again
+				for (int i = 0; i < rc; i++)
+				{
+					*dst = (TCHAR) ansinum[i];
+					dst++;
+				}
+			}
+			break;
+
+		case 'l':
+		case 'I':
+			{
+				if ((src[0] == 'l' && src[1] != 'l') ||
+					(src[0] == 'I' && (src[1] != '6' || src[2] != '4')))
+				{
+					printf("Unknown percent [%lc%lc%lc] in Android:vswprintf() [%s]\n.", src[0], src[1], src[2], TCHAR_TO_ANSI(fmt));
+					src++;  // skip it, I guess.
+					break;
+				}
+
+				// Yes, this is lame.
+				int cpyidx = 0;
+				unsigned long long val = va_arg(args, unsigned long long);
+				ANSICHAR ansinum[60];
+				ANSICHAR fmtbuf[30];
+				if (src[0] == 'l')
+				{
+					src += 3;
+				}
+				else
+				{
+					src += 4;
+					strcpy(fmtbuf, "%L");
+					percent_ptr += 4;
+					cpyidx = 2;
+				}
+
+				while (percent_ptr < src)
+				{
+					fmtbuf[cpyidx] = (ANSICHAR) *percent_ptr;
+					percent_ptr++;
+					cpyidx++;
+				}
+				fmtbuf[cpyidx] = 0;
+
+				int rc = snprintf(ansinum, sizeof (ansinum), fmtbuf, val);
+				if ((dst + rc) > enddst)
+					return -1;	// Fail - the app needs to create a larger buffer and try again
+				for (int i = 0; i < rc; i++)
+				{
+					*dst = (TCHAR) ansinum[i];
+					dst++;
+				}
+			}
+			break;
+
+		case 'f':
+		case 'e':
+		case 'g':
+			{
+				src++;
+				double val = va_arg(args, double);
+				ANSICHAR ansinum[30];
+				ANSICHAR fmtbuf[30];
+
+				// Yes, this is lame.
+				int cpyidx = 0;
+				while (percent_ptr < src)
+				{
+					fmtbuf[cpyidx] = (ANSICHAR) *percent_ptr;
+					percent_ptr++;
+					cpyidx++;
+				}
+				fmtbuf[cpyidx] = 0;
+
+				int rc = snprintf(ansinum, sizeof (ansinum), fmtbuf, val);
+				if ((dst + rc) > enddst)
+					return -1;	// Fail - the app needs to create a larger buffer and try again
+				for (int i = 0; i < rc; i++)
+				{
+					*dst = (TCHAR) ansinum[i];
+					dst++;
+				}
+			}
+			break;
+
+		case 's':
+			{
+				src++;
+				static const TCHAR* Null = TEXT("(null)");
+				const TCHAR *val = va_arg(args, TCHAR *);
+				if (val == NULL)
+					val = Null;
+
+				int rc = FAndroidPlatformString::Strlen(val);
+				int spaces = FPlatformMath::Max(FPlatformMath::Abs(fieldlen) - rc, 0);
+				if ((dst + rc + spaces) > enddst)
+					return -1;	// Fail - the app needs to create a larger buffer and try again
+				if (spaces > 0 && fieldlen > 0)
+				{
+					for (int i = 0; i < spaces; i++)
+					{
+						*dst = TEXT(' ');
+						dst++;
+					}
+				}
+				for (int i = 0; i < rc; i++)
+				{
+					*dst = *val;
+					dst++;
+					val++;
+				}
+				if (spaces > 0 && fieldlen < 0)
+				{
+					for (int i = 0; i < spaces; i++)
+					{
+						*dst = TEXT(' ');
+						dst++;
+					}
+				}
+			}
+			break;
+
+		default:
+			printf("Unknown percent [%%%c] in Android:vswprintf().\n", *src);
+			src++;  // skip char, I guess.
+			break;
+		}
+	}
+
+	// Check if we were able to finish the entire format string
+	// If not, the app needs to create a larger buffer and try again
+	if(*src)
+		return -1;
+
+	*dst = 0;  // null terminate the new string.
+	return(dst - buf);
+}
+
+#endif
\ No newline at end of file
diff --git a/Engine/Source/Runtime/Core/Private/Android/AndroidProcess.cpp b/Engine/Source/Runtime/Core/Private/Android/AndroidProcess.cpp
deleted file mode 100644
diff --git a/Engine/Source/Runtime/Core/Private/Android/AndroidStats.cpp b/Engine/Source/Runtime/Core/Private/Android/AndroidStats.cpp
--- a/Engine/Source/Runtime/Core/Private/Android/AndroidStats.cpp
+++ b/Engine/Source/Runtime/Core/Private/Android/AndroidStats.cpp
@@ -2,7 +2,7 @@
 
 #include "Android/AndroidStats.h"
 #include "CoreMinimal.h"
-#include "Android/AndroidMisc.h"
+#include "Android/AndroidPlatformMisc.h"
 #include "Stats/Stats.h"
 #include "HAL/IConsoleManager.h"
 
diff --git a/Engine/Source/Runtime/Core/Private/Android/AndroidString.cpp b/Engine/Source/Runtime/Core/Private/Android/AndroidString.cpp
deleted file mode 100644
diff --git a/Engine/Source/Runtime/Core/Private/IOS/IOSPlatformFile.cpp b/Engine/Source/Runtime/Core/Private/IOS/IOSPlatformFile.cpp
--- a/Engine/Source/Runtime/Core/Private/IOS/IOSPlatformFile.cpp
+++ b/Engine/Source/Runtime/Core/Private/IOS/IOSPlatformFile.cpp
@@ -1,6 +1,6 @@
 // Copyright 1998-2019 Epic Games, Inc. All Rights Reserved.
 
-#include "IOSPlatformFile.h"
+#include "IOS/IOSPlatformFile.h"
 #include "HAL/PlatformTLS.h"
 #include "Containers/StringConv.h"
 #include "HAL/PlatformTime.h"
diff --git a/Engine/Source/Runtime/Core/Private/IOS/IOSPlatformFile.h b/Engine/Source/Runtime/Core/Private/IOS/IOSPlatformFile.h
deleted file mode 100644
diff --git a/Engine/Source/Runtime/Core/Private/Lumin/LuminPlatformFile.cpp b/Engine/Source/Runtime/Core/Private/Lumin/LuminPlatformFile.cpp
--- a/Engine/Source/Runtime/Core/Private/Lumin/LuminPlatformFile.cpp
+++ b/Engine/Source/Runtime/Core/Private/Lumin/LuminPlatformFile.cpp
@@ -2,7 +2,7 @@
 // Copyright 2016 Magic Leap, Inc. All Rights Reserved.
 
 #include "Lumin/LuminPlatformFile.h"
-#include "Android/AndroidMisc.h"
+#include "Android/AndroidPlatformMisc.h"
 #include "Lumin/LuminLifecycle.h"
 #include "HAL/PlatformProcess.h"
 #include "Misc/Parse.h"
diff --git a/Engine/Source/Runtime/Core/Private/Mac/MacPlatformFile.cpp b/Engine/Source/Runtime/Core/Private/Mac/MacPlatformFile.cpp
--- a/Engine/Source/Runtime/Core/Private/Mac/MacPlatformFile.cpp
+++ b/Engine/Source/Runtime/Core/Private/Mac/MacPlatformFile.cpp
@@ -1,7 +1,6 @@
 // Copyright 1998-2019 Epic Games, Inc. All Rights Reserved.
 
-#include "MacPlatformFile.h"
-#include "Apple/ApplePlatformFile.h"
+#include "Mac/MacPlatformFile.h"
 
 IPlatformFile& IPlatformFile::GetPlatformPhysical()
 {
diff --git a/Engine/Source/Runtime/Core/Private/Mac/MacPlatformFile.h b/Engine/Source/Runtime/Core/Private/Mac/MacPlatformFile.h
deleted file mode 100644
diff --git a/Engine/Source/Runtime/Core/Public/Android/AndroidAffinity.h b/Engine/Source/Runtime/Core/Public/Android/AndroidAffinity.h
deleted file mode 100644
diff --git a/Engine/Source/Runtime/Core/Public/Android/AndroidAtomics.h b/Engine/Source/Runtime/Core/Public/Android/AndroidAtomics.h
deleted file mode 100644
diff --git a/Engine/Source/Runtime/Core/Public/Android/AndroidCompilerSetup.h b/Engine/Source/Runtime/Core/Public/Android/AndroidCompilerSetup.h
deleted file mode 100644
diff --git a/Engine/Source/Runtime/Core/Public/Android/AndroidFile.h b/Engine/Source/Runtime/Core/Public/Android/AndroidFile.h
deleted file mode 100644
diff --git a/Engine/Source/Runtime/Core/Public/Android/AndroidMath.h b/Engine/Source/Runtime/Core/Public/Android/AndroidMath.h
deleted file mode 100644
diff --git a/Engine/Source/Runtime/Core/Public/Android/AndroidMemory.h b/Engine/Source/Runtime/Core/Public/Android/AndroidMemory.h
deleted file mode 100644
diff --git a/Engine/Source/Runtime/Core/Public/Android/AndroidMisc.h b/Engine/Source/Runtime/Core/Public/Android/AndroidMisc.h
deleted file mode 100644
diff --git a/Engine/Source/Runtime/Core/Public/Android/AndroidOutputDevices.h b/Engine/Source/Runtime/Core/Public/Android/AndroidOutputDevices.h
deleted file mode 100644
diff --git a/Engine/Source/Runtime/Core/Public/Android/AndroidPlatformAffinity.h b/Engine/Source/Runtime/Core/Public/Android/AndroidPlatformAffinity.h
new file mode 100644
--- /dev/null
+++ b/Engine/Source/Runtime/Core/Public/Android/AndroidPlatformAffinity.h
@@ -0,0 +1,87 @@
+// Copyright 1998-2019 Epic Games, Inc. All Rights Reserved.
+
+/*=============================================================================================
+AndroidAffinity.h: Android affinity profile masks definitions.
+==============================================================================================*/
+
+#pragma once
+
+#include "GenericPlatform/GenericPlatformAffinity.h"
+
+class FAndroidAffinity : public FGenericPlatformAffinity
+{
+private:
+	static uint64 GetLittleCoreMask();
+	const static uint64 AllCores = 0xFFFFFFFFFF;
+public:
+	static const CORE_API uint64 GetMainGameMask()
+	{
+		return GameThreadMask;
+	}
+
+	static const CORE_API uint64 GetRenderingThreadMask()
+	{
+		return RenderingThreadMask;
+	}
+
+	static const CORE_API uint64 GetRHIThreadMask()
+	{
+		return AllCores;
+	}
+
+	static const CORE_API uint64 GetRTHeartBeatMask()
+	{
+		return GetLittleCoreMask();
+	}
+
+	static const CORE_API uint64 GetPoolThreadMask()
+	{
+		return GetLittleCoreMask();
+	}
+
+	static const CORE_API uint64 GetTaskGraphThreadMask()
+	{
+		return GetLittleCoreMask();
+	}
+
+	static const CORE_API uint64 GetStatsThreadMask()
+	{
+		return GetLittleCoreMask();
+	}
+
+	static const CORE_API uint64 GetAudioThreadMask()
+	{
+		return GetLittleCoreMask();
+	}
+
+	static const CORE_API uint64 GetTaskGraphBackgroundTaskMask()
+	{
+		return GetLittleCoreMask();
+	}
+
+	static const CORE_API uint64 GetTaskGraphHighPriorityTaskMask()
+	{
+		return AllCores;
+	}
+
+	static const CORE_API uint64 GetAsyncLoadingThreadMask()
+	{
+		return AllCores;
+	}
+
+	static EThreadPriority GetRenderingThreadPriority()
+	{
+		return TPri_SlightlyBelowNormal;
+	}
+
+	static EThreadPriority GetRHIThreadPriority()
+	{
+		return TPri_Normal;
+	}
+
+public:
+	static int64 GameThreadMask;
+	static int64 RenderingThreadMask;
+};
+
+typedef FAndroidAffinity FPlatformAffinity;
diff --git a/Engine/Source/Runtime/Core/Public/Android/AndroidPlatformAtomics.h b/Engine/Source/Runtime/Core/Public/Android/AndroidPlatformAtomics.h
new file mode 100644
--- /dev/null
+++ b/Engine/Source/Runtime/Core/Public/Android/AndroidPlatformAtomics.h
@@ -0,0 +1,11 @@
+// Copyright 1998-2019 Epic Games, Inc. All Rights Reserved.
+
+/*=============================================================================================
+	AndroidPlatformAtomics.h: Android platform Atomics functions
+==============================================================================================*/
+
+#pragma once
+
+#include "Clang/ClangPlatformAtomics.h"
+
+typedef FClangPlatformAtomics FPlatformAtomics;
diff --git a/Engine/Source/Runtime/Core/Public/Android/AndroidPlatformCompilerSetup.h b/Engine/Source/Runtime/Core/Public/Android/AndroidPlatformCompilerSetup.h
new file mode 100644
--- /dev/null
+++ b/Engine/Source/Runtime/Core/Public/Android/AndroidPlatformCompilerSetup.h
@@ -0,0 +1,11 @@
+// Copyright 1998-2019 Epic Games, Inc. All Rights Reserved.
+
+
+/*=============================================================================================
+	AndroidCompilerSetup.h: pragmas, version checks and other things for the Android compiler
+==============================================================================================*/
+
+#pragma once
+
+// The register keyword is deprecated
+#define register
\ No newline at end of file
diff --git a/Engine/Source/Runtime/Core/Public/Android/AndroidPlatformFile.h b/Engine/Source/Runtime/Core/Public/Android/AndroidPlatformFile.h
new file mode 100644
--- /dev/null
+++ b/Engine/Source/Runtime/Core/Public/Android/AndroidPlatformFile.h
@@ -0,0 +1,44 @@
+// Copyright 1998-2019 Epic Games, Inc. All Rights Reserved.
+
+/*=============================================================================================
+	AndroidFile.h: Android platform File functions
+==============================================================================================*/
+
+#pragma once
+
+#include "GenericPlatform/GenericPlatformFile.h"
+#if USE_ANDROID_JNI
+#include <jni.h>
+#endif
+
+/**
+	Android File I/O implementation with additional utilities to deal
+	with Java side access.
+**/
+class CORE_API IAndroidPlatformFile : public IPhysicalPlatformFile
+{
+public:
+	static IAndroidPlatformFile & GetPlatformPhysical();
+
+#if USE_ANDROID_JNI
+	// Get the android.content.res.AssetManager that Java code
+	// should use to open APK assets.
+	virtual jobject GetAssetManager() = 0;
+#endif
+
+	// Get detailed information for a file that
+	// we can hand to other Android media classes for access.
+
+	// Is file embedded as an asset in the APK?
+	virtual bool IsAsset(const TCHAR* Filename) = 0;
+
+	// Offset within file or asset where its data starts.
+	// Note, offsets for assets is relative to complete APK file
+	// and matches what is returned by AssetFileDescriptor.getStartOffset().
+	virtual int64 FileStartOffset(const TCHAR* Filename) = 0;
+
+	// Get the root, i.e. underlying, path for the file. This
+	// can be any of: a resolved file path, an OBB path, an
+	// asset path.
+	virtual FString FileRootPath(const TCHAR* Filename) = 0;
+};
diff --git a/Engine/Source/Runtime/Core/Public/Android/AndroidPlatformMath.h b/Engine/Source/Runtime/Core/Public/Android/AndroidPlatformMath.h
new file mode 100644
--- /dev/null
+++ b/Engine/Source/Runtime/Core/Public/Android/AndroidPlatformMath.h
@@ -0,0 +1,10 @@
+// Copyright 1998-2019 Epic Games, Inc. All Rights Reserved.
+
+/*=============================================================================================
+	AndroidMath.h: Android platform Math functions
+==============================================================================================*/
+
+#pragma once
+#include "Clang/ClangPlatformMath.h"
+
+typedef FClangPlatformMath FPlatformMath;
diff --git a/Engine/Source/Runtime/Core/Public/Android/AndroidPlatformMemory.h b/Engine/Source/Runtime/Core/Public/Android/AndroidPlatformMemory.h
new file mode 100644
--- /dev/null
+++ b/Engine/Source/Runtime/Core/Public/Android/AndroidPlatformMemory.h
@@ -0,0 +1,41 @@
+// Copyright 1998-2019 Epic Games, Inc. All Rights Reserved.
+
+
+/*=============================================================================================
+	AndroidPlatformMemory.h: Android platform memory functions
+==============================================================================================*/
+
+#pragma once
+#include "GenericPlatform/GenericPlatformMemory.h"
+
+/**
+ *	Android implementation of the FGenericPlatformMemoryStats.
+ */
+struct FPlatformMemoryStats : public FGenericPlatformMemoryStats
+{};
+
+/**
+* Android implementation of the memory OS functions
+**/
+struct CORE_API FAndroidPlatformMemory : public FGenericPlatformMemory
+{
+	//~ Begin FGenericPlatformMemory Interface
+	static void Init();
+	static FPlatformMemoryStats GetStats();
+	static uint64 GetMemoryUsedFast();
+	static const FPlatformMemoryConstants& GetConstants();
+	static EPlatformMemorySizeBucket GetMemorySizeBucket();
+	static FMalloc* BaseAllocator();
+	static void* BinnedAllocFromOS( SIZE_T Size );
+	static void BinnedFreeToOS( void* Ptr, SIZE_T Size );
+	static bool MemoryRangeDecommit(void* Ptr, SIZE_T Size);
+
+
+	static bool GetLLMAllocFunctions(void*(*&OutAllocFunction)(size_t), void(*&OutFreeFunction)(void*, size_t), int32& OutAlignment);
+	//~ End FGenericPlatformMemory Interface
+};
+
+typedef FAndroidPlatformMemory FPlatformMemory;
+
+
+
diff --git a/Engine/Source/Runtime/Core/Public/Android/AndroidPlatformMisc.h b/Engine/Source/Runtime/Core/Public/Android/AndroidPlatformMisc.h
new file mode 100644
--- /dev/null
+++ b/Engine/Source/Runtime/Core/Public/Android/AndroidPlatformMisc.h
@@ -0,0 +1,268 @@
+// Copyright 1998-2019 Epic Games, Inc. All Rights Reserved.
+
+
+/*=============================================================================================
+	AndroidMisc.h: Android platform misc functions
+==============================================================================================*/
+
+#pragma once
+#include "CoreTypes.h"
+#include "Android/AndroidSystemIncludes.h"
+#include "GenericPlatform/GenericPlatformMisc.h"
+//@todo android: this entire file
+
+template <typename FuncType>
+class TFunction;
+
+#if PLATFORM_ANDROID_ARM64
+	#define PLATFORM_BREAK()	__asm__(".inst 0xd4200000")
+#elif PLATFORM_ANDROID_ARM
+	#define PLATFORM_BREAK()	__asm__("trap")
+#else
+	#define PLATFORM_BREAK()	__asm__("int $3")
+#endif
+
+#define UE_DEBUG_BREAK_IMPL()	PLATFORM_BREAK()
+
+/**
+ * Android implementation of the misc OS functions
+ */
+struct CORE_API FAndroidMisc : public FGenericPlatformMisc
+{
+	static void RequestExit( bool Force );
+	static void LocalPrint(const TCHAR *Message);
+	static bool IsLocalPrintThreadSafe() { return true; }
+	static void PlatformPreInit();
+	static void PlatformInit();
+	static void PlatformTearDown();
+	static void PlatformHandleSplashScreen(bool ShowSplashScreen);
+    static EDeviceScreenOrientation GetDeviceOrientation() { return DeviceOrientation; }
+    
+	FORCEINLINE static int32 GetMaxPathLength()
+	{
+		return ANDROID_MAX_PATH;
+	}
+
+	UE_DEPRECATED(4.21, "void FPlatformMisc::GetEnvironmentVariable(Name, Result, Length) is deprecated. Use FString FPlatformMisc::GetEnvironmentVariable(Name) instead.")
+	static void GetEnvironmentVariable(const TCHAR* VariableName, TCHAR* Result, int32 ResultLength);
+
+	static FString GetEnvironmentVariable(const TCHAR* VariableName);
+	static const TCHAR* GetSystemErrorMessage(TCHAR* OutBuffer, int32 BufferCount, int32 Error);
+	static EAppReturnType::Type MessageBoxExt( EAppMsgType::Type MsgType, const TCHAR* Text, const TCHAR* Caption );
+	static bool UseRenderThread();
+	static bool HasPlatformFeature(const TCHAR* FeatureName);
+	static bool ShouldDisablePluginAtRuntime(const FString& PluginName);
+	static void SetThreadName(const char* name);
+	static bool SupportsES30();
+
+public:
+
+	static bool AllowThreadHeartBeat()
+	{
+		return false;
+	}
+
+	struct FCPUStatTime{
+		uint64_t			TotalTime;
+		uint64_t			UserTime;
+		uint64_t			NiceTime;
+		uint64_t			SystemTime;
+		uint64_t			SoftIRQTime;
+		uint64_t			IRQTime;
+		uint64_t			IdleTime;
+		uint64_t			IOWaitTime;
+	};
+
+	struct FCPUState
+	{
+		const static int32			MaxSupportedCores = 16; //Core count 16 is maximum for now
+		int32						CoreCount;
+		int32						ActivatedCoreCount;
+		ANSICHAR					Name[6];
+		FAndroidMisc::FCPUStatTime	CurrentUsage[MaxSupportedCores]; 
+		FAndroidMisc::FCPUStatTime	PreviousUsage[MaxSupportedCores];
+		int32						Status[MaxSupportedCores];
+		double						Utilization[MaxSupportedCores];
+		double						AverageUtilization;
+		
+	};
+
+	static FCPUState& GetCPUState();
+	static int32 NumberOfCores();
+	static int32 NumberOfCoresIncludingHyperthreads();
+	static bool SupportsLocalCaching();
+	static void SetCrashHandler(void (* CrashHandler)(const FGenericCrashContext& Context));
+	// NOTE: THIS FUNCTION IS DEFINED IN ANDROIDOPENGL.CPP
+	static void GetValidTargetPlatforms(class TArray<class FString>& TargetPlatformNames);
+	static bool GetUseVirtualJoysticks();
+	static bool SupportsTouchInput();
+	static bool IsStandaloneStereoOnlyDevice();
+	static const TCHAR* GetDefaultDeviceProfileName() { return TEXT("Android_Default"); }
+	static bool GetVolumeButtonsHandledBySystem();
+	static void SetVolumeButtonsHandledBySystem(bool enabled);
+	// Returns current volume, 0-15
+	static int GetVolumeState(double* OutTimeOfChangeInSec = nullptr);
+
+	static int32 GetDeviceVolume();
+
+#if USE_ANDROID_FILE
+	static const TCHAR* GamePersistentDownloadDir();
+	static FString GetLoginId();
+#endif
+#if USE_ANDROID_JNI
+	static FString GetDeviceId();
+	static FString GetUniqueAdvertisingId();
+#endif
+	static FString GetCPUVendor();
+	static FString GetCPUBrand();
+	static FString GetCPUChipset();
+	static FString GetPrimaryGPUBrand();
+	static void GetOSVersions(FString& out_OSVersionLabel, FString& out_OSSubVersionLabel);
+	static bool GetDiskTotalAndFreeSpace(const FString& InPath, uint64& TotalNumberOfBytes, uint64& NumberOfFreeBytes);
+	
+	enum EBatteryState
+	{
+		BATTERY_STATE_UNKNOWN = 1,
+		BATTERY_STATE_CHARGING,
+		BATTERY_STATE_DISCHARGING,
+		BATTERY_STATE_NOT_CHARGING,
+		BATTERY_STATE_FULL
+	};
+	struct FBatteryState
+	{
+		FAndroidMisc::EBatteryState	State;
+		int							Level;          // in range [0,100]
+		float						Temperature;    // in degrees of Celsius
+	};
+
+	static FBatteryState GetBatteryState();
+	static int GetBatteryLevel();
+	static bool IsRunningOnBattery();
+	static bool IsInLowPowerMode();
+	static float GetDeviceTemperatureLevel();
+	static bool AreHeadPhonesPluggedIn();
+	static ENetworkConnectionType GetNetworkConnectionType();
+#if USE_ANDROID_JNI
+	static bool HasActiveWiFiConnection();
+#endif
+
+	static void RegisterForRemoteNotifications();
+	static void UnregisterForRemoteNotifications();
+
+	/** @return Memory representing a true type or open type font provided by the platform as a default font for unreal to consume; empty array if the default font failed to load. */
+	static TArray<uint8> GetSystemFontBytes();
+
+	static IPlatformChunkInstall* GetPlatformChunkInstall();
+
+	static void PrepareMobileHaptics(EMobileHapticsType Type);
+	static void TriggerMobileHaptics();
+	static void ReleaseMobileHaptics();
+	static void ShareURL(const FString& URL, const FText& Description, int32 LocationHintX, int32 LocationHintY);
+
+	static FString LoadTextFileFromPlatformPackage(const FString& RelativePath);
+
+	// ANDROID ONLY:
+	static void SetVersionInfo(FString AndroidVersion, FString DeviceMake, FString DeviceModel, FString DeviceBuildNumber, FString OSLanguage);
+	static const FString GetAndroidVersion();
+	static const FString GetDeviceMake();
+	static const FString GetDeviceModel();
+	static const FString GetOSLanguage();
+	static const FString GetDeviceBuildNumber();
+	static const FString GetProjectVersion();
+	static FString GetDefaultLocale();
+	static FString GetGPUFamily();
+	static FString GetGLVersion();
+	static bool SupportsFloatingPointRenderTargets();
+	static bool SupportsShaderFramebufferFetch();
+	static bool SupportsShaderIOBlocks();
+#if USE_ANDROID_JNI
+	static int GetAndroidBuildVersion();
+#endif
+	static TMap<FString, FString> GetConfigRulesTMap();
+	static FString* GetConfigRulesVariable(const FString& Key);
+
+	/* HasVulkanDriverSupport
+	 * @return true if this Android device supports a Vulkan API Unreal could use
+	 */
+	static bool HasVulkanDriverSupport();
+
+	/* IsVulkanAvailable
+	 * @return	true if there is driver support, we have an RHI, we are packaged with Vulkan support,
+	 *			and not we are not forcing GLES with a command line switch
+	 */
+	static bool IsVulkanAvailable();
+
+	/* ShouldUseVulkan
+	 * @return true if Vulkan is available, and not disabled by device profile cvar
+	 */
+	static bool ShouldUseVulkan();
+	static bool ShouldUseDesktopVulkan();
+	static FString GetVulkanVersion();
+	static bool IsDaydreamApplication();
+	typedef TFunction<void(void* NewNativeHandle)> ReInitWindowCallbackType;
+	static ReInitWindowCallbackType GetOnReInitWindowCallback();
+	static void SetOnReInitWindowCallback(ReInitWindowCallbackType InOnReInitWindowCallback);
+	static FString GetOSVersion();
+	static bool GetOverrideResolution(int32 &ResX, int32& ResY) { return false; }
+	typedef TFunction<void()> OnPauseCallBackType;
+	static OnPauseCallBackType GetOnPauseCallback();
+	static void SetOnPauseCallback(OnPauseCallBackType InOnPauseCallback);
+
+#if !UE_BUILD_SHIPPING
+	static bool IsDebuggerPresent();
+#endif
+
+	FORCEINLINE static void MemoryBarrier()
+	{
+		__sync_synchronize();
+	}
+
+
+#if STATS || ENABLE_STATNAMEDEVENTS
+	static void BeginNamedEventFrame();
+	static void BeginNamedEvent(const struct FColor& Color, const TCHAR* Text);
+	static void BeginNamedEvent(const struct FColor& Color, const ANSICHAR* Text);
+	static void EndNamedEvent();
+	static void CustomNamedStat(const TCHAR* Text, float Value, const TCHAR* Graph, const TCHAR* Unit);
+	static void CustomNamedStat(const ANSICHAR* Text, float Value, const ANSICHAR* Graph, const ANSICHAR* Unit);
+#endif
+
+#if (STATS || ENABLE_STATNAMEDEVENTS)
+	static int32 TraceMarkerFileDescriptor;
+#endif
+	
+	// run time compatibility information
+	static FString AndroidVersion; // version of android we are running eg "4.0.4"
+	static FString DeviceMake; // make of the device we are running on eg. "samsung"
+	static FString DeviceModel; // model of the device we are running on eg "SAMSUNG-SGH-I437"
+	static FString DeviceBuildNumber; // platform image build number of device "R16NW.G960NKSU1ARD6"
+	static FString OSLanguage; // language code the device is set to
+
+	// Build version of Android, i.e. API level.
+	static int32 AndroidBuildVersion;
+
+	// Key/Value pair variables from the optional configuration.txt
+	static TMap<FString, FString> ConfigRulesVariables;
+
+	static bool VolumeButtonsHandledBySystem;
+
+	static bool bNeedsRestartAfterPSOPrecompile;
+
+	enum class ECoreFrequencyProperty
+	{
+		CurrentFrequency,
+		MaxFrequency,
+		MinFrequency,
+	};
+
+	static uint32 GetCoreFrequency(int32 CoreIndex, ECoreFrequencyProperty CoreFrequencyProperty);
+    
+    static void SetDeviceOrientation(EDeviceScreenOrientation NewDeviceOrentation) { DeviceOrientation = NewDeviceOrentation; }
+    
+private:
+    static EDeviceScreenOrientation DeviceOrientation;
+};
+
+#if !PLATFORM_LUMIN
+typedef FAndroidMisc FPlatformMisc;
+#endif
diff --git a/Engine/Source/Runtime/Core/Public/Android/AndroidPlatformOutputDevices.h b/Engine/Source/Runtime/Core/Public/Android/AndroidPlatformOutputDevices.h
new file mode 100644
--- /dev/null
+++ b/Engine/Source/Runtime/Core/Public/Android/AndroidPlatformOutputDevices.h
@@ -0,0 +1,17 @@
+// Copyright 1998-2019 Epic Games, Inc. All Rights Reserved.
+
+
+/*=============================================================================================
+	AndroidOutputDevices.h: Android platform OutputDevices functions
+==============================================================================================*/
+
+#pragma once
+#include "GenericPlatform/GenericPlatformOutputDevices.h"
+
+struct CORE_API FAndroidOutputDevices : public FGenericPlatformOutputDevices
+{
+};
+
+#if !PLATFORM_LUMIN
+typedef FAndroidOutputDevices FPlatformOutputDevices;
+#endif
diff --git a/Engine/Source/Runtime/Core/Public/Android/AndroidPlatformProcess.h b/Engine/Source/Runtime/Core/Public/Android/AndroidPlatformProcess.h
new file mode 100644
--- /dev/null
+++ b/Engine/Source/Runtime/Core/Public/Android/AndroidPlatformProcess.h
@@ -0,0 +1,49 @@
+// Copyright 1998-2019 Epic Games, Inc. All Rights Reserved.
+
+
+/*=============================================================================================
+	AndroidProcess.h: Android platform Process functions
+==============================================================================================*/
+
+#pragma once
+#include "GenericPlatform/GenericPlatformProcess.h"
+
+/** Dummy process handle for platforms that use generic implementation. */
+struct FProcHandle : public TProcHandle<void*, nullptr>
+{
+public:
+	/** Default constructor. */
+	FORCEINLINE FProcHandle()
+		: TProcHandle()
+	{}
+
+	/** Initialization constructor. */
+	FORCEINLINE explicit FProcHandle( HandleType Other )
+		: TProcHandle( Other )
+	{}
+};
+
+/**
+ * Android implementation of the Process OS functions
+ **/
+struct CORE_API FAndroidPlatformProcess : public FGenericPlatformProcess
+{
+	static void* GetDllHandle(const TCHAR* Filename);
+	static void FreeDllHandle(void* DllHandle);
+	static void* GetDllExport(void* DllHandle, const TCHAR* ProcName);
+	static const TCHAR* ComputerName();
+	static void SetThreadAffinityMask( uint64 AffinityMask );
+	static uint32 GetCurrentProcessId();
+	static uint32 GetCurrentCoreNumber();
+	static const TCHAR* BaseDir();
+	static const TCHAR* ExecutableName(bool bRemoveExtension = true);
+	static class FRunnableThread* CreateRunnableThread();
+	static bool CanLaunchURL(const TCHAR* URL);
+	static void LaunchURL(const TCHAR* URL, const TCHAR* Parms, FString* Error);
+	static FString GetGameBundleId();
+};
+
+#if !PLATFORM_LUMIN
+typedef FAndroidPlatformProcess FPlatformProcess;
+#endif
+
diff --git a/Engine/Source/Runtime/Core/Public/Android/AndroidPlatformProperties.h b/Engine/Source/Runtime/Core/Public/Android/AndroidPlatformProperties.h
new file mode 100644
--- /dev/null
+++ b/Engine/Source/Runtime/Core/Public/Android/AndroidPlatformProperties.h
@@ -0,0 +1,114 @@
+// Copyright 1998-2019 Epic Games, Inc. All Rights Reserved.
+
+/*================================================================================
+	AndroidProperties.h - Basic static properties of a platform 
+	These are shared between:
+		the runtime platform - via FPlatformProperties
+		the target platforms - via ITargetPlatform
+==================================================================================*/
+
+#pragma once
+
+#include "CoreTypes.h"
+#include "GenericPlatform/GenericPlatformProperties.h"
+
+
+/**
+ * Implements Android platform properties.
+ */
+struct FAndroidPlatformProperties
+	: public FGenericPlatformProperties
+{
+	static FORCEINLINE const char* GetPhysicsFormat( )
+	{
+		return "PhysXGeneric";		//@todo android: physx format
+	}
+
+	static FORCEINLINE bool HasEditorOnlyData( )
+	{
+		return false;
+	}
+
+	static FORCEINLINE const char* PlatformName()
+	{
+		return "Android";
+	}
+
+	static FORCEINLINE const char* IniPlatformName()
+	{
+		return "Android";
+	}
+
+	static FORCEINLINE bool IsGameOnly( )
+	{
+		return true;
+	}
+
+	static FORCEINLINE bool RequiresCookedData( )
+	{
+		return true;
+	}
+
+	static FORCEINLINE bool SupportsBuildTarget( EBuildTargets::Type BuildTarget )
+	{
+		return (BuildTarget == EBuildTargets::Game);
+	}
+
+	static FORCEINLINE bool SupportsAutoSDK()
+	{
+		return true;
+	}
+
+	static FORCEINLINE bool SupportsHighQualityLightmaps()
+	{
+		return true; // always true because of Vulkan
+	}
+
+	static FORCEINLINE bool SupportsLowQualityLightmaps()
+	{
+#if PLATFORM_ANDROIDESDEFERRED || PLATFORM_ANDROIDGL4
+		return false;
+#else
+		return true;
+#endif
+	}
+
+	static FORCEINLINE bool SupportsDistanceFieldShadows()
+	{
+		return true;
+	}
+
+	static FORCEINLINE bool SupportsTextureStreaming()
+	{
+		return true;
+	}
+
+	static FORCEINLINE bool SupportsMinimize()
+	{
+		return true;
+	}
+
+	static FORCEINLINE bool HasFixedResolution()
+	{
+		return true;
+	}
+
+	static FORCEINLINE bool AllowsFramerateSmoothing()
+	{
+		return true;
+	}
+
+	static FORCEINLINE bool AllowsCallStackDumpDuringAssert()
+	{
+		return true;
+	}
+
+	static FORCEINLINE bool SupportsAudioStreaming()
+	{
+		return true;
+	}
+};
+
+#ifdef PROPERTY_HEADER_SHOULD_DEFINE_TYPE
+typedef FAndroidPlatformProperties FPlatformProperties;
+#endif
diff --git a/Engine/Source/Runtime/Core/Public/Android/AndroidPlatformRunnableThread.h b/Engine/Source/Runtime/Core/Public/Android/AndroidPlatformRunnableThread.h
--- a/Engine/Source/Runtime/Core/Public/Android/AndroidPlatformRunnableThread.h
+++ b/Engine/Source/Runtime/Core/Public/Android/AndroidPlatformRunnableThread.h
@@ -7,7 +7,7 @@
 #pragma once
 
 #include "Runtime/Core/Private/HAL/PThreadRunnableThread.h"
-#include "Android/AndroidMisc.h"
+#include "Android/AndroidPlatformMisc.h"
 
 /**
 * Android implementation of the pthread functions
diff --git a/Engine/Source/Runtime/Core/Public/Android/AndroidPlatformString.h b/Engine/Source/Runtime/Core/Public/Android/AndroidPlatformString.h
new file mode 100644
--- /dev/null
+++ b/Engine/Source/Runtime/Core/Public/Android/AndroidPlatformString.h
@@ -0,0 +1,584 @@
+// Copyright 1998-2019 Epic Games, Inc. All Rights Reserved.
+
+
+/*=============================================================================================
+	AndroidString.h: Android platform string classes
+==============================================================================================*/
+
+#pragma once
+
+#include "Misc/Build.h"
+#include "HAL/Platform.h"
+
+/**
+ * Android string implementation
+ **/
+
+#if PLATFORM_TCHAR_IS_CHAR16
+
+// By default we now use 2-byte strings on Android.
+
+#include "GenericPlatform/GenericWidePlatformString.h"
+
+typedef FGenericWidePlatformString FAndroidPlatformString;
+
+#else
+
+// NOTE: This legacy 4-byte implementation will probably be removed in the future.
+
+// @todo android: probably rewrite/simplify most of this.  currently it converts all wide chars
+// to ANSI.  The Android NDK appears to have wchar_t support, but many of the functions appear
+// to just be stubs to the non-wide versions.  This doesn't work for obvious reasons.
+
+#include "Misc/Char.h"
+#include "GenericPlatform/GenericPlatformMemory.h"
+#include "GenericPlatform/GenericPlatformString.h"
+
+struct FAndroidPlatformString : public FGenericPlatformString
+{
+	template <typename CharType>
+	static inline CharType* Strupr(CharType* Dest, SIZE_T DestCount)
+	{
+		for (CharType* Char = Dest; *Char && DestCount > 0; Char++, DestCount--)
+		{
+			*Char = TChar<CharType>::ToUpper(*Char);
+		}
+		return Dest;
+	}
+
+	/** 
+	 * Widechar implementation 
+	 **/
+	static FORCEINLINE WIDECHAR* Strcpy(WIDECHAR* Dest, SIZE_T DestCount, const WIDECHAR* Src)
+	{
+		if (!Dest || !Src)
+			return NULL;
+
+		int Pos = 0;
+		while (Src[Pos])
+		{
+			Dest[Pos] = Src[Pos];
+			++Pos;
+		}
+
+		Dest[Pos] = 0;
+
+		return Dest;
+	}
+
+	static FORCEINLINE WIDECHAR* Strncpy(WIDECHAR* Dest, const WIDECHAR* Src, SIZE_T MaxLen)
+	{
+		if (!Dest || !Src)
+			return NULL;
+
+		int Pos = 0;
+		while ((Pos < MaxLen) && Src[Pos])
+		{
+			Dest[Pos] = Src[Pos];
+			++Pos;
+		}
+
+		while (Pos < MaxLen)
+		{
+			Dest[Pos] = 0;
+			++Pos;
+		}
+
+		Dest[MaxLen-1]=0;
+		return Dest;
+	}
+
+	static FORCEINLINE WIDECHAR* Strcat(WIDECHAR* Dest, SIZE_T DestCount, const WIDECHAR* Src)
+	{
+		if (!Dest || !Src)
+			return NULL;
+
+		int DestPos = Strlen(Dest);
+		int SrcPos = 0;
+
+		while (Src[SrcPos])
+		{
+			Dest[DestPos] = Src[SrcPos];
+			++SrcPos;
+			++DestPos;
+		}
+
+		Dest[DestPos] = 0;
+
+		return Dest;
+	}
+
+	static FORCEINLINE int32 Strcmp( const WIDECHAR* String1, const WIDECHAR* String2 )
+	{
+		for (; *String1 || *String2; String1++, String2++)
+		{
+			if (*String1 != *String2)
+			{
+				return *String1 - *String2;
+			}
+		}
+		return 0;
+	}
+
+	static FORCEINLINE int32 Strncmp( const WIDECHAR* String1, const WIDECHAR* String2, SIZE_T Count )
+	{
+		for (; (*String1 || *String2) && Count > 0; String1++, String2++, Count--)
+		{
+			if (*String1 != *String2)
+			{
+				return *String1 - *String2;
+			}
+		}
+		return 0;
+	}
+
+	static FORCEINLINE int32 Strlen( const WIDECHAR* String )
+	{
+		if (!String)
+			return 0;
+
+		int Len = 0;
+		while (String[Len])
+		{
+			++Len;
+		}
+
+		return Len;
+	}
+
+	static FORCEINLINE void CopyWideToAnsi(ANSICHAR* Dest, const WIDECHAR* Src)
+	{
+		if (!Src || !Dest)
+			return;
+
+		int Pos = 0;
+		while (Src[Pos])
+		{
+			if (Src[Pos] <= 255)
+			{
+				Dest[Pos] = Src[Pos];
+			}
+			else
+			{
+				Dest[Pos] = '?';
+			}
+
+			++Pos;
+		}
+
+		Dest[Pos] = 0;
+	}
+
+	static FORCEINLINE void CopyAnsiToWide(WIDECHAR* Dest, const ANSICHAR* Src)
+	{
+		if (!Src || !Dest)
+			return;
+
+		int Pos = 0;
+		while (Src[Pos])
+		{
+			Dest[Pos] = Src[Pos];
+			++Pos;
+		}
+
+		Dest[Pos] = 0;
+	}
+
+	static FORCEINLINE const WIDECHAR* Strstr( const WIDECHAR* String, const WIDECHAR* Find)
+	{
+		WIDECHAR FindChar = *Find;
+
+		// Always find an empty string
+		if (!FindChar)
+			return String;
+
+		++Find;
+		size_t MemCmpLen = wcslen(Find);
+		WIDECHAR* FoundChar;
+		for (;;)
+		{
+			FoundChar = wcschr(String, FindChar);
+			if (!FoundChar)
+			{
+				//	No more instances of FindChar in String, Find does not exist in String
+				break;
+			}
+
+			// Found FindChar, now set String to one character after FoundChar...
+			// We can now compare characters after FoundChar to Find, which is pointing at original Find[1], to see if the rest of the characters match
+			String = FoundChar + 1;
+			if (!wmemcmp(String, Find, MemCmpLen))
+			{
+				// Strings match, return pointer to beginning of Find instance in String
+				return FoundChar;
+			}
+
+			// No match, String is already ready to loop again to find the next instance of FindChar
+		}
+
+		return nullptr;
+	}
+
+	static FORCEINLINE const WIDECHAR* Strchr( const WIDECHAR* String, WIDECHAR C)
+	{
+		if (!String)
+			return NULL;
+
+		int Pos = 0;
+		while (String[Pos])
+		{
+			if(String[Pos] == C)
+			{
+				return &(String[Pos]);
+			}
+
+			++Pos;
+		}
+
+		if (C == 0)
+		{
+			return &(String[Pos]);
+		}
+
+		return NULL;
+	}
+
+	static FORCEINLINE const WIDECHAR* Strrchr( const WIDECHAR* String, WIDECHAR C)
+	{
+		if (!String)
+			return NULL;
+
+		const WIDECHAR* Last = NULL;
+
+		int Pos = 0;
+		while (String[Pos])
+		{
+			if(String[Pos] == C)
+			{
+				Last = &(String[Pos]);
+			}
+
+			++Pos;
+		}
+
+		if (C == 0)
+		{
+			Last = &(String[Pos]);
+		}
+
+		return Last;
+	}
+
+	static FORCEINLINE int32 Atoi(const WIDECHAR* String)
+	{
+		int StringLen = Strlen(String);
+		ANSICHAR* AnsiString = (ANSICHAR*)FMemory_Alloca(StringLen+1);
+		CopyWideToAnsi(AnsiString, String);
+
+		return atoi(AnsiString);
+	}
+
+	static FORCEINLINE int64 Atoi64(const WIDECHAR* String)
+	{
+		int StringLen = Strlen(String);
+		ANSICHAR* AnsiString = (ANSICHAR*)FMemory_Alloca(StringLen+1);
+		CopyWideToAnsi(AnsiString, String);
+
+		return strtoll(AnsiString, NULL, 10);
+	}
+
+	static FORCEINLINE float Atof(const WIDECHAR* String)
+	{
+		int StringLen = Strlen(String);
+		ANSICHAR* AnsiString = (ANSICHAR*)FMemory_Alloca(StringLen+1);
+		CopyWideToAnsi(AnsiString, String);
+
+		return (float)atof(AnsiString);
+	}
+
+	static FORCEINLINE double Atod(const WIDECHAR* String)
+	{
+		int StringLen = Strlen(String);
+		ANSICHAR* AnsiString = (ANSICHAR*)FMemory_Alloca(StringLen+1);
+		CopyWideToAnsi(AnsiString, String);
+
+		return atof(AnsiString);
+	}
+
+	static FORCEINLINE int32 Strtoi( const WIDECHAR* Start, WIDECHAR** End, int32 Base ) 
+	{
+		int StartLen = Strlen(Start);
+		ANSICHAR* AnsiStart = (ANSICHAR*)FMemory_Alloca(StartLen+1);
+		CopyWideToAnsi(AnsiStart, Start);
+
+		ANSICHAR* AnsiEnd = NULL;
+
+		int32 Res = strtol(AnsiStart, &AnsiEnd, Base);
+
+		if (End)
+		{
+			if (AnsiEnd == NULL)
+			{
+				*End = NULL;
+			}
+			else
+			{
+				*End = (WIDECHAR*)(Start + (AnsiEnd - AnsiStart));
+			}
+		}
+
+		return Res;
+	}
+
+	static FORCEINLINE int64 Strtoi64( const WIDECHAR* Start, WIDECHAR** End, int32 Base ) 
+	{
+		int StartLen = Strlen(Start);
+		ANSICHAR* AnsiStart = (ANSICHAR*)FMemory_Alloca(StartLen+1);
+		CopyWideToAnsi(AnsiStart, Start);
+
+		ANSICHAR* AnsiEnd = NULL;
+
+		uint64 Res = strtoll(AnsiStart, &AnsiEnd, Base);
+
+		if (End)
+		{
+			if (AnsiEnd == NULL)
+			{
+				*End = NULL;
+			}
+			else
+			{
+				*End = (WIDECHAR*)(Start + (AnsiEnd - AnsiStart));
+			}
+		}
+
+		return Res;
+	}
+
+	static FORCEINLINE uint64 Strtoui64( const WIDECHAR* Start, WIDECHAR** End, int32 Base ) 
+	{
+		int StartLen = Strlen(Start);
+		ANSICHAR* AnsiStart = (ANSICHAR*)FMemory_Alloca(StartLen+1);
+		CopyWideToAnsi(AnsiStart, Start);
+
+		ANSICHAR* AnsiEnd = NULL;
+
+		uint64 Res = strtoull(AnsiStart, &AnsiEnd, Base);
+
+		if (End)
+		{
+			if (AnsiEnd == NULL)
+			{
+				*End = NULL;
+			}
+			else
+			{
+				*End = (WIDECHAR*)(Start + (AnsiEnd - AnsiStart));
+			}
+		}
+
+		return Res;
+	}
+
+	static FORCEINLINE WIDECHAR* Strtok(WIDECHAR* StrToken, const WIDECHAR* Delim, WIDECHAR** Context)
+	{
+		int StrTokenLen = Strlen(StrToken);
+		int DelimLen = Strlen(Delim);
+		ANSICHAR* AnsiStrToken = (ANSICHAR*)FMemory_Alloca(StrTokenLen+1);
+		ANSICHAR* AnsiDelim = (ANSICHAR*)FMemory_Alloca(DelimLen+1);
+		CopyWideToAnsi(AnsiStrToken, StrToken);
+		CopyWideToAnsi(AnsiDelim, Delim);
+
+		ANSICHAR* Pos = strtok(AnsiStrToken, AnsiDelim);
+
+		if (!Pos)
+		{
+			return NULL;
+		}
+
+		return StrToken + (Pos - AnsiStrToken);
+	}
+
+	UE_DEPRECATED(4.22, "GetVarArgs with DestSize and Count arguments has been deprecated - only DestSize should be passed")
+	static FORCEINLINE int32 GetVarArgs( WIDECHAR* Dest, SIZE_T DestSize, int32 Count, const WIDECHAR*& Fmt, va_list ArgPtr )
+	{
+		return GetVarArgs(Dest, DestSize, Fmt, ArgPtr);
+	}
+
+	static FORCEINLINE int32 GetVarArgs( WIDECHAR* Dest, SIZE_T DestSize, const WIDECHAR*& Fmt, va_list ArgPtr )
+	{
+#if PLATFORM_USE_LS_SPEC_FOR_WIDECHAR
+		// fix up the Fmt string, as fast as possible, without using an FString
+		const WIDECHAR* OldFormat = Fmt;
+		WIDECHAR* NewFormat = (WIDECHAR*)FMemory_Alloca((Strlen(Fmt) * 2 + 1) * sizeof(WIDECHAR));
+		
+		int32 NewIndex = 0;
+
+		for (; *OldFormat != 0; NewIndex++, OldFormat++)
+		{
+			// fix up %s -> %ls
+			if (OldFormat[0] == LITERAL(WIDECHAR, '%'))
+			{
+				NewFormat[NewIndex++] = *OldFormat++;
+
+				if (*OldFormat == LITERAL(WIDECHAR, '%'))
+				{
+					NewFormat[NewIndex] = *OldFormat;
+				}
+				else
+				{
+					const WIDECHAR* NextChar = OldFormat;
+					
+					while(*NextChar != 0 && !FChar::IsAlpha(*NextChar)) 
+					{ 
+						NewFormat[NewIndex++] = *NextChar;
+						++NextChar; 
+					};
+
+					if (*NextChar == LITERAL(WIDECHAR, 's'))
+					{
+						NewFormat[NewIndex++] = LITERAL(WIDECHAR, 'l');
+						NewFormat[NewIndex] = *NextChar;
+					}
+					else
+					{
+						NewFormat[NewIndex] = *NextChar;
+					}
+
+					OldFormat = NextChar;
+				}
+			}
+			else
+			{
+				NewFormat[NewIndex] = *OldFormat;
+			}
+		}
+		NewFormat[NewIndex] = 0;
+#endif // PLATFORM_USE_LS_SPEC_FOR_WIDECHAR
+		int32 Result = vswprintf( Dest, DestSize, NewFormat, ArgPtr);
+		va_end( ArgPtr );
+		return Result;
+	}
+
+	/** 
+	 * Ansi implementation 
+	 **/
+	static FORCEINLINE ANSICHAR* Strcpy(ANSICHAR* Dest, SIZE_T DestCount, const ANSICHAR* Src)
+	{
+		return strcpy( Dest, Src );
+	}
+
+	static FORCEINLINE ANSICHAR* Strncpy(ANSICHAR* Dest, const ANSICHAR* Src, int32 MaxLen)
+	{
+		::strncpy(Dest, Src, MaxLen);
+		Dest[MaxLen-1]=0;
+		return Dest;
+	}
+
+	static FORCEINLINE ANSICHAR* Strcat(ANSICHAR* Dest, SIZE_T DestCount, const ANSICHAR* Src)
+	{
+		return strcat( Dest, Src );
+	}
+
+	static FORCEINLINE int32 Strcmp( const ANSICHAR* String1, const ANSICHAR* String2 )
+	{
+		return strcmp(String1, String2);
+	}
+
+	static FORCEINLINE int32 Strncmp( const ANSICHAR* String1, const ANSICHAR* String2, SIZE_T Count )
+	{
+		return strncmp( String1, String2, Count );
+	}
+
+	static FORCEINLINE int32 Strlen( const ANSICHAR* String )
+	{
+		return strlen( String ); 
+	}
+
+	static FORCEINLINE const ANSICHAR* Strstr( const ANSICHAR* String, const ANSICHAR* Find)
+	{
+		return strstr(String, Find);
+	}
+
+	static FORCEINLINE const ANSICHAR* Strchr( const ANSICHAR* String, ANSICHAR C)
+	{
+		return strchr(String, C);
+	}
+
+	static FORCEINLINE const ANSICHAR* Strrchr( const ANSICHAR* String, ANSICHAR C)
+	{
+		return strrchr(String, C);
+	}
+
+	static FORCEINLINE int32 Atoi(const ANSICHAR* String)
+	{
+		return atoi( String ); 
+	}
+
+	static FORCEINLINE int64 Atoi64(const ANSICHAR* String)
+	{
+		return strtoll( String, NULL, 10 );
+	}
+
+	static FORCEINLINE float Atof(const ANSICHAR* String)
+	{
+		return (float)atof( String ); 
+	}
+
+	static FORCEINLINE double Atod(const ANSICHAR* String)
+	{
+		return atof( String ); 
+	}
+
+	static FORCEINLINE int32 Strtoi( const ANSICHAR* Start, ANSICHAR** End, int32 Base ) 
+	{
+		return strtol( Start, End, Base ); 
+	}
+
+	static FORCEINLINE int64 Strtoi64( const ANSICHAR* Start, ANSICHAR** End, int32 Base ) 
+	{
+		return strtoll(Start, End, Base);;
+	}
+
+	static FORCEINLINE uint64 Strtoui64( const ANSICHAR* Start, ANSICHAR** End, int32 Base ) 
+	{
+		return strtoull(Start, End, Base);;
+	}
+
+	static FORCEINLINE ANSICHAR* Strtok(ANSICHAR* StrToken, const ANSICHAR* Delim, ANSICHAR** Context)
+	{
+		return strtok(StrToken, Delim);
+	}
+
+	UE_DEPRECATED(4.22, "GetVarArgs with DestSize and Count arguments has been deprecated - only DestSize should be passed")
+	static FORCEINLINE int32 GetVarArgs(ANSICHAR* Dest, SIZE_T DestSize, int32 Count, const ANSICHAR*& Fmt, va_list ArgPtr)
+	{
+		return GetVarArgs(Dest, DestSize, Count, Fmt, ArgPtr);
+	}
+
+	static FORCEINLINE int32 GetVarArgs( ANSICHAR* Dest, SIZE_T DestSize, const ANSICHAR*& Fmt, va_list ArgPtr )
+	{
+		int32 Result = vsnprintf(Dest, DestSize, Fmt, ArgPtr);
+		va_end( ArgPtr );
+		return (Result != -1 && Result < (int32)DestSize) ? Result : -1;
+	}
+
+	/** 
+	 * UCS2 implementation 
+	 **/
+
+	static FORCEINLINE int32 Strlen( const UCS2CHAR* String )
+	{
+		int32 Result = 0;
+		while (*String++)
+		{
+			++Result;
+		}
+
+		return Result;
+	}
+};
+#endif
+
+typedef FAndroidPlatformString FPlatformString;
+
diff --git a/Engine/Source/Runtime/Core/Public/Android/AndroidPlatformTLS.h b/Engine/Source/Runtime/Core/Public/Android/AndroidPlatformTLS.h
new file mode 100644
--- /dev/null
+++ b/Engine/Source/Runtime/Core/Public/Android/AndroidPlatformTLS.h
@@ -0,0 +1,73 @@
+// Copyright 1998-2019 Epic Games, Inc. All Rights Reserved.
+
+
+/*=============================================================================================
+	AndroidTLS.h: Android platform TLS (Thread local storage and thread ID) functions
+==============================================================================================*/
+
+#pragma once
+#include "GenericPlatform/GenericPlatformTLS.h"
+#include "Android/AndroidSystemIncludes.h"
+
+/**
+ * Android implementation of the TLS OS functions
+ */
+struct CORE_API FAndroidTLS : public FGenericPlatformTLS
+{
+	/**
+	 * Returns the currently executing thread's id
+	 */
+	static FORCEINLINE uint32 GetCurrentThreadId(void)
+	{
+		static_assert(sizeof(pid_t) <= sizeof(uint32), "pid_t is larger than uint32, reconsider implementation of GetCurrentThreadId()");
+		return static_cast<uint32>(gettid());
+//		return pthread_self();
+	}
+
+	/**
+	 * Allocates a thread local store slot
+	 */
+	static FORCEINLINE uint32 AllocTlsSlot(void)
+	{
+		// allocate a per-thread mem slot
+		pthread_key_t Key = 0;
+		if (pthread_key_create(&Key, NULL) != 0)
+		{
+			Key = 0xFFFFFFFF;  // matches the Windows TlsAlloc() retval //@todo android: should probably check for this below, or assert out instead
+		}
+		return Key;
+	}
+
+	/**
+	 * Sets a value in the specified TLS slot
+	 *
+	 * @param SlotIndex the TLS index to store it in
+	 * @param Value the value to store in the slot
+	 */
+	static FORCEINLINE void SetTlsValue(uint32 SlotIndex,void* Value)
+	{
+		pthread_setspecific((pthread_key_t)SlotIndex, Value);
+	}
+
+	/**
+	 * Reads the value stored at the specified TLS slot
+	 *
+	 * @return the value stored in the slot
+	 */
+	static FORCEINLINE void* GetTlsValue(uint32 SlotIndex)
+	{
+		return pthread_getspecific((pthread_key_t)SlotIndex);
+	}
+
+	/**
+	 * Frees a previously allocated TLS slot
+	 *
+	 * @param SlotIndex the TLS index to store it in
+	 */
+	static FORCEINLINE void FreeTlsSlot(uint32 SlotIndex)
+	{
+		pthread_key_delete((pthread_key_t)SlotIndex);
+	}
+};
+
+typedef FAndroidTLS FPlatformTLS;
diff --git a/Engine/Source/Runtime/Core/Public/Android/AndroidPlatformTime.h b/Engine/Source/Runtime/Core/Public/Android/AndroidPlatformTime.h
new file mode 100644
--- /dev/null
+++ b/Engine/Source/Runtime/Core/Public/Android/AndroidPlatformTime.h
@@ -0,0 +1,42 @@
+// Copyright 1998-2019 Epic Games, Inc. All Rights Reserved.
+
+
+/*=============================================================================================
+	AndroidTime.h: Android platform Time functions
+==============================================================================================*/
+
+#pragma once
+#include "GenericPlatform/GenericPlatformTime.h"
+#include <sys/time.h>
+
+//@todo android: this entire file
+
+/**
+ * Android implementation of the Time OS functions
+ */
+struct CORE_API FAndroidTime : public FGenericPlatformTime
+{
+	// android uses BSD time code from GenericPlatformTime
+	static FORCEINLINE double Seconds()
+	{
+		struct timespec ts;
+		clock_gettime(CLOCK_MONOTONIC, &ts);
+		return ((double) ts.tv_sec) + (((double) ts.tv_nsec) / 1000000000.0);
+	}
+
+	static FORCEINLINE uint32 Cycles()
+	{
+		struct timespec ts;
+		clock_gettime(CLOCK_MONOTONIC, &ts);
+		return (uint32) ((((uint64)ts.tv_sec) * 1000000ULL) + (((uint64)ts.tv_nsec) / 1000ULL));
+	}
+
+	static FORCEINLINE uint64 Cycles64()
+	{
+		struct timespec ts;
+		clock_gettime(CLOCK_MONOTONIC, &ts);
+		return ((((uint64)ts.tv_sec) * 1000000ULL) + (((uint64)ts.tv_nsec) / 1000ULL));
+	}
+};
+
+typedef FAndroidTime FPlatformTime;
diff --git a/Engine/Source/Runtime/Core/Public/Android/AndroidProcess.h b/Engine/Source/Runtime/Core/Public/Android/AndroidProcess.h
deleted file mode 100644
diff --git a/Engine/Source/Runtime/Core/Public/Android/AndroidProperties.h b/Engine/Source/Runtime/Core/Public/Android/AndroidProperties.h
deleted file mode 100644
diff --git a/Engine/Source/Runtime/Core/Public/Android/AndroidStats.h b/Engine/Source/Runtime/Core/Public/Android/AndroidStats.h
--- a/Engine/Source/Runtime/Core/Public/Android/AndroidStats.h
+++ b/Engine/Source/Runtime/Core/Public/Android/AndroidStats.h
@@ -2,7 +2,7 @@
 #pragma once
 
 #include "CoreMinimal.h"
-#include "Android/AndroidMisc.h"
+#include "Android/AndroidPlatformMisc.h"
 
 class FAndroidStats
 {
diff --git a/Engine/Source/Runtime/Core/Public/Android/AndroidString.h b/Engine/Source/Runtime/Core/Public/Android/AndroidString.h
deleted file mode 100644
diff --git a/Engine/Source/Runtime/Core/Public/Android/AndroidSystemIncludes.h b/Engine/Source/Runtime/Core/Public/Android/AndroidSystemIncludes.h
--- a/Engine/Source/Runtime/Core/Public/Android/AndroidSystemIncludes.h
+++ b/Engine/Source/Runtime/Core/Public/Android/AndroidSystemIncludes.h
@@ -14,7 +14,7 @@
 #include <unistd.h>
 
 // Set up compiler pragmas, etc
-#include "Android/AndroidCompilerSetup.h"
+#include "Android/AndroidPlatformCompilerSetup.h"
 
 //@todo android: verify malloc alignment or change?
 #define _aligned_malloc(Size,Align) malloc(Size)
diff --git a/Engine/Source/Runtime/Core/Public/Android/AndroidTLS.h b/Engine/Source/Runtime/Core/Public/Android/AndroidTLS.h
deleted file mode 100644
diff --git a/Engine/Source/Runtime/Core/Public/Android/AndroidTime.h b/Engine/Source/Runtime/Core/Public/Android/AndroidTime.h
deleted file mode 100644
diff --git a/Engine/Source/Runtime/Core/Public/HAL/CriticalSection.h b/Engine/Source/Runtime/Core/Public/HAL/CriticalSection.h
--- a/Engine/Source/Runtime/Core/Public/HAL/CriticalSection.h
+++ b/Engine/Source/Runtime/Core/Public/HAL/CriticalSection.h
@@ -4,8 +4,4 @@
 
 #include "CoreTypes.h"
 
-#if PLATFORM_UNIX
-	#include "Unix/UnixCriticalSection.h"
-#else
-	#include COMPILED_PLATFORM_HEADER(CriticalSection.h)
-#endif
+#include COMPILED_PLATFORM_HEADER(CriticalSection.h)
diff --git a/Engine/Source/Runtime/Core/Public/HAL/Platform.h b/Engine/Source/Runtime/Core/Public/HAL/Platform.h
--- a/Engine/Source/Runtime/Core/Public/HAL/Platform.h
+++ b/Engine/Source/Runtime/Core/Public/HAL/Platform.h
@@ -902,15 +902,8 @@ namespace TypeTests
 }
 
 // Platform specific compiler setup.
-#if PLATFORM_PS4
-	#include "PS4/PS4CompilerSetup.h"
-#elif PLATFORM_XBOXONE
-	#include "XboxOne/XboxOneCompilerSetup.h"
-#elif PLATFORM_ANDROID
-	#include "Android/AndroidCompilerSetup.h"
-#else
 #include COMPILED_PLATFORM_HEADER(PlatformCompilerSetup.h)
-#endif
+
 
 // If we don't have a platform-specific define for the TEXT macro, define it now.
 #if !defined(TEXT) && !UE_BUILD_DOCS
diff --git a/Engine/Source/Runtime/Core/Public/HAL/PlatformAffinity.h b/Engine/Source/Runtime/Core/Public/HAL/PlatformAffinity.h
--- a/Engine/Source/Runtime/Core/Public/HAL/PlatformAffinity.h
+++ b/Engine/Source/Runtime/Core/Public/HAL/PlatformAffinity.h
@@ -4,14 +4,4 @@
 
 #include "CoreTypes.h"
 
-#if PLATFORM_XBOXONE
-#include "XboxOne/XBoxOneAffinity.h"
-#elif PLATFORM_PS4
-#include "PS4/PS4Affinity.h"
-#elif PLATFORM_LUMIN
-#include "Lumin/LuminAffinity.h"
-#elif PLATFORM_ANDROID
-#include "Android/AndroidAffinity.h"
-#else
 #include COMPILED_PLATFORM_HEADER(PlatformAffinity.h)
-#endif
diff --git a/Engine/Source/Runtime/Core/Public/HAL/PlatformAtomics.h b/Engine/Source/Runtime/Core/Public/HAL/PlatformAtomics.h
--- a/Engine/Source/Runtime/Core/Public/HAL/PlatformAtomics.h
+++ b/Engine/Source/Runtime/Core/Public/HAL/PlatformAtomics.h
@@ -5,18 +5,4 @@
 #include "CoreTypes.h"
 #include "GenericPlatform/GenericPlatformAtomics.h"
 
-#if PLATFORM_PS4
-#include "PS4/PS4Atomics.h"
-#elif PLATFORM_XBOXONE
-#include "XboxOne/XboxOneAtomics.h"
-#elif PLATFORM_MAC
-#include "Apple/ApplePlatformAtomics.h"
-#elif PLATFORM_IOS
-#include "Apple/ApplePlatformAtomics.h"
-#elif PLATFORM_ANDROID
-#include "Android/AndroidAtomics.h"
-#elif PLATFORM_UNIX
-#include "Unix/UnixPlatformAtomics.h"
-#else
 #include COMPILED_PLATFORM_HEADER(PlatformAtomics.h)
-#endif
diff --git a/Engine/Source/Runtime/Core/Public/HAL/PlatformFile.h b/Engine/Source/Runtime/Core/Public/HAL/PlatformFile.h
--- a/Engine/Source/Runtime/Core/Public/HAL/PlatformFile.h
+++ b/Engine/Source/Runtime/Core/Public/HAL/PlatformFile.h
@@ -4,20 +4,4 @@
 #include "CoreTypes.h"
 #include "GenericPlatform/GenericPlatformFile.h"
 
-#if PLATFORM_PS4
-#include "PS4/PS4File.h"
-#elif PLATFORM_XBOXONE
-#include "XboxOne/XboxOneFile.h"
-#elif PLATFORM_MAC
-#include "Apple/ApplePlatformFile.h"
-#elif PLATFORM_IOS
-#include "Apple/ApplePlatformFile.h"
-#elif PLATFORM_ANDROID
-#include "Android/AndroidFile.h"
-#elif PLATFORM_HTML5
-//#include "HTML5PlatformFile.h"
-#elif PLATFORM_LINUX
-#include "Linux/LinuxPlatformFile.h"
-#else
 #include COMPILED_PLATFORM_HEADER(PlatformFile.h)
-#endif
diff --git a/Engine/Source/Runtime/Core/Public/HAL/PlatformMath.h b/Engine/Source/Runtime/Core/Public/HAL/PlatformMath.h
--- a/Engine/Source/Runtime/Core/Public/HAL/PlatformMath.h
+++ b/Engine/Source/Runtime/Core/Public/HAL/PlatformMath.h
@@ -4,16 +4,4 @@
 #include "CoreTypes.h"
 #include "GenericPlatform/GenericPlatformMath.h"
 
-#if PLATFORM_PS4
-#include "PS4/PS4Math.h"
-#elif PLATFORM_XBOXONE
-#include "XboxOne/XboxOneMath.h"
-#elif PLATFORM_MAC
-#include "Mac/MacPlatformMath.h"
-#elif PLATFORM_ANDROID
-#include "Android/AndroidMath.h"
-#elif PLATFORM_UNIX
-#include "Unix/UnixPlatformMath.h"
-#else
 #include COMPILED_PLATFORM_HEADER(PlatformMath.h)
-#endif
diff --git a/Engine/Source/Runtime/Core/Public/HAL/PlatformMemory.h b/Engine/Source/Runtime/Core/Public/HAL/PlatformMemory.h
--- a/Engine/Source/Runtime/Core/Public/HAL/PlatformMemory.h
+++ b/Engine/Source/Runtime/Core/Public/HAL/PlatformMemory.h
@@ -4,19 +4,4 @@
 #include "CoreTypes.h"
 #include "GenericPlatform/GenericPlatformMemory.h"
 
-#if PLATFORM_PS4
-
-#if USE_NEW_PS4_MEMORY_SYSTEM
-#include "PS4/PS4Memory2.h"
-#else
-#include "PS4/PS4Memory.h"
-#endif
-#elif PLATFORM_XBOXONE
-#include "XboxOne/XboxOneMemory.h"
-#elif PLATFORM_ANDROID
-#include "Android/AndroidMemory.h"
-#elif PLATFORM_UNIX
-#include "Unix/UnixPlatformMemory.h"
-#else
 #include COMPILED_PLATFORM_HEADER(PlatformMemory.h)
-#endif
diff --git a/Engine/Source/Runtime/Core/Public/HAL/PlatformMisc.h b/Engine/Source/Runtime/Core/Public/HAL/PlatformMisc.h
--- a/Engine/Source/Runtime/Core/Public/HAL/PlatformMisc.h
+++ b/Engine/Source/Runtime/Core/Public/HAL/PlatformMisc.h
@@ -4,17 +4,8 @@
 #include "CoreTypes.h"
 #include "GenericPlatform/GenericPlatformMisc.h"
 
-#ifdef PLATFORM_HEADER_MISC
-#include PLATFORM_HEADER_MISC
-#elif PLATFORM_PS4
-#include "PS4/PS4Misc.h"
-#elif PLATFORM_XBOXONE
-#include "XboxOne/XboxOneMisc.h"
-#elif PLATFORM_ANDROID
-#include "Android/AndroidMisc.h"
-#else
 #include COMPILED_PLATFORM_HEADER(PlatformMisc.h)
-#endif
+
 
 #ifndef UE_DEBUG_BREAK
 #error UE_DEBUG_BREAK is not defined for this platform
diff --git a/Engine/Source/Runtime/Core/Public/HAL/PlatformOutputDevices.h b/Engine/Source/Runtime/Core/Public/HAL/PlatformOutputDevices.h
--- a/Engine/Source/Runtime/Core/Public/HAL/PlatformOutputDevices.h
+++ b/Engine/Source/Runtime/Core/Public/HAL/PlatformOutputDevices.h
@@ -3,14 +3,4 @@
 
 #include "CoreTypes.h"
 
-#if PLATFORM_PS4
-#include "PS4/PS4OutputDevices.h"
-#elif PLATFORM_XBOXONE
-#include "XboxOne/XboxOneOutputDevices.h"
-#elif PLATFORM_ANDROID
-#include "Android/AndroidOutputDevices.h"
-#elif PLATFORM_UNIX
-#include "Unix/UnixPlatformOutputDevices.h"
-#else
 #include COMPILED_PLATFORM_HEADER(PlatformOutputDevices.h)
-#endif
diff --git a/Engine/Source/Runtime/Core/Public/HAL/PlatformProcess.h b/Engine/Source/Runtime/Core/Public/HAL/PlatformProcess.h
--- a/Engine/Source/Runtime/Core/Public/HAL/PlatformProcess.h
+++ b/Engine/Source/Runtime/Core/Public/HAL/PlatformProcess.h
@@ -4,12 +4,4 @@
 
 #include "CoreTypes.h"
 
-#if PLATFORM_PS4
-#include "PS4/PS4Process.h"
-#elif PLATFORM_XBOXONE
-#include "XboxOne/XboxOneProcess.h"
-#elif PLATFORM_ANDROID
-#include "Android/AndroidProcess.h"
-#else
 #include COMPILED_PLATFORM_HEADER(PlatformProcess.h)
-#endif
diff --git a/Engine/Source/Runtime/Core/Public/HAL/PlatformProperties.h b/Engine/Source/Runtime/Core/Public/HAL/PlatformProperties.h
--- a/Engine/Source/Runtime/Core/Public/HAL/PlatformProperties.h
+++ b/Engine/Source/Runtime/Core/Public/HAL/PlatformProperties.h
@@ -8,14 +8,6 @@
 // the Properties.h files are included in all other places, so just use #ifdef not #if in this special case
 #define PROPERTY_HEADER_SHOULD_DEFINE_TYPE
 
-#if PLATFORM_PS4
-#include "PS4/PS4Properties.h"
-#elif PLATFORM_XBOXONE
-#include "XboxOne/XboxOneProperties.h"
-#elif PLATFORM_ANDROID
-#include "Android/AndroidProperties.h"
-#else
 #include COMPILED_PLATFORM_HEADER(PlatformProperties.h)
-#endif
 
 #undef PROPERTY_HEADER_SHOULD_DEFINE_TYPE
diff --git a/Engine/Source/Runtime/Core/Public/HAL/PlatformStackWalk.h b/Engine/Source/Runtime/Core/Public/HAL/PlatformStackWalk.h
--- a/Engine/Source/Runtime/Core/Public/HAL/PlatformStackWalk.h
+++ b/Engine/Source/Runtime/Core/Public/HAL/PlatformStackWalk.h
@@ -3,10 +3,4 @@
 
 #include "CoreTypes.h"
 
-#if PLATFORM_PS4
-#include "PS4/PS4StackWalk.h"
-#elif PLATFORM_XBOXONE
-#include "XboxOne/XboxOneStackWalk.h"
-#else
 #include COMPILED_PLATFORM_HEADER(PlatformStackWalk.h)
-#endif
diff --git a/Engine/Source/Runtime/Core/Public/HAL/PlatformString.h b/Engine/Source/Runtime/Core/Public/HAL/PlatformString.h
--- a/Engine/Source/Runtime/Core/Public/HAL/PlatformString.h
+++ b/Engine/Source/Runtime/Core/Public/HAL/PlatformString.h
@@ -3,19 +3,5 @@
 
 #include "CoreTypes.h"
 
-#if PLATFORM_PS4
-#include "PS4/PS4String.h"
-#elif PLATFORM_XBOXONE
-#include "XboxOne/XboxOneString.h"
-#elif PLATFORM_MAC
-#include "Apple/ApplePlatformString.h"
-#elif PLATFORM_IOS
-#include "Apple/ApplePlatformString.h"
-#elif PLATFORM_ANDROID
-#include "Android/AndroidString.h"
-#elif PLATFORM_UNIX
-#include "Unix/UnixPlatformString.h"
-#else
 #include COMPILED_PLATFORM_HEADER(PlatformString.h)
-#endif
 
diff --git a/Engine/Source/Runtime/Core/Public/HAL/PlatformTLS.h b/Engine/Source/Runtime/Core/Public/HAL/PlatformTLS.h
--- a/Engine/Source/Runtime/Core/Public/HAL/PlatformTLS.h
+++ b/Engine/Source/Runtime/Core/Public/HAL/PlatformTLS.h
@@ -5,18 +5,4 @@
 #include "CoreTypes.h"
 #include "GenericPlatform/GenericPlatformTLS.h"
 
-#if PLATFORM_PS4
-#include "PS4/PS4TLS.h"
-#elif PLATFORM_XBOXONE
-#include "XboxOne/XboxOneTLS.h"
-#elif PLATFORM_MAC
-#include "Apple/ApplePlatformTLS.h"
-#elif PLATFORM_IOS
-#include "Apple/ApplePlatformTLS.h"
-#elif PLATFORM_ANDROID
-#include "Android/AndroidTLS.h"
-#elif PLATFORM_UNIX
-#include "Unix/UnixPlatformTLS.h"
-#else
 #include COMPILED_PLATFORM_HEADER(PlatformTLS.h)
-#endif
diff --git a/Engine/Source/Runtime/Core/Public/HAL/PlatformTime.h b/Engine/Source/Runtime/Core/Public/HAL/PlatformTime.h
--- a/Engine/Source/Runtime/Core/Public/HAL/PlatformTime.h
+++ b/Engine/Source/Runtime/Core/Public/HAL/PlatformTime.h
@@ -4,18 +4,4 @@
 
 #include "CoreTypes.h"
 
-#if PLATFORM_PS4
-#include "PS4/PS4Time.h"
-#elif PLATFORM_XBOXONE
-#include "XboxOne/XboxOneTime.h"
-#elif PLATFORM_MAC
-#include "Apple/ApplePlatformTime.h"
-#elif PLATFORM_IOS
-#include "Apple/ApplePlatformTime.h"
-#elif PLATFORM_ANDROID
-#include "Android/AndroidTime.h"
-#elif PLATFORM_UNIX
-#include "Unix/UnixPlatformTime.h"
-#else
 #include COMPILED_PLATFORM_HEADER(PlatformTime.h)
-#endif
diff --git a/Engine/Source/Runtime/Core/Public/HTML5/HTML5PlatformFile.h b/Engine/Source/Runtime/Core/Public/HTML5/HTML5PlatformFile.h
new file mode 100644
--- /dev/null
+++ b/Engine/Source/Runtime/Core/Public/HTML5/HTML5PlatformFile.h
@@ -0,0 +1,3 @@
+// Copyright 1998-2019 Epic Games, Inc. All Rights Reserved.
+
+#pragma once
diff --git a/Engine/Source/Runtime/Core/Public/IOS/IOSPlatformFile.h b/Engine/Source/Runtime/Core/Public/IOS/IOSPlatformFile.h
new file mode 100644
--- /dev/null
+++ b/Engine/Source/Runtime/Core/Public/IOS/IOSPlatformFile.h
@@ -0,0 +1,66 @@
+// Copyright 1998-2019 Epic Games, Inc. All Rights Reserved.
+
+#pragma once
+
+#include "Apple/ApplePlatformFile.h"
+#include "Containers/UnrealString.h"
+
+class IFileHandle;
+class FApplePlatformFile;
+
+
+/**
+ * iOS File I/O implementation
+**/
+class CORE_API FIOSPlatformFile : public FApplePlatformFile
+{
+protected:
+	virtual FString NormalizeFilename(const TCHAR* Filename);
+	virtual FString NormalizeDirectory(const TCHAR* Directory);
+
+public:
+	FIOSPlatformFile();
+
+	//virtual bool Initialize(IPlatformFile* Inner, const TCHAR* CommandLineParam) override;
+
+	virtual bool FileExists(const TCHAR* Filename) override;
+	virtual int64 FileSize(const TCHAR* Filename) override;
+	virtual bool DeleteFile(const TCHAR* Filename) override;
+	virtual bool IsReadOnly(const TCHAR* Filename) override;
+	virtual bool MoveFile(const TCHAR* To, const TCHAR* From) override;
+	virtual bool SetReadOnly(const TCHAR* Filename, bool bNewReadOnlyValue) override;
+
+	virtual FDateTime GetTimeStamp(const TCHAR* Filename) override;
+	virtual void SetTimeStamp(const TCHAR* Filename, const FDateTime DateTime) override;
+	virtual FDateTime GetAccessTimeStamp(const TCHAR* Filename) override;
+	virtual FString GetFilenameOnDisk(const TCHAR* Filename) override;
+    
+    virtual FString ConvertToAbsolutePathForExternalAppForRead( const TCHAR* Filename ) override;
+    virtual FString ConvertToAbsolutePathForExternalAppForWrite( const TCHAR* Filename ) override;
+
+	virtual FFileStatData GetStatData(const TCHAR* FilenameOrDirectory) override;
+
+	virtual IFileHandle* OpenRead(const TCHAR* Filename, bool bAllowWrite = false) override;
+	virtual IFileHandle* OpenWrite(const TCHAR* Filename, bool bAppend = false, bool bAllowRead = false) override;
+	virtual IMappedFileHandle* OpenMapped(const TCHAR* Filename) override;
+
+	virtual bool DirectoryExists(const TCHAR* Directory) override;
+	virtual bool CreateDirectory(const TCHAR* Directory) override;
+	virtual bool DeleteDirectory(const TCHAR* Directory) override;
+
+	virtual bool IterateDirectory(const TCHAR* Directory, FDirectoryVisitor& Visitor) override;
+	virtual bool IterateDirectoryStat(const TCHAR* Directory, FDirectoryStatVisitor& Visitor) override;
+
+	virtual bool DoesCreatePublicFiles() override;
+	virtual void SetCreatePublicFiles(bool bCreatePublicFilesIn) override;
+	
+private:
+	bool IterateDirectoryCommon(const TCHAR* Directory, const TFunctionRef<bool(struct dirent*)>& Visitor);
+
+	FString ConvertToIOSPath(const FString& Filename, bool bForWrite, bool bIsPublicWrite);
+
+	FString ReadDirectory;
+	FString WriteDirectory;
+	
+	bool bCreatePublicFiles;
+};
diff --git a/Engine/Source/Runtime/Core/Public/IOS/IOSPlatformString.h b/Engine/Source/Runtime/Core/Public/IOS/IOSPlatformString.h
new file mode 100644
--- /dev/null
+++ b/Engine/Source/Runtime/Core/Public/IOS/IOSPlatformString.h
@@ -0,0 +1,5 @@
+// Copyright 1998-2019 Epic Games, Inc. All Rights Reserved.
+
+#pragma once
+
+#include "Apple/ApplePlatformString.h"
diff --git a/Engine/Source/Runtime/Core/Public/IOS/IOSPlatformTLS.h b/Engine/Source/Runtime/Core/Public/IOS/IOSPlatformTLS.h
new file mode 100644
--- /dev/null
+++ b/Engine/Source/Runtime/Core/Public/IOS/IOSPlatformTLS.h
@@ -0,0 +1,5 @@
+// Copyright 1998-2019 Epic Games, Inc. All Rights Reserved.
+
+#pragma once
+
+#include "Apple/ApplePlatformTLS.h"
diff --git a/Engine/Source/Runtime/Core/Public/IOS/IOSPlatformTime.h b/Engine/Source/Runtime/Core/Public/IOS/IOSPlatformTime.h
new file mode 100644
--- /dev/null
+++ b/Engine/Source/Runtime/Core/Public/IOS/IOSPlatformTime.h
@@ -0,0 +1,5 @@
+// Copyright 1998-2019 Epic Games, Inc. All Rights Reserved.
+
+#pragma once
+
+#include "Apple/ApplePlatformTime.h"
diff --git a/Engine/Source/Runtime/Core/Public/Linux/LinuxCriticalSection.h b/Engine/Source/Runtime/Core/Public/Linux/LinuxCriticalSection.h
new file mode 100644
--- /dev/null
+++ b/Engine/Source/Runtime/Core/Public/Linux/LinuxCriticalSection.h
@@ -0,0 +1,5 @@
+// Copyright 1998-2019 Epic Games, Inc. All Rights Reserved.
+
+#pragma once
+
+#include "Unix/UnixCriticalSection.h"
diff --git a/Engine/Source/Runtime/Core/Public/Linux/LinuxPlatformAffinity.h b/Engine/Source/Runtime/Core/Public/Linux/LinuxPlatformAffinity.h
new file mode 100644
--- /dev/null
+++ b/Engine/Source/Runtime/Core/Public/Linux/LinuxPlatformAffinity.h
@@ -0,0 +1,65 @@
+// Copyright 1998-2019 Epic Games, Inc. All Rights Reserved.
+
+/*=============================================================================================
+LuminAffinity.h: Lumin affinity profile masks definitions.
+==============================================================================================*/
+
+#pragma once
+
+#include "GenericPlatform/GenericPlatformAffinity.h"
+
+class FLuminAffinity : public FGenericPlatformAffinity
+{
+	const static uint64 ArmCores = MAKEAFFINITYMASK2(3,4);
+	const static uint64 DenverCores = MAKEAFFINITYMASK1(2);
+public:
+	static const CORE_API uint64 GetMainGameMask() {
+		return ArmCores;
+	}
+
+	static const CORE_API uint64 GetRenderingThreadMask() {
+		return ArmCores;
+	}
+
+	static const CORE_API uint64 GetRHIThreadMask() {
+		return ArmCores;
+	}
+
+	static const CORE_API uint64 GetRTHeartBeatMask() {
+		return ArmCores;
+	}
+
+	static const CORE_API uint64 GetPoolThreadMask() {
+		return DenverCores;
+	}
+
+	static const CORE_API uint64 GetTaskGraphThreadMask() {
+		return ArmCores;
+	}
+
+	static const CORE_API uint64 GetStatsThreadMask() {
+		return ArmCores;
+	}
+
+	static const CORE_API uint64 GetAudioThreadMask() {
+		return DenverCores;
+	}
+
+	static const CORE_API uint64 GetNoAffinityMask() {
+		return 0xFFFFFFFFFFFFFFFF;
+	}
+
+	static const CORE_API uint64 GetTaskGraphBackgroundTaskMask() {
+		return ArmCores;
+	}
+
+	static EThreadPriority GetRenderingThreadPriority() {
+		return TPri_Normal;
+	}
+
+	static EThreadPriority GetRHIThreadPriority() {
+		return TPri_SlightlyBelowNormal;
+	}
+};
+
+typedef FLuminAffinity FPlatformAffinity;
diff --git a/Engine/Source/Runtime/Core/Public/Linux/LinuxPlatformAtomics.h b/Engine/Source/Runtime/Core/Public/Linux/LinuxPlatformAtomics.h
new file mode 100644
--- /dev/null
+++ b/Engine/Source/Runtime/Core/Public/Linux/LinuxPlatformAtomics.h
@@ -0,0 +1,5 @@
+// Copyright 1998-2019 Epic Games, Inc. All Rights Reserved.
+
+#pragma once
+
+#include "Unix/UnixPlatformAtomics.h"
diff --git a/Engine/Source/Runtime/Core/Public/Linux/LinuxPlatformMath.h b/Engine/Source/Runtime/Core/Public/Linux/LinuxPlatformMath.h
new file mode 100644
--- /dev/null
+++ b/Engine/Source/Runtime/Core/Public/Linux/LinuxPlatformMath.h
@@ -0,0 +1,5 @@
+// Copyright 1998-2019 Epic Games, Inc. All Rights Reserved.
+
+#pragma once
+
+#include "Unix/UnixPlatformMath.h"
diff --git a/Engine/Source/Runtime/Core/Public/Linux/LinuxPlatformMemory.h b/Engine/Source/Runtime/Core/Public/Linux/LinuxPlatformMemory.h
new file mode 100644
--- /dev/null
+++ b/Engine/Source/Runtime/Core/Public/Linux/LinuxPlatformMemory.h
@@ -0,0 +1,5 @@
+// Copyright 1998-2019 Epic Games, Inc. All Rights Reserved.
+
+#pragma once
+
+#include "Unix/UnixPlatformMemory.h"
diff --git a/Engine/Source/Runtime/Core/Public/Linux/LinuxPlatformOutputDevices.h b/Engine/Source/Runtime/Core/Public/Linux/LinuxPlatformOutputDevices.h
new file mode 100644
--- /dev/null
+++ b/Engine/Source/Runtime/Core/Public/Linux/LinuxPlatformOutputDevices.h
@@ -0,0 +1,5 @@
+// Copyright 1998-2019 Epic Games, Inc. All Rights Reserved.
+
+#pragma once
+
+#include "Unix/UnixPlatformOutputDevices.h"
diff --git a/Engine/Source/Runtime/Core/Public/Linux/LinuxPlatformStackWalk.h b/Engine/Source/Runtime/Core/Public/Linux/LinuxPlatformStackWalk.h
new file mode 100644
--- /dev/null
+++ b/Engine/Source/Runtime/Core/Public/Linux/LinuxPlatformStackWalk.h
@@ -0,0 +1,5 @@
+// Copyright 1998-2019 Epic Games, Inc. All Rights Reserved.
+
+#pragma once
+
+#include "Unix/UnixPlatformStackWalk.h"
diff --git a/Engine/Source/Runtime/Core/Public/Linux/LinuxPlatformString.h b/Engine/Source/Runtime/Core/Public/Linux/LinuxPlatformString.h
new file mode 100644
--- /dev/null
+++ b/Engine/Source/Runtime/Core/Public/Linux/LinuxPlatformString.h
@@ -0,0 +1,5 @@
+// Copyright 1998-2019 Epic Games, Inc. All Rights Reserved.
+
+#pragma once
+
+#include "Unix/UnixPlatformString.h"
diff --git a/Engine/Source/Runtime/Core/Public/Linux/LinuxPlatformTLS.h b/Engine/Source/Runtime/Core/Public/Linux/LinuxPlatformTLS.h
new file mode 100644
--- /dev/null
+++ b/Engine/Source/Runtime/Core/Public/Linux/LinuxPlatformTLS.h
@@ -0,0 +1,5 @@
+// Copyright 1998-2019 Epic Games, Inc. All Rights Reserved.
+
+#pragma once
+
+#include "Unix/UnixPlatformTLS.h"
diff --git a/Engine/Source/Runtime/Core/Public/Linux/LinuxPlatformTime.h b/Engine/Source/Runtime/Core/Public/Linux/LinuxPlatformTime.h
new file mode 100644
--- /dev/null
+++ b/Engine/Source/Runtime/Core/Public/Linux/LinuxPlatformTime.h
@@ -0,0 +1,5 @@
+// Copyright 1998-2019 Epic Games, Inc. All Rights Reserved.
+
+#pragma once
+
+#include "Unix/UnixPlatformTime.h"
diff --git a/Engine/Source/Runtime/Core/Public/Lumin/LuminAffinity.h b/Engine/Source/Runtime/Core/Public/Lumin/LuminAffinity.h
deleted file mode 100644
diff --git a/Engine/Source/Runtime/Core/Public/Lumin/LuminPlatform.h b/Engine/Source/Runtime/Core/Public/Lumin/LuminPlatform.h
new file mode 100644
--- /dev/null
+++ b/Engine/Source/Runtime/Core/Public/Lumin/LuminPlatform.h
@@ -0,0 +1,5 @@
+// Copyright 1998-2019 Epic Games, Inc. All Rights Reserved.
+
+#pragma once
+
+#include "Android/AndroidPlatform.h"
diff --git a/Engine/Source/Runtime/Core/Public/Lumin/LuminPlatformAffinity.h b/Engine/Source/Runtime/Core/Public/Lumin/LuminPlatformAffinity.h
new file mode 100644
--- /dev/null
+++ b/Engine/Source/Runtime/Core/Public/Lumin/LuminPlatformAffinity.h
@@ -0,0 +1,65 @@
+// Copyright 1998-2019 Epic Games, Inc. All Rights Reserved.
+
+/*=============================================================================================
+LuminAffinity.h: Lumin affinity profile masks definitions.
+==============================================================================================*/
+
+#pragma once
+
+#include "GenericPlatform/GenericPlatformAffinity.h"
+
+class FLuminAffinity : public FGenericPlatformAffinity
+{
+	const static uint64 ArmCores = MAKEAFFINITYMASK2(3,4);
+	const static uint64 DenverCores = MAKEAFFINITYMASK1(2);
+public:
+	static const CORE_API uint64 GetMainGameMask() {
+		return ArmCores;
+	}
+
+	static const CORE_API uint64 GetRenderingThreadMask() {
+		return ArmCores;
+	}
+
+	static const CORE_API uint64 GetRHIThreadMask() {
+		return ArmCores;
+	}
+
+	static const CORE_API uint64 GetRTHeartBeatMask() {
+		return ArmCores;
+	}
+
+	static const CORE_API uint64 GetPoolThreadMask() {
+		return DenverCores;
+	}
+
+	static const CORE_API uint64 GetTaskGraphThreadMask() {
+		return ArmCores;
+	}
+
+	static const CORE_API uint64 GetStatsThreadMask() {
+		return ArmCores;
+	}
+
+	static const CORE_API uint64 GetAudioThreadMask() {
+		return DenverCores;
+	}
+
+	static const CORE_API uint64 GetNoAffinityMask() {
+		return 0xFFFFFFFFFFFFFFFF;
+	}
+
+	static const CORE_API uint64 GetTaskGraphBackgroundTaskMask() {
+		return ArmCores;
+	}
+
+	static EThreadPriority GetRenderingThreadPriority() {
+		return TPri_Normal;
+	}
+
+	static EThreadPriority GetRHIThreadPriority() {
+		return TPri_SlightlyBelowNormal;
+	}
+};
+
+typedef FLuminAffinity FPlatformAffinity;
diff --git a/Engine/Source/Runtime/Core/Public/Lumin/LuminPlatformAtomics.h b/Engine/Source/Runtime/Core/Public/Lumin/LuminPlatformAtomics.h
new file mode 100644
--- /dev/null
+++ b/Engine/Source/Runtime/Core/Public/Lumin/LuminPlatformAtomics.h
@@ -0,0 +1,5 @@
+// Copyright 1998-2019 Epic Games, Inc. All Rights Reserved.
+
+#pragma once
+
+#include "Android/AndroidPlatformAtomics.h"
diff --git a/Engine/Source/Runtime/Core/Public/Lumin/LuminPlatformCompilerPreSetup.h b/Engine/Source/Runtime/Core/Public/Lumin/LuminPlatformCompilerPreSetup.h
new file mode 100644
--- /dev/null
+++ b/Engine/Source/Runtime/Core/Public/Lumin/LuminPlatformCompilerPreSetup.h
@@ -0,0 +1,5 @@
+// Copyright 1998-2019 Epic Games, Inc. All Rights Reserved.
+
+#pragma once
+
+#include "Android/AndroidPlatformCompilerPreSetup.h"
diff --git a/Engine/Source/Runtime/Core/Public/Lumin/LuminPlatformCompilerSetup.h b/Engine/Source/Runtime/Core/Public/Lumin/LuminPlatformCompilerSetup.h
new file mode 100644
--- /dev/null
+++ b/Engine/Source/Runtime/Core/Public/Lumin/LuminPlatformCompilerSetup.h
@@ -0,0 +1,5 @@
+// Copyright 1998-2019 Epic Games, Inc. All Rights Reserved.
+
+#pragma once
+
+#include "Android/AndroidPlatformCompilerSetup.h"
diff --git a/Engine/Source/Runtime/Core/Public/Lumin/LuminPlatformCriticalSection.h b/Engine/Source/Runtime/Core/Public/Lumin/LuminPlatformCriticalSection.h
new file mode 100644
--- /dev/null
+++ b/Engine/Source/Runtime/Core/Public/Lumin/LuminPlatformCriticalSection.h
@@ -0,0 +1,5 @@
+// Copyright 1998-2019 Epic Games, Inc. All Rights Reserved.
+
+#pragma once
+
+#include "Android/AndroidPlatformCriticalSection.h"
diff --git a/Engine/Source/Runtime/Core/Public/Lumin/LuminPlatformMath.h b/Engine/Source/Runtime/Core/Public/Lumin/LuminPlatformMath.h
new file mode 100644
--- /dev/null
+++ b/Engine/Source/Runtime/Core/Public/Lumin/LuminPlatformMath.h
@@ -0,0 +1,5 @@
+// Copyright 1998-2019 Epic Games, Inc. All Rights Reserved.
+
+#pragma once
+
+#include "Android/AndroidPlatformMath.h"
diff --git a/Engine/Source/Runtime/Core/Public/Lumin/LuminPlatformMemory.h b/Engine/Source/Runtime/Core/Public/Lumin/LuminPlatformMemory.h
new file mode 100644
--- /dev/null
+++ b/Engine/Source/Runtime/Core/Public/Lumin/LuminPlatformMemory.h
@@ -0,0 +1,5 @@
+// Copyright 1998-2019 Epic Games, Inc. All Rights Reserved.
+
+#pragma once
+
+#include "Android/AndroidPlatformMemory.h"
diff --git a/Engine/Source/Runtime/Core/Public/Lumin/LuminPlatformMisc.h b/Engine/Source/Runtime/Core/Public/Lumin/LuminPlatformMisc.h
--- a/Engine/Source/Runtime/Core/Public/Lumin/LuminPlatformMisc.h
+++ b/Engine/Source/Runtime/Core/Public/Lumin/LuminPlatformMisc.h
@@ -1,7 +1,7 @@
 // Copyright 1998-2019 Epic Games, Inc. All Rights Reserved.
 
 #pragma once
-#include "Android/AndroidMisc.h"
+#include "Android/AndroidPlatformMisc.h"
 #include "Logging/LogVerbosity.h"
 
 struct CORE_API FLuminPlatformMisc : public FAndroidMisc
diff --git a/Engine/Source/Runtime/Core/Public/Lumin/LuminPlatformOutputDevices.h b/Engine/Source/Runtime/Core/Public/Lumin/LuminPlatformOutputDevices.h
new file mode 100644
--- /dev/null
+++ b/Engine/Source/Runtime/Core/Public/Lumin/LuminPlatformOutputDevices.h
@@ -0,0 +1,5 @@
+// Copyright 1998-2019 Epic Games, Inc. All Rights Reserved.
+
+#pragma once
+
+#include "Android/AndroidPlatformOutputDevices.h"
diff --git a/Engine/Source/Runtime/Core/Public/Lumin/LuminPlatformProcess.h b/Engine/Source/Runtime/Core/Public/Lumin/LuminPlatformProcess.h
--- a/Engine/Source/Runtime/Core/Public/Lumin/LuminPlatformProcess.h
+++ b/Engine/Source/Runtime/Core/Public/Lumin/LuminPlatformProcess.h
@@ -2,7 +2,7 @@
 // Copyright 2016 Magic Leap, Inc. All Rights Reserved.
 
 #pragma once
-#include "Android/AndroidProcess.h"
+#include "Android/AndroidPlatformProcess.h"
 
 /**
 * Android implementation of the Process OS functions
diff --git a/Engine/Source/Runtime/Core/Public/Lumin/LuminPlatformStackWalk.h b/Engine/Source/Runtime/Core/Public/Lumin/LuminPlatformStackWalk.h
new file mode 100644
--- /dev/null
+++ b/Engine/Source/Runtime/Core/Public/Lumin/LuminPlatformStackWalk.h
@@ -0,0 +1,5 @@
+// Copyright 1998-2019 Epic Games, Inc. All Rights Reserved.
+
+#pragma once
+
+#include "Android/AndroidPlatformStackWalk.h"
diff --git a/Engine/Source/Runtime/Core/Public/Lumin/LuminPlatformString.h b/Engine/Source/Runtime/Core/Public/Lumin/LuminPlatformString.h
new file mode 100644
--- /dev/null
+++ b/Engine/Source/Runtime/Core/Public/Lumin/LuminPlatformString.h
@@ -0,0 +1,5 @@
+// Copyright 1998-2019 Epic Games, Inc. All Rights Reserved.
+
+#pragma once
+
+#include "Android/AndroidPlatformString.h"
diff --git a/Engine/Source/Runtime/Core/Public/Lumin/LuminPlatformTLS.h b/Engine/Source/Runtime/Core/Public/Lumin/LuminPlatformTLS.h
new file mode 100644
--- /dev/null
+++ b/Engine/Source/Runtime/Core/Public/Lumin/LuminPlatformTLS.h
@@ -0,0 +1,5 @@
+// Copyright 1998-2019 Epic Games, Inc. All Rights Reserved.
+
+#pragma once
+
+#include "Android/AndroidPlatformTLS.h"
diff --git a/Engine/Source/Runtime/Core/Public/Lumin/LuminPlatformTime.h b/Engine/Source/Runtime/Core/Public/Lumin/LuminPlatformTime.h
new file mode 100644
--- /dev/null
+++ b/Engine/Source/Runtime/Core/Public/Lumin/LuminPlatformTime.h
@@ -0,0 +1,5 @@
+// Copyright 1998-2019 Epic Games, Inc. All Rights Reserved.
+
+#pragma once
+
+#include "Android/AndroidPlatformTime.h"
diff --git a/Engine/Source/Runtime/Core/Public/Mac/MacPlatformFile.h b/Engine/Source/Runtime/Core/Public/Mac/MacPlatformFile.h
new file mode 100644
--- /dev/null
+++ b/Engine/Source/Runtime/Core/Public/Mac/MacPlatformFile.h
@@ -0,0 +1,4 @@
+// Copyright 1998-2019 Epic Games, Inc. All Rights Reserved.
+
+#pragma once
+
diff --git a/Engine/Source/Runtime/Core/Public/Mac/MacPlatformString.h b/Engine/Source/Runtime/Core/Public/Mac/MacPlatformString.h
new file mode 100644
--- /dev/null
+++ b/Engine/Source/Runtime/Core/Public/Mac/MacPlatformString.h
@@ -0,0 +1,5 @@
+// Copyright 1998-2019 Epic Games, Inc. All Rights Reserved.
+
+#pragma once
+
+#include "Apple/ApplePlatformString.h"
diff --git a/Engine/Source/Runtime/Core/Public/Mac/MacPlatformTLS.h b/Engine/Source/Runtime/Core/Public/Mac/MacPlatformTLS.h
new file mode 100644
--- /dev/null
+++ b/Engine/Source/Runtime/Core/Public/Mac/MacPlatformTLS.h
@@ -0,0 +1,5 @@
+// Copyright 1998-2019 Epic Games, Inc. All Rights Reserved.
+
+#pragma once
+
+#include "Apple/ApplePlatformTLS.h"
diff --git a/Engine/Source/Runtime/Core/Public/Mac/MacPlatformTime.h b/Engine/Source/Runtime/Core/Public/Mac/MacPlatformTime.h
new file mode 100644
--- /dev/null
+++ b/Engine/Source/Runtime/Core/Public/Mac/MacPlatformTime.h
@@ -0,0 +1,5 @@
+// Copyright 1998-2019 Epic Games, Inc. All Rights Reserved.
+
+#pragma once
+
+#include "Apple/ApplePlatformTime.h"
diff --git a/Engine/Source/Runtime/Online/BackgroundHTTP/Public/IOS/IOSPlatformBackgroundHttp.h b/Engine/Source/Runtime/Online/BackgroundHTTP/Public/IOS/IOSPlatformBackgroundHttp.h
new file mode 100644
--- /dev/null
+++ b/Engine/Source/Runtime/Online/BackgroundHTTP/Public/IOS/IOSPlatformBackgroundHttp.h
@@ -0,0 +1,5 @@
+// Copyright 1998-2019 Epic Games, Inc. All Rights Reserved.
+
+#pragma once
+
+#include "Apple/ApplePlatformBackgroundHttp.h"
diff --git a/Engine/Source/Runtime/Online/BackgroundHTTP/Public/Linux/LinuxPlatformBackgroundHttp.h b/Engine/Source/Runtime/Online/BackgroundHTTP/Public/Linux/LinuxPlatformBackgroundHttp.h
new file mode 100644
--- /dev/null
+++ b/Engine/Source/Runtime/Online/BackgroundHTTP/Public/Linux/LinuxPlatformBackgroundHttp.h
@@ -0,0 +1,5 @@
+// Copyright 1998-2019 Epic Games, Inc. All Rights Reserved.
+
+#pragma once
+
+#include "Unix/UnixPlatformBackgroundHttp.h"
diff --git a/Engine/Source/Runtime/Online/BackgroundHTTP/Public/Lumin/LuminPlatformBackgroundHttp.h b/Engine/Source/Runtime/Online/BackgroundHTTP/Public/Lumin/LuminPlatformBackgroundHttp.h
new file mode 100644
--- /dev/null
+++ b/Engine/Source/Runtime/Online/BackgroundHTTP/Public/Lumin/LuminPlatformBackgroundHttp.h
@@ -0,0 +1,5 @@
+// Copyright 1998-2019 Epic Games, Inc. All Rights Reserved.
+
+#pragma once
+
+#include "Android/AndroidPlatformBackgroundHttp.h"
diff --git a/Engine/Source/Runtime/Online/BackgroundHTTP/Public/PlatformBackgroundHttp.h b/Engine/Source/Runtime/Online/BackgroundHTTP/Public/PlatformBackgroundHttp.h
--- a/Engine/Source/Runtime/Online/BackgroundHTTP/Public/PlatformBackgroundHttp.h
+++ b/Engine/Source/Runtime/Online/BackgroundHTTP/Public/PlatformBackgroundHttp.h
@@ -4,12 +4,4 @@
 
 #include "CoreMinimal.h"
 
-//On all un-implemented platforms, just uses GenericPlatform/GenericPlatformBackgroundHttp.h
-
-#if PLATFORM_IOS                         // of FBackgroundURLSessionHandler. For now everything is setup generically "Apple" so that
-#include "IOS/ApplePlatformBackgroundHttp.h" // we can transition easier if we decide to implement this on Mac. So for now ApplePlatformX is actually only iOS in background http
-#elif PLATFORM_UNIX
-#include "Unix/UnixPlatformBackgroundHttp.h"
-#else
 #include COMPILED_PLATFORM_HEADER(PlatformBackgroundHttp.h)
-#endif
diff --git a/Engine/Source/Runtime/Online/HTTP/Private/Android/AndroidHttp.cpp b/Engine/Source/Runtime/Online/HTTP/Private/Android/AndroidHttp.cpp
deleted file mode 100644
diff --git a/Engine/Source/Runtime/Online/HTTP/Private/Android/AndroidPlatformHttp.cpp b/Engine/Source/Runtime/Online/HTTP/Private/Android/AndroidPlatformHttp.cpp
new file mode 100644
--- /dev/null
+++ b/Engine/Source/Runtime/Online/HTTP/Private/Android/AndroidPlatformHttp.cpp
@@ -0,0 +1,50 @@
+// Copyright 1998-2019 Epic Games, Inc. All Rights Reserved.
+
+#include "Android/AndroidPlatformHttp.h"
+#include "Curl/CurlHttp.h"
+#include "Curl/CurlHttpManager.h"
+
+void FAndroidPlatformHttp::Init()
+{
+	FCurlHttpManager::InitCurl();
+}
+
+class FHttpManager * FAndroidPlatformHttp::CreatePlatformHttpManager()
+{
+	return new FCurlHttpManager();
+}
+
+void FAndroidPlatformHttp::Shutdown()
+{
+	FCurlHttpManager::ShutdownCurl();
+}
+
+IHttpRequest* FAndroidPlatformHttp::ConstructRequest()
+{
+	return new FCurlHttpRequest();
+}
+
+TOptional<FString> FAndroidPlatformHttp::GetOperatingSystemProxyAddress()
+{
+	FString ProxyAddress;
+
+#if USE_ANDROID_JNI
+	extern int32 AndroidThunkCpp_GetMetaDataInt(const FString& Key);
+	extern FString AndroidThunkCpp_GetMetaDataString(const FString& Key);
+
+	FString ProxyHost = AndroidThunkCpp_GetMetaDataString(TEXT("ue4.http.proxy.proxyHost"));
+	int32 ProxyPort = AndroidThunkCpp_GetMetaDataInt(TEXT("ue4.http.proxy.proxyPort"));
+
+	if (ProxyPort != -1 && !ProxyHost.IsEmpty())
+	{
+		ProxyAddress = FString::Printf(TEXT("%s:%d"), *ProxyHost, ProxyPort);
+	}
+#endif
+
+	return ProxyAddress;
+}
+
+bool FAndroidPlatformHttp::IsOperatingSystemProxyInformationSupported()
+{
+	return true;
+}
\ No newline at end of file
diff --git a/Engine/Source/Runtime/Online/HTTP/Public/Android/AndroidHttp.h b/Engine/Source/Runtime/Online/HTTP/Public/Android/AndroidHttp.h
deleted file mode 100644
diff --git a/Engine/Source/Runtime/Online/HTTP/Public/Android/AndroidPlatformHttp.h b/Engine/Source/Runtime/Online/HTTP/Public/Android/AndroidPlatformHttp.h
new file mode 100644
--- /dev/null
+++ b/Engine/Source/Runtime/Online/HTTP/Public/Android/AndroidPlatformHttp.h
@@ -0,0 +1,55 @@
+// Copyright 1998-2019 Epic Games, Inc. All Rights Reserved.
+
+#pragma once
+#include "GenericPlatform/GenericPlatformHttp.h"
+
+/**
+ * Platform specific Http implementations
+ */
+class FAndroidPlatformHttp : public FGenericPlatformHttp
+{
+public:
+
+	/**
+	 * Platform initialization step
+	 */
+	static void Init();
+
+	/**
+	 * Creates a platform-specific HTTP manager.
+	 *
+	 * @return NULL if default implementation is to be used
+	 */
+	static FHttpManager* CreatePlatformHttpManager();
+
+	/**
+	 * Platform shutdown step
+	 */
+	static void Shutdown();
+
+	/**
+	 * Creates a new Http request instance for the current platform
+	 *
+	 * @return request object
+	 */
+	static IHttpRequest* ConstructRequest();
+
+	/**
+	 * Get the proxy address specified by the operating system
+	 *
+	 * @return optional FString: If unset: we are unable to get information from the operating system. If set: the proxy address set by the operating system (may be blank)
+	 */
+	static TOptional<FString> GetOperatingSystemProxyAddress();
+
+	/**
+	 * Check if getting proxy information from the current operating system is supported
+	 * Useful for "Network Settings" type pages.  GetProxyAddress may return an empty or populated string but that does not imply
+	 * the operating system does or does not support proxies (or that it has been implemented here)
+	 * 
+	 * @return true if we are able to get proxy information from the current operating system, false if not
+	 */
+	static bool IsOperatingSystemProxyInformationSupported();
+};
+
+
+typedef FAndroidPlatformHttp FPlatformHttp;
\ No newline at end of file
diff --git a/Engine/Source/Runtime/Online/HTTP/Public/IOS/IOSPlatformHttp.h b/Engine/Source/Runtime/Online/HTTP/Public/IOS/IOSPlatformHttp.h
new file mode 100644
--- /dev/null
+++ b/Engine/Source/Runtime/Online/HTTP/Public/IOS/IOSPlatformHttp.h
@@ -0,0 +1,5 @@
+// Copyright 1998-2019 Epic Games, Inc. All Rights Reserved.
+
+#pragma once
+
+#include "Apple/ApplePlatformHttp.h"
diff --git a/Engine/Source/Runtime/Online/HTTP/Public/Linux/LinuxPlatformHttp.h b/Engine/Source/Runtime/Online/HTTP/Public/Linux/LinuxPlatformHttp.h
new file mode 100644
--- /dev/null
+++ b/Engine/Source/Runtime/Online/HTTP/Public/Linux/LinuxPlatformHttp.h
@@ -0,0 +1,5 @@
+// Copyright 1998-2019 Epic Games, Inc. All Rights Reserved.
+
+#pragma once
+
+#include "Unix/UnixPlatformHttp.h"
diff --git a/Engine/Source/Runtime/Online/HTTP/Public/Lumin/LuminPlatformHttp.h b/Engine/Source/Runtime/Online/HTTP/Public/Lumin/LuminPlatformHttp.h
new file mode 100644
--- /dev/null
+++ b/Engine/Source/Runtime/Online/HTTP/Public/Lumin/LuminPlatformHttp.h
@@ -0,0 +1,5 @@
+// Copyright 1998-2019 Epic Games, Inc. All Rights Reserved.
+
+#pragma once
+
+#include "Android/AndroidPlatformHttp.h"
diff --git a/Engine/Source/Runtime/Online/HTTP/Public/Mac/MacPlatformHttp.h b/Engine/Source/Runtime/Online/HTTP/Public/Mac/MacPlatformHttp.h
new file mode 100644
--- /dev/null
+++ b/Engine/Source/Runtime/Online/HTTP/Public/Mac/MacPlatformHttp.h
@@ -0,0 +1,5 @@
+// Copyright 1998-2019 Epic Games, Inc. All Rights Reserved.
+
+#pragma once
+
+#include "Apple/ApplePlatformHttp.h"
diff --git a/Engine/Source/Runtime/Online/HTTP/Public/PlatformHttp.h b/Engine/Source/Runtime/Online/HTTP/Public/PlatformHttp.h
--- a/Engine/Source/Runtime/Online/HTTP/Public/PlatformHttp.h
+++ b/Engine/Source/Runtime/Online/HTTP/Public/PlatformHttp.h
@@ -3,16 +3,4 @@
 
 #include "CoreMinimal.h"
 
-#if PLATFORM_XBOXONE
-#include "XboxOne/XboxOneHttp.h"
-#elif PLATFORM_MAC
-#include "Apple/ApplePlatformHttp.h"
-#elif PLATFORM_IOS
-#include "Apple/ApplePlatformHttp.h"
-#elif PLATFORM_ANDROID
-#include "Android/AndroidHttp.h"
-#elif PLATFORM_UNIX
-#include "Unix/UnixPlatformHttp.h"
-#else
 #include COMPILED_PLATFORM_HEADER(PlatformHttp.h)
-#endif
