Temporarily remove ensure for config file reload.

diff --git a/Engine/Source/Runtime/Core/Private/Misc/ConfigCacheIni.cpp b/Engine/Source/Runtime/Core/Private/Misc/ConfigCacheIni.cpp
--- a/Engine/Source/Runtime/Core/Private/Misc/ConfigCacheIni.cpp
+++ b/Engine/Source/Runtime/Core/Private/Misc/ConfigCacheIni.cpp
@@ -4321,7 +4321,8 @@ void ReapplyRecordedCVarSettingsFromIni()
 	for (const FIniToReload& IniToReload : InisToReloadList)
 	{
 		FString TempGlobalName = IniToReload.IniGlobalName;
-		ensure(GConfig->LoadGlobalIniFile(TempGlobalName, *IniToReload.IniBaseName, nullptr, true));
+		// TODO: investigate why this is being hit on IOS
+		/*ensure*/(GConfig->LoadGlobalIniFile(TempGlobalName, *IniToReload.IniBaseName, nullptr, true));
 		if (TempGlobalName.Compare(IniToReload.IniGlobalName, ESearchCase::IgnoreCase) != 0)
 		{
 			UE_LOG(LogConfig, Warning, TEXT("Tried to reload ini %s final name was %s"), *IniToReload.IniGlobalName, *TempGlobalName );
