- Fixed some CIS errors

diff --git a/Engine/Source/Runtime/Core/Private/Lumin/LuminPlatformMisc.cpp b/Engine/Source/Runtime/Core/Private/Lumin/LuminPlatformMisc.cpp
--- a/Engine/Source/Runtime/Core/Private/Lumin/LuminPlatformMisc.cpp
+++ b/Engine/Source/Runtime/Core/Private/Lumin/LuminPlatformMisc.cpp
@@ -37,7 +37,7 @@ void FLuminPlatformMisc::InitLifecycle()
 
 void FLuminPlatformMisc::RequestExit(bool Force)
 {
-	UE_LOG(LogWindows, Log, TEXT("FLuminPlatformMisc::RequestExit(%i)"), Force);
+	UE_LOG(LogLumin, Log, TEXT("FLuminPlatformMisc::RequestExit(%i)"), Force);
 	if (Force)
 	{
 		_exit(1);
diff --git a/Engine/Source/Runtime/Core/Private/Lumin/LuminPlatformProcess.cpp b/Engine/Source/Runtime/Core/Private/Lumin/LuminPlatformProcess.cpp
--- a/Engine/Source/Runtime/Core/Private/Lumin/LuminPlatformProcess.cpp
+++ b/Engine/Source/Runtime/Core/Private/Lumin/LuminPlatformProcess.cpp
@@ -120,7 +120,7 @@ void* FLuminPlatformProcess::GetDllHandle(const TCHAR* Filename)
 	void *Handle = dlopen(TCHAR_TO_UTF8(*AbsolutePath), DlOpenMode);
 	if (!Handle)
 	{
-		UE_LOG(LogLinux, Warning, TEXT("dlopen failed: %s"), UTF8_TO_TCHAR(dlerror()));
+		UE_LOG(LogLumin, Warning, TEXT("dlopen failed: %s"), UTF8_TO_TCHAR(dlerror()));
 	}
 	return Handle;
 }
