- FIxed LogLumin compile errors

diff --git a/Engine/Source/Runtime/Core/Private/Misc/CoreGlobals.cpp b/Engine/Source/Runtime/Core/Private/Misc/CoreGlobals.cpp
--- a/Engine/Source/Runtime/Core/Private/Misc/CoreGlobals.cpp
+++ b/Engine/Source/Runtime/Core/Private/Misc/CoreGlobals.cpp
@@ -592,10 +592,13 @@ DEFINE_LOG_CATEGORY(LogMemory);
 DEFINE_LOG_CATEGORY(LogProfilingDebugging);
 DEFINE_LOG_CATEGORY(LogTemp);
 
-#ifdef PLATFORM_GLOBAL_LOG_CATEGORY
 // need another layer of macro to help using a define in a define
 #define DEFINE_LOG_CATEGORY_HELPER(A) DEFINE_LOG_CATEGORY(A)
-DEFINE_LOG_CATEGORY_HELPER(PLATFORM_GLOBAL_LOG_CATEGORY);
+#ifdef PLATFORM_GLOBAL_LOG_CATEGORY
+	DEFINE_LOG_CATEGORY_HELPER(PLATFORM_GLOBAL_LOG_CATEGORY);
+#endif
+#ifdef PLATFORM_GLOBAL_LOG_CATEGORY_ALT
+	DEFINE_LOG_CATEGORY_HELPER(PLATFORM_GLOBAL_LOG_CATEGORY_ALT);
 #endif
 
 #undef LOCTEXT_NAMESPACE
diff --git a/Engine/Source/Runtime/Core/Public/CoreGlobals.h b/Engine/Source/Runtime/Core/Public/CoreGlobals.h
--- a/Engine/Source/Runtime/Core/Public/CoreGlobals.h
+++ b/Engine/Source/Runtime/Core/Public/CoreGlobals.h
@@ -44,10 +44,13 @@ CORE_API DECLARE_LOG_CATEGORY_EXTERN(LogLoad, Log, All);
 CORE_API DECLARE_LOG_CATEGORY_EXTERN(LogTemp, Log, All);
 
 // Platform specific logs, set here to make it easier to use them from anywhere
-#ifdef PLATFORM_GLOBAL_LOG_CATEGORY
 // need another layer of macro to help using a define in a define
 #define DECLARE_LOG_CATEGORY_EXTERN_HELPER(A,B,C) DECLARE_LOG_CATEGORY_EXTERN(A,B,C)
-CORE_API DECLARE_LOG_CATEGORY_EXTERN_HELPER(PLATFORM_GLOBAL_LOG_CATEGORY, Log, All);
+#ifdef PLATFORM_GLOBAL_LOG_CATEGORY
+	CORE_API DECLARE_LOG_CATEGORY_EXTERN_HELPER(PLATFORM_GLOBAL_LOG_CATEGORY, Log, All);
+#endif
+#ifdef PLATFORM_GLOBAL_LOG_CATEGORY_ALT
+	CORE_API DECLARE_LOG_CATEGORY_EXTERN_HELPER(PLATFORM_GLOBAL_LOG_CATEGORY_ALT, Log, All);
 #endif
 
 
diff --git a/Engine/Source/Runtime/Core/Public/Lumin/LuminPlatform.h b/Engine/Source/Runtime/Core/Public/Lumin/LuminPlatform.h
--- a/Engine/Source/Runtime/Core/Public/Lumin/LuminPlatform.h
+++ b/Engine/Source/Runtime/Core/Public/Lumin/LuminPlatform.h
@@ -3,3 +3,5 @@
 #pragma once
 
 #include "Android/AndroidPlatform.h"
+
+PLATFORM_GLOBAL_LOG_CATEGORY_ALT LogLumin
