- Fixed rookie mistake ;) (missed a #define in a line, breaking Lumin builds)

diff --git a/Engine/Source/Runtime/Core/Public/Lumin/LuminPlatform.h b/Engine/Source/Runtime/Core/Public/Lumin/LuminPlatform.h
--- a/Engine/Source/Runtime/Core/Public/Lumin/LuminPlatform.h
+++ b/Engine/Source/Runtime/Core/Public/Lumin/LuminPlatform.h
@@ -4,4 +4,4 @@
 
 #include "Android/AndroidPlatform.h"
 
-PLATFORM_GLOBAL_LOG_CATEGORY_ALT LogLumin
+#define PLATFORM_GLOBAL_LOG_CATEGORY_ALT LogLumin
