Added support for fatal hangs on PC to be hotfixed.

diff --git a/Engine/Source/Runtime/Core/Private/HAL/ThreadHeartBeat.cpp b/Engine/Source/Runtime/Core/Private/HAL/ThreadHeartBeat.cpp
--- a/Engine/Source/Runtime/Core/Private/HAL/ThreadHeartBeat.cpp
+++ b/Engine/Source/Runtime/Core/Private/HAL/ThreadHeartBeat.cpp
@@ -72,6 +72,7 @@ FThreadHeartBeat::FThreadHeartBeat()
 	, HangDurationMultiplier(1.0)
 	, LastHangCallstackCRC(0)
 	, LastHungThreadId(InvalidThreadId)
+	, bHangsAreFatal(false)
 	, Clock(HangDetectorClock_MaxTimeStep_MS / 1000)
 {
 	// Start with the frame-present based hang detection disabled. This will be automatically enabled on
@@ -235,19 +236,35 @@ void FORCENOINLINE FThreadHeartBeat::OnHang(double HangDuration, uint32 ThreadTh
 		}
 
 		const FString ErrorMessage = FString::Printf(TEXT("Hang detected on %s:%s%s%sCheck log for full callstack."), *ThreadName, LINE_TERMINATOR, *StackTrimmed, LINE_TERMINATOR);
-#if UE_ASSERT_ON_HANG
-		UE_LOG(LogCore, Fatal, TEXT("%s"), *ErrorMessage);
-#else
-		UE_LOG(LogCore, Error, TEXT("%s"), *ErrorMessage);
 
 #if PLATFORM_DESKTOP
+		UE_LOG(LogCore, Error, TEXT("%s"), *ErrorMessage);
 		GLog->PanicFlushThreadedLogs();
 
 		// Skip macros and FDebug, we always want this to fire
 		ReportHang(*ErrorMessage, StackFrames, NumStackFrames, ThreadThatHung);
-#endif // PLATFORM_DESKTOP
 
-#endif // UE_ASSERT_ON_HANG == 0
+		if (bHangsAreFatal)
+		{
+			if (FApp::CanEverRender())
+			{
+				FPlatformMisc::MessageBoxExt(EAppMsgType::Ok,
+					*NSLOCTEXT("MessageDialog", "ReportHangError_Body", "The application has hung and will now close. We apologize for the inconvenience.").ToString(),
+					*NSLOCTEXT("MessageDialog", "ReportHangError_Title", "Application Hang Detected").ToString());
+			}
+
+			FPlatformMisc::RequestExit(true);
+		}
+#else
+		if (bHangsAreFatal)
+		{
+			UE_LOG(LogCore, Fatal, TEXT("%s"), *ErrorMessage);
+		}
+		else
+		{
+			UE_LOG(LogCore, Error, TEXT("%s"), *ErrorMessage);
+		}
+#endif
 	}
 #endif // MINIMAL_FATAL_HANG_DETECTION == 0
 }
@@ -309,11 +326,13 @@ void FThreadHeartBeat::InitSettings()
 	// Default to 25 seconds if not overridden in config.
 	double NewHangDuration = 25.0;
 	double NewPresentDuration = 0.0;
+	bool bNewHangsAreFatal = !!(UE_ASSERT_ON_HANG);
 
 	if (GConfig)
 	{
 		GConfig->GetDouble(TEXT("Core.System"), TEXT("HangDuration"), NewHangDuration, GEngineIni);
 		GConfig->GetDouble(TEXT("Core.System"), TEXT("PresentHangDuration"), NewPresentDuration, GEngineIni);
+		GConfig->GetBool(TEXT("Core.System"), TEXT("HangsAreFatal"), bNewHangsAreFatal, GEngineIni);
 
 		const double MinHangDuration = 5.0;
 		if (NewHangDuration > 0.0 && NewHangDuration < MinHangDuration)
@@ -335,6 +354,8 @@ void FThreadHeartBeat::InitSettings()
 
 	CurrentHangDuration = ConfigHangDuration * HangDurationMultiplier;
 	CurrentPresentDuration = ConfigPresentDuration * HangDurationMultiplier;
+
+	bHangsAreFatal = bNewHangsAreFatal;
 }
 
 void FThreadHeartBeat::HeartBeat(bool bReadConfig)
diff --git a/Engine/Source/Runtime/Core/Public/HAL/ThreadHeartBeat.h b/Engine/Source/Runtime/Core/Public/HAL/ThreadHeartBeat.h
--- a/Engine/Source/Runtime/Core/Public/HAL/ThreadHeartBeat.h
+++ b/Engine/Source/Runtime/Core/Public/HAL/ThreadHeartBeat.h
@@ -100,6 +100,8 @@ class CORE_API FThreadHeartBeat : public FRunnable
 	/** Id of the last thread that hung */
 	uint32 LastHungThreadId;
 
+	bool bHangsAreFatal;
+
 	/** Global suspended count */
 	FThreadSafeCounter GlobalSuspendCount;
 
