Adding LLM category for Video Streaming

diff --git a/Engine/Plugins/Media/WmfMedia/Source/WmfMedia/Private/Wmf/WmfMediaHardwareVideoDecodingRendering.cpp b/Engine/Plugins/Media/WmfMedia/Source/WmfMedia/Private/Wmf/WmfMediaHardwareVideoDecodingRendering.cpp
--- a/Engine/Plugins/Media/WmfMedia/Source/WmfMedia/Private/Wmf/WmfMediaHardwareVideoDecodingRendering.cpp
+++ b/Engine/Plugins/Media/WmfMedia/Source/WmfMedia/Private/Wmf/WmfMediaHardwareVideoDecodingRendering.cpp
@@ -42,6 +42,7 @@ struct FRHICommandCopyResource final : public FRHICommand<FRHICommandCopyResourc
 
 	void Execute(FRHICommandListBase& CmdList)
 	{
+		LLM_SCOPE(ELLMTag::VideoStreaming);
 		ID3D11Device* D3D11Device = static_cast<ID3D11Device*>(GDynamicRHI->RHIGetNativeDevice());
 		ID3D11DeviceContext* D3D11DeviceContext = nullptr;
 
@@ -103,6 +104,7 @@ struct FRHICommandCopyResource final : public FRHICommand<FRHICommandCopyResourc
 
 bool FWmfMediaHardwareVideoDecodingParameters::ConvertTextureFormat_RenderThread(FWmfMediaHardwareVideoDecodingTextureSample* InSample, FTexture2DRHIRef InDstTexture)
 {
+	LLM_SCOPE(ELLMTag::VideoStreaming);
 	if (InSample == nullptr || !InDstTexture.IsValid())
 	{
 		return false;
diff --git a/Engine/Source/Runtime/Core/Private/HAL/LowLevelMemTracker.cpp b/Engine/Source/Runtime/Core/Private/HAL/LowLevelMemTracker.cpp
--- a/Engine/Source/Runtime/Core/Private/HAL/LowLevelMemTracker.cpp
+++ b/Engine/Source/Runtime/Core/Private/HAL/LowLevelMemTracker.cpp
@@ -102,6 +102,7 @@ DECLARE_LLM_MEMORY_STAT(TEXT("MaterialInstance"), STAT_MaterialInstanceLLM, STAT
 DECLARE_LLM_MEMORY_STAT(TEXT("SkeletalMesh"), STAT_SkeletalMeshLLM, STATGROUP_LLMFULL);
 DECLARE_LLM_MEMORY_STAT(TEXT("InstancedMesh"), STAT_InstancedMeshLLM, STATGROUP_LLMFULL);
 DECLARE_LLM_MEMORY_STAT(TEXT("Landscape"), STAT_LandscapeLLM, STATGROUP_LLMFULL);
+DECLARE_LLM_MEMORY_STAT(TEXT("VideoStreaming"), STAT_VideoStreamingLLM, STATGROUP_LLMFULL);
 
 /*
 * LLM Summary stats referenced by ELLMTagNames
diff --git a/Engine/Source/Runtime/Core/Public/HAL/LowLevelMemTracker.h b/Engine/Source/Runtime/Core/Public/HAL/LowLevelMemTracker.h
--- a/Engine/Source/Runtime/Core/Public/HAL/LowLevelMemTracker.h
+++ b/Engine/Source/Runtime/Core/Public/HAL/LowLevelMemTracker.h
@@ -215,7 +215,8 @@ enum class ELLMTagSet : uint8
 	macro(SkeletalMesh,							"SkeletalMesh",					GET_STATFNAME(STAT_SkeletalMeshLLM),						GET_STATFNAME(STAT_EngineSummaryLLM),			ELLMTag::Meshes)\
 	macro(InstancedMesh,						"InstancedMesh",				GET_STATFNAME(STAT_InstancedMeshLLM),						GET_STATFNAME(STAT_EngineSummaryLLM),			ELLMTag::Meshes)\
 	macro(Landscape,							"Landscape",					GET_STATFNAME(STAT_LandscapeLLM),							GET_STATFNAME(STAT_EngineSummaryLLM),			ELLMTag::Meshes)\
-	macro(CsvProfiler,							"CsvProfiler",					GET_STATFNAME(STAT_CsvProfilerLLM),							GET_STATFNAME(STAT_EngineSummaryLLM),			-1)
+	macro(CsvProfiler,							"CsvProfiler",					GET_STATFNAME(STAT_CsvProfilerLLM),							GET_STATFNAME(STAT_EngineSummaryLLM),			-1)\
+	macro(VideoStreaming,						"VideoStreaming",				GET_STATFNAME(STAT_VideoStreamingLLM),						GET_STATFNAME(STAT_EngineSummaryLLM),			-1)\
 
 /*
  * Enum values to be passed in to LLM_SCOPE() macro
