Duplicate [CL 6412607] - Added missing include to resolve compile issue in some projects.

diff --git a/Engine/Source/Runtime/Core/Private/HAL/ThreadHeartBeat.cpp b/Engine/Source/Runtime/Core/Private/HAL/ThreadHeartBeat.cpp
--- a/Engine/Source/Runtime/Core/Private/HAL/ThreadHeartBeat.cpp
+++ b/Engine/Source/Runtime/Core/Private/HAL/ThreadHeartBeat.cpp
@@ -16,6 +16,7 @@
 #include "Stats/Stats.h"
 #include "Async/TaskGraphInterfaces.h"
 #include "ProfilingDebugging/CsvProfiler.h"
+#include "Misc/App.h"
 
 #if PLATFORM_SWITCH
 #include "Switch/SwitchPlatformCrashContext.h"
