stop existing FramePro recording when trying to start when it's already started.

diff --git a/Engine/Source/Runtime/Core/Private/FramePro/FrameProProfiler.cpp b/Engine/Source/Runtime/Core/Private/FramePro/FrameProProfiler.cpp
--- a/Engine/Source/Runtime/Core/Private/FramePro/FrameProProfiler.cpp
+++ b/Engine/Source/Runtime/Core/Private/FramePro/FrameProProfiler.cpp
@@ -357,6 +357,11 @@ void FFrameProProfiler::StartFrameProRecordingScopeOverrideFromCommand(const TAr
 
 FString FFrameProProfiler::StartFrameProRecording(const FString& FilenameRoot, int32 MinScopeTime)
 {
+	if (GFrameProIsRecording)
+	{
+		StopFrameProRecording();
+	}
+
 	FString RelPathName = FPaths::ProfilingDir() + TEXT("FramePro/");
 	bool bSuccess = IFileManager::Get().MakeDirectory(*RelPathName, true); // ensure folder exists
 
