reduce the thread stack size on IOS. Helps reduce VM usage

diff --git a/Engine/Source/Runtime/Core/Private/IOS/IOSPlatformMisc.cpp b/Engine/Source/Runtime/Core/Private/IOS/IOSPlatformMisc.cpp
--- a/Engine/Source/Runtime/Core/Private/IOS/IOSPlatformMisc.cpp
+++ b/Engine/Source/Runtime/Core/Private/IOS/IOSPlatformMisc.cpp
@@ -679,7 +679,7 @@ FIOSPlatformMisc::EIOSDevice FIOSPlatformMisc::GetIOSDeviceType()
 
 int FIOSPlatformMisc::GetDefaultStackSize()
 {
-	return 4 * 1024 * 1024;
+	return 512 * 1024;
 }
 
 void FIOSPlatformMisc::SetMemoryWarningHandler(void (* InHandler)(const FGenericMemoryWarningContext& Context))
