TArray slack changes from Arne and Core

diff --git a/Engine/Source/Runtime/Core/Public/Containers/Array.h b/Engine/Source/Runtime/Core/Public/Containers/Array.h
--- a/Engine/Source/Runtime/Core/Public/Containers/Array.h
+++ b/Engine/Source/Runtime/Core/Public/Containers/Array.h
@@ -1195,6 +1195,28 @@ public:
 		Ar.CountBytes(ArrayNum*sizeof(ElementType), ArrayMax*sizeof(ElementType));
 	}
 
+	/**
+	 * Adds an uninitialized element into the array.
+	 *
+	 * Caution, AddUninitialized() will create elements without calling
+	 * the constructor and this is not appropriate for element types that
+	 * require a constructor to function properly.
+	 *
+	 * @returns Number of elements in array before addition.
+	 */
+	FORCEINLINE int32 AddUninitialized()
+	{
+		CheckInvariants();
+
+		const int32 OldNum = ArrayNum;
+		const int32 NewNum = ArrayNum += 1;
+		if (NewNum > ArrayMax)
+		{
+			ResizeGrow(OldNum);
+		}
+		return OldNum;
+	}
+
 	/**
 	 * Adds a given number of uninitialized elements into the array.
 	 *
@@ -1205,19 +1227,50 @@ public:
 	 * @param Count Number of elements to add.
 	 * @returns Number of elements in array before addition.
 	 */
-	FORCEINLINE int32 AddUninitialized(int32 Count = 1)
+	FORCEINLINE int32 AddUninitialized(int32 Count)
 	{
 		CheckInvariants();
 		checkSlow(Count >= 0);
 
 		const int32 OldNum = ArrayNum;
-		if ((ArrayNum += Count) > ArrayMax)
+		const int32 NewNum = ArrayNum += Count;
+		if (OldNum == 0)
+		{
+			Reserve(Count);
+		}
+		else if (NewNum > ArrayMax)
 		{
 			ResizeGrow(OldNum);
 		}
 		return OldNum;
 	}
 
+	/**
+	 * Inserts a given number of uninitialized elements into the array at given
+	 * location.
+	 *
+	 * Caution, InsertUninitialized() will create elements without calling the
+	 * constructor and this is not appropriate for element types that require
+	 * a constructor to function properly.
+	 *
+	 * @param Index Tells where to insert the new elements.
+	 * @see Insert, InsertZeroed, InsertDefaulted
+	 */
+	void InsertUninitialized(int32 Index)
+	{
+		CheckInvariants();
+		checkSlow((Index >= 0) & (Index <= ArrayNum));
+
+		const int32 OldNum = ArrayNum;
+		const int32 NewNum = ArrayNum += 1;
+		if (NewNum > ArrayMax)
+		{
+			ResizeGrow(OldNum);
+		}
+		ElementType* Data = GetData() + Index;
+		RelocateConstructItems<ElementType>(Data + 1, Data, OldNum - Index);
+	}
+
 	/**
 	 * Inserts a given number of uninitialized elements into the array at given
 	 * location.
@@ -1230,13 +1283,18 @@ public:
 	 * @param Count Number of elements to add.
 	 * @see Insert, InsertZeroed, InsertDefaulted
 	 */
-	void InsertUninitialized(int32 Index, int32 Count = 1)
+	void InsertUninitialized(int32 Index, int32 Count)
 	{
 		CheckInvariants();
 		checkSlow((Count >= 0) & (Index >= 0) & (Index <= ArrayNum));
 
 		const int32 OldNum = ArrayNum;
-		if ((ArrayNum += Count) > ArrayMax)
+		const int32 NewNum = ArrayNum += Count;
+		if (OldNum == 0)
+		{
+			Reserve(Count);
+		}
+		else if (NewNum > ArrayMax)
 		{
 			ResizeGrow(OldNum);
 		}
@@ -1244,6 +1302,22 @@ public:
 		RelocateConstructItems<ElementType>(Data + Count, Data, OldNum - Index);
 	}
 
+	/**
+	 * Inserts a zeroed element into the array at given location.
+	 *
+	 * Caution, InsertZeroed() will create elements without calling the
+	 * constructor and this is not appropriate for element types that require
+	 * a constructor to function properly.
+	 *
+	 * @param Index Tells where to insert the new elements.
+	 * @see Insert, InsertUninitialized, InsertDefaulted
+	 */
+	void InsertZeroed(int32 Index)
+	{
+		InsertUninitialized(Index);
+		FMemory::Memzero(GetData() + Index, sizeof(ElementType));
+	}
+
 	/**
 	 * Inserts a given number of zeroed elements into the array at given
 	 * location.
@@ -1256,7 +1330,7 @@ public:
 	 * @param Count Number of elements to add.
 	 * @see Insert, InsertUninitialized, InsertDefaulted
 	 */
-	void InsertZeroed(int32 Index, int32 Count = 1)
+	void InsertZeroed(int32 Index, int32 Count)
 	{
 		InsertUninitialized(Index, Count);
 		FMemory::Memzero(GetData() + Index, Count * sizeof(ElementType));
@@ -1275,12 +1349,24 @@ public:
 	 */
 	ElementType& InsertZeroed_GetRef(int32 Index)
 	{
-		InsertUninitialized(Index, 1);
+		InsertUninitialized(Index);
 		ElementType* Ptr = GetData() + Index;
 		FMemory::Memzero(Ptr, sizeof(ElementType));
 		return *Ptr;
 	}
 
+	/**
+	 * Inserts a default-constructed elements into the array at a given location.
+	 *
+	 * @param Index Tells where to insert the new elements.
+	 * @see Insert, InsertUninitialized, InsertZeroed
+	 */
+	void InsertDefaulted(int32 Index)
+	{
+		InsertUninitialized(Index);
+		DefaultConstructItems<ElementType>(GetData() + Index, 1);
+	}
+
 	/**
 	 * Inserts a given number of default-constructed elements into the array at a given
 	 * location.
@@ -1289,7 +1375,7 @@ public:
 	 * @param Count Number of elements to add.
 	 * @see Insert, InsertUninitialized, InsertZeroed
 	 */
-	void InsertDefaulted(int32 Index, int32 Count = 1)
+	void InsertDefaulted(int32 Index, int32 Count)
 	{
 		InsertUninitialized(Index, Count);
 		DefaultConstructItems<ElementType>(GetData() + Index, Count);
@@ -1305,7 +1391,7 @@ public:
 	 */
 	ElementType& InsertDefaulted_GetRef(int32 Index)
 	{
-		InsertUninitialized(Index, 1);
+		InsertUninitialized(Index);
 		ElementType* Ptr = GetData() + Index;
 		DefaultConstructItems<ElementType>(Ptr, 1);
 		return *Ptr;
@@ -1416,7 +1502,7 @@ public:
 
 		// construct a copy in place at Index (this new operator will insert at 
 		// Index, then construct that memory with Item)
-		InsertUninitialized(Index, 1);
+		InsertUninitialized(Index);
 		new(GetData() + Index) ElementType(MoveTempIfPossible(Item));
 		return Index;
 	}
@@ -1435,7 +1521,7 @@ public:
 
 		// construct a copy in place at Index (this new operator will insert at 
 		// Index, then construct that memory with Item)
-		InsertUninitialized(Index, 1);
+		InsertUninitialized(Index);
 		new(GetData() + Index) ElementType(Item);
 		return Index;
 	}
@@ -1455,7 +1541,7 @@ public:
 
 		// construct a copy in place at Index (this new operator will insert at 
 		// Index, then construct that memory with Item)
-		InsertUninitialized(Index, 1);
+		InsertUninitialized(Index);
 		ElementType* Ptr = GetData() + Index;
 		new(Ptr) ElementType(MoveTempIfPossible(Item));
 		return *Ptr;
@@ -1475,7 +1561,7 @@ public:
 
 		// construct a copy in place at Index (this new operator will insert at 
 		// Index, then construct that memory with Item)
-		InsertUninitialized(Index, 1);
+		InsertUninitialized(Index);
 		ElementType* Ptr = GetData() + Index;
 		new(Ptr) ElementType(Item);
 		return *Ptr;
@@ -1656,6 +1742,7 @@ public:
 	{
 		if (NewNum > Num())
 		{
+			ResizeTo(NewNum);
 			const int32 Diff = NewNum - ArrayNum;
 			const int32 Index = AddUninitialized(Diff);
 			DefaultConstructItems<ElementType>((uint8*)AllocatorInstance.GetAllocation() + Index * sizeof(ElementType), Diff);
@@ -1675,6 +1762,7 @@ public:
 	{
 		if (NewNum > Num())
 		{
+			ResizeTo(NewNum);
 			AddZeroed(NewNum - Num());
 		}
 		else if (NewNum < Num())
@@ -1692,6 +1780,7 @@ public:
 	{
 		if (NewNum > Num())
 		{
+			ResizeTo(NewNum);
 			AddUninitialized(NewNum - Num());
 		}
 		else if (NewNum < Num())
@@ -1840,7 +1929,7 @@ public:
 	template <typename... ArgsType>
 	FORCEINLINE int32 Emplace(ArgsType&&... Args)
 	{
-		const int32 Index = AddUninitialized(1);
+		const int32 Index = AddUninitialized();
 		new(GetData() + Index) ElementType(Forward<ArgsType>(Args)...);
 		return Index;
 	}
@@ -1854,7 +1943,7 @@ public:
 	template <typename... ArgsType>
 	FORCEINLINE ElementType& Emplace_GetRef(ArgsType&&... Args)
 	{
-		const int32 Index = AddUninitialized(1);
+		const int32 Index = AddUninitialized();
 		ElementType* Ptr = GetData() + Index;
 		new(Ptr) ElementType(Forward<ArgsType>(Args)...);
 		return *Ptr;
@@ -1869,7 +1958,7 @@ public:
 	template <typename... ArgsType>
 	FORCEINLINE void EmplaceAt(int32 Index, ArgsType&&... Args)
 	{
-		InsertUninitialized(Index, 1);
+		InsertUninitialized(Index);
 		new(GetData() + Index) ElementType(Forward<ArgsType>(Args)...);
 	}
 
@@ -1883,7 +1972,7 @@ public:
 	template <typename... ArgsType>
 	FORCEINLINE ElementType& EmplaceAt_GetRef(int32 Index, ArgsType&&... Args)
 	{
-		InsertUninitialized(Index, 1);
+		InsertUninitialized(Index);
 		ElementType* Ptr = GetData() + Index;
 		new(Ptr) ElementType(Forward<ArgsType>(Args)...);
 		return *Ptr;
@@ -1945,6 +2034,24 @@ public:
 		return Emplace_GetRef(Item);
 	}
 
+	/**
+	 * Adds a new item to the end of the array, possibly reallocating the whole
+	 * array to fit. The new item will be zeroed.
+	 *
+	 * Caution, AddZeroed() will create elements without calling the
+	 * constructor and this is not appropriate for element types that require
+	 * a constructor to function properly.
+	 *
+	 * @return Index of the new item.
+	 * @see Add, AddDefaulted, AddUnique, Append, Insert
+	 */
+	int32 AddZeroed()
+	{
+		const int32 Index = AddUninitialized();
+		FMemory::Memzero((uint8*)AllocatorInstance.GetAllocation() + Index*sizeof(ElementType), sizeof(ElementType));
+		return Index;
+	}
+
 	/**
 	 * Adds new items to the end of the array, possibly reallocating the whole
 	 * array to fit. The new items will be zeroed.
@@ -1957,7 +2064,7 @@ public:
 	 * @return Index to the first of the new items.
 	 * @see Add, AddDefaulted, AddUnique, Append, Insert
 	 */
-	int32 AddZeroed(int32 Count = 1)
+	int32 AddZeroed(int32 Count)
 	{
 		const int32 Index = AddUninitialized(Count);
 		FMemory::Memzero((uint8*)AllocatorInstance.GetAllocation() + Index*sizeof(ElementType), Count*sizeof(ElementType));
@@ -1977,12 +2084,26 @@ public:
 	 */
 	ElementType& AddZeroed_GetRef()
 	{
-		const int32 Index = AddUninitialized(1);
+		const int32 Index = AddUninitialized();
 		ElementType* Ptr = GetData() + Index;
 		FMemory::Memzero(Ptr, sizeof(ElementType));
 		return *Ptr;
 	}
 
+	/**
+	 * Adds a new item to the end of the array, possibly reallocating the whole
+	 * array to fit. The new item will be default-constructed.
+	 *
+	 * @return Index of the new item.
+	 * @see Add, AddZeroed, AddUnique, Append, Insert
+	 */
+	int32 AddDefaulted()
+	{
+		const int32 Index = AddUninitialized();
+		DefaultConstructItems<ElementType>((uint8*)AllocatorInstance.GetAllocation() + Index * sizeof(ElementType), 1);
+		return Index;
+	}
+
 	/**
 	 * Adds new items to the end of the array, possibly reallocating the whole
 	 * array to fit. The new items will be default-constructed.
@@ -1991,7 +2112,7 @@ public:
 	 * @return Index to the first of the new items.
 	 * @see Add, AddZeroed, AddUnique, Append, Insert
 	 */
-	int32 AddDefaulted(int32 Count = 1)
+	int32 AddDefaulted(int32 Count)
 	{
 		const int32 Index = AddUninitialized(Count);
 		DefaultConstructItems<ElementType>((uint8*)AllocatorInstance.GetAllocation() + Index * sizeof(ElementType), Count);
@@ -2007,7 +2128,7 @@ public:
 	 */
 	ElementType& AddDefaulted_GetRef()
 	{
-		const int32 Index = AddUninitialized(1);
+		const int32 Index = AddUninitialized();
 		ElementType* Ptr = GetData() + Index;
 		DefaultConstructItems<ElementType>(Ptr, 1);
 		return *Ptr;
