Merging CL #5362180 from Dev-Core:

Removed global locks when logging through FOutputDeviceRedirector to avoid hangs when logging crashes.

+ ~10% Logging performance increase (tested by multiple attempts at logging 1M lines of text):
   before: 365984ms (min)
   after:    324901ms (max)

diff --git a/Engine/Source/Runtime/Core/Private/Misc/OutputDeviceRedirector.cpp b/Engine/Source/Runtime/Core/Private/Misc/OutputDeviceRedirector.cpp
--- a/Engine/Source/Runtime/Core/Private/Misc/OutputDeviceRedirector.cpp
+++ b/Engine/Source/Runtime/Core/Private/Misc/OutputDeviceRedirector.cpp
@@ -4,6 +4,7 @@
 #include "Misc/ScopeLock.h"
 #include "Stats/Stats.h"
 #include "Misc/CoreStats.h"
+#include "HAL/PlatformProcess.h"
 
 /*-----------------------------------------------------------------------------
 	FOutputDeviceRedirector.
@@ -29,18 +30,31 @@ FOutputDeviceRedirector* FOutputDeviceRedirector::Get()
  */
 void FOutputDeviceRedirector::AddOutputDevice( FOutputDevice* OutputDevice )
 {
-	FScopeLock ScopeLock( &SynchronizationObject );
-
-	if( OutputDevice )
+	if (OutputDevice)
 	{
-		if (OutputDevice->CanBeUsedOnMultipleThreads())
-		{
-			UnbufferedOutputDevices.AddUnique(OutputDevice);
-		}
-		else
+		bool bAdded = false;
+		do
 		{
-			BufferedOutputDevices.AddUnique(OutputDevice);
-		}
+			{
+				FScopeLock ScopeLock(&OutputDevicesMutex);
+				if (OutputDevicesLockCounter.GetValue() == 0)
+				{
+					if (OutputDevice->CanBeUsedOnMultipleThreads())
+					{
+						UnbufferedOutputDevices.AddUnique(OutputDevice);
+					}
+					else
+					{
+						BufferedOutputDevices.AddUnique(OutputDevice);
+					}
+					bAdded = true;
+				}
+			}
+			if (!bAdded)
+			{
+				FPlatformProcess::Sleep(0);
+			}
+		} while (!bAdded);
 	}
 }
 
@@ -51,9 +65,23 @@ void FOutputDeviceRedirector::AddOutputDevice( FOutputDevice* OutputDevice )
  */
 void FOutputDeviceRedirector::RemoveOutputDevice( FOutputDevice* OutputDevice )
 {
-	FScopeLock ScopeLock( &SynchronizationObject );
-	BufferedOutputDevices.Remove( OutputDevice );
-	UnbufferedOutputDevices.Remove(OutputDevice);
+	bool bRemoved = false;
+	do
+	{
+		{
+			FScopeLock ScopeLock(&OutputDevicesMutex);
+			if (OutputDevicesLockCounter.GetValue() == 0)
+			{
+				BufferedOutputDevices.Remove(OutputDevice);
+				UnbufferedOutputDevices.Remove(OutputDevice);
+				bRemoved = true;
+			}
+		}
+		if (!bRemoved)
+		{
+			FPlatformProcess::Sleep(0);
+		}
+	} while (!bRemoved);
 }
 
 /**
@@ -64,32 +92,49 @@ void FOutputDeviceRedirector::RemoveOutputDevice( FOutputDevice* OutputDevice )
  */
 bool FOutputDeviceRedirector::IsRedirectingTo( FOutputDevice* OutputDevice )
 {
-	FScopeLock ScopeLock( &SynchronizationObject );
-
+	// For performance reasons whe're not using the FOutputDevicesLock here
+	FScopeLock OutputDevicesLock(&OutputDevicesMutex);
 	return BufferedOutputDevices.Contains(OutputDevice) || UnbufferedOutputDevices.Contains(OutputDevice);
 }
 
+void FOutputDeviceRedirector::InternalFlushThreadedLogs(bool bUseAllDevices)
+{
+	TLocalOutputDevicesArray LocalBufferedDevices;
+	TLocalOutputDevicesArray LocalUnbufferedDevices;
+	FOutputDevicesLock OutputDevicesLock(this, LocalBufferedDevices, LocalUnbufferedDevices);
+
+	InternalFlushThreadedLogs(LocalBufferedDevices, bUseAllDevices);
+}
+
 /**
  * The unsynchronized version of FlushThreadedLogs.
- * Assumes that the caller holds a lock on SynchronizationObject.
  */
-void FOutputDeviceRedirector::UnsynchronizedFlushThreadedLogs( bool bUseAllDevices )
-{
-	for(int32 LineIndex = 0;LineIndex < BufferedLines.Num();LineIndex++)
+void FOutputDeviceRedirector::InternalFlushThreadedLogs(TLocalOutputDevicesArray& InBufferedDevices, bool bUseAllDevices)
+{	
+	if (BufferedLines.Num())
 	{
-		FBufferedLine BufferedLine = BufferedLines[LineIndex];
+		TArray<FBufferedLine, TInlineAllocator<64>> LocalBufferedLines;
+		{
+			FScopeLock ScopeLock(&BufferSynchronizationObject);
+			LocalBufferedLines.AddUninitialized(BufferedLines.Num());
+			for (int32 LineIndex = 0; LineIndex < BufferedLines.Num(); LineIndex++)
+			{
+				new(&LocalBufferedLines[LineIndex]) FBufferedLine(BufferedLines[LineIndex], FBufferedLine::EMoveCtor);
+			}
+			BufferedLines.Empty();
+		}
 
-		for( int32 OutputDeviceIndex=0; OutputDeviceIndex<BufferedOutputDevices.Num(); OutputDeviceIndex++ )
+		for (FBufferedLine& LocalBufferedLine : LocalBufferedLines)
 		{
-			FOutputDevice* OutputDevice = BufferedOutputDevices[OutputDeviceIndex];
-			if( OutputDevice->CanBeUsedOnAnyThread() || bUseAllDevices )
+			for (FOutputDevice* OutputDevice : InBufferedDevices)
 			{
-				OutputDevice->Serialize( *BufferedLine.Data, BufferedLine.Verbosity, BufferedLine.Category, BufferedLine.Time );
+				if (OutputDevice->CanBeUsedOnAnyThread() || bUseAllDevices)
+				{
+					OutputDevice->Serialize(*LocalBufferedLine.Data, LocalBufferedLine.Verbosity, LocalBufferedLine.Category, LocalBufferedLine.Time);
+				}
 			}
 		}
 	}
-
-	BufferedLines.Empty();
 }
 
 /**
@@ -98,37 +143,34 @@ void FOutputDeviceRedirector::UnsynchronizedFlushThreadedLogs( bool bUseAllDevic
 void FOutputDeviceRedirector::FlushThreadedLogs()
 {
 	SCOPE_CYCLE_COUNTER(STAT_FlushThreadedLogs);
-	// Acquire a lock on SynchronizationObject and call the unsynchronized worker function.
-	FScopeLock ScopeLock( &SynchronizationObject );
 	check(IsInGameThread());
-	UnsynchronizedFlushThreadedLogs( true );
+	InternalFlushThreadedLogs(true);
 }
 
 void FOutputDeviceRedirector::PanicFlushThreadedLogs()
 {
 //	SCOPE_CYCLE_COUNTER(STAT_FlushThreadedLogs);
-	// Acquire a lock on SynchronizationObject and call the unsynchronized worker function.
-	FScopeLock ScopeLock( &SynchronizationObject );
-	
+
+	TLocalOutputDevicesArray LocalBufferedDevices;
+	TLocalOutputDevicesArray LocalUnbufferedDevices;
+	FOutputDevicesLock OutputDevicesLock(this, LocalBufferedDevices, LocalUnbufferedDevices);
+
 	// Flush threaded logs, but use the safe version.
-	UnsynchronizedFlushThreadedLogs( false );
+	InternalFlushThreadedLogs(LocalBufferedDevices, false);
 
 	// Flush devices.
-	for (int32 OutputDeviceIndex = 0; OutputDeviceIndex<BufferedOutputDevices.Num(); OutputDeviceIndex++)
+	for (FOutputDevice* OutputDevice : LocalBufferedDevices)
 	{
-		FOutputDevice* OutputDevice = BufferedOutputDevices[OutputDeviceIndex];
 		if (OutputDevice->CanBeUsedOnAnyThread())
 		{
 			OutputDevice->Flush();
 		}
 	}
 
-	for (int32 OutputDeviceIndex = 0; OutputDeviceIndex < UnbufferedOutputDevices.Num(); OutputDeviceIndex++)
+	for (FOutputDevice* OutputDevice : LocalUnbufferedDevices)
 	{
-		UnbufferedOutputDevices[OutputDeviceIndex]->Flush();
+		OutputDevice->Flush();
 	}
-
-	BufferedLines.Empty();
 }
 
 /**
@@ -152,7 +194,7 @@ void FOutputDeviceRedirector::SerializeBacklog( FOutputDevice* OutputDevice )
  */
 void FOutputDeviceRedirector::EnableBacklog( bool bEnable )
 {
-	FScopeLock ScopeLock( &SynchronizationObject );
+	FScopeLock ScopeLock(&SynchronizationObject);
 
 	bEnableBacklog = bEnable;
 	if ( bEnableBacklog == false )
@@ -167,20 +209,38 @@ void FOutputDeviceRedirector::EnableBacklog( bool bEnable )
  */
 void FOutputDeviceRedirector::SetCurrentThreadAsMasterThread()
 {
-	FScopeLock ScopeLock( &SynchronizationObject );
-
 	// make sure anything queued up is flushed out, this may be called from a background thread, so use the safe version.
-	UnsynchronizedFlushThreadedLogs( false );
+	InternalFlushThreadedLogs( false );
 
-	// set the current thread as the master thread
-	MasterThreadID = FPlatformTLS::GetCurrentThreadId();
+	{
+		FScopeLock ScopeLock(&SynchronizationObject);
+		// set the current thread as the master thread
+		MasterThreadID = FPlatformTLS::GetCurrentThreadId();
+	}
+}
+
+void FOutputDeviceRedirector::LockOutputDevices(TLocalOutputDevicesArray& OutBufferedDevices, TLocalOutputDevicesArray& OutUnbufferedDevices)
+{
+	FScopeLock OutputDevicesLock(&OutputDevicesMutex);
+	OutputDevicesLockCounter.Increment();
+	OutBufferedDevices.Append(BufferedOutputDevices);
+	OutUnbufferedDevices.Append(UnbufferedOutputDevices);
+}
+
+void FOutputDeviceRedirector::UnlockOutputDevices()
+{
+	FScopeLock OutputDevicesLock(&OutputDevicesMutex);
+	int32 LockValue = OutputDevicesLockCounter.Decrement();
+	check(LockValue >= 0);
 }
 
 void FOutputDeviceRedirector::Serialize( const TCHAR* Data, ELogVerbosity::Type Verbosity, const class FName& Category, const double Time )
 {
 	const double RealTime = Time == -1.0f ? FPlatformTime::Seconds() - GStartTime : Time;
 
-	FScopeLock ScopeLock( &SynchronizationObject );
+	TLocalOutputDevicesArray LocalBufferedDevices;
+	TLocalOutputDevicesArray LocalUnbufferedDevices;
+	FOutputDevicesLock OutputDevicesLock(this, LocalBufferedDevices, LocalUnbufferedDevices);
 
 #if PLATFORM_DESKTOP
 	// this is for errors which occur after shutdown we might be able to salvage information from stdout 
@@ -197,30 +257,31 @@ void FOutputDeviceRedirector::Serialize( const TCHAR* Data, ELogVerbosity::Type
 #endif
 
 	// Serialize directly to any output devices which don't require buffering
-	for (int32 OutputDeviceIndex = 0; OutputDeviceIndex < UnbufferedOutputDevices.Num(); OutputDeviceIndex++)
+	for (FOutputDevice* OutputDevice : LocalUnbufferedDevices)
 	{
-		UnbufferedOutputDevices[OutputDeviceIndex]->Serialize(Data, Verbosity, Category, RealTime);
+		OutputDevice->Serialize(Data, Verbosity, Category, RealTime);
 	}
 
 
 	if ( bEnableBacklog )
 	{
+		FScopeLock ScopeLock(&SynchronizationObject);
 		new(BacklogLines)FBufferedLine( Data, Category, Verbosity, RealTime );
 	}
 
-	if(FPlatformTLS::GetCurrentThreadId() != MasterThreadID || BufferedOutputDevices.Num() == 0)
+	if(FPlatformTLS::GetCurrentThreadId() != MasterThreadID || LocalBufferedDevices.Num() == 0)
 	{
+		FScopeLock ScopeLock(&BufferSynchronizationObject);
 		new(BufferedLines)FBufferedLine( Data, Category, Verbosity, RealTime );
 	}
 	else
 	{
 		// Flush previously buffered lines from secondary threads.
-		// Since we already hold a lock on SynchronizationObject, call the unsynchronized version.
-		UnsynchronizedFlushThreadedLogs( true );
+		InternalFlushThreadedLogs(LocalBufferedDevices, true);
 
-		for( int32 OutputDeviceIndex=0; OutputDeviceIndex<BufferedOutputDevices.Num(); OutputDeviceIndex++ )
+		for (FOutputDevice* OutputDevice : LocalBufferedDevices)
 		{
-			BufferedOutputDevices[OutputDeviceIndex]->Serialize( Data, Verbosity, Category, RealTime );
+			OutputDevice->Serialize(Data, Verbosity, Category, RealTime);
 		}
 	}
 }
@@ -235,23 +296,25 @@ void FOutputDeviceRedirector::Serialize( const TCHAR* Data, ELogVerbosity::Type
  */
 void FOutputDeviceRedirector::Flush()
 {
-	FScopeLock ScopeLock( &SynchronizationObject );
+	TLocalOutputDevicesArray LocalBufferedDevices;
+	TLocalOutputDevicesArray LocalUnbufferedDevices;
+	FOutputDevicesLock OutputDevicesLock(this, LocalBufferedDevices, LocalUnbufferedDevices);
 
 	if(FPlatformTLS::GetCurrentThreadId() == MasterThreadID)
 	{
 		// Flush previously buffered lines from secondary threads.
 		// Since we already hold a lock on SynchronizationObject, call the unsynchronized version.
-		UnsynchronizedFlushThreadedLogs( true );
+		InternalFlushThreadedLogs(true);
 
-		for( int32 OutputDeviceIndex=0; OutputDeviceIndex<BufferedOutputDevices.Num(); OutputDeviceIndex++ )
+		for (FOutputDevice* OutputDevice : LocalBufferedDevices)
 		{
-			BufferedOutputDevices[OutputDeviceIndex]->Flush();
+			OutputDevice->Flush();
 		}
 	}
 
-	for (int32 OutputDeviceIndex = 0; OutputDeviceIndex < UnbufferedOutputDevices.Num(); OutputDeviceIndex++)
+	for (FOutputDevice* OutputDevice : LocalUnbufferedDevices)
 	{
-		UnbufferedOutputDevices[OutputDeviceIndex]->Flush();
+		OutputDevice->Flush();
 	}
 }
 
@@ -262,32 +325,39 @@ void FOutputDeviceRedirector::Flush()
  */
 void FOutputDeviceRedirector::TearDown()
 {
+	FScopeLock SyncLock(&SynchronizationObject);
 	check(FPlatformTLS::GetCurrentThreadId() == MasterThreadID);
 
-	FScopeLock ScopeLock( &SynchronizationObject );
+	TLocalOutputDevicesArray LocalBufferedDevices;
+	TLocalOutputDevicesArray LocalUnbufferedDevices;
+	
+	{
+		// We need to lock the mutex here so that it gets unlocked after we empty the devices arrays
+		FScopeLock GlobalOutputDevicesLock(&OutputDevicesMutex);
+		LockOutputDevices(LocalBufferedDevices, LocalUnbufferedDevices);
+		BufferedOutputDevices.Empty();
+		UnbufferedOutputDevices.Empty();	
+	}
 
 	// Flush previously buffered lines from secondary threads.
-	// Since we already hold a lock on SynchronizationObject, call the unsynchronized version.
-	UnsynchronizedFlushThreadedLogs( false );
+	InternalFlushThreadedLogs(LocalBufferedDevices, false);
 
-	for( int32 OutputDeviceIndex=0; OutputDeviceIndex<BufferedOutputDevices.Num(); OutputDeviceIndex++ )
+	for (FOutputDevice* OutputDevice : LocalBufferedDevices)
 	{
-		FOutputDevice* OutputDevice = BufferedOutputDevices[OutputDeviceIndex];
 		if (OutputDevice->CanBeUsedOnAnyThread())
 		{
 			OutputDevice->Flush();
 		}
 		OutputDevice->TearDown();
-	}
-	BufferedOutputDevices.Empty();
+	}	
 
-	for (int32 OutputDeviceIndex = 0; OutputDeviceIndex < UnbufferedOutputDevices.Num(); OutputDeviceIndex++)
+	for (FOutputDevice* OutputDevice : LocalUnbufferedDevices)
 	{
-		FOutputDevice* OutputDevice = UnbufferedOutputDevices[OutputDeviceIndex];
 		OutputDevice->Flush();
 		OutputDevice->TearDown();
 	}
-	UnbufferedOutputDevices.Empty();
+
+	UnlockOutputDevices();
 }
 
 CORE_API FOutputDeviceRedirector* GetGlobalLogSingleton()
diff --git a/Engine/Source/Runtime/Core/Public/Misc/OutputDeviceRedirector.h b/Engine/Source/Runtime/Core/Public/Misc/OutputDeviceRedirector.h
--- a/Engine/Source/Runtime/Core/Public/Misc/OutputDeviceRedirector.h
+++ b/Engine/Source/Runtime/Core/Public/Misc/OutputDeviceRedirector.h
@@ -15,6 +15,11 @@ FOutputDeviceRedirector.
 /** The type of lines buffered by secondary threads. */
 struct FBufferedLine
 {
+	enum EBufferedLineInit
+	{
+		EMoveCtor = 0
+	};
+
 	const FString Data;
 	const FName Category;
 	const double Time;
@@ -27,6 +32,19 @@ struct FBufferedLine
 		, Time(InTime)
 		, Verbosity(InVerbosity)
 	{}
+
+	/** Default constructor. */
+	FBufferedLine()
+		: Time(0.0)
+		, Verbosity(ELogVerbosity::NoLogging)
+	{}
+
+	FBufferedLine(FBufferedLine& InBufferedLine, EBufferedLineInit Unused)
+		: Data(MoveTemp(const_cast<FString&>(InBufferedLine.Data)))
+		, Category(InBufferedLine.Category)
+		, Time(InBufferedLine.Time)
+		, Verbosity(InBufferedLine.Verbosity)
+	{}
 };
 
 /**
@@ -34,6 +52,10 @@ struct FBufferedLine
 */
 class CORE_API FOutputDeviceRedirector : public FOutputDevice
 {
+public:
+
+	typedef TArray<FOutputDevice*, TInlineAllocator<16> > TLocalOutputDevicesArray;
+
 private:
 	/** A FIFO of lines logged by non-master threads. */
 	TArray<FBufferedLine> BufferedLines;
@@ -53,16 +75,42 @@ private:
 	/** Whether backlogging is enabled. */
 	bool bEnableBacklog;
 
-	/** Object used for synchronization via a scoped lock */
+	/** Objects used for synchronization via a scoped lock */
 	FCriticalSection	SynchronizationObject;
+	FCriticalSection	BufferSynchronizationObject;
+	FCriticalSection	OutputDevicesMutex;
+	FThreadSafeCounter	OutputDevicesLockCounter;
 
 	/**
 	* The unsynchronized version of FlushThreadedLogs.
 	* Assumes that the caller holds a lock on SynchronizationObject.
 	* @param bUseAllDevices - if true this method will use all output devices
 	*/
-	void UnsynchronizedFlushThreadedLogs(bool bUseAllDevices);
+	void InternalFlushThreadedLogs(TLocalOutputDevicesArray& InBufferedDevices, bool bUseAllDevices);
+	void InternalFlushThreadedLogs(bool bUseAllDevices);
+
+	/** Locks OutputDevices arrays so that nothing can be added or removed from them */
+	void LockOutputDevices(TLocalOutputDevicesArray& OutBufferedDevices, TLocalOutputDevicesArray& OutUnbufferedDevices);
 
+	/** Unlocks OutputDevices arrays */
+	void UnlockOutputDevices();
+
+	friend struct FOutputDevicesLock;
+
+	/** Helper struct for scope locking OutputDevices arrays */
+	struct FOutputDevicesLock
+	{
+		FOutputDeviceRedirector* Redirector;
+		FOutputDevicesLock(FOutputDeviceRedirector* InRedirector, TLocalOutputDevicesArray& OutBufferedDevices, TLocalOutputDevicesArray& OutUnbufferedDevices)
+			: Redirector(InRedirector)
+		{
+			Redirector->LockOutputDevices(OutBufferedDevices, OutUnbufferedDevices);
+		}
+		~FOutputDevicesLock()
+		{
+			Redirector->UnlockOutputDevices();
+		}
+	};
 public:
 
 	/** Initialization constructor. */
