[FramePro] update to the latest version of FramePro runtime integration code (1.5.15.0)

diff --git a/Engine/Source/Runtime/Core/Private/Android/AndroidPlatformMisc.cpp b/Engine/Source/Runtime/Core/Private/Android/AndroidPlatformMisc.cpp
--- a/Engine/Source/Runtime/Core/Private/Android/AndroidPlatformMisc.cpp
+++ b/Engine/Source/Runtime/Core/Private/Android/AndroidPlatformMisc.cpp
@@ -1995,12 +1995,12 @@ void FAndroidMisc::EndNamedEvent()
 
 void FAndroidMisc::CustomNamedStat(const TCHAR* Text, float Value, const TCHAR* Graph, const TCHAR* Unit)
 {
-	FRAMEPRO_DYNAMIC_CUSTOM_STAT(TCHAR_TO_WCHAR(Text), Value, TCHAR_TO_WCHAR(Graph), TCHAR_TO_WCHAR(Unit));
+	FRAMEPRO_DYNAMIC_CUSTOM_STAT(TCHAR_TO_WCHAR(Text), Value, TCHAR_TO_WCHAR(Graph), TCHAR_TO_WCHAR(Unit), FRAMEPRO_COLOUR(255,255,255) );
 }
 
 void FAndroidMisc::CustomNamedStat(const ANSICHAR* Text, float Value, const ANSICHAR* Graph, const ANSICHAR* Unit)
 {
-	FRAMEPRO_DYNAMIC_CUSTOM_STAT(Text, Value, Graph, Unit);
+	FRAMEPRO_DYNAMIC_CUSTOM_STAT(Text, Value, Graph, Unit, FRAMEPRO_COLOUR(255,255,255));
 }
 
 #endif // STATS || ENABLE_STATNAMEDEVENTS
diff --git a/Engine/Source/Runtime/Core/Private/Apple/ApplePlatformMisc.cpp b/Engine/Source/Runtime/Core/Private/Apple/ApplePlatformMisc.cpp
--- a/Engine/Source/Runtime/Core/Private/Apple/ApplePlatformMisc.cpp
+++ b/Engine/Source/Runtime/Core/Private/Apple/ApplePlatformMisc.cpp
@@ -350,12 +350,12 @@ void FApplePlatformMisc::EndNamedEvent()
 
 void FApplePlatformMisc::CustomNamedStat(const TCHAR* Text, float Value, const TCHAR* Graph, const TCHAR* Unit)
 {
-	FRAMEPRO_DYNAMIC_CUSTOM_STAT(TCHAR_TO_WCHAR(Text), Value, TCHAR_TO_WCHAR(Graph), TCHAR_TO_WCHAR(Unit));
+	FRAMEPRO_DYNAMIC_CUSTOM_STAT(TCHAR_TO_WCHAR(Text), Value, TCHAR_TO_WCHAR(Graph), TCHAR_TO_WCHAR(Unit), FRAMEPRO_COLOUR(255,255,255));
 }
 
 void FApplePlatformMisc::CustomNamedStat(const ANSICHAR* Text, float Value, const ANSICHAR* Graph, const ANSICHAR* Unit)
 {
-	FRAMEPRO_DYNAMIC_CUSTOM_STAT(Text, Value, Graph, Unit);
+	FRAMEPRO_DYNAMIC_CUSTOM_STAT(Text, Value, Graph, Unit, FRAMEPRO_COLOUR(255,255,255));
 }
 
 #endif // STATS || ENABLE_STATNAMEDEVENTS || APPLE_PROFILING_ENABLED
