[FramePro] update to the latest version of FramePro runtime integration code (1.5.15.0)

diff --git a/Engine/Source/Runtime/Core/Private/FramePro/FramePro.cpp b/Engine/Source/Runtime/Core/Private/FramePro/FramePro.cpp
--- a/Engine/Source/Runtime/Core/Private/FramePro/FramePro.cpp
+++ b/Engine/Source/Runtime/Core/Private/FramePro/FramePro.cpp
@@ -22,123 +22,35 @@
 	Add FramePro.cpp to your project to allow FramePro to communicate with your application.
 */
 
-// BEGIN EPIC 
-#include "FramePro/FramePro.h"
-#include "CoreGlobals.h"
-#include "HAL/FileManager.h"
-#include "Misc/Paths.h"
-//END EPIC 
-//------------------------------------------------------------------------
-// EventTraceWin32.cpp
-
-//------------------------------------------------------------------------
-// EventTraceWin32.hpp
-#ifndef FRAMEPRO_EVENTTRACEWIN32_H_INCLUDED
-#define FRAMEPRO_EVENTTRACEWIN32_H_INCLUDED
+#if defined(__UNREAL__)
+	#include "FramePro/FramePro.h"
+#else
+	#include "FramePro.h"
+#endif
 
 //------------------------------------------------------------------------
+#if FRAMEPRO_ENABLED
 
-//------------------------------------------------------------------------
-// FrameProLib.hpp
-#ifndef FRAMEPROLIB_H_INCLUDED
-#define FRAMEPROLIB_H_INCLUDED
 
-//------------------------------------------------------------------------
 
 //------------------------------------------------------------------------
-#if FRAMEPRO_ENABLED
+//
+// FrameProLib.hpp
+//
 
 //------------------------------------------------------------------------
 #include <stdio.h>
 #include <stdlib.h>
 #include <string.h>
 #include <wchar.h>
+#include <limits.h>
 #include <new>
 
-#if FRAMEPRO_WIN_BASED_PLATFORM
-	#include <tchar.h>
-#endif
-
-// BEGIN EPIC
-#if FRAMEPRO_UE4_BASED_PLATFORM
-	#include "HAL/PlatformProcess.h"
-	#include "HAL/PlatformTLS.h"
-	#include "HAL/Event.h"
-	#include "HAL/CriticalSection.h"
-	#include "HAL/Runnable.h"
-	#include "HAL/RunnableThread.h"
-	#include "Templates/UniquePtr.h"
-#endif
-// END EPIC
-
-// BEGIN EPIC 
-// Remove FRAMEPRO_UNIX_BASED_PLATFORM and FRAMEPRO_TIMER_QUERY_PERFORMANCE_COUNTER blocks
-// END EPIC
-//------------------------------------------------------------------------
-#if FRAMEPRO_SOCKETS_ENABLED
-	#if FRAMEPRO_WIN_BASED_PLATFORM
-		#if defined(AF_IPX) && !defined(_WINSOCK2API_)
-			#error winsock already defined. Please include winsock2.h before including windows.h or use WIN32_LEAN_AND_MEAN. See the FAQ for more info.
-		#endif
-		#if FRAMEPRO_TOOLSET_UE4
-			// BEGIN EPIC 
-			#include "Windows/AllowWindowsPlatformTypes.h"
-			// END EPIC
-		#endif
-		#include <winsock2.h>
-		#include <ws2tcpip.h>
-		#if FRAMEPRO_TOOLSET_UE4
-			// BEGIN EPIC 
-			#include "Windows/HideWindowsPlatformTypes.h"
-			// END EPIC
-		#endif
-	// BEGIN EPIC
-	#elif FRAMEPRO_PLATFORM_SWITCH
-		#include "Switch/SwitchPlatformFramePro.h"
-	// END EPIC
-	#else
-		#include <sys/socket.h>
-		#include <netinet/in.h>
-	#endif
-#endif
-
-//------------------------------------------------------------------------
-#define FRAMEPRO_MAX_INLINE_STRING_LENGTH 256
-
-//------------------------------------------------------------------------
-// BEGIN EPIC 
-#if PLATFORM_64BITS
-// END EPIC
-	#define FRAMEPRO_X64 1
-#else
-	#define FRAMEPRO_X64 0
-#endif
-
-//------------------------------------------------------------------------
-// BEGIN EPIC
-//#if FRAMEPRO_WIN_BASED_PLATFORM
-//	#define FRAMEPRO_THREAD_LOCAL __declspec(thread)
-//#endif
-// END EPIC
-
-//------------------------------------------------------------------------
-// BEGIN EPIC 
-#define FRAMEPRO_NO_INLINE 
-#define FRAMEPRO_FORCE_INLINE FORCEINLINE
-// END EPIC
-
-//------------------------------------------------------------------------
-#if FRAMEPRO_WIN_BASED_PLATFORM
-	#define FRAMEPRO_TCHAR _TCHAR
-#else
-	#define FRAMEPRO_TCHAR char
-#endif
-
 //------------------------------------------------------------------------
 namespace FramePro
 {
 	//------------------------------------------------------------------------
-	enum { g_FrameProLibVersion = 11 };
+	enum { g_FrameProLibVersion = 14 };
 
 	//------------------------------------------------------------------------
 	namespace StringLiteralType
@@ -150,6 +62,7 @@ namespace FramePro
 			SourceInfo,
 			GeneralString,
 			StringLiteralTimerName,
+			GeneralStringW,
 		};
 	}
 
@@ -157,32 +70,24 @@ namespace FramePro
 	class FrameProTLS;
 
 	//------------------------------------------------------------------------
-// BEGIN EPIC
-//	extern FRAMEPRO_THREAD_LOCAL FrameProTLS* gp_FrameProTLS;
-// END EPIC
 	extern FRAMEPRO_NO_INLINE FrameProTLS* CreateFrameProTLS();
 	extern FRAMEPRO_NO_INLINE void DestroyFrameProTLS(FrameProTLS* p_framepro_tls);
 
-// BEGIN EPIC
-	static uint32 GetFrameProTLSSlot()
-	{
-		static uint32 TlsSlot = FPlatformTLS::AllocTlsSlot();
-		return TlsSlot;
-	}
-// END EPIC
+	#if FRAMEPRO_USE_TLS_SLOTS
+		extern FRAMEPRO_NO_INLINE uint GetFrameProTLSSlot();
+	#else
+		inline uint GetFrameProTLSSlot() { return 0; } // help the compiler optimise this away if we are not using tls TLS slots
+	#endif
 
 	//------------------------------------------------------------------------
-	FrameProTLS* GetFrameProTLS()
+	FRAMEPRO_FORCE_INLINE FrameProTLS* GetFrameProTLS()
 	{
-// BEGIN EPIC
-//		FrameProTLS* p_framepro_tls = gp_FrameProTLS;
-		FrameProTLS* p_framepro_tls = (FrameProTLS*)FPlatformTLS::GetTlsValue(GetFrameProTLSSlot());
-// END EPIC
+		FrameProTLS* p_framepro_tls = (FrameProTLS*)Platform::GetTLSValue(GetFrameProTLSSlot());
 		return p_framepro_tls ? p_framepro_tls : CreateFrameProTLS();
 	}
 
 	//------------------------------------------------------------------------
-	void DebugWrite(const FRAMEPRO_TCHAR* p_str, ...);
+	void DebugWrite(const char* p_str, ...);
 
 	//------------------------------------------------------------------------
 	inline bool IsPow2(int value)
@@ -212,41 +117,6 @@ namespace FramePro
 		return a > b ? a : b;
 	}
 
-	//------------------------------------------------------------------------
-	template<class T>
-	inline T* New(Allocator* p_allocator)
-	{
-		T* p = (T*)p_allocator->Alloc(sizeof(T));
-		new (p)T();
-		return p;
-	}
-
-	//------------------------------------------------------------------------
-	template<class T, typename Targ1>
-	inline T* New(Allocator* p_allocator, Targ1 arg1)
-	{
-		T* p = (T*)p_allocator->Alloc(sizeof(T));
-		new (p)T(arg1);
-		return p;
-	}
-
-	//------------------------------------------------------------------------
-	template<class T, typename Targ1, typename Targ2, typename TArg3>
-	inline T* New(Allocator* p_allocator, Targ1 arg1, Targ2 arg2, TArg3 arg3)
-	{
-		T* p = (T*)p_allocator->Alloc(sizeof(T));
-		new (p)T(arg1, arg2, arg3);
-		return p;
-	}
-
-	//------------------------------------------------------------------------
-	template<typename T>
-	inline void Delete(Allocator* p_allocator, T* p)
-	{
-		p->~T();
-		p_allocator->Free(p);
-	}
-
 	//------------------------------------------------------------------------
 	template<typename T>
 	inline void Swap(T& a, T& b)
@@ -331,104 +201,20 @@ namespace FramePro
 	};
 
 	//------------------------------------------------------------------------
-	#if !FRAMEPRO_WIN_BASED_PLATFORM
-		#define sprintf_s sprintf
-
-		inline void strncpy_s(char *p_dest, size_t element_count, const char *p_source, size_t count)
-		{
-			(void)(element_count);
-			strncpy(p_dest, p_source, count);
-		}
-
-		inline void wcsncpy_s(wchar_t *p_dest, size_t element_count, const wchar_t *p_source, size_t count)
-		{
-			(void)(element_count);
-			wcsncpy(p_dest, p_source, count);
-		}
-
-		inline void strcpy_s(char *p_dest, const char *p_source)
-		{
-			strcpy(p_dest, p_source);
-		}
-
-		inline void strcpy_s(char *p_dest, size_t dest_length, const char *p_source)
-		{
-			(void)(dest_length);
-			strcpy(p_dest, p_source);
-		}
-
-		inline void _vstprintf_s(FRAMEPRO_TCHAR* const p_buffer, size_t const buffer_size, FRAMEPRO_TCHAR const* const p_format, va_list arg_list)
-		{
-			(void)(buffer_size);
-			vsprintf(p_buffer, p_format, arg_list);
-		}
-
-		inline void OutputDebugString(const FRAMEPRO_TCHAR* p_string)
-		{
-			printf("%s", p_string);
-		}
-
-		inline uint64 GetCurrentThreadId()
-		{
-// START EPIC
-#if FRAMEPRO_UE4_BASED_PLATFORM
-			return FPlatformTLS::GetCurrentThreadId();
-#else
-			return (uint64)pthread_self();
-#endif
-// END EPIC
-		}
-
-		inline int fopen_s(FILE **pp_file, const char* p_filename, const char* p_mode)
-		{
-			*pp_file = fopen(p_filename, p_mode);
-			return *pp_file ? 0 : 1;
-		}
-
-		inline int GetCurrentProcessId()
-		{
-// START EPIC
-#if FRAMEPRO_UE4_BASED_PLATFORM
-			return FPlatformProcess::GetCurrentProcessId();
-#else
-			return getpid();
-#endif
-// END EPIC
-		}
-
-		inline void localtime_s(struct tm* p_tm, const time_t *p_time)
-		{
-			tm* p_local_tm = localtime(p_time);
-			*p_tm = *p_local_tm;
-		}
-	#endif
-
-	//------------------------------------------------------------------------
-	#if FRAMEPRO_WIN_BASED_PLATFORM
-		#define FRAMEPRO_SOCKET SOCKET
-		#define FRAMEPRO_INVALID_SOCKET INVALID_SOCKET
-		#define FRAMEPRO_MAX_PATH MAX_PATH
-		#define FRAMEPRO_STRING _T
-		#define FRAMEPRO_SOCKET_ERROR SOCKET_ERROR
-// START EPIC
-	#else
-// END EPIC
-		typedef int FRAMEPRO_SOCKET;
-		#define FRAMEPRO_INVALID_SOCKET -1
-		#define FRAMEPRO_MAX_PATH 256
-		#define FRAMEPRO_STRING
-		#define FRAMEPRO_SOCKET_ERROR -1
-	#endif
+	void SPrintf(char* p_buffer, size_t const buffer_size, const char* p_format, ...);
 }
 
+
 //------------------------------------------------------------------------
-#endif		// #if FRAMEPRO_ENABLED
+//
+// EventTraceWin32.hpp
+//
 
 //------------------------------------------------------------------------
-#endif		// #ifndef FRAMEPROLIB_H_INCLUDED
+#if FRAMEPRO_PLATFORM_WIN
 
 //------------------------------------------------------------------------
-#if FRAMEPRO_ENABLED
+#if FRAMEPRO_ENABLE_CONTEXT_SWITCH_TRACKING
 
 //------------------------------------------------------------------------
 namespace FramePro
@@ -455,32 +241,35 @@ namespace FramePro
 
 		void Flush();
 
+		static void* Create(Allocator* p_allocator);
+
+		static void Destroy(void* p_context_switch_recorder, Allocator* p_allocator);
+
+		static bool Start(void* p_context_switch_recorder, Platform::ContextSwitchCallbackFunction p_callback, void* p_context, DynamicString& error);
+
+		static void Stop(void* p_context_switch_recorder);
+	
+		static void Flush(void* p_context_switch_recorder);
+
 		//------------------------------------------------------------------------
 		// data
 	private:
 		EventTraceWin32Imp* mp_Imp;
-
 		Allocator* mp_Allocator;
 	};
 }
 
 //------------------------------------------------------------------------
-#endif		// #if FRAMEPRO_ENABLED
+#endif		// #if FRAMEPRO_ENABLE_CONTEXT_SWITCH_TRACKING
 
 //------------------------------------------------------------------------
-#endif		// #ifndef FRAMEPRO_EVENTTRACEWIN32_H_INCLUDED
+#endif		// #if FRAMEPRO_PLATFORM_WIN
 
 
 //------------------------------------------------------------------------
+//
 // CriticalSection.hpp
-#ifndef FRAMEPRO_CRITICALSECTION_H_INCLUDED
-#define FRAMEPRO_CRITICALSECTION_H_INCLUDED
-
-//------------------------------------------------------------------------
-
-
-//------------------------------------------------------------------------
-#if FRAMEPRO_ENABLED
+//
 
 //------------------------------------------------------------------------
 namespace FramePro
@@ -489,111 +278,67 @@ namespace FramePro
 	class CriticalSection
 	{
 	public:
+		//------------------------------------------------------------------------
 		CriticalSection()
 		{
-#if FRAMEPRO_WIN_BASED_PLATFORM
-			InitializeSRWLock(&lock);
-// START EPIC
-#elif FRAMEPRO_UE4_BASED_PLATFORM
-			//...
-// END EPIC
-#else
-			pthread_mutexattr_t attr;
-			pthread_mutexattr_init(&attr);
-			pthread_mutexattr_settype(&attr, PTHREAD_MUTEX_RECURSIVE);
-			pthread_mutex_init(&cs, &attr);
-#endif
+			Platform::CreateLock(m_OSLockMem, sizeof(m_OSLockMem));
 
-#if FRAMEPRO_DEBUG
-			m_Locked = false;
-
-			#if FRAMEPRO_WIN_BASED_PLATFORM
-				m_LockedOnThread = 0xffffffff;
+			#if FRAMEPRO_DEBUG
+				m_Locked = false;
+				m_LockedOnThread = 0xffffffffffffffffULL;
 			#endif
-#endif
 		}
 
+		//------------------------------------------------------------------------
 		~CriticalSection()
 		{
-// START EPIC
-#if !FRAMEPRO_WIN_BASED_PLATFORM && !FRAMEPRO_UE4_BASED_PLATFORM
-			pthread_mutex_destroy(&cs);
-#endif
-// END EPIC
+			Platform::DestroyLock(m_OSLockMem);
 		}
 
+		//------------------------------------------------------------------------
 		void Enter()
 		{
-			#if FRAMEPRO_WIN_BASED_PLATFORM
-				FRAMEPRO_ASSERT(GetCurrentThreadId() != m_LockedOnThread);
-			#endif
+			FRAMEPRO_ASSERT(Platform::GetCurrentThreadId() != m_LockedOnThread);
 
-#if FRAMEPRO_WIN_BASED_PLATFORM
-			AcquireSRWLockExclusive(&lock);
-// START EPIC
-#elif FRAMEPRO_UE4_BASED_PLATFORM
-			cs.Lock();
-// END EPIC
-#else
-			pthread_mutex_lock(&cs);
-#endif
+			Platform::TakeLock(m_OSLockMem);
 
-#if FRAMEPRO_DEBUG
-			m_Locked = true;
-			#if FRAMEPRO_WIN_BASED_PLATFORM
-				m_LockedOnThread = GetCurrentThreadId();
+			#if FRAMEPRO_DEBUG
+				m_Locked = true;
+				m_LockedOnThread = Platform::GetCurrentThreadId();
 			#endif
-#endif
 		}
 
+		//------------------------------------------------------------------------
 		void Leave()
 		{
-#if FRAMEPRO_DEBUG
-			m_Locked = false;
-			#if FRAMEPRO_WIN_BASED_PLATFORM
-				m_LockedOnThread = 0xffffffff;
+			#if FRAMEPRO_DEBUG
+				m_Locked = false;
+				m_LockedOnThread = 0xffffffffffffffffULL;
 			#endif
-#endif
 
-#if FRAMEPRO_WIN_BASED_PLATFORM
-			ReleaseSRWLockExclusive(&lock);
-// START EPIC
-#elif FRAMEPRO_UE4_BASED_PLATFORM
-			cs.Unlock();
-// END EPIC
-#else
-			pthread_mutex_unlock(&cs);
-#endif
+			Platform::ReleaseLock(m_OSLockMem);
 		}
 
 		//------------------------------------------------------------------------
-#if FRAMEPRO_DEBUG
-		bool Locked() const		// only safe to use in an assert to check that it IS locked
-		{
-			return m_Locked;
-		}
-#endif
+		#if FRAMEPRO_DEBUG
+			bool Locked() const		// only safe to use in an assert to check that it IS locked
+			{
+				return m_Locked;
+			}
+		#endif
 
 		//------------------------------------------------------------------------
 		// data
 	private:
-#if FRAMEPRO_WIN_BASED_PLATFORM
-		SRWLOCK lock;
-// START EPIC
-#elif FRAMEPRO_UE4_BASED_PLATFORM
-		FCriticalSection cs;
-// END EPIC
-#else
-		pthread_mutex_t cs;
-#endif
+		static const int m_OSLockMaxSize = 40;
+
+		char m_OSLockMem[m_OSLockMaxSize];
 
 #if FRAMEPRO_DEBUG
 		RelaxedAtomic<bool> m_Locked;
-		#if FRAMEPRO_WIN_BASED_PLATFORM
-			unsigned int m_LockedOnThread;
-		#endif
+		uint64 m_LockedOnThread;
 #endif
-	};
+	} FRAMEPRO_ALIGN_STRUCT(16);
 
 	//------------------------------------------------------------------------
 	class CriticalSectionScope
@@ -608,48 +353,38 @@ namespace FramePro
 	};
 }
 
-//------------------------------------------------------------------------
-#endif		// #if FRAMEPRO_ENABLED
-
-//------------------------------------------------------------------------
-#endif		// #ifndef FRAMEPRO_CRITICALSECTION_H_INCLUDED
 
 //------------------------------------------------------------------------
+//
 // HashMap.hpp
-#ifndef HASHMAP_H_INCLUDED
-#define HASHMAP_H_INCLUDED
-
-//------------------------------------------------------------------------
+//
 
 //------------------------------------------------------------------------
-#if FRAMEPRO_ENABLED
-
-//-----------------------------------------------------------------
 #include <new>
 
-//-----------------------------------------------------------------
-#define FRAMEPRO_PROFILE_HashMap 0
+//------------------------------------------------------------------------
+#define FRAMEPRO_PROFILE_HASHMAP 0
 
-//-----------------------------------------------------------------
+//------------------------------------------------------------------------
 namespace FramePro
 {
-	//-----------------------------------------------------------------
+	//------------------------------------------------------------------------
 	typedef unsigned int uint;
 	typedef char byte;
 
-	//-----------------------------------------------------------------
+	//------------------------------------------------------------------------
 	template<typename TKey, typename TValue>
 	class HashMap
 	{
 	public:
-		//-----------------------------------------------------------------
+		//------------------------------------------------------------------------
 		struct Pair
 		{
 			TKey m_Key;
 			TValue m_Value;
 		};
 
-		//-----------------------------------------------------------------
+		//------------------------------------------------------------------------
 		// The default capacity of the set. The capacity is the number
 		// of elements that the set is expected to hold. The set will resized
 		// when the item count is greater than the capacity;
@@ -660,15 +395,15 @@ namespace FramePro
 			mp_ItemPool(NULL),
 			mp_FreePair(NULL),
 			mp_Allocator(p_allocator)
-#if FRAMEPRO_PROFILE_HashMap
-			,m_IterAcc(0)
-			,m_IterCount(0)
-#endif
+			#if FRAMEPRO_PROFILE_HASHMAP
+				,m_IterAcc(0)
+				,m_IterCount(0)
+			#endif
 		{
 			AllocTable(GetNextPow2((256 * m_DefaultCapacity) / m_Margin));
 		}
 
-		//-----------------------------------------------------------------
+		//------------------------------------------------------------------------
 		~HashMap()
 		{
 			Clear();
@@ -677,13 +412,13 @@ namespace FramePro
 			FreePools();
 		}
 
-		//-----------------------------------------------------------------
+		//------------------------------------------------------------------------
 		void Clear()
 		{
 			RemoveAll();
 		}
 
-		//-----------------------------------------------------------------
+		//------------------------------------------------------------------------
 		void RemoveAll()
 		{
 			for(int i=0; i<m_Capacity; ++i)
@@ -698,7 +433,7 @@ namespace FramePro
 			m_Count = 0;
 		}
 
-		//-----------------------------------------------------------------
+		//------------------------------------------------------------------------
 		// Add a value to this set.
 		// If this set already contains the value does nothing.
 		void Add(const TKey& key, const TValue& value)
@@ -729,7 +464,7 @@ namespace FramePro
 			}
 		}
 
-		//-----------------------------------------------------------------
+		//------------------------------------------------------------------------
 		// if this set contains the value set value to the existing value and
 		// return true, otherwise set to the default value and return false.
 		bool TryGetValue(const TKey& key, TValue& value) const
@@ -749,13 +484,13 @@ namespace FramePro
 			}
 		}
 
-		//-----------------------------------------------------------------
+		//------------------------------------------------------------------------
 		int GetCount() const
 		{
 			return m_Count;
 		}
 
-		//-----------------------------------------------------------------
+		//------------------------------------------------------------------------
 		void Resize(int new_capacity)
 		{
 			new_capacity = GetNextPow2(new_capacity);
@@ -782,7 +517,7 @@ namespace FramePro
 			mp_Allocator->Free(p_old_table);
 		}
 
-		//-----------------------------------------------------------------
+		//------------------------------------------------------------------------
 		size_t GetMemorySize() const
 		{
 			size_t table_memory = m_Capacity * sizeof(Pair*);
@@ -799,7 +534,7 @@ namespace FramePro
 		}
 
 	private:
-		//-----------------------------------------------------------------
+		//------------------------------------------------------------------------
 		static int GetNextPow2(int value)
 		{
 			int p = 2;
@@ -808,7 +543,7 @@ namespace FramePro
 			return p;
 		}
 
-		//-----------------------------------------------------------------
+		//------------------------------------------------------------------------
 		void AllocTable(const int capacity)
 		{
 			FRAMEPRO_ASSERT(capacity < m_MaxCapacity);
@@ -823,13 +558,13 @@ namespace FramePro
 			}
 		}
 
-		//-----------------------------------------------------------------
+		//------------------------------------------------------------------------
 		bool IsItemInUse(const int index) const
 		{
 			return mp_Table[index] != NULL;
 		}
 
-		//-----------------------------------------------------------------
+		//------------------------------------------------------------------------
 		int GetItemIndex(const TKey& key) const
 		{
 			FRAMEPRO_ASSERT(mp_Table);
@@ -838,31 +573,34 @@ namespace FramePro
 			while(IsItemInUse(srch_index) && !(mp_Table[srch_index]->m_Key == key))
 			{
 				srch_index = (srch_index + 1) & (m_Capacity-1);
-#if FRAMEPRO_PROFILE_HashMap
-				++m_IterAcc;
-#endif
+				#if FRAMEPRO_PROFILE_HASHMAP
+					++m_IterAcc;
+				#endif
 			}
 
-#if FRAMEPRO_PROFILE_HashMap
-			++m_IterCount;
-			double average = m_IterAcc / (double)m_IterCount;
-			if(average > 2.0)
-			{
-				static int last_write_time = 0;
-				int now = GetTickCount();
-				if(now - last_write_time > 1000)
+			#if FRAMEPRO_PROFILE_HASHMAP
+				++m_IterCount;
+				double average = m_IterAcc / (double)m_IterCount;
+				if(average > 2.0)
 				{
-					last_write_time = now;
-					FRAMEPRO_TCHAR temp[64];
-					_stprintf_s(temp, FRAMEPRO_STRING("WARNING: HashMap average: %f\n"), (float)average);
-					OutputDebugString(temp);
+					static int64 last_write_time = 0;
+					int64 now = 0;
+					FRAMEPRO_GET_CLOCK_COUNT(now);
+					static int64 timer_freq = Platform::GetTimerFrequency();
+					if(now - last_write_time > timer_freq)
+					{
+						last_write_time = now;
+						char temp[64];
+						SPrintf(temp, sizeof(temp), "WARNING: HashMap average: %f\n", (float)average);
+						Platform::DebugWrite(temp);
+					}
 				}
-			}
-#endif
+			#endif
+
 			return srch_index;
 		}
 
-		//-----------------------------------------------------------------
+		//------------------------------------------------------------------------
 		static bool InRange(
 			const int index,
 			const int start_index,
@@ -873,7 +611,7 @@ namespace FramePro
 				index >= start_index || index <= end_index;
 		}
 
-		//-----------------------------------------------------------------
+		//------------------------------------------------------------------------
 		void FreePools()
 		{
 			byte* p_pool = mp_ItemPool;
@@ -887,7 +625,7 @@ namespace FramePro
 			mp_FreePair = NULL;
 		}
 
-		//-----------------------------------------------------------------
+		//------------------------------------------------------------------------
 		Pair* AllocPair()
 		{
 			if(!mp_FreePair)
@@ -919,7 +657,7 @@ namespace FramePro
 			return p_pair;
 		}
 
-		//-----------------------------------------------------------------
+		//------------------------------------------------------------------------
 		void FreePair(Pair* p_pair)
 		{
 			p_pair->~Pair();
@@ -928,7 +666,7 @@ namespace FramePro
 			mp_FreePair = (byte*)p_pair;
 		}
 
-		//-----------------------------------------------------------------
+		//------------------------------------------------------------------------
 		// data
 	private:
 		enum { m_DefaultCapacity = 32 };
@@ -946,37 +684,18 @@ namespace FramePro
 
 		Allocator* mp_Allocator;
 
-#if FRAMEPRO_PROFILE_HashMap
-		mutable int64 m_IterAcc;
-		mutable int64 m_IterCount;
-#endif
+		#if FRAMEPRO_PROFILE_HASHMAP
+			mutable int64 m_IterAcc;
+			mutable int64 m_IterCount;
+		#endif
 	};
 }
 
-//-----------------------------------------------------------------
-#endif		// #if FRAMEPRO_ENABLED
-
-//-----------------------------------------------------------------
-#endif		// #ifndef HASHMAP_H_INCLUDED
-
-//------------------------------------------------------------------------
-// FrameProString.hpp
-#ifndef FRAMEPROSTRING_H_INCLUDED
-#define FRAMEPROSTRING_H_INCLUDED
-
-//------------------------------------------------------------------------
-
 
 //------------------------------------------------------------------------
+//
 // IncrementingBlockAllocator.hpp
-#ifndef INCREMENTINGBLOCKALLOCATOR_H_INCLUDED
-#define INCREMENTINGBLOCKALLOCATOR_H_INCLUDED
-
-//------------------------------------------------------------------------
-
-
-//------------------------------------------------------------------------
-#if FRAMEPRO_ENABLED
+//
 
 //------------------------------------------------------------------------
 namespace FramePro
@@ -1024,14 +743,11 @@ namespace FramePro
 	};
 }
 
-//------------------------------------------------------------------------
-#endif		// #if FRAMEPRO_ENABLED
-
-//------------------------------------------------------------------------
-#endif		// #ifndef INCREMENTINGBLOCKALLOCATOR_H_INCLUDED
 
 //------------------------------------------------------------------------
-#if FRAMEPRO_ENABLED
+//
+// FrameProString.hpp
+//
 
 //------------------------------------------------------------------------
 #include <string.h>
@@ -1107,6 +823,26 @@ namespace FramePro
 		return MurmurHash2(p_str, (int)wcslen(p_str) * sizeof(wchar_t), prime);
 	}
 
+	//------------------------------------------------------------------------
+	inline void StringCopy(char* p_dest, size_t dest_size, const char* p_source)
+	{
+		size_t length = strlen(p_source) + 1;
+		FRAMEPRO_ASSERT(length <= dest_size);
+		FRAMEPRO_UNREFERENCED(dest_size);
+		memcpy(p_dest, p_source, length);
+	}
+
+	//------------------------------------------------------------------------
+	inline void StringCopy(char* p_dest, size_t dest_size, const char* p_source, size_t source_length)
+	{
+		size_t length = strlen(p_source) + 1;
+		FRAMEPRO_ASSERT(length <= dest_size && source_length < length);
+		FRAMEPRO_UNREFERENCED(length);
+		FRAMEPRO_UNREFERENCED(dest_size);
+		memcpy(p_dest, p_source, source_length);
+		p_dest[source_length] = '\0';
+	}
+
 	//------------------------------------------------------------------------
 	// this string class is meant to be as light weight as possible and does
 	// not clean up its allocations.
@@ -1137,10 +873,10 @@ namespace FramePro
 			const char* p_old_value = mp_Value;
 			size_t len = strlen(p_old_value);
 			char* p_new_value = (char*)allocator.Alloc(len+1);
-			strncpy_s(p_new_value, len+1, mp_Value, len);
+			StringCopy(p_new_value, len+1, mp_Value, len);
 			mp_Value = p_new_value;
 #else
-			(void)allocator;
+			FRAMEPRO_UNREFERENCED(allocator);
 #endif
 		}
 
@@ -1200,10 +936,10 @@ namespace FramePro
 			const wchar_t* p_old_value = mp_Value;
 			size_t len = wcslen(p_old_value);
 			wchar_t* p_new_value = (wchar_t*)allocator.Alloc((len+1)*sizeof(wchar_t));
-			wcsncpy_s(p_new_value, len+1, mp_Value, len);
+			StringCopy(p_new_value, len+1, mp_Value, len);
 			mp_Value = p_new_value;
 #else
-			(void)allocator;
+			FRAMEPRO_UNREFERENCED(allocator);
 #endif
 		}
 
@@ -1249,7 +985,7 @@ namespace FramePro
 			FRAMEPRO_ASSERT(!mp_Value);
 			size_t len = strlen(p_value);
 			mp_Value = (char*)mp_Allocator->Alloc(len + 1);
-			strncpy_s(mp_Value, len + 1, p_value, len);
+			StringCopy(mp_Value, len + 1, p_value, len);
 		}
 
 		//------------------------------------------------------------------------
@@ -1266,11 +1002,11 @@ namespace FramePro
 			{
 				size_t len = strlen(mp_Value);
 				len = FramePro_Min(len, max_length - 1);
-				strncpy_s(p_dest, max_length, mp_Value, len);
+				StringCopy(p_dest, max_length, mp_Value, len);
 			}
 			else
 			{
-				strcpy_s(p_dest, max_length, "");
+				StringCopy(p_dest, max_length, "");
 			}
 		}
 
@@ -1280,36 +1016,130 @@ namespace FramePro
 		char* mp_Value;
 		Allocator* mp_Allocator;
 	};
-}
 
-//------------------------------------------------------------------------
-#endif		// #if FRAMEPRO_ENABLED
-
-//------------------------------------------------------------------------
-#endif		// #ifndef FRAMEPROSTRING_H_INCLUDED
-
-//------------------------------------------------------------------------
-#if FRAMEPRO_ENABLED
-
-//------------------------------------------------------------------------
-#if FRAMEPRO_EVENT_TRACE_WIN32
-	#include <wmistr.h>
-	#define INITGUID
-	#include <evntrace.h>
-	#include <evntcons.h>
-	#include <tchar.h>
-	#include <Tdh.h>
-
-	#pragma comment(lib, "tdh.lib")
-	#pragma comment(lib, "Advapi32.lib")
-#endif
-
-//------------------------------------------------------------------------
-#if FRAMEPRO_EVENT_TRACE_WIN32
-namespace FramePro
-{
 	//------------------------------------------------------------------------
-	struct ThreadIdKey
+	class DynamicWString
+	{
+	public:
+		//------------------------------------------------------------------------
+		DynamicWString()
+		:	mp_Value(NULL),
+			mp_Allocator(NULL)
+		{
+		}
+
+		//------------------------------------------------------------------------
+		~DynamicWString()
+		{
+			if (mp_Value)
+			{
+				FRAMEPRO_ASSERT(mp_Allocator);
+				mp_Allocator->Free(mp_Value);
+			}
+		}
+		//------------------------------------------------------------------------
+		void SetAllocator(Allocator* p_allocator)
+		{
+			FRAMEPRO_ASSERT(!mp_Allocator);
+			FRAMEPRO_ASSERT(p_allocator);
+
+			mp_Allocator = p_allocator;
+		}
+
+		//------------------------------------------------------------------------
+		void Clear()
+		{
+			if (mp_Value)
+			{
+				FRAMEPRO_ASSERT(mp_Allocator);
+				mp_Allocator->Free(mp_Value);
+				mp_Value = NULL;
+			}
+		}
+
+		//------------------------------------------------------------------------
+		void operator=(const char* p_value)
+		{
+			FRAMEPRO_ASSERT(!mp_Value);
+			FRAMEPRO_ASSERT(mp_Allocator);
+
+			int len = (int)strlen(p_value);
+			mp_Value = (wchar_t*)mp_Allocator->Alloc((len + 1) * sizeof(wchar_t));
+			for(int i = 0; i < len; ++i)
+				mp_Value[i] = p_value[i];
+			mp_Value[len - 1] = L'\0';
+		}
+
+		//------------------------------------------------------------------------
+		void operator=(const wchar_t* p_value)
+		{
+			FRAMEPRO_ASSERT(!mp_Value);
+			FRAMEPRO_ASSERT(mp_Allocator);
+
+			int len = (int)wcslen(p_value);
+			mp_Value = (wchar_t*)mp_Allocator->Alloc((len + 1) * sizeof(wchar_t));
+			for(int i = 0; i < len; ++i)
+				mp_Value[i] = p_value[i];
+			mp_Value[len - 1] = L'\0';
+		}
+
+		//------------------------------------------------------------------------
+		const wchar_t* c_str() const
+		{
+			return mp_Value ? mp_Value : L"";
+		}
+
+		//------------------------------------------------------------------------
+		// data
+	private:
+		wchar_t* mp_Value;
+		Allocator* mp_Allocator;
+	};
+}
+
+
+//------------------------------------------------------------------------
+//
+// EventTraceWin32.cpp
+//
+
+//------------------------------------------------------------------------
+#if FRAMEPRO_PLATFORM_WIN
+
+//------------------------------------------------------------------------
+#if FRAMEPRO_ENABLE_CONTEXT_SWITCH_TRACKING
+
+//------------------------------------------------------------------------
+#if FRAMEPRO_PLATFORM_UE4
+	#include "Windows/AllowWindowsPlatformTypes.h"
+#endif
+#pragma warning(push)
+#pragma warning(disable:4668)
+#ifndef WIN32_LEAN_AND_MEAN
+	#define WIN32_LEAN_AND_MEAN
+#endif
+#include <windows.h>
+#include <wmistr.h>
+#define INITGUID
+#include <evntrace.h>
+#include <evntcons.h>
+#include <tchar.h>
+#include <Tdh.h>
+#pragma warning(pop)
+#if FRAMEPRO_PLATFORM_UE4
+	#include "Windows/HideWindowsPlatformTypes.h"
+#endif
+
+#include <intrin.h>
+
+#pragma comment(lib, "tdh.lib")
+#pragma comment(lib, "Advapi32.lib")
+
+//------------------------------------------------------------------------
+namespace FramePro
+{
+	//------------------------------------------------------------------------
+	struct ThreadIdKey
 	{
 		ThreadIdKey() :	m_ThreadId(0) {}
 
@@ -1554,9 +1384,7 @@ namespace FramePro
 			default:
 			{
 				char temp[128];
-// START EPIC 
 				sprintf_s(temp, "Error code: %ld", error_code);
-// END EPIC 
 				error_string = temp;
 			}
 		}
@@ -1679,18 +1507,13 @@ namespace FramePro
 		#endif
 	}
 }
-#endif		// #if FRAMEPRO_EVENT_TRACE_WIN32
 
 //------------------------------------------------------------------------
 namespace FramePro
 {
 	//------------------------------------------------------------------------
 	EventTraceWin32::EventTraceWin32(Allocator* p_allocator)
-#if FRAMEPRO_EVENT_TRACE_WIN32
 	:	mp_Imp(New<EventTraceWin32Imp>(p_allocator, p_allocator)),
-#else
-	:	mp_Imp(NULL),
-#endif
 		mp_Allocator(p_allocator)
 	{
 	}
@@ -1698,78 +1521,107 @@ namespace FramePro
 	//------------------------------------------------------------------------
 	EventTraceWin32::~EventTraceWin32()
 	{
-#if FRAMEPRO_EVENT_TRACE_WIN32
 		g_ShuttingDown = true;
 
 		Delete(mp_Allocator, mp_Imp);
-#endif
 	}
 
 	//------------------------------------------------------------------------
 	bool EventTraceWin32::Start(ContextSwitchCallback p_context_switch_callback, void* p_context_switch_callback_param, DynamicString& error)
 	{
-#if FRAMEPRO_EVENT_TRACE_WIN32
 		return mp_Imp->Start(p_context_switch_callback, p_context_switch_callback_param, error);
-#else
-		return false;
-#endif
 	}
 
 	//------------------------------------------------------------------------
 	void EventTraceWin32::Stop()
 	{
-#if FRAMEPRO_EVENT_TRACE_WIN32
 		mp_Imp->Stop();
-#endif
 	}
 
 	//------------------------------------------------------------------------
 	void EventTraceWin32::Flush()
 	{
-#if FRAMEPRO_EVENT_TRACE_WIN32
 		mp_Imp->Flush();
-#endif
 	}
-}
 
-//------------------------------------------------------------------------
-#endif		// #if FRAMEPRO_ENABLED
-//------------------------------------------------------------------------
-// FramePro.cpp
+	//------------------------------------------------------------------------
+	void* EventTraceWin32::Create(Allocator* p_allocator)
+	{
+		return New<EventTraceWin32>(p_allocator, p_allocator);
+	}
+
+	//------------------------------------------------------------------------
+	void EventTraceWin32::Destroy(void* p_context_switch_recorder, Allocator* p_allocator)
+	{
+		if (p_context_switch_recorder)
+		{
+			EventTraceWin32* p_event_trace_win32 = (EventTraceWin32*)p_context_switch_recorder;
+			Delete(p_allocator, p_event_trace_win32);
+		}
+	}
 
+	//------------------------------------------------------------------------
+	bool EventTraceWin32::Start(
+		void* p_context_switch_recorder,
+		Platform::ContextSwitchCallbackFunction p_callback,
+		void* p_context,
+		DynamicString& error)
+	{
+		if (!p_context_switch_recorder)
+			return false;
 
-//------------------------------------------------------------------------
-// FrameProTLS.hpp
-#ifndef FRAMEPROTLS_H_INCLUDED
-#define FRAMEPROTLS_H_INCLUDED
+		EventTraceWin32* p_event_trace_win32 = (EventTraceWin32*)p_context_switch_recorder;
 
-//------------------------------------------------------------------------
+		bool started = p_event_trace_win32->Start(p_callback, p_context, error);
 
+		if(!started)
+			FramePro::DebugWrite("FramePro Warning: Failed to start recording context switches. Please make sure that you are running with administrator privileges.\n");
 
-//------------------------------------------------------------------------
-#if FRAMEPRO_ENABLED
+		return started;
+	}
 
-//------------------------------------------------------------------------
+	//------------------------------------------------------------------------
+	void EventTraceWin32::Stop(void* p_context_switch_recorder)
+	{
+		if (p_context_switch_recorder)
+		{
+			EventTraceWin32* p_event_trace_win32 = (EventTraceWin32*)p_context_switch_recorder;
+			p_event_trace_win32->Stop();
+		}
+	}
+
+	//------------------------------------------------------------------------
+	void EventTraceWin32::Flush(void* p_context_switch_recorder)
+	{
+		if (p_context_switch_recorder)
+		{
+			EventTraceWin32* p_event_trace_win32 = (EventTraceWin32*)p_context_switch_recorder;
+			p_event_trace_win32->Flush();
+		}
+	}
+}
 
 //------------------------------------------------------------------------
-// Socket.hpp
-#ifndef FRAMEPRO_SOCKET_H_INCLUDED
-#define FRAMEPRO_SOCKET_H_INCLUDED
+#endif			// #if FRAMEPRO_ENABLE_CONTEXT_SWITCH_TRACKING
 
 //------------------------------------------------------------------------
+#endif		// #if FRAMEPRO_PLATFORM_WIN
 
 
 //------------------------------------------------------------------------
-#if FRAMEPRO_ENABLED
+//
+// FrameProTLS.hpp
+//
 
 //------------------------------------------------------------------------
-#if FRAMEPRO_SOCKETS_ENABLED
 
 //------------------------------------------------------------------------
-#if FRAMEPRO_WIN_BASED_PLATFORM
-	#pragma warning(push)
-	#pragma warning(disable : 4100)
-#endif
+//
+// Socket.hpp
+//
+
+//------------------------------------------------------------------------
+#if FRAMEPRO_SOCKETS_ENABLED
 
 //------------------------------------------------------------------------
 namespace FramePro
@@ -1781,9 +1633,9 @@ namespace FramePro
 	class Socket
 	{
 	public:
-		inline Socket();
+		Socket();
 
-		inline ~Socket();
+		~Socket();
 
 		void Disconnect();
 
@@ -1797,110 +1649,30 @@ namespace FramePro
 
 		bool Send(const void* p_buffer, size_t size);
 
-		inline bool IsValid() const { return m_Socket != FRAMEPRO_INVALID_SOCKET; }
-
-		static void HandleError();
-
-	private:
-		bool InitialiseWSA();
-
-		void CleanupWSA();
+		bool IsValid() const;
 
 		//------------------------------------------------------------------------
 		// data
-		FRAMEPRO_SOCKET m_Socket;
+	private:
+		static const int m_OSSocketMemMaxSize = 8;
+		char m_OSSocketMem[m_OSSocketMemMaxSize];
 
 		bool m_Listening;
 	};
-
-	//------------------------------------------------------------------------
-	Socket::Socket()
-	:	m_Socket(FRAMEPRO_INVALID_SOCKET),
-		m_Listening(false)
-	{
-	}
-
-	//------------------------------------------------------------------------
-	Socket::~Socket()
-	{
-		CleanupWSA();
-	}
 }
 
-//------------------------------------------------------------------------
-#if FRAMEPRO_WIN_BASED_PLATFORM
-	#pragma warning(pop)
-#endif
-
 //------------------------------------------------------------------------
 #endif		// #if FRAMEPRO_SOCKETS_ENABLED
 
-//------------------------------------------------------------------------
-#endif		// #if FRAMEPRO_ENABLED
-
-//------------------------------------------------------------------------
-#endif		// #ifndef FRAMEPRO_SOCKET_H_INCLUDED
 
 //------------------------------------------------------------------------
+//
 // Packets.hpp
-#ifndef PACKETS_H_INCLUDED
-#define PACKETS_H_INCLUDED
-
-//------------------------------------------------------------------------
-
-
-//------------------------------------------------------------------------
-#if FRAMEPRO_ENABLED
+//
 
 //------------------------------------------------------------------------
 namespace FramePro
 {
-	//------------------------------------------------------------------------
-	namespace PacketType
-	{
-		enum Enum
-		{
-			Connect = 0xaabb,
-			FrameStart,
-			TimeSpan,
-			TimeSpanW,
-			NamedTimeSpan,
-			StringLiteralNamedTimeSpan,
-			ThreadName,
-			ThreadOrder,
-			StringPacket,
-			WStringPacket,
-			NameAndSourceInfoPacket,
-			NameAndSourceInfoPacketW,
-			SourceInfoPacket,
-			MainThreadPacket,
-			RequestStringLiteralPacket,
-			SetConditionalScopeMinTimePacket,
-			ConnectResponsePacket,
-			SessionInfoPacket,
-			RequestRecordedDataPacket,
-			SessionDetailsPacket,
-			ContextSwitchPacket,
-			ContextSwitchRecordingStartedPacket,
-			ProcessNamePacket,
-			CustomStatPacket,
-			StringLiteralTimerNamePacket,
-			HiResTimerScopePacket,
-			LogPacket,
-			EventPacket,
-			StartWaitEventPacket,
-			StopWaitEventPacket,
-			TriggerWaitEventPacket,
-			TimeSpanCustomStatPacket,
-			TimeSpanWithCallstack,
-			TimeSpanWWithCallstack,
-			NamedTimeSpanWithCallstack,
-			StringLiteralNamedTimeSpanWithCallstack,
-			ModulePacket,
-			SetCallstackRecordingEnabledPacket,
-		};
-	};
-
 	//------------------------------------------------------------------------
 	// send packets
 	//------------------------------------------------------------------------
@@ -1911,19 +1683,6 @@ namespace FramePro
 		#pragma clang diagnostic ignored "-Wunused-private-field"
 	#endif
 
-	//------------------------------------------------------------------------
-	namespace Platform
-	{
-		enum Enum
-		{
-			Windows = 0,
-			Windows_UWP,
-			XBoxOne,
-			XBox360,
-			Unix,
-		};
-	}
-
 	//------------------------------------------------------------------------
 	namespace CustomStatValueType
 	{
@@ -1992,7 +1751,6 @@ namespace FramePro
 		int m_ValueType;	// CustomStatValueType enum
 		int m_Padding;
 		StringId m_Name;
-		StringId m_Unit;
 		int64 m_ValueInt64;
 		double m_ValueDouble;
 		int64 m_Time;
@@ -2155,8 +1913,6 @@ namespace FramePro
 		int m_Count;
 		StringId m_Name;
 		int64 m_Value;
-		StringId m_Graph;
-		StringId m_Unit;
 	};
 
 	//------------------------------------------------------------------------
@@ -2166,8 +1922,6 @@ namespace FramePro
 		int m_Count;
 		StringId m_Name;
 		double m_Value;
-		StringId m_Graph;
-		StringId m_Unit;
 	};
 
 	//------------------------------------------------------------------------
@@ -2227,16 +1981,28 @@ namespace FramePro
 	};
 
 	//------------------------------------------------------------------------
-	struct ModulePacket
+	struct ScopeColourPacket
 	{
 		PacketType::Enum m_PacketType;
-		int m_UseLookupFunctionForBaseAddress;
-		int64 m_ModuleBase;
-		char m_Sig[16];
-		int m_Age;
-		int m_Padding;
-		char m_ModuleName[FRAMEPRO_MAX_INLINE_STRING_LENGTH];
-		char m_SymbolFilename[FRAMEPRO_MAX_INLINE_STRING_LENGTH];
+		uint m_Colour;
+		StringId m_Name;
+	};
+
+	//------------------------------------------------------------------------
+	struct CustomStatInfoPacket
+	{
+		PacketType::Enum m_PacketType;
+		uint m_Padding;
+		StringId m_Name;
+		StringId m_Value;
+	};
+
+	//------------------------------------------------------------------------
+	struct CustomStatColourPacket
+	{
+		PacketType::Enum m_PacketType;
+		uint m_Colour;
+		StringId m_Name;
 	};
 
 	//------------------------------------------------------------------------
@@ -2281,22 +2047,11 @@ namespace FramePro
 	#endif
 }
 
-//------------------------------------------------------------------------
-#endif		// #if FRAMEPRO_ENABLED
-
-//------------------------------------------------------------------------
-#endif		// #ifndef PACKETS_H_INCLUDED
 
 //------------------------------------------------------------------------
+//
 // PointerSet.hpp
-#ifndef FRAMEPRO_POINTERSET_H_INCLUDED
-#define FRAMEPRO_POINTERSET_H_INCLUDED
-
-//------------------------------------------------------------------------
-
-
-//------------------------------------------------------------------------
-#if FRAMEPRO_ENABLED
+//
 
 //------------------------------------------------------------------------
 #define FRAMEPRO_PRIME 0x01000193
@@ -2353,24 +2108,11 @@ namespace FramePro
 	};
 }
 
-//------------------------------------------------------------------------
-#endif		// #if FRAMEPRO_ENABLED
-
-//------------------------------------------------------------------------
-#endif		// #ifndef FRAMEPRO_POINTERSET_H_INCLUDED
-
-
 
 //------------------------------------------------------------------------
+//
 // SendBuffer.hpp
-#ifndef SENDBUFFER_H_INCLUDED
-#define SENDBUFFER_H_INCLUDED
-
-//------------------------------------------------------------------------
-
-
-//------------------------------------------------------------------------
-#if FRAMEPRO_ENABLED
+//
 
 //------------------------------------------------------------------------
 namespace FramePro
@@ -2441,24 +2183,11 @@ namespace FramePro
 	};
 }
 
-//------------------------------------------------------------------------
-#endif		// #if FRAMEPRO_ENABLED
-
-//------------------------------------------------------------------------
-#endif		// #ifndef SENDBUFFER_H_INCLUDED
-
-
 
 //------------------------------------------------------------------------
+//
 // Buffer.hpp
-#ifndef BUFFER_H_INCLUDED
-#define BUFFER_H_INCLUDED
-
-//------------------------------------------------------------------------
-
-
-//------------------------------------------------------------------------
-#if FRAMEPRO_ENABLED
+//
 
 //------------------------------------------------------------------------
 namespace FramePro
@@ -2555,31 +2284,11 @@ namespace FramePro
 	};
 }
 
-//------------------------------------------------------------------------
-#endif		// #if FRAMEPRO_ENABLED
-
-//------------------------------------------------------------------------
-#endif		// #ifndef BUFFER_H_INCLUDED
-
-//------------------------------------------------------------------------
-// ConditionalParentScope.hpp
-#ifndef FRAMEPRO_CONDITIONALPARENTSCOPE_H_INCLUDED
-#define FRAMEPRO_CONDITIONALPARENTSCOPE_H_INCLUDED
-
-//------------------------------------------------------------------------
-
-
 
 //------------------------------------------------------------------------
+//
 // List.hpp
-#ifndef FRAMEPRO_LIST_H_INCLUDED
-#define FRAMEPRO_LIST_H_INCLUDED
-
-//------------------------------------------------------------------------
-
-
-//------------------------------------------------------------------------
-#if FRAMEPRO_ENABLED
+//
 
 //------------------------------------------------------------------------
 namespace FramePro
@@ -2730,14 +2439,11 @@ namespace FramePro
 	};
 }
 
-//------------------------------------------------------------------------
-#endif		// #if FRAMEPRO_ENABLED
-
-//------------------------------------------------------------------------
-#endif		// #ifndef FRAMEPRO_LIST_H_INCLUDED
 
 //------------------------------------------------------------------------
-#if FRAMEPRO_ENABLED
+//
+// ConditionalParentScope.hpp
+//
 
 //------------------------------------------------------------------------
 namespace FramePro
@@ -2771,180 +2477,11 @@ namespace FramePro
 	};
 }
 
-//------------------------------------------------------------------------
-#endif		// #if FRAMEPRO_ENABLED
-
-//------------------------------------------------------------------------
-#endif		// #ifndef FRAMEPRO_CONDITIONALPARENTSCOPE_H_INCLUDED
-
 
 //------------------------------------------------------------------------
-// Array.hpp
-#ifndef ARRAY_H_INCLUDED
-#define ARRAY_H_INCLUDED
-
-//------------------------------------------------------------------------
-
-
-//------------------------------------------------------------------------
-#if FRAMEPRO_ENABLED
-
-//------------------------------------------------------------------------
-namespace FramePro
-{
-	//--------------------------------------------------------------------
-	template<typename T>
-	class Array
-	{
-	public:
-		//--------------------------------------------------------------------
-		Array()
-		:	mp_Array(NULL),
-			m_Count(0),
-			m_Capacity(0),
-			mp_Allocator(NULL)
-		{
-		}
-
-		//--------------------------------------------------------------------
-		~Array()
-		{
-			FRAMEPRO_ASSERT(!mp_Array);
-		}
-
-		//--------------------------------------------------------------------
-		int GetCount() const
-		{
-			return m_Count;
-		}
-
-		//--------------------------------------------------------------------
-		void Clear()
-		{
-			if(mp_Array)
-			{
-				mp_Allocator->Free(mp_Array);
-				mp_Array = NULL;
-			}
-
-			m_Count = 0;
-			m_Capacity = 0;
-		}
-
-		//--------------------------------------------------------------------
-		void ClearNoFree()
-		{
-			m_Count = 0;
-		}
-
-		//--------------------------------------------------------------------
-		void SetAllocator(Allocator* p_allocator)
-		{
-			FRAMEPRO_ASSERT(mp_Allocator == p_allocator || !(mp_Allocator != NULL && p_allocator != NULL));
-			mp_Allocator = p_allocator;
-		}
-
-		//--------------------------------------------------------------------
-		void Add(const T& value)
-		{
-			if(m_Count == m_Capacity)
-				Grow();
-
-			mp_Array[m_Count] = value;
-			++m_Count;
-		}
-
-		//--------------------------------------------------------------------
-		const T& operator[](int index) const
-		{
-			FRAMEPRO_ASSERT(index >= 0 && index < m_Count);
-			return mp_Array[index];
-		}
-
-		//--------------------------------------------------------------------
-		T& operator[](int index)
-		{
-			FRAMEPRO_ASSERT(index >= 0 && index < m_Count);
-			return mp_Array[index];
-		}
-
-		//--------------------------------------------------------------------
-		void RemoveAt(int index)
-		{
-			FRAMEPRO_ASSERT(index >= 0 && index < m_Count);
-		
-			if(index < m_Count - 1)
-				memmove(mp_Array + index, mp_Array + index + 1, (m_Count - 1 - index) * sizeof(T));
-
-			--m_Count;
-		}
-
-		//--------------------------------------------------------------------
-		T RemoveLast()
-		{
-			FRAMEPRO_ASSERT(m_Count);
-
-			return mp_Array[--m_Count];
-		}
-
-		//--------------------------------------------------------------------
-		bool Contains(const T& value) const
-		{
-			for(int i=0; i<m_Count; ++i)
-				if(mp_Array[i] == value)
-					return true;
-			return false;
-		}
-
-	private:
-		//--------------------------------------------------------------------
-		void Grow()
-		{
-			m_Capacity = m_Capacity ? 2*m_Capacity : 32;
-			T* p_new_array = (T*)mp_Allocator->Alloc(sizeof(T)*m_Capacity);
-			if(mp_Array)
-				memcpy(p_new_array, mp_Array, sizeof(T)*m_Count);
-			mp_Allocator->Free(mp_Array);
-			mp_Array = p_new_array;
-		}
-
-		//------------------------------------------------------------------------
-		// data
-	private:
-		T* mp_Array;
-		
-		int m_Count;
-		int m_Capacity;
-
-		Allocator* mp_Allocator;
-	};
-}
-
-//------------------------------------------------------------------------
-#endif		// #if FRAMEPRO_ENABLED
-
-//------------------------------------------------------------------------
-#endif		// #ifndef ARRAY_H_INCLUDED
-
-//------------------------------------------------------------------------
-// FrameProStackTrace.hpp
-#ifndef FRAMEPRO_STACKTRACE_H_INCLUDED
-#define FRAMEPRO_STACKTRACE_H_INCLUDED
-
-//------------------------------------------------------------------------
-
-
-//------------------------------------------------------------------------
-// FrameProCallstackSet.hpp
-#ifndef FRAMEPRO_CALLSTACKSET_H_INCLUDED
-#define FRAMEPRO_CALLSTACKSET_H_INCLUDED
-
-//------------------------------------------------------------------------
-
-
-
-//------------------------------------------------------------------------
-#if FRAMEPRO_ENABLED
+//
+// FrameProCallstackSet.hpp
+//
 
 //------------------------------------------------------------------------
 namespace FramePro
@@ -2994,21 +2531,15 @@ namespace FramePro
 	};
 }
 
-//------------------------------------------------------------------------
-#endif		// #if FRAMEPRO_ENABLED
-
-//------------------------------------------------------------------------
-#endif		// #ifndef FRAMEPRO_CALLSTACKSET_H_INCLUDED
 
 //------------------------------------------------------------------------
-#if FRAMEPRO_ENABLED
+//
+// FrameProStackTrace.hpp
+//
 
 //------------------------------------------------------------------------
 #if FRAMEPRO_ENABLE_CALLSTACKS
 
-//------------------------------------------------------------------------
-#define FRAMEPRO_STACK_TRACE_SIZE 128
-
 //------------------------------------------------------------------------
 namespace FramePro
 {
@@ -3044,12 +2575,6 @@ namespace FramePro
 //------------------------------------------------------------------------
 #endif		// #if FRAMEPRO_ENABLE_CALLSTACKS
 
-//------------------------------------------------------------------------
-#endif		// #if FRAMEPRO_ENABLED
-
-//------------------------------------------------------------------------
-#endif		// #ifndef FRAMEPRO_STACKTRACE_H_INCLUDED
-#include <atomic>
 
 //------------------------------------------------------------------------
 namespace FramePro
@@ -3192,6 +2717,14 @@ namespace FramePro
 
 		void SendEventPacket(const char* p_name, uint colour);
 
+		void SendScopeColourPacket(StringId name, uint colour);
+		
+		void SendCustomStatGraphPacket(StringId name, StringId graph);
+
+		void SendCustomStatUnitPacket(StringId name, StringId unit);
+
+		void SendCustomStatColourPacket(StringId name, uint colour);
+
 		FRAMEPRO_FORCE_INLINE void StartHiResTimer(const char* p_name)
 		{
 			FRAMEPRO_ASSERT(IsOnTLSThread());
@@ -3282,13 +2815,23 @@ namespace FramePro
 		int64 GetWaitEventMinTime() const { return m_WaitEventMinTime; }
 #endif
 
-		void SetCustomTimeSpanStat(StringId name, int64 value, const char* p_unit);
+		void SetCustomTimeSpanStat(StringId name, int64 value);
+
+		void SetCustomTimeSpanStat(StringId name, double value);
+
+		void SetCustomTimeSpanStatW(StringId name, int64 value);
+
+		void SetCustomTimeSpanStatW(StringId name, double value);
+
+		void SetCustomStatInfo(const char* p_name, const char* p_graph, const char* p_unit, uint colour);
+		
+		void SetCustomStatInfo(const wchar_t* p_name, const wchar_t* p_graph, const wchar_t* p_unit, uint colour);
 
-		void SetCustomTimeSpanStat(StringId name, int64 value, const wchar_t* p_unit);
+		void SetCustomStatInfo(StringId name, const char* p_graph, const char* p_unit, uint colour);
 
-		void SetCustomTimeSpanStat(StringId name, double value, const char* p_unit);
+		void SetCustomStatInfo(StringId name, const wchar_t* p_graph, const wchar_t* p_unit, uint colour);
 
-		void SetCustomTimeSpanStat(StringId name, double value, const wchar_t* p_unit);
+		void SetCustomStatInfo(StringId name, StringId graph, StringId unit, uint colour);
 
 #if FRAMEPRO_ENABLE_CALLSTACKS
 		bool ShouldSendCallstacks() const { return m_SendCallstacks; }
@@ -3344,8 +2887,12 @@ namespace FramePro
 
 		void SendRootHiResTimerList();
 
+		bool HaveSentCustomStatInfo(StringId name);
+
+		void SetHaveSentCustomStatInfo(StringId name);
+
 #if FRAMEPRO_DEBUG
-		bool IsOnTLSThread() const { return GetCurrentThreadId() == m_OSThreadId; }
+		bool IsOnTLSThread() const { return Platform::GetCurrentThreadId() == m_ThreadId; }
 #endif
 		template<typename T>
 		void DeleteListItems(List<T>& list)
@@ -3381,14 +2928,6 @@ namespace FramePro
 		int m_CurrentSendBufferSize;					// access must be protected with m_CurrentSendBufferCriticalSection
 		
 		int m_ThreadId;
-		uint64 m_OSThreadId;
-
-// START EPIC
-		#if !FRAMEPRO_WIN_BASED_PLATFORM && !FRAMEPRO_UE4_BASED_PLATFORM
-// END EPIC
-
-			static int m_NewThreadId;
-		#endif
 
 		int64 m_HiResTimerScopeStartTime;
 
@@ -3415,6 +2954,8 @@ namespace FramePro
 		HashMap<String, StringId> m_StringHashMap;
 		HashMap<WString, StringId> m_WStringHashMap;
 
+		Array<bool> m_InitialisedCustomStats;
+
 		Buffer m_SessionInfoBuffer;
 		CriticalSection m_SessionInfoBufferLock;
 		RelaxedAtomic<size_t> m_SessionInfoBufferMemorySize;
@@ -3425,8 +2966,6 @@ namespace FramePro
 
 		IncrementingBlockAllocator m_StringAllocator;
 
-		char m_TempStringBuffer[FRAMEPRO_MAX_PATH];
-
 		RelaxedAtomic<size_t> m_SendBufferMemorySize;
 		RelaxedAtomic<size_t> m_StringMemorySize;
 
@@ -3450,40 +2989,11 @@ namespace FramePro
 	};
 }
 
-//------------------------------------------------------------------------
-#endif		// #if FRAMEPRO_ENABLED
-
-//------------------------------------------------------------------------
-#endif		// #ifndef FRAMEPROTLS_H_INCLUDED
-
-//------------------------------------------------------------------------
-// FrameProSession.hpp
-#ifndef FRAMEPROSESSION_H_INCLUDED
-#define FRAMEPROSESSION_H_INCLUDED
-
-//------------------------------------------------------------------------
-
-
-
-
-//------------------------------------------------------------------------
-// Thread.hpp
-#ifndef FRAMEPRO_THREAD_H_INCLUDED
-#define FRAMEPRO_THREAD_H_INCLUDED
-
-//------------------------------------------------------------------------
-
 
 //------------------------------------------------------------------------
+//
 // Event.hpp
-#ifndef FRAMEPRO_EVENT_H_INCLUDED
-#define FRAMEPRO_EVENT_H_INCLUDED
-
-//------------------------------------------------------------------------
-
-
-//------------------------------------------------------------------------
-#if FRAMEPRO_ENABLED
+//
 
 //------------------------------------------------------------------------
 namespace FramePro
@@ -3495,196 +3005,57 @@ namespace FramePro
 		//--------------------------------------------------------------------
 		Event(bool initial_state, bool auto_reset)
 		{
-#if FRAMEPRO_WIN_BASED_PLATFORM
-			m_Handle = CreateEvent(NULL, !auto_reset, initial_state, NULL);
-// START EPIC
-#elif FRAMEPRO_UE4_BASED_PLATFORM
-			m_Event = FPlatformProcess::CreateSynchEvent( !auto_reset );
-			if( initial_state )
-				Set();
-// END EPIC
-#else
-			pthread_cond_init(&m_Cond, NULL);
-			pthread_mutex_init(&m_Mutex, NULL);
-			m_Signalled = false;
-			m_AutoReset = auto_reset;
-	
-			if(initial_state)
-				Set();
-#endif
+			Platform::CreateEventX(m_OSEventMem, sizeof(m_OSEventMem), initial_state, auto_reset);
 		}
 
 		//--------------------------------------------------------------------
 		~Event()
 		{
-#if FRAMEPRO_WIN_BASED_PLATFORM
-			CloseHandle(m_Handle);
-// START EPIC
-#elif FRAMEPRO_UE4_BASED_PLATFORM
-			delete m_Event;
-			m_Event = nullptr;
-// END EPIC
-#else
-			pthread_mutex_destroy(&m_Mutex);
-			pthread_cond_destroy(&m_Cond);
-#endif
+			Platform::DestroyEvent(m_OSEventMem);
 		}
 
 		//--------------------------------------------------------------------
 		void Set() const
 		{
-#if FRAMEPRO_WIN_BASED_PLATFORM
-			SetEvent(m_Handle);
-// START EPIC
-#elif FRAMEPRO_UE4_BASED_PLATFORM
-			m_Event->Trigger();
-// END EPIC
-#else
-			pthread_mutex_lock(&m_Mutex);
-			m_Signalled = true;
-			pthread_mutex_unlock(&m_Mutex);
-			pthread_cond_signal(&m_Cond);
-#endif
+			Platform::SetEvent(m_OSEventMem);
 		}
 
 		//--------------------------------------------------------------------
 		void Reset()
 		{
-#if FRAMEPRO_WIN_BASED_PLATFORM
-			ResetEvent(m_Handle);
-// START EPIC
-#elif FRAMEPRO_UE4_BASED_PLATFORM
-			m_Event->Reset();
-// END EPIC
-#else
-			pthread_mutex_lock(&m_Mutex);
-			m_Signalled = false;
-			pthread_mutex_unlock(&m_Mutex);
-#endif
+			Platform::ResetEvent(m_OSEventMem);
 		}
 
 		//--------------------------------------------------------------------
 		int Wait(int timeout=-1) const
 		{
-#if FRAMEPRO_WIN_BASED_PLATFORM
-			return WaitForSingleObject(m_Handle, timeout) == 0/*WAIT_OBJECT_0*/;
-// START EPIC
-#elif FRAMEPRO_UE4_BASED_PLATFORM
-			return (timeout==-1) ? m_Event->Wait() : m_Event->Wait( timeout );
-// END EPIC
-#else
-			pthread_mutex_lock(&m_Mutex);
-	
-			if(m_Signalled)
-			{
-				m_Signalled = false;
-				pthread_mutex_unlock(&m_Mutex);
-				return true;
-			}
-	
-			if(timeout == -1)
-			{
-				while(!m_Signalled)
-					pthread_cond_wait(&m_Cond, &m_Mutex);
-		
-				if(!m_AutoReset)
-					m_Signalled = false;
-
-				pthread_mutex_unlock(&m_Mutex);
-		
-				return true;
-			}
-			else
-			{		
-				timeval curr;
-				gettimeofday(&curr, NULL);
-		
-				timespec time;
-				time.tv_sec  = curr.tv_sec + timeout / 1000;
-				time.tv_nsec = (curr.tv_usec * 1000) + ((timeout % 1000) * 1000000);
-		
-				time.tv_sec += time.tv_nsec / 1000000000L;
-				time.tv_nsec = time.tv_nsec % 1000000000L;
-
-				int ret = 0;
-				do
-				{
-					ret = pthread_cond_timedwait(&m_Cond, &m_Mutex, &time);
-
-				} while(!m_Signalled && ret != ETIMEDOUT);
-		
-				if(m_Signalled)
-				{
-					if(!m_AutoReset)
-						m_Signalled = false;
-
-					pthread_mutex_unlock(&m_Mutex);
-					return true;
-				}
-		
-				pthread_mutex_unlock(&m_Mutex);
-				return false;
-			}
-#endif
+			return Platform::WaitEvent(m_OSEventMem, timeout);
 		}
 
 		//------------------------------------------------------------------------
 		// data
 	private:
-#if FRAMEPRO_WIN_BASED_PLATFORM
-		HANDLE m_Handle;
-// START EPIC
-#elif FRAMEPRO_UE4_BASED_PLATFORM
-		FEvent* m_Event;
-// END EPIC
-#else
-		mutable pthread_cond_t  m_Cond;
-		mutable pthread_mutex_t m_Mutex;
-		mutable volatile bool m_Signalled;
-		bool m_AutoReset;
-#endif
+		static const int m_OSEventMemMaxSize = 96;
+		mutable char m_OSEventMem[m_OSEventMemMaxSize];
 	};
 }
 
-//------------------------------------------------------------------------
-#endif		// #if FRAMEPRO_ENABLED
-
-//------------------------------------------------------------------------
-#endif		// #ifndef FRAMEPRO_EVENT_H_INCLUDED
-
-//------------------------------------------------------------------------
-#if FRAMEPRO_ENABLED
-
-//------------------------------------------------------------------------
-#if FRAMEPRO_WIN_BASED_PLATFORM
-	#pragma warning(push)
-	#pragma warning(disable : 4100)
-#endif
 
 //------------------------------------------------------------------------
-#if !FRAMEPRO_WIN_BASED_PLATFORM
-	#include <pthread.h>
-#endif
+//
+// Thread.hpp
+//
 
 //------------------------------------------------------------------------
 namespace FramePro
 {
-	//------------------------------------------------------------------------
-	typedef int (*ThreadMain)(void*);
-
 	//------------------------------------------------------------------------
 	class Thread
-// BEGIN EPIC
-#if FRAMEPRO_UE4_BASED_PLATFORM
-	: public FRunnable
-#endif		
-// END EPIC
-
 	{
 	public:
 		Thread();
 
-		void CreateThread(ThreadMain p_thread_main, void* p_param=NULL);
+		void CreateThread(ThreadMain p_thread_main, void* p_param, Allocator* p_allocator);
 
 		bool IsAlive() const { return m_Alive; }
 
@@ -3695,72 +3066,155 @@ namespace FramePro
 		void WaitForThreadToTerminate(int timeout);
 
 	private:
-		#if FRAMEPRO_WIN_BASED_PLATFORM
-			static unsigned long WINAPI PlatformThreadMain(void* p_param);
-		// BEGIN EPIC
-		#elif FRAMEPRO_UE4_BASED_PLATFORM
-			virtual uint32 Run() override;
-		// END EPIC
-		#else
-			static void* PlatformThreadMain(void* p_param);
-		#endif
+		static int ThreadMain(void* p_context);
 
 		//------------------------------------------------------------------------
 		// data
 	private:
-#if FRAMEPRO_WIN_BASED_PLATFORM
-		mutable HANDLE m_Handle;
-// BEGIN EPIC
-#elif FRAMEPRO_UE4_BASED_PLATFORM
-		mutable TUniquePtr<FRunnableThread> m_Runnable;
-// END EPIC
-#else
-		mutable pthread_t m_Thread;
-#endif
-		mutable bool m_Alive;
+		static const int m_OSThreadMaxSize = 16;
+		char m_OSThread[m_OSThreadMaxSize];
 
-		mutable ThreadMain mp_ThreadMain;
-		mutable void* mp_Param;
+		bool m_Alive;
+
+		FramePro::ThreadMain mp_ThreadMain;
+		void* mp_Param;
 
 		Event m_ThreadTerminatedEvent;
 	};
 }
 
-//------------------------------------------------------------------------
-#if FRAMEPRO_WIN_BASED_PLATFORM
-	#pragma warning(pop)
-#endif
-
-//------------------------------------------------------------------------
-#endif		// #if FRAMEPRO_ENABLED
-
-//------------------------------------------------------------------------
-#endif		// #ifndef FRAMEPRO_THREAD_H_INCLUDED
-
-
-
-
-
-
-
-
-//------------------------------------------------------------------------
-#if FRAMEPRO_ENABLED
 
 //------------------------------------------------------------------------
-#include <atomic>
+//
+// File.hpp
+//
 
 //------------------------------------------------------------------------
 namespace FramePro
 {
 	//------------------------------------------------------------------------
-	class FrameProTLS;
 	class Allocator;
-	struct ThreadCustomStats;
+
+	//------------------------------------------------------------------------
+	class File
+	{
+	public:
+		//------------------------------------------------------------------------
+		File()
+		:	m_Opened(false)
+		{
+		}
+
+		//------------------------------------------------------------------------
+		void SetAllocator(Allocator* p_allocator)
+		{
+			m_Filename.SetAllocator(p_allocator);
+		}
+
+		//------------------------------------------------------------------------
+		bool OpenForRead(const char* p_filename)
+		{
+			FRAMEPRO_ASSERT(!m_Opened);
+			m_Filename = p_filename;
+			m_Opened = Platform::OpenFileForRead(m_OSFile, m_OSFileMaxSize, p_filename);
+			return m_Opened;
+		}
+
+		//------------------------------------------------------------------------
+		bool OpenForRead(const wchar_t* p_filename)
+		{
+			FRAMEPRO_ASSERT(!m_Opened);
+			m_Filename = p_filename;
+			m_Opened = Platform::OpenFileForRead(m_OSFile, m_OSFileMaxSize, p_filename);
+			return m_Opened;
+		}
+
+		//------------------------------------------------------------------------
+		bool OpenForWrite(const char* p_filename)
+		{
+			FRAMEPRO_ASSERT(!m_Opened);
+			m_Filename = p_filename;
+			m_Opened = Platform::OpenFileForWrite(m_OSFile, m_OSFileMaxSize, p_filename);
+			return m_Opened;
+		}
+
+		//------------------------------------------------------------------------
+		bool OpenForWrite(const wchar_t* p_filename)
+		{
+			FRAMEPRO_ASSERT(!m_Opened);
+			m_Filename = p_filename;
+			m_Opened = Platform::OpenFileForWrite(m_OSFile, m_OSFileMaxSize, p_filename);
+			return m_Opened;
+		}
+
+		//------------------------------------------------------------------------
+		void Close()
+		{
+			FRAMEPRO_ASSERT(m_Opened);
+			m_Filename.Clear();
+			Platform::CloseFile(m_OSFile);
+			m_Opened = false;
+		}
+
+		//------------------------------------------------------------------------
+		void Read(void* p_data, size_t size)
+		{
+			FRAMEPRO_ASSERT(m_Opened);
+			Platform::ReadFromFile(m_OSFile, p_data, size);
+		}
+
+		//------------------------------------------------------------------------
+		void Write(const void* p_data, size_t size)
+		{
+			FRAMEPRO_ASSERT(m_Opened);
+			Platform::WriteToFile(m_OSFile, p_data, size);
+		}
+
+		//------------------------------------------------------------------------
+		bool IsOpened() const
+		{
+			return m_Opened;
+		}
+
+		//------------------------------------------------------------------------
+		size_t GetSize() const
+		{
+			FRAMEPRO_ASSERT(m_Opened);
+			return Platform::GetFileSize(m_OSFile);
+		}
+
+		//------------------------------------------------------------------------
+		const DynamicWString& GetFilename() const
+		{
+			return m_Filename;
+		}
+
+		//------------------------------------------------------------------------
+		// data
+	private:
+		static const int m_OSFileMaxSize = 16;
+		char m_OSFile[m_OSFileMaxSize];
+
+		bool m_Opened;
+
+		DynamicWString m_Filename;
+	};
+}
+
+
+//------------------------------------------------------------------------
+//
+// FrameProSession.hpp
+//
+
+//------------------------------------------------------------------------
+namespace FramePro
+{
+	//------------------------------------------------------------------------
+	class FrameProTLS;
+	class Allocator;
+	struct ThreadCustomStats;
 	class SendBuffer;
-#if FRAMEPRO_EVENT_TRACE_WIN32
-	class EventTraceWin32;
-#endif
 
 	//------------------------------------------------------------------------
 	class FrameProSession
@@ -3796,9 +3250,9 @@ namespace FramePro
 
 		void SetThreadName(const char* p_name);
 
-// START EPIC
-		void StartRecording(const FString& p_filename, bool context_switches, int64 max_file_size);
-// END EPIC
+		void StartRecording(const char* p_filename, bool context_switches, bool callstacks, int64 max_file_size);
+
+		void StartRecording(const wchar_t* p_filename, bool context_switches, bool callstacks, int64 max_file_size);
 
 		void StopRecording();
 
@@ -3818,31 +3272,34 @@ namespace FramePro
 
 		bool CallConditionalParentScopeCallback(ConditionalParentScopeCallback p_callback, const char* p_name, int64 start_time, int64 end_time);
 
-		void EnumerateLoadedModulesCallback(
-			int64 module_base,
-			const char* p_module_name,
-			bool m_UseLookupFunctionForBaseAddress);
-
 		void SetConditionalScopeMinTimeInMicroseconds(int64 value);
 
+		void SetScopeColour(StringId name, uint colour);
+
+		void SetCustomStatGraph(StringId name, StringId graph);
+
+		void SetCustomStatUnit(StringId name, StringId unit);
+
+		void SetCustomStatColour(StringId name, uint colour);
+
 	private:
 		void Initialise(FrameProTLS* p_framepro_tls);
 
+		void StopRecording_NoLock();
+
 		void SendSessionDetails(StringId name, StringId build_id);
 
 		Allocator* CreateDefaultAllocator();
 
+		void SetAllocatorInternal(Allocator* p_allocator);
+
 		void InitialiseConnection(FrameProTLS* p_framepro_tls);
 
 		FRAMEPRO_FORCE_INLINE void CalculateTimerFrequency();
 
 		void SetConnected(bool value);
 
-		bool SendSendBuffer(SendBuffer* p_send_buffer, Socket& socket);
-
-// START EPIC
-		void WriteSendBuffer(SendBuffer* p_send_buffer, FArchive* p_file, int64& file_size);
-// END EPIC
+		void WriteSendBuffer(SendBuffer* p_send_buffer, File& file, int64& file_size);
 		
 		void SendFrameBuffer();
 
@@ -3850,9 +3307,9 @@ namespace FramePro
 
 		int SendThreadMain();
 
-		void HandleDisconnect();
+		void HandleDisconnect(bool wait_for_threads_to_exit=true);
 
-		void HandleDisconnect_NoLock();
+		void HandleDisconnect_NoLock(bool wait_for_threads_to_exit=true);
 
 		void SendRecordedDataAndDisconnect();
 
@@ -3862,9 +3319,8 @@ namespace FramePro
 
 		bool HasSetThreadName(int thread_id) const;
 
-// START EPIC
-		void OnConnectionChanged(const FString& ProfileName, bool connected) const;
-// END EPIC
+		void OnConnectionChanged(bool connected, const DynamicWString& filename) const;
+
 		int GetConnectionChangedCallbackIndex(ConnectionChangedCallback p_callback);
 
 		size_t GetMemoryUsage() const;
@@ -3881,37 +3337,49 @@ namespace FramePro
 
 		void ClearGlobalHiResTimers();
 
-		void EnumerateModules();
-
 		void SendExtraModuleInfo(int64 ModuleBase, FrameProTLS* p_framepro_tls);
 
-#if FRAMEPRO_SOCKETS_ENABLED
-		bool InitialiseFileCache();
+		#if FRAMEPRO_SOCKETS_ENABLED
+			bool SendSendBuffer(SendBuffer* p_send_buffer, Socket& socket);
 
-		static int StaticConnectThreadMain(void*);
+			bool InitialiseFileCache();
 
-		int ConnectThreadMain();
+			static int StaticConnectThreadMain(void*);
 
-		void SendOnMainThread(void* p_src, int size);
+			int ConnectThreadMain();
+
+			static int StaticReceiveThreadMain(void*);
 
-		static int StaticReceiveThreadMain(void*);
+			int OnReceiveThreadExit();
 
-		int OnReceiveThreadExit();
+			int ReceiveThreadMain();
 
-		int ReceiveThreadMain();
+			void OpenListenSocket();
 
-		void OpenListenSocket();
+			void StartConnectThread();
 
-		void StartConnectThread();
+			void CreateReceiveThread();
+		#endif
 
-		void CreateReceiveThread();
+		void SendOnMainThread(void* p_src, int size);
 
 		template<typename T>
 		void SendOnMainThread(T& packet)
 		{
 			SendOnMainThread(&packet, sizeof(packet));
 		}
-#endif
+
+		#if FRAMEPRO_ENABLE_CALLSTACKS
+			void SetCallstacksEnabled(bool enabled);
+		#endif
+
+		void SendScopeColours();
+
+		void SendCustomStatGraphs();
+		
+		void SendCustomStatUnits();
+
+		void SendCustomStatColours();
 
 		//------------------------------------------------------------------------
 		// data
@@ -3931,6 +3399,10 @@ namespace FramePro
 
 		std::atomic<bool> m_StartContextSwitchRecording;
 
+		#if FRAMEPRO_ENABLE_CALLSTACKS
+			std::atomic<bool> m_StartRecordingCallstacks;
+		#endif
+			
 		int64 m_ClockFrequency;
 
 		mutable CriticalSection m_TLSListCriticalSection;
@@ -3949,18 +3421,14 @@ namespace FramePro
 		CriticalSection m_SendFrameBufferCriticalSection;
 
 		RelaxedAtomic<bool> m_Interactive;
-// START EPIC
-		FArchive* mp_NonInteractiveRecordingFile;
-// END EPIC
+		File m_NonInteractiveRecordingFile;
 		int64 m_NonInteractiveRecordingFileSize;
 
 		int64 m_LastSessionInfoSendTime;
 
 		Array<int> m_NamedThreads;
 
-// START EPIC
-		FArchive* mp_RecordingFile;
-// END EPIC
+		File m_RecordingFile;
 		int64 m_RecordingFileSize;
 		int64 m_MaxRecordingFileSize;
 
@@ -3969,11 +3437,11 @@ namespace FramePro
 		bool m_ThreadAffinitySet;
 		int m_ThreadAffinity;
 
-#if FRAMEPRO_SOCKETS_ENABLED
-		Thread m_ConnectThread;
-		Socket m_ListenSocket;
-		Socket m_ClientSocket;
-#endif
+		#if FRAMEPRO_SOCKETS_ENABLED
+			Thread m_ConnectThread;
+			Socket m_ListenSocket;
+			Socket m_ClientSocket;
+		#endif
 
 		std::atomic<bool> m_SendThreadExit;
 		Event m_SendThreadFinished;
@@ -4002,37 +3470,38 @@ namespace FramePro
 
 		Array<ModulePacket*> m_ModulePackets;
 
-#if FRAMEPRO_EVENT_TRACE_WIN32
-		EventTraceWin32* mp_EventTraceWin32;
-#endif
-
-#if FRAMEPRO_ENABLE_CALLSTACKS
-		bool m_SendModules;
-#endif
-	};
-}
+		void* mp_ContextSwitchRecorder;
 
-//------------------------------------------------------------------------
-#endif		// #if FRAMEPRO_ENABLED
+		struct ScopeColour
+		{
+			StringId m_Name;
+			uint m_Colour;
+		};
+		Array<ScopeColour> m_ScopeColours;
+		CriticalSection m_ScopeColoursLock;
 
-//------------------------------------------------------------------------
-#endif		// #ifndef FRAMEPROSESSION_H_INCLUDED
+		struct CustomStatInfo
+		{
+			StringId m_Name;
+			StringId m_Value;
+		};
 
-//------------------------------------------------------------------------
-#if FRAMEPRO_ENABLED
+		Array<CustomStatInfo> m_CustomStatGraphs;
+		Array<CustomStatInfo> m_CustomStatUnits;
+		Array<ScopeColour> m_CustomStatColours;
+		CriticalSection m_CustomStatInfoLock;
 
-//------------------------------------------------------------------------
-#include <atomic>
+		#if FRAMEPRO_ENABLE_CALLSTACKS
+			bool m_SendModules;
+		#endif
+	};
+}
 
-//------------------------------------------------------------------------
-#if FRAMEPRO_WIN_BASED_PLATFORM
-	#include <intrin.h>
-#endif
 
 //------------------------------------------------------------------------
-#if FRAMEPRO_SOCKETS_ENABLED && FRAMEPRO_WIN_BASED_PLATFORM
-	#pragma comment(lib, "Ws2_32.lib")
-#endif
+//
+// FramePro.cpp
+//
 
 //------------------------------------------------------------------------
 namespace FramePro
@@ -4043,43 +3512,6 @@ namespace FramePro
 	RelaxedAtomic<unsigned int> g_ConditionalScopeMinTime = UINT_MAX;
 
 	//------------------------------------------------------------------------
-
-// START EPIC
-#if FRAMEPRO_PLATFORM_ANDROID
-	#include <sys/syscall.h>
-#elif FRAMEPRO_PLATFORM_MAC
-	#include <cpuid.h>
-#endif
-
-	FRAMEPRO_FORCE_INLINE int GetCore()
-	{
-	#if FRAMEPRO_PLATFORM_XBOXONE || FRAMEPRO_PLATFORM_UWP
-		int cpu_info[4];
-		__cpuid(cpu_info, 1);
-		return (cpu_info[1] >> 24) & 0xff;
-	#elif FRAMEPRO_PLATFORM_WIN
-		return GetCurrentProcessorNumber();
-	#elif FRAMEPRO_PLATFORM_UNIX
-		return sched_getcpu();
-	#elif FRAMEPRO_PLATFORM_ANDROID
-		unsigned cpu;
-		int err = syscall(__NR_getcpu, &cpu, NULL, NULL);
-		return (!err) ? (int)cpu : 0;
-	#elif FRAMEPRO_PLATFORM_IOS
-		return 0; // TODO
-	#elif FRAMEPRO_PLATFORM_MAC
-		int cpu_info[4];
-		__cpuid(1, cpu_info[0], cpu_info[1], cpu_info[2], cpu_info[3]);
-		return (cpu_info[1] >> 24) & 0xff;
-	#elif FRAMEPRO_UE4_BASED_PLATFORM
-		return FPlatformProcess::GetCurrentCoreNumber();
-	#else
-		#error
-	#endif
-	}
-// END EPIC
-
-	//------------------------------------------------------------------------
 	FrameProSession& GetFrameProSession()
 	{
 		static FrameProSession session;
@@ -4087,9 +3519,13 @@ namespace FramePro
 	}
 		
 	//------------------------------------------------------------------------
-// BEGIN EPIC
-//	FRAMEPRO_THREAD_LOCAL FrameProTLS* gp_FrameProTLS = NULL;
-// END EPIC
+	#if FRAMEPRO_USE_TLS_SLOTS
+		uint GetFrameProTLSSlot()
+		{
+			static uint slot = Platform::AllocateTLSSlot();
+			return slot;
+		}
+	#endif
 
 	//------------------------------------------------------------------------
 	FRAMEPRO_NO_INLINE FrameProTLS* CreateFrameProTLS()
@@ -4103,10 +3539,13 @@ namespace FramePro
 
 		framepro_session.AddFrameProTLS(p_framepro_tls);
 
-// BEGIN EPIC
-		FPlatformTLS::SetTlsValue(GetFrameProTLSSlot(), (void*)p_framepro_tls);
-//		gp_FrameProTLS = p_framepro_tls;
-// END EPIC
+		#if FRAMEPRO_USE_TLS_SLOTS
+			int slot = GetFrameProTLSSlot();
+		#else
+			int slot = 0;
+		#endif
+
+		Platform::SetTLSValue(slot, p_framepro_tls);
 
 		return p_framepro_tls;
 	}
@@ -4137,7 +3576,7 @@ namespace FramePro
 
 		p_packet->m_PacketType = packet_type;
 		p_packet->m_Thread = p_framepro_tls->GetThreadId();
-		p_packet->m_Core = GetCore();
+		p_packet->m_Core = Platform::GetCore();
 		p_packet->m_EventId = event_id;
 		p_packet->m_Time = time;
 	}
@@ -4150,10 +3589,12 @@ void FramePro::SetAllocator(Allocator* p_allocator)
 }
 
 //------------------------------------------------------------------------
-// START EPIC
-// Remove QueryPerformanceCounter and DebugBreak
-// END EPIC
+void FramePro::DebugBreak()
+{
+	Platform::DebugBreak();
+}
 
+//------------------------------------------------------------------------
 void FramePro::Shutdown()
 {
 	GetFrameProSession().Shutdown();
@@ -4177,12 +3618,6 @@ void FramePro::UnregisterConnectionChangedcallback(ConnectionChangedCallback p_c
 	GetFrameProSession().UnregisterConnectionChangedCallback(p_callback);
 }
 
-//------------------------------------------------------------------------
-#if FRAMEPRO_WIN_BASED_PLATFORM
-	#pragma warning(push)
-	#pragma warning(disable : 4127)
-#endif
-
 //------------------------------------------------------------------------
 void FramePro::AddTimeSpan(const char* p_name_and_source_info, int64 start_time, int64 end_time)
 {
@@ -4213,7 +3648,7 @@ void FramePro::AddTimeSpan(const char* p_name_and_source_info, int64 start_time,
 
 			TimeSpanPacket* p_packet = (TimeSpanPacket*)p_framepro_tls->AllocateSpaceInBuffer(send_size);
 
-			p_packet->m_PacketType_AndCore = PacketType::TimeSpanWithCallstack | (FramePro::GetCore() << 16);
+			p_packet->m_PacketType_AndCore = PacketType::TimeSpanWithCallstack | (FramePro::Platform::GetCore() << 16);
 			p_packet->m_ThreadId = p_framepro_tls->GetThreadId();
 			p_packet->m_NameAndSourceInfo = (StringId)p_name_and_source_info;
 			p_packet->m_StartTime = start_time;
@@ -4240,7 +3675,7 @@ void FramePro::AddTimeSpan(const char* p_name_and_source_info, int64 start_time,
 
 		TimeSpanPacket* p_packet = p_framepro_tls->AllocateSpaceInBuffer<TimeSpanPacket>();
 
-		p_packet->m_PacketType_AndCore = PacketType::TimeSpan | (FramePro::GetCore() << 16);
+		p_packet->m_PacketType_AndCore = PacketType::TimeSpan | (FramePro::Platform::GetCore() << 16);
 		p_packet->m_ThreadId = p_framepro_tls->GetThreadId();
 		p_packet->m_NameAndSourceInfo = (StringId)p_name_and_source_info;
 		p_packet->m_StartTime = start_time;
@@ -4280,7 +3715,7 @@ void FramePro::AddTimeSpan(const wchar_t* p_name_and_source_info, int64 start_ti
 
 			TimeSpanPacket* p_packet = (TimeSpanPacket*)p_framepro_tls->AllocateSpaceInBuffer(send_size);
 
-			p_packet->m_PacketType_AndCore = PacketType::TimeSpanWWithCallstack | (FramePro::GetCore() << 16);
+			p_packet->m_PacketType_AndCore = PacketType::TimeSpanWWithCallstack | (FramePro::Platform::GetCore() << 16);
 			p_packet->m_ThreadId = p_framepro_tls->GetThreadId();
 			p_packet->m_NameAndSourceInfo = (StringId)p_name_and_source_info;
 			p_packet->m_StartTime = start_time;
@@ -4307,7 +3742,7 @@ void FramePro::AddTimeSpan(const wchar_t* p_name_and_source_info, int64 start_ti
 
 		TimeSpanPacket* p_packet = p_framepro_tls->AllocateSpaceInBuffer<TimeSpanPacket>();
 
-		p_packet->m_PacketType_AndCore = PacketType::TimeSpanW | (FramePro::GetCore() << 16);
+		p_packet->m_PacketType_AndCore = PacketType::TimeSpanW | (FramePro::Platform::GetCore() << 16);
 		p_packet->m_ThreadId = p_framepro_tls->GetThreadId();
 		p_packet->m_NameAndSourceInfo = (StringId)p_name_and_source_info;
 		p_packet->m_StartTime = start_time;
@@ -4347,7 +3782,7 @@ void FramePro::AddTimeSpan(StringId name, const char* p_source_info, int64 start
 
 			NamedTimeSpanPacket* p_packet = (NamedTimeSpanPacket*)p_framepro_tls->AllocateSpaceInBuffer(send_size);
 
-			p_packet->m_PacketType_AndCore = PacketType::NamedTimeSpanWithCallstack | (FramePro::GetCore() << 16);
+			p_packet->m_PacketType_AndCore = PacketType::NamedTimeSpanWithCallstack | (FramePro::Platform::GetCore() << 16);
 			p_packet->m_ThreadId = p_framepro_tls->GetThreadId();
 			p_packet->m_Name = name;
 			p_packet->m_SourceInfo = (StringId)p_source_info;
@@ -4375,7 +3810,7 @@ void FramePro::AddTimeSpan(StringId name, const char* p_source_info, int64 start
 
 		NamedTimeSpanPacket* p_packet = p_framepro_tls->AllocateSpaceInBuffer<NamedTimeSpanPacket>();
 
-		p_packet->m_PacketType_AndCore = PacketType::NamedTimeSpan | (FramePro::GetCore() << 16);
+		p_packet->m_PacketType_AndCore = PacketType::NamedTimeSpan | (FramePro::Platform::GetCore() << 16);
 		p_packet->m_ThreadId = p_framepro_tls->GetThreadId();
 		p_packet->m_Name = name;
 		p_packet->m_SourceInfo = (StringId)p_source_info;
@@ -4487,7 +3922,7 @@ void FramePro::AddTimeSpan(const char* p_name, const char* p_source_info, int64
 
 			NamedTimeSpanPacket* p_packet = (NamedTimeSpanPacket*)p_framepro_tls->AllocateSpaceInBuffer(send_size);
 
-			p_packet->m_PacketType_AndCore = PacketType::StringLiteralNamedTimeSpanWithCallstack | (FramePro::GetCore() << 16);
+			p_packet->m_PacketType_AndCore = PacketType::StringLiteralNamedTimeSpanWithCallstack | (FramePro::Platform::GetCore() << 16);
 			p_packet->m_ThreadId = p_framepro_tls->GetThreadId();
 			p_packet->m_Name = (StringId)p_name;
 			p_packet->m_SourceInfo = (StringId)p_source_info;
@@ -4515,7 +3950,7 @@ void FramePro::AddTimeSpan(const char* p_name, const char* p_source_info, int64
 
 		NamedTimeSpanPacket* p_packet = p_framepro_tls->AllocateSpaceInBuffer<NamedTimeSpanPacket>();
 
-		p_packet->m_PacketType_AndCore = PacketType::StringLiteralNamedTimeSpan | (FramePro::GetCore() << 16);
+		p_packet->m_PacketType_AndCore = PacketType::StringLiteralNamedTimeSpan | (FramePro::Platform::GetCore() << 16);
 		p_packet->m_ThreadId = p_framepro_tls->GetThreadId();
 		p_packet->m_Name = (StringId)p_name;
 		p_packet->m_SourceInfo = (StringId)p_source_info;
@@ -4525,23 +3960,21 @@ void FramePro::AddTimeSpan(const char* p_name, const char* p_source_info, int64
 }
 
 //------------------------------------------------------------------------
-void FramePro::AddCustomStat(const char* p_name, int value, const char* p_graph, const char* p_unit)
+void FramePro::AddCustomStat(const char* p_name, int value, const char* p_graph, const char* p_unit, uint colour)
 {
-	AddCustomStat(p_name, (int64)value, p_graph, p_unit);
+	AddCustomStat(p_name, (int64)value, p_graph, p_unit, colour);
 }
 
 //------------------------------------------------------------------------
-void FramePro::AddCustomStat(const char* p_name, int64 value, const char* p_graph, const char* p_unit)
+void FramePro::AddCustomStat(const char* p_name, int64 value, const char* p_graph, const char* p_unit, uint colour)
 {
 	FrameProTLS* p_framepro_tls = GetFrameProTLS();
 
 	// if we are connecting to FramePro, FramePro will ask for the string value later, otherwise we need to send it now
 	if (p_framepro_tls->SendStringsImmediately())
-	{
 		p_framepro_tls->SendString(p_name, PacketType::StringPacket);
-		p_framepro_tls->SendString(p_graph, PacketType::StringPacket);
-		p_framepro_tls->SendString(p_unit, PacketType::StringPacket);
-	}
+
+	p_framepro_tls->SetCustomStatInfo(p_name, p_graph, p_unit, colour);	// only sends first time
 
 	CriticalSectionScope lock(p_framepro_tls->GetCurrentSendBufferCriticalSection());
 
@@ -4553,28 +3986,24 @@ void FramePro::AddCustomStat(const char* p_name, int64 value, const char* p_grap
 	p_packet->m_Count = 1;
 	p_packet->m_Name = (StringId)p_name;
 	p_packet->m_Value = value;
-	p_packet->m_Graph = (StringId)p_graph;
-	p_packet->m_Unit = (StringId)p_unit;
 }
 
 //------------------------------------------------------------------------
-void FramePro::AddCustomStat(const char* p_name, float value, const char* p_graph, const char* p_unit)
+void FramePro::AddCustomStat(const char* p_name, float value, const char* p_graph, const char* p_unit, uint colour)
 {
-	AddCustomStat(p_name, (double)value, p_graph, p_unit);
+	AddCustomStat(p_name, (double)value, p_graph, p_unit, colour);
 }
 
 //------------------------------------------------------------------------
-void FramePro::AddCustomStat(const char* p_name, double value, const char* p_graph, const char* p_unit)
+void FramePro::AddCustomStat(const char* p_name, double value, const char* p_graph, const char* p_unit, uint colour)
 {
 	FrameProTLS* p_framepro_tls = GetFrameProTLS();
 
 	// if we are connecting to FramePro, FramePro will ask for the string value later, otherwise we need to send it now
 	if (p_framepro_tls->SendStringsImmediately())
-	{
 		p_framepro_tls->SendString(p_name, PacketType::StringPacket);
-		p_framepro_tls->SendString(p_graph, PacketType::StringPacket);
-		p_framepro_tls->SendString(p_unit, PacketType::StringPacket);
-	}
+
+	p_framepro_tls->SetCustomStatInfo(p_name, p_graph, p_unit, colour);	// only sends first time
 
 	CriticalSectionScope lock(p_framepro_tls->GetCurrentSendBufferCriticalSection());
 
@@ -4586,28 +4015,24 @@ void FramePro::AddCustomStat(const char* p_name, double value, const char* p_gra
 	p_packet->m_Count = 1;
 	p_packet->m_Name = (StringId)p_name;
 	p_packet->m_Value = value;
-	p_packet->m_Graph = (StringId)p_graph;
-	p_packet->m_Unit = (StringId)p_unit;
 }
 
 //------------------------------------------------------------------------
-void FramePro::AddCustomStat(const wchar_t* p_name, int value, const wchar_t* p_graph, const wchar_t* p_unit)
+void FramePro::AddCustomStat(const wchar_t* p_name, int value, const wchar_t* p_graph, const wchar_t* p_unit, uint colour)
 {
-	AddCustomStat(p_name, (int64)value, p_graph, p_unit);
+	AddCustomStat(p_name, (int64)value, p_graph, p_unit, colour);
 }
 
 //------------------------------------------------------------------------
-void FramePro::AddCustomStat(const wchar_t* p_name, int64 value, const wchar_t* p_graph, const wchar_t* p_unit)
+void FramePro::AddCustomStat(const wchar_t* p_name, int64 value, const wchar_t* p_graph, const wchar_t* p_unit, uint colour)
 {
 	FrameProTLS* p_framepro_tls = GetFrameProTLS();
 
 	// if we are connecting to FramePro, FramePro will ask for the string value later, otherwise we need to send it now
 	if (p_framepro_tls->SendStringsImmediately())
-	{
-		p_framepro_tls->SendString(p_name, PacketType::StringPacket);
-		p_framepro_tls->SendString(p_graph, PacketType::StringPacket);
-		p_framepro_tls->SendString(p_unit, PacketType::StringPacket);
-	}
+		p_framepro_tls->SendString(p_name, PacketType::WStringPacket);
+
+	p_framepro_tls->SetCustomStatInfo(p_name, p_graph, p_unit, colour);	// only sends first time
 
 	CriticalSectionScope lock(p_framepro_tls->GetCurrentSendBufferCriticalSection());
 
@@ -4615,32 +4040,28 @@ void FramePro::AddCustomStat(const wchar_t* p_name, int64 value, const wchar_t*
 
 	CustomStatValueType::Enum value_type = CustomStatValueType::Int64;
 
-	p_packet->m_PacketTypeAndValueType = PacketType::CustomStatPacket | (value_type << 16);
+	p_packet->m_PacketTypeAndValueType = PacketType::CustomStatPacketW | (value_type << 16);
 	p_packet->m_Count = 1;
 	p_packet->m_Name = (StringId)p_name;
 	p_packet->m_Value = value;
-	p_packet->m_Graph = (StringId)p_graph;
-	p_packet->m_Unit = (StringId)p_unit;
 }
 
 //------------------------------------------------------------------------
-void FramePro::AddCustomStat(const wchar_t* p_name, float value, const wchar_t* p_graph, const wchar_t* p_unit)
+void FramePro::AddCustomStat(const wchar_t* p_name, float value, const wchar_t* p_graph, const wchar_t* p_unit, uint colour)
 {
-	AddCustomStat(p_name, (double)value, p_graph, p_unit);
+	AddCustomStat(p_name, (double)value, p_graph, p_unit, colour);
 }
 
 //------------------------------------------------------------------------
-void FramePro::AddCustomStat(const wchar_t* p_name, double value, const wchar_t* p_graph, const wchar_t* p_unit)
+void FramePro::AddCustomStat(const wchar_t* p_name, double value, const wchar_t* p_graph, const wchar_t* p_unit, uint colour)
 {
 	FrameProTLS* p_framepro_tls = GetFrameProTLS();
 
 	// if we are connecting to FramePro, FramePro will ask for the string value later, otherwise we need to send it now
 	if (p_framepro_tls->SendStringsImmediately())
-	{
 		p_framepro_tls->SendString(p_name, PacketType::StringPacket);
-		p_framepro_tls->SendString(p_graph, PacketType::StringPacket);
-		p_framepro_tls->SendString(p_unit, PacketType::StringPacket);
-	}
+
+	p_framepro_tls->SetCustomStatInfo(p_name, p_graph, p_unit, colour);	// only sends first time
 
 	CriticalSectionScope lock(p_framepro_tls->GetCurrentSendBufferCriticalSection());
 
@@ -4648,31 +4069,24 @@ void FramePro::AddCustomStat(const wchar_t* p_name, double value, const wchar_t*
 
 	CustomStatValueType::Enum value_type = CustomStatValueType::Double;
 
-	p_packet->m_PacketTypeAndValueType = PacketType::CustomStatPacket | (value_type << 16);
+	p_packet->m_PacketTypeAndValueType = PacketType::CustomStatPacketW | (value_type << 16);
 	p_packet->m_Count = 1;
 	p_packet->m_Name = (StringId)p_name;
 	p_packet->m_Value = value;
-	p_packet->m_Graph = (StringId)p_graph;
-	p_packet->m_Unit = (StringId)p_unit;
 }
 
 //------------------------------------------------------------------------
-void FramePro::AddCustomStat(StringId name, int value, const char* p_graph, const char* p_unit)
+void FramePro::AddCustomStat(StringId name, int value, const char* p_graph, const char* p_unit, uint colour)
 {
-	AddCustomStat(name, (int64)value, p_graph, p_unit);
+	AddCustomStat(name, (int64)value, p_graph, p_unit, colour);
 }
 
 //------------------------------------------------------------------------
-void FramePro::AddCustomStat(StringId name, int64 value, const char* p_graph, const char* p_unit)
+void FramePro::AddCustomStat(StringId name, int64 value, const char* p_graph, const char* p_unit, uint colour)
 {
 	FrameProTLS* p_framepro_tls = GetFrameProTLS();
 
-	// if we are connecting to FramePro, FramePro will ask for the string value later, otherwise we need to send it now
-	if (p_framepro_tls->SendStringsImmediately())
-	{
-		p_framepro_tls->SendString(p_graph, PacketType::StringPacket);
-		p_framepro_tls->SendString(p_unit, PacketType::StringPacket);
-	}
+	p_framepro_tls->SetCustomStatInfo(name, p_graph, p_unit, colour);	// only sends first time
 
 	CriticalSectionScope lock(p_framepro_tls->GetCurrentSendBufferCriticalSection());
 
@@ -4684,27 +4098,20 @@ void FramePro::AddCustomStat(StringId name, int64 value, const char* p_graph, co
 	p_packet->m_Count = 1;
 	p_packet->m_Name = name;
 	p_packet->m_Value = value;
-	p_packet->m_Graph = (StringId)p_graph;
-	p_packet->m_Unit = (StringId)p_unit;
 }
 
 //------------------------------------------------------------------------
-void FramePro::AddCustomStat(StringId name, float value, const char* p_graph, const char* p_unit)
+void FramePro::AddCustomStat(StringId name, float value, const char* p_graph, const char* p_unit, uint colour)
 {
-	AddCustomStat(name, (double)value, p_graph, p_unit);
+	AddCustomStat(name, (double)value, p_graph, p_unit, colour);
 }
 
 //------------------------------------------------------------------------
-void FramePro::AddCustomStat(StringId name, double value, const char* p_graph, const char* p_unit)
+void FramePro::AddCustomStat(StringId name, double value, const char* p_graph, const char* p_unit, uint colour)
 {
 	FrameProTLS* p_framepro_tls = GetFrameProTLS();
 
-	// if we are connecting to FramePro, FramePro will ask for the string value later, otherwise we need to send it now
-	if (p_framepro_tls->SendStringsImmediately())
-	{
-		p_framepro_tls->SendString(p_graph, PacketType::StringPacket);
-		p_framepro_tls->SendString(p_unit, PacketType::StringPacket);
-	}
+	p_framepro_tls->SetCustomStatInfo(name, p_graph, p_unit, colour);	// only sends first time
 
 	CriticalSectionScope lock(p_framepro_tls->GetCurrentSendBufferCriticalSection());
 
@@ -4716,28 +4123,20 @@ void FramePro::AddCustomStat(StringId name, double value, const char* p_graph, c
 	p_packet->m_Count = 1;
 	p_packet->m_Name = name;
 	p_packet->m_Value = value;
-	p_packet->m_Graph = (StringId)p_graph;
-	p_packet->m_Unit = (StringId)p_unit;
 }
 
-// START EPIC 
 //------------------------------------------------------------------------
-void FramePro::AddCustomStat(StringId name, int value, const wchar_t* p_graph, const wchar_t* p_unit)
+void FramePro::AddCustomStat(StringId name, int value, const wchar_t* p_graph, const wchar_t* p_unit, uint colour)
 {
-	AddCustomStat(name, (int64)value, p_graph, p_unit);
+	AddCustomStat(name, (int64)value, p_graph, p_unit, colour);
 }
 
 //------------------------------------------------------------------------
-void FramePro::AddCustomStat(StringId name, int64 value, const wchar_t* p_graph, const wchar_t* p_unit)
+void FramePro::AddCustomStat(StringId name, int64 value, const wchar_t* p_graph, const wchar_t* p_unit, uint colour)
 {
 	FrameProTLS* p_framepro_tls = GetFrameProTLS();
 
-	// if we are connecting to FramePro, FramePro will ask for the string value later, otherwise we need to send it now
-	if (p_framepro_tls->SendStringsImmediately())
-	{
-		p_framepro_tls->SendString(p_graph, PacketType::StringPacket);
-		p_framepro_tls->SendString(p_unit, PacketType::StringPacket);
-	}
+	p_framepro_tls->SetCustomStatInfo(name, p_graph, p_unit, colour);	// only sends first time
 
 	CriticalSectionScope lock(p_framepro_tls->GetCurrentSendBufferCriticalSection());
 
@@ -4749,27 +4148,20 @@ void FramePro::AddCustomStat(StringId name, int64 value, const wchar_t* p_graph,
 	p_packet->m_Count = 1;
 	p_packet->m_Name = name;
 	p_packet->m_Value = value;
-	p_packet->m_Graph = (StringId)p_graph;
-	p_packet->m_Unit = (StringId)p_unit;
 }
 
 //------------------------------------------------------------------------
-void FramePro::AddCustomStat(StringId name, float value, const wchar_t* p_graph, const wchar_t* p_unit)
+void FramePro::AddCustomStat(StringId name, float value, const wchar_t* p_graph, const wchar_t* p_unit, uint colour)
 {
-	AddCustomStat(name, (double)value, p_graph, p_unit);
+	AddCustomStat(name, (double)value, p_graph, p_unit, colour);
 }
 
 //------------------------------------------------------------------------
-void FramePro::AddCustomStat(StringId name, double value, const wchar_t* p_graph, const wchar_t* p_unit)
+void FramePro::AddCustomStat(StringId name, double value, const wchar_t* p_graph, const wchar_t* p_unit, uint colour)
 {
 	FrameProTLS* p_framepro_tls = GetFrameProTLS();
 
-	// if we are connecting to FramePro, FramePro will ask for the string value later, otherwise we need to send it now
-	if (p_framepro_tls->SendStringsImmediately())
-	{
-		p_framepro_tls->SendString(p_graph, PacketType::StringPacket);
-		p_framepro_tls->SendString(p_unit, PacketType::StringPacket);
-	}
+	p_framepro_tls->SetCustomStatInfo(name, p_graph, p_unit, colour);	// only sends first time
 
 	CriticalSectionScope lock(p_framepro_tls->GetCurrentSendBufferCriticalSection());
 
@@ -4781,15 +4173,57 @@ void FramePro::AddCustomStat(StringId name, double value, const wchar_t* p_graph
 	p_packet->m_Count = 1;
 	p_packet->m_Name = name;
 	p_packet->m_Value = value;
-	p_packet->m_Graph = (StringId)p_graph;
-	p_packet->m_Unit = (StringId)p_unit;
 }
-// END EPIC
 
 //------------------------------------------------------------------------
-#if FRAMEPRO_WIN_BASED_PLATFORM
-	#pragma warning(pop)
-#endif
+void FramePro::AddCustomStat(StringId name, int value, StringId graph, StringId unit, uint colour)
+{
+	AddCustomStat(name, (int64)value, graph, unit, colour);
+}
+
+//------------------------------------------------------------------------
+void FramePro::AddCustomStat(StringId name, int64 value, StringId graph, StringId unit, uint colour)
+{
+	FrameProTLS* p_framepro_tls = GetFrameProTLS();
+
+	p_framepro_tls->SetCustomStatInfo(name, graph, unit, colour);	// only sends first time
+
+	CriticalSectionScope lock(p_framepro_tls->GetCurrentSendBufferCriticalSection());
+
+	CustomStatPacketInt64* p_packet = p_framepro_tls->AllocateSpaceInBuffer<CustomStatPacketInt64>();
+
+	CustomStatValueType::Enum value_type = CustomStatValueType::Int64;
+
+	p_packet->m_PacketTypeAndValueType = PacketType::CustomStatPacket | (value_type << 16);
+	p_packet->m_Count = 1;
+	p_packet->m_Name = name;
+	p_packet->m_Value = value;
+}
+
+//------------------------------------------------------------------------
+void FramePro::AddCustomStat(StringId name, float value, StringId graph, StringId unit, uint colour)
+{
+	AddCustomStat(name, (double)value, graph, unit, colour);
+}
+
+//------------------------------------------------------------------------
+void FramePro::AddCustomStat(StringId name, double value, StringId graph, StringId unit, uint colour)
+{
+	FrameProTLS* p_framepro_tls = GetFrameProTLS();
+
+	p_framepro_tls->SetCustomStatInfo(name, graph, unit, colour);	// only sends first time
+
+	CriticalSectionScope lock(p_framepro_tls->GetCurrentSendBufferCriticalSection());
+
+	CustomStatPacketDouble* p_packet = p_framepro_tls->AllocateSpaceInBuffer<CustomStatPacketDouble>();
+
+	CustomStatValueType::Enum value_type = CustomStatValueType::Double;
+
+	p_packet->m_PacketTypeAndValueType = PacketType::CustomStatPacket | (value_type << 16);
+	p_packet->m_Count = 1;
+	p_packet->m_Name = name;
+	p_packet->m_Value = value;
+}
 
 //------------------------------------------------------------------------
 void FramePro::SetThreadName(const char* p_name)
@@ -4816,12 +4250,16 @@ FramePro::StringId FramePro::RegisterString(const wchar_t* p_str)
 }
 
 //------------------------------------------------------------------------
-// START EPIC
-void FramePro::StartRecording(const FString& p_filename, bool context_switches, int64 max_file_size)
+void FramePro::StartRecording(const char* p_filename, bool context_switches, bool callstacks, int64 max_file_size)
+{
+	GetFrameProSession().StartRecording(p_filename, context_switches, callstacks, max_file_size);
+}
+
+//------------------------------------------------------------------------
+void FramePro::StartRecording(const wchar_t* p_filename, bool context_switches, bool callstacks, int64 max_file_size)
 {
-	GetFrameProSession().StartRecording(p_filename, context_switches, max_file_size);
+	GetFrameProSession().StartRecording(p_filename, context_switches, callstacks, max_file_size);
 }
-// END EPIC
 
 //------------------------------------------------------------------------
 void FramePro::StopRecording()
@@ -4960,123 +4398,108 @@ void FramePro::TriggerWaitEvent(int64 event_id)
 }
 
 //------------------------------------------------------------------------
-void FramePro::SetScopeCustomStat(const char* p_name, int64 value, const char* p_graph, const char* p_unit)
+void FramePro::SetScopeCustomStat(const char* p_name, int64 value, const char* p_graph, const char* p_unit, uint colour)
 {
 	FrameProTLS* p_framepro_tls = GetFrameProTLS();
 
 	if (p_framepro_tls->SendStringsImmediately())
-	{
 		p_framepro_tls->SendString(p_name, PacketType::StringPacket);
-		p_framepro_tls->SendString(p_graph, PacketType::StringPacket);
-		p_framepro_tls->SendString(p_unit, PacketType::StringPacket);
-	}
 
-	p_framepro_tls->SetCustomTimeSpanStat((StringId)p_name, value, p_unit);
+	p_framepro_tls->SetCustomTimeSpanStat((StringId)p_name, value);
 
-	AddCustomStat(p_name, value, p_graph, p_unit);
+	AddCustomStat(p_name, value, p_graph, p_unit, colour);
 }
 
 //------------------------------------------------------------------------
-void FramePro::SetScopeCustomStat(const wchar_t* p_name, int64 value, const wchar_t* p_graph, const wchar_t* p_unit)
+void FramePro::SetScopeCustomStat(const wchar_t* p_name, int64 value, const wchar_t* p_graph, const wchar_t* p_unit, uint colour)
 {
 	FrameProTLS* p_framepro_tls = GetFrameProTLS();
 
 	if (p_framepro_tls->SendStringsImmediately())
-	{
 		p_framepro_tls->SendString(p_name, PacketType::StringPacket);
-		p_framepro_tls->SendString(p_graph, PacketType::StringPacket);
-		p_framepro_tls->SendString(p_unit, PacketType::StringPacket);
-	}
 
-	p_framepro_tls->SetCustomTimeSpanStat((StringId)p_name, value, p_unit);
+	p_framepro_tls->SetCustomTimeSpanStatW((StringId)p_name, value);
 
-	AddCustomStat(p_name, value, p_graph, p_unit);
+	AddCustomStat(p_name, value, p_graph, p_unit, colour);
 }
 
 //------------------------------------------------------------------------
-void FramePro::SetScopeCustomStat(StringId name, int64 value, const char* p_graph, const char* p_unit)
+void FramePro::SetScopeCustomStat(StringId name, int64 value, StringId graph, StringId unit, uint colour)
 {
-	GetFrameProTLS()->SetCustomTimeSpanStat(name, value, p_unit);
+	// don't care about whether it is W or not because string has already been registerd
+	GetFrameProTLS()->SetCustomTimeSpanStat(name, value);
 
-	AddCustomStat(name, value, p_graph, p_unit);
+	AddCustomStat(name, value, graph, unit, colour);
 }
 
 //------------------------------------------------------------------------
-void FramePro::SetScopeCustomStat(const char* p_name, int value, const char* p_graph, const char* p_unit)
+void FramePro::SetScopeCustomStat(const char* p_name, int value, const char* p_graph, const char* p_unit, uint colour)
 {
-	SetScopeCustomStat(p_name, (int64)value, p_graph, p_unit);
+	SetScopeCustomStat(p_name, (int64)value, p_graph, p_unit, colour);
 }
 
 //------------------------------------------------------------------------
-void FramePro::SetScopeCustomStat(const wchar_t* p_name, int value, const wchar_t* p_graph, const wchar_t* p_unit)
+void FramePro::SetScopeCustomStat(const wchar_t* p_name, int value, const wchar_t* p_graph, const wchar_t* p_unit, uint colour)
 {
-	SetScopeCustomStat(p_name, (int64)value, p_graph, p_unit);
+	SetScopeCustomStat(p_name, (int64)value, p_graph, p_unit, colour);
 }
 
 //------------------------------------------------------------------------
-void FramePro::SetScopeCustomStat(StringId name, int value, const char* p_graph, const char* p_unit)
+void FramePro::SetScopeCustomStat(StringId name, int value, StringId graph, StringId unit, uint colour)
 {
-	SetScopeCustomStat(name, (int64)value, p_graph, p_unit);
+	SetScopeCustomStat(name, (int64)value, graph, unit, colour);
 }
 
 //------------------------------------------------------------------------
-void FramePro::SetScopeCustomStat(const char* p_name, float value, const char* p_graph, const char* p_unit)
+void FramePro::SetScopeCustomStat(const char* p_name, float value, const char* p_graph, const char* p_unit, uint colour)
 {
-	SetScopeCustomStat(p_name, (double)value, p_graph, p_unit);
+	SetScopeCustomStat(p_name, (double)value, p_graph, p_unit, colour);
 }
 
 //------------------------------------------------------------------------
-void FramePro::SetScopeCustomStat(const wchar_t* p_name, float value, const wchar_t* p_graph, const wchar_t* p_unit)
+void FramePro::SetScopeCustomStat(const wchar_t* p_name, float value, const wchar_t* p_graph, const wchar_t* p_unit, uint colour)
 {
-	SetScopeCustomStat(p_name, (double)value, p_graph, p_unit);
+	SetScopeCustomStat(p_name, (double)value, p_graph, p_unit, colour);
 }
 
 //------------------------------------------------------------------------
-void FramePro::SetScopeCustomStat(StringId name, float value, const char* p_graph, const char* p_unit)
+void FramePro::SetScopeCustomStat(StringId name, float value, StringId graph, StringId unit, uint colour)
 {
-	SetScopeCustomStat(name, (double)value, p_graph, p_unit);
+	SetScopeCustomStat(name, (double)value, graph, unit, colour);
 }
 
 //------------------------------------------------------------------------
-void FramePro::SetScopeCustomStat(const char* p_name, double value, const char* p_graph, const char* p_unit)
+void FramePro::SetScopeCustomStat(const char* p_name, double value, const char* p_graph, const char* p_unit, uint colour)
 {
 	FrameProTLS* p_framepro_tls = GetFrameProTLS();
 
 	if (p_framepro_tls->SendStringsImmediately())
-	{
 		p_framepro_tls->SendString(p_name, PacketType::StringPacket);
-		p_framepro_tls->SendString(p_graph, PacketType::StringPacket);
-		p_framepro_tls->SendString(p_unit, PacketType::StringPacket);
-	}
 
-	p_framepro_tls->SetCustomTimeSpanStat((StringId)p_name, value, p_unit);
+	p_framepro_tls->SetCustomTimeSpanStat((StringId)p_name, value);
 
-	AddCustomStat(p_name, value, p_graph, p_unit);
+	AddCustomStat(p_name, value, p_graph, p_unit, colour);
 }
 
 //------------------------------------------------------------------------
-void FramePro::SetScopeCustomStat(const wchar_t* p_name, double value, const wchar_t* p_graph, const wchar_t* p_unit)
+void FramePro::SetScopeCustomStat(const wchar_t* p_name, double value, const wchar_t* p_graph, const wchar_t* p_unit, uint colour)
 {
 	FrameProTLS* p_framepro_tls = GetFrameProTLS();
 
 	if (p_framepro_tls->SendStringsImmediately())
-	{
 		p_framepro_tls->SendString(p_name, PacketType::StringPacket);
-		p_framepro_tls->SendString(p_graph, PacketType::StringPacket);
-		p_framepro_tls->SendString(p_unit, PacketType::StringPacket);
-	}
 
-	p_framepro_tls->SetCustomTimeSpanStat((StringId)p_name, value, p_unit);
+	p_framepro_tls->SetCustomTimeSpanStat((StringId)p_name, value);
 
-	AddCustomStat(p_name, value, p_graph, p_unit);
+	AddCustomStat(p_name, value, p_graph, p_unit, colour);
 }
 
 //------------------------------------------------------------------------
-void FramePro::SetScopeCustomStat(StringId name, double value, const char* p_graph, const char* p_unit)
+void FramePro::SetScopeCustomStat(StringId name, double value, StringId graph, StringId unit, uint colour)
 {
-	GetFrameProTLS()->SetCustomTimeSpanStat(name, value, p_unit);
+	GetFrameProTLS()->SetCustomTimeSpanStat(name, value);
 
-	AddCustomStat(name, value, p_graph, p_unit);
+	AddCustomStat(name, value, graph, unit, colour);
 }
 
 //------------------------------------------------------------------------
@@ -5086,14 +4509,34 @@ void FramePro::SetConditionalScopeMinTimeInMicroseconds(int64 value)
 }
 
 //------------------------------------------------------------------------
-#endif		// #if FRAMEPRO_ENABLED
+void FramePro::SetScopeColour(StringId name, uint colour)
+{
+	GetFrameProSession().SetScopeColour(name, colour);
+}
+
 //------------------------------------------------------------------------
-// FrameProCallstackSet.cpp
+void FramePro::SetCustomStatGraph(StringId name, StringId graph)
+{
+	GetFrameProSession().SetCustomStatGraph(name, graph);
+}
+
+//------------------------------------------------------------------------
+void FramePro::SetCustomStatUnit(StringId name, StringId unit)
+{
+	GetFrameProSession().SetCustomStatUnit(name, unit);
+}
 
+//------------------------------------------------------------------------
+void FramePro::SetCustomStatColour(StringId name, uint colour)
+{
+	GetFrameProSession().SetCustomStatColour(name, colour);
+}
 
 
 //------------------------------------------------------------------------
-#if FRAMEPRO_ENABLED
+//
+// FrameProCallstackSet.cpp
+//
 
 //------------------------------------------------------------------------
 namespace FramePro
@@ -5116,119 +4559,116 @@ namespace FramePro
 
 		return true;
 	}
-}
 
-//------------------------------------------------------------------------
-FramePro::CallstackSet::CallstackSet(Allocator* p_allocator)
-:	mp_Data((Callstack**)p_allocator->Alloc(g_CallstackSetInitialCapacity*sizeof(Callstack*))),
-	m_CapacityMask(g_CallstackSetInitialCapacity-1),
-	m_Count(0),
-	m_Capacity(g_CallstackSetInitialCapacity),
-	mp_Allocator(p_allocator),
-	m_BlockAllocator(p_allocator)
-{
-	memset(mp_Data, 0, g_CallstackSetInitialCapacity*sizeof(Callstack*));
-}
+	//------------------------------------------------------------------------
+	CallstackSet::CallstackSet(Allocator* p_allocator)
+	:	mp_Data((Callstack**)p_allocator->Alloc(g_CallstackSetInitialCapacity*sizeof(Callstack*))),
+		m_CapacityMask(g_CallstackSetInitialCapacity-1),
+		m_Count(0),
+		m_Capacity(g_CallstackSetInitialCapacity),
+		mp_Allocator(p_allocator),
+		m_BlockAllocator(p_allocator)
+	{
+		memset(mp_Data, 0, g_CallstackSetInitialCapacity*sizeof(Callstack*));
+	}
 
-//------------------------------------------------------------------------
-FramePro::CallstackSet::~CallstackSet()
-{
-	Clear();
-}
+	//------------------------------------------------------------------------
+	CallstackSet::~CallstackSet()
+	{
+		Clear();
+	}
 
-//------------------------------------------------------------------------
-void FramePro::CallstackSet::Grow()
-{
-	int old_capacity = m_Capacity;
-	Callstack** p_old_data = mp_Data;
+	//------------------------------------------------------------------------
+	void CallstackSet::Grow()
+	{
+		int old_capacity = m_Capacity;
+		Callstack** p_old_data = mp_Data;
 
-	// allocate a new set
-	m_Capacity *= 2;
-	m_CapacityMask = m_Capacity - 1;
-	int size = m_Capacity * sizeof(Callstack*);
-	mp_Data = (Callstack**)mp_Allocator->Alloc(size);
-	memset(mp_Data, 0, size);
+		// allocate a new set
+		m_Capacity *= 2;
+		m_CapacityMask = m_Capacity - 1;
+		int size = m_Capacity * sizeof(Callstack*);
+		mp_Data = (Callstack**)mp_Allocator->Alloc(size);
+		memset(mp_Data, 0, size);
 
-	// transfer callstacks from old set
-	m_Count = 0;
-	for(int i=0; i<old_capacity; ++i)
-	{
-		Callstack* p_callstack = p_old_data[i];
-		if(p_callstack)
-			Add(p_callstack);
-	}
+		// transfer callstacks from old set
+		m_Count = 0;
+		for(int i=0; i<old_capacity; ++i)
+		{
+			Callstack* p_callstack = p_old_data[i];
+			if(p_callstack)
+				Add(p_callstack);
+		}
 
-	// release old buffer
-	mp_Allocator->Free(p_old_data);
-}
+		// release old buffer
+		mp_Allocator->Free(p_old_data);
+	}
 
-//------------------------------------------------------------------------
-FramePro::Callstack* FramePro::CallstackSet::Get(uint64* p_stack, int stack_size, unsigned int hash)
-{
-	int index = hash & m_CapacityMask;
+	//------------------------------------------------------------------------
+	FramePro::Callstack* CallstackSet::Get(uint64* p_stack, int stack_size, unsigned int hash)
+	{
+		int index = hash & m_CapacityMask;
 
-	while(mp_Data[index] && !StacksMatch(mp_Data[index], p_stack, stack_size, hash))
-		index = (index + 1) & m_CapacityMask;
+		while(mp_Data[index] && !StacksMatch(mp_Data[index], p_stack, stack_size, hash))
+			index = (index + 1) & m_CapacityMask;
 
-	return mp_Data[index];
-}
+		return mp_Data[index];
+	}
 
-//------------------------------------------------------------------------
-FramePro::Callstack* FramePro::CallstackSet::Add(uint64* p_stack, int stack_size, unsigned int hash)
-{
-	// grow the set if necessary
-	if(m_Count > m_Capacity/4)
-		Grow();
+	//------------------------------------------------------------------------
+	FramePro::Callstack* CallstackSet::Add(uint64* p_stack, int stack_size, unsigned int hash)
+	{
+		// grow the set if necessary
+		if(m_Count > m_Capacity/4)
+			Grow();
 
-	// create a new callstack
-	Callstack* p_callstack = (Callstack*)m_BlockAllocator.Alloc(sizeof(Callstack));
-	p_callstack->m_ID = m_Count;
-	p_callstack->m_Size = stack_size;
-	p_callstack->mp_Stack = (uint64*)m_BlockAllocator.Alloc(stack_size*sizeof(uint64));
-	p_callstack->m_Hash = hash;
-	memcpy(p_callstack->mp_Stack, p_stack, stack_size*sizeof(uint64));
+		// create a new callstack
+		Callstack* p_callstack = (Callstack*)m_BlockAllocator.Alloc(sizeof(Callstack));
+		p_callstack->m_ID = m_Count;
+		p_callstack->m_Size = stack_size;
+		p_callstack->mp_Stack = (uint64*)m_BlockAllocator.Alloc(stack_size*sizeof(uint64));
+		p_callstack->m_Hash = hash;
+		memcpy(p_callstack->mp_Stack, p_stack, stack_size*sizeof(uint64));
 
-	Add(p_callstack);
+		Add(p_callstack);
 
-	return p_callstack;
-}
+		return p_callstack;
+	}
 
-//------------------------------------------------------------------------
-void FramePro::CallstackSet::Add(Callstack* p_callstack)
-{
-	// find a clear index
-	int index = p_callstack->m_Hash & m_CapacityMask;
-	while(mp_Data[index])
-		index = (index + 1) & m_CapacityMask;
+	//------------------------------------------------------------------------
+	void CallstackSet::Add(Callstack* p_callstack)
+	{
+		// find a clear index
+		int index = p_callstack->m_Hash & m_CapacityMask;
+		while(mp_Data[index])
+			index = (index + 1) & m_CapacityMask;
 
-	mp_Data[index] = p_callstack;
+		mp_Data[index] = p_callstack;
 
-	++m_Count;
-}
+		++m_Count;
+	}
 
-//------------------------------------------------------------------------
-void FramePro::CallstackSet::Clear()
-{
-	m_BlockAllocator.Clear();
+	//------------------------------------------------------------------------
+	void CallstackSet::Clear()
+	{
+		m_BlockAllocator.Clear();
 
-	mp_Allocator->Free(mp_Data);
+		mp_Allocator->Free(mp_Data);
 
-	size_t size = g_CallstackSetInitialCapacity*sizeof(Callstack*);
-	mp_Data = (Callstack**)mp_Allocator->Alloc((int)size);
-	memset(mp_Data, 0, size);
-	m_CapacityMask = g_CallstackSetInitialCapacity-1;
-	m_Count = 0;
-	m_Capacity = g_CallstackSetInitialCapacity;
+		size_t size = g_CallstackSetInitialCapacity*sizeof(Callstack*);
+		mp_Data = (Callstack**)mp_Allocator->Alloc((int)size);
+		memset(mp_Data, 0, size);
+		m_CapacityMask = g_CallstackSetInitialCapacity-1;
+		m_Count = 0;
+		m_Capacity = g_CallstackSetInitialCapacity;
+	}
 }
 
-//------------------------------------------------------------------------
-#endif		// #ifdef ENABLE_MEMPRO
-//------------------------------------------------------------------------
-// FrameProLib.cpp
-
 
 //------------------------------------------------------------------------
-#if FRAMEPRO_ENABLED
+//
+// FrameProLib.cpp
+//
 
 //------------------------------------------------------------------------
 #include <stdarg.h>
@@ -5237,78 +4677,42 @@ void FramePro::CallstackSet::Clear()
 namespace FramePro
 {
 	//------------------------------------------------------------------------
-	void DebugWrite(const FRAMEPRO_TCHAR* p_str, ...)
+	void SPrintf(char* p_buffer, size_t const buffer_size, const char* p_format, ...)
 	{
 		va_list args;
-		va_start(args, p_str);
+		va_start(args, p_format);
 
-		static FRAMEPRO_TCHAR g_TempString[1024];
-		_vstprintf_s(g_TempString, sizeof(g_TempString)/sizeof(FRAMEPRO_TCHAR), p_str, args);
-		OutputDebugString(g_TempString);
+		Platform::VSPrintf(p_buffer, buffer_size, p_format, args);
 
 		va_end(args);
 	}
-}
-
-//------------------------------------------------------------------------
-#endif		// #if FRAMEPRO_ENABLED
-//------------------------------------------------------------------------
-// FrameProSession.cpp
-
-
-
 
+	//------------------------------------------------------------------------
+	void DebugWrite(const char* p_str, ...)
+	{
+		va_list args;
+		va_start(args, p_str);
 
-#include <ctime>
+		static char temp_string[1024];
+		Platform::VSPrintf(temp_string, sizeof(temp_string), p_str, args);
 
-//------------------------------------------------------------------------
-#if FRAMEPRO_ENABLED
+		Platform::DebugWrite(temp_string);
 
-//------------------------------------------------------------------------
-#if FRAMEPRO_PLATFORM_WIN
-	#if FRAMEPRO_TOOLSET_UE4
-		#include "Windows/AllowWindowsPlatformTypes.h"
-	#endif
-	#include <psapi.h>
-	#if FRAMEPRO_TOOLSET_UE4
-		#include "Windows/HideWindowsPlatformTypes.h"
-	#endif
-#endif
+		va_end(args);
+	}
+}
 
-//------------------------------------------------------------------------
-// if you are having problems compiling this on your platform undefine FRAMEPRO_ENUMERATE_MODULES and it send info for just the main module
-#define FRAMEPRO_ENUMERATE_MODULES (!FRAMEPRO_PLATFORM_XBOXONE && FRAMEPRO_ENABLE_CALLSTACKS && FRAMEPRO_X64 && 1)
 
 //------------------------------------------------------------------------
-#if FRAMEPRO_ENUMERATE_MODULES
-    #if FRAMEPRO_WIN_BASED_PLATFORM
-
-		#ifdef __UNREAL__
-			#include "AllowWindowsPlatformTypes.h"
-		#endif
-
-		#pragma warning(push)
-		#pragma warning(disable : 4091)
-			#include <Dbghelp.h>
-		#pragma warning(pop)
-
-		#pragma comment(lib, "Dbghelp.lib")
-
-		#ifdef __UNREAL__
-			#include "HideWindowsPlatformTypes.h"
-		#endif
-	#else
-        #include <link.h>
-    #endif
-#endif
+//
+// FrameProSession.cpp
+//
 
 //------------------------------------------------------------------------
 namespace FramePro
 {
 	//------------------------------------------------------------------------
-// START EPIC
-	const TCHAR* g_NonInteractiveRecordingFilePath = TEXT("framepro_recording.bin");
-// END EPIC
+	const char* g_NonInteractiveRecordingFilePath = "framepro_recording.bin";
 
 	//------------------------------------------------------------------------
 	FrameProSession* FrameProSession::mp_Inst = NULL;
@@ -5328,50 +4732,21 @@ namespace FramePro
 		time(&rawtime);
 
 		tm timeinfo;
-		localtime_s(&timeinfo, &rawtime);
+		Platform::GetLocalTime(&timeinfo, &rawtime);
 
 		strftime(p_date, len, "%d-%m-%Y %I:%M:%S", &timeinfo);
 	}
 
-	//------------------------------------------------------------------------
-	void BaseAddressLookupFunction()
-	{
-	}
-
-	//------------------------------------------------------------------------
-	bool GetProcessName(int process_id, char* p_name, int max_name_length)
-	{
-#if FRAMEPRO_PLATFORM_WIN
-		HANDLE process = OpenProcess(PROCESS_ALL_ACCESS, true, process_id);
-		if(process)
-		{
-			unsigned long result = GetProcessImageFileNameA(process, p_name, max_name_length);
-			CloseHandle(process);
-
-			if(result)
-			{
-				int total_length = (int)strlen(p_name);
-				char* p_filename = strrchr(p_name, '\\');
-				if(p_filename && p_filename[1])
-				{
-					++p_filename;
-					memmove(p_name, p_filename, p_name + total_length + 1 - p_filename);
-				}
-
-				return true;
-			}
-		}
-#endif
-		return false;
-	}
-
 	//------------------------------------------------------------------------
 	FrameProSession::FrameProSession()
-	:	mp_Allocator(NULL),
+	:	mp_Allocator(NULL),	// must be first (before any call to GetAllocator)
 		m_CreatedAllocator(false),
 		m_Initialised(false),
 		m_InitialiseConnectionNextFrame(false),
 		m_StartContextSwitchRecording(false),
+#if FRAMEPRO_ENABLE_CALLSTACKS
+		m_StartRecordingCallstacks(false),
+#endif
 		m_ClockFrequency(0),
 		m_MainThreadId(-1),
 		m_SendThreadStarted(false, true),
@@ -5379,10 +4754,8 @@ namespace FramePro
 		m_SendComplete(false, false),
 		m_ReceiveThreadTerminatedEvent(false, false),
 		m_Interactive(true),
-		mp_NonInteractiveRecordingFile(NULL),
 		m_NonInteractiveRecordingFileSize(0),
 		m_LastSessionInfoSendTime(0),
-		mp_RecordingFile(NULL),
 		m_RecordingFileSize(0),
 		m_MaxRecordingFileSize(0),
 		m_ThreadPrioritySet(false),
@@ -5392,17 +4765,16 @@ namespace FramePro
 		m_SendThreadExit(false),
 		m_SendThreadFinished(false, true),
 		m_SocketsBlocked(FRAMEPRO_SOCKETS_BLOCKED_BY_DEFAULT),
-		mp_GlobalHiResTimers(NULL)
-#if FRAMEPRO_EVENT_TRACE_WIN32
-		,mp_EventTraceWin32(NULL)
-#endif
+		mp_GlobalHiResTimers(NULL),
+		mp_ContextSwitchRecorder(NULL)
 #if FRAMEPRO_ENABLE_CALLSTACKS
 		,m_SendModules(false)
 #endif
 	{
 		mp_Inst = this;
 
-		strcpy_s(m_Port, FRAMEPRO_PORT);
+		FRAMEPRO_ASSERT(sizeof(m_Port) >= strlen(FRAMEPRO_PORT) + 1);
+		memcpy(m_Port, FRAMEPRO_PORT, strlen(FRAMEPRO_PORT) + 1);
 
 		CalculateTimerFrequency();
 	}
@@ -5410,14 +4782,12 @@ namespace FramePro
 	//------------------------------------------------------------------------
 	FrameProSession::~FrameProSession()
 	{
-		HandleDisconnect();
+		HandleDisconnect(false);
 
 		m_NamedThreads.Clear();
 
-#if FRAMEPRO_EVENT_TRACE_WIN32
-		if(mp_EventTraceWin32)
-			Delete(mp_Allocator, mp_EventTraceWin32);
-#endif
+		Platform::DestroyContextSwitchRecorder(mp_ContextSwitchRecorder, GetAllocator());
+
 		// must clear all arrays and buffers and detach the allocator before deleting the allocator
 		m_ProcessIds.Clear();
 		m_ProcessIds.SetAllocator(NULL);
@@ -5431,6 +4801,18 @@ namespace FramePro
 		m_ModulePackets.Clear();
 		m_ModulePackets.SetAllocator(NULL);
 
+		m_ScopeColours.Clear();
+		m_ScopeColours.SetAllocator(NULL);
+
+		m_CustomStatGraphs.Clear();
+		m_CustomStatGraphs.SetAllocator(NULL);
+
+		m_CustomStatUnits.Clear();
+		m_CustomStatUnits.SetAllocator(NULL);
+
+		m_CustomStatColours.Clear();
+		m_CustomStatColours.SetAllocator(NULL);
+
 		m_NamedThreads.Clear();
 		m_NamedThreads.SetAllocator(NULL);
 
@@ -5444,17 +4826,28 @@ namespace FramePro
 	//------------------------------------------------------------------------
 	void FrameProSession::SetPort(int port)
 	{
-		sprintf_s(m_Port, "%d", port);
+		Platform::ToString(port, m_Port, sizeof(m_Port));
 	}
 
 	//------------------------------------------------------------------------
 	void FrameProSession::SetAllocator(Allocator* p_allocator)
 	{
 		if(mp_Allocator)
-			FRAMEPRO_BREAK();		// allocator already set. You must call Allocator BEFORE calling FRAMEPRO_FRAME_START
+			FramePro::DebugBreak();		// allocator already set. You must call Allocator BEFORE calling FRAMEPRO_FRAME_START
+
+		SetAllocatorInternal(p_allocator);
+	}
 
+	//------------------------------------------------------------------------
+	void FrameProSession::SetAllocatorInternal(Allocator* p_allocator)
+	{
+		FRAMEPRO_ASSERT(!mp_Allocator);
 		FRAMEPRO_ASSERT(p_allocator);
+
 		mp_Allocator = p_allocator;
+
+		m_NonInteractiveRecordingFile.SetAllocator(p_allocator);
+		m_RecordingFile.SetAllocator(p_allocator);
 	}
 
 	//------------------------------------------------------------------------
@@ -5464,7 +4857,7 @@ namespace FramePro
 
 		if(!mp_Allocator)
 		{
-			mp_Allocator = new DefaultAllocator();
+			SetAllocatorInternal(new DefaultAllocator());
 			m_CreatedAllocator = true;
 		}
 
@@ -5474,9 +4867,7 @@ namespace FramePro
 	//------------------------------------------------------------------------
 	void FrameProSession::CalculateTimerFrequency()
 	{
-// START EPIC
-		m_ClockFrequency = (1.0 / FPlatformTime::GetSecondsPerCycle());
-// END EPIC
+		m_ClockFrequency = Platform::GetTimerFrequency();
 	}
 
 	//------------------------------------------------------------------------
@@ -5542,7 +4933,7 @@ namespace FramePro
 
 		{
 			CriticalSectionScope lock(m_CriticalSection);
-			if(mp_RecordingFile)
+			if(m_RecordingFile.IsOpened())
 			{
 				m_ListenSocket.Disconnect();		// don't allow connections while recording
 				return 0;
@@ -5572,7 +4963,6 @@ namespace FramePro
 #endif
 
 	//------------------------------------------------------------------------
-#if FRAMEPRO_SOCKETS_ENABLED
 	void FrameProSession::SendOnMainThread(void* p_src, int size)
 	{
 		CriticalSectionScope tls_lock(m_MainThreadSendBufferLock);
@@ -5580,7 +4970,6 @@ namespace FramePro
 		void* p_dst = m_MainThreadSendBuffer.Allocate(size);
 		memcpy(p_dst, p_src, size);
 	}
-#endif
 
 	//------------------------------------------------------------------------
 #if FRAMEPRO_SOCKETS_ENABLED
@@ -5646,16 +5035,9 @@ namespace FramePro
 
 						if(!packet.m_Interactive)
 						{
-							FRAMEPRO_ASSERT(!mp_NonInteractiveRecordingFile);
-// START EPIC
-							FString FileName = FPaths::ProfilingDir() + TEXT("FramePro/") + g_NonInteractiveRecordingFilePath;
-
-#if ALLOW_DEBUG_FILES
-							mp_NonInteractiveRecordingFile = IFileManager::Get().CreateDebugFileWriter(*FileName);
-#else
-							mp_NonInteractiveRecordingFile = IFileManager::Get().CreateFileWriter(*FileName);
-#endif
-// END EPIC
+							bool opened = m_NonInteractiveRecordingFile.OpenForWrite(g_NonInteractiveRecordingFilePath);
+							FRAMEPRO_ASSERT(opened);
+							FRAMEPRO_UNREFERENCED(opened);
 						}
 
 						m_Interactive = packet.m_Interactive ? 1 : 0;
@@ -5679,24 +5061,13 @@ namespace FramePro
 
 				case PacketType::SetCallstackRecordingEnabledPacket:
 				{
-					#if FRAMEPRO_ENABLE_CALLSTACKS
-						SetCallstackRecordingEnabledPacket packet;
-						if (m_ClientSocket.Receive(&packet, sizeof(packet)) == sizeof(packet))
-						{
-							// enumerate modules
-							if (!m_SendModules)
-							{
-								EnumerateModules();
-								m_SendModules = true;
-							}
-
-							{
-								CriticalSectionScope tls_lock(m_TLSListCriticalSection);
-								for (FrameProTLS* p_tls = m_FrameProTLSList.GetHead(); p_tls != NULL; p_tls = p_tls->GetNext())
-									p_tls->SetSendCallstacks(packet.m_Enabled != 0);
-							}
-						}
-					#endif
+					SetCallstackRecordingEnabledPacket packet;
+					if (m_ClientSocket.Receive(&packet, sizeof(packet)) == sizeof(packet))
+					{
+						#if FRAMEPRO_ENABLE_CALLSTACKS
+							SetCallstacksEnabled(packet.m_Enabled != 0);
+						#endif
+					}
 				} break;
 			}
 		}
@@ -5707,32 +5078,58 @@ namespace FramePro
 #endif
 
 	//------------------------------------------------------------------------
-	void FrameProSession::StartRecordingContextSitches()
+#if FRAMEPRO_ENABLE_CALLSTACKS
+	void FrameProSession::SetCallstacksEnabled(bool enabled)
 	{
-#if FRAMEPRO_EVENT_TRACE_WIN32
-		if(!mp_EventTraceWin32)
-			mp_EventTraceWin32 = New<EventTraceWin32>(mp_Allocator, mp_Allocator);
+		// enumerate modules
+		if (enabled && !m_SendModules)
+		{
+			Platform::EnumerateModules(m_ModulePackets, GetAllocator());
+
+			// send module packets
+			for (int i = 0; i < m_ModulePackets.GetCount(); ++i)
+			{
+				ModulePacket* p_module_packet = m_ModulePackets[i];
+				SendImmediate(p_module_packet, sizeof(ModulePacket), GetFrameProTLS());
+				mp_Allocator->Free(p_module_packet);
+			}
+			m_ModulePackets.Clear();
+
+			m_SendModules = true;
+		}
+
+		{
+			CriticalSectionScope tls_lock(m_TLSListCriticalSection);
+			for (FrameProTLS* p_tls = m_FrameProTLSList.GetHead(); p_tls != NULL; p_tls = p_tls->GetNext())
+				p_tls->SetSendCallstacks(enabled);
+		}
+	}
+#endif
 
+	//------------------------------------------------------------------------
+	void FrameProSession::StartRecordingContextSitches()
+	{
 		DynamicString error(GetAllocator());
-		bool started = mp_EventTraceWin32->Start(ContextSwitchCallback_Static, this, error);
+
+		if (!mp_ContextSwitchRecorder)
+			mp_ContextSwitchRecorder = Platform::CreateContextSwitchRecorder(GetAllocator());
+
+		bool started = Platform::StartRecordingContextSitches(mp_ContextSwitchRecorder, ContextSwitchCallback_Static, this, error);
+
+		if (!started)
+		{
+			Platform::DestroyContextSwitchRecorder(mp_ContextSwitchRecorder, GetAllocator());
+			mp_ContextSwitchRecorder = NULL;
+		}
 
 		// send the context switch started packet
 		ContextSwitchRecordingStartedPacket response_packet;
 		response_packet.m_PacketType = PacketType::ContextSwitchRecordingStartedPacket;
 		response_packet.m_StartedSucessfully = started;
+		
 		error.CopyTo(response_packet.m_Error, sizeof(response_packet.m_Error));
 
 		SendOnMainThread(response_packet);
-
-		if(!started)
-		{
-			#if FRAMEPRO_PLATFORM_WIN || FRAMEPRO_PLATFORM_UWP
-				FramePro::DebugWrite(FRAMEPRO_STRING("FramePro Warning: Failed to start recording context switches. Please make sure that you are running with administrator privileges.\n"));
-			#else
-				FramePro::DebugWrite(FRAMEPRO_STRING("FramePro Warning: Failed to start recording context switches. Context switches may not be supported for this platform\n"));
-			#endif
-		}
-#endif
 	}
 
 	//------------------------------------------------------------------------
@@ -5749,7 +5146,7 @@ namespace FramePro
 	{
 		m_CriticalSection.Leave();
 
-		m_SendThread.CreateThread(StaticSendThreadMain, this);
+		m_SendThread.CreateThread(StaticSendThreadMain, this, GetAllocator());
 
 		if(m_ThreadPrioritySet)
 			m_SendThread.SetPriority(m_ThreadPriority);
@@ -5768,7 +5165,7 @@ namespace FramePro
 	{
 		m_ReceiveThreadTerminatedEvent.Reset();
 
-		m_ReceiveThread.CreateThread(StaticReceiveThreadMain, this);
+		m_ReceiveThread.CreateThread(StaticReceiveThreadMain, this, GetAllocator());
 
 		if(m_ThreadPrioritySet)
 			m_ReceiveThread.SetPriority(m_ThreadPriority);
@@ -5798,7 +5195,7 @@ namespace FramePro
 
 			const int max_process_name_length = 260;
 			char process_name[max_process_name_length];
-			if(GetProcessName(context_switch.m_ProcessId, process_name, max_process_name_length))
+			if(Platform::GetProcessName(context_switch.m_ProcessId, process_name, max_process_name_length))
 			{
 				StringId name_id = RegisterString(process_name);
 				p_framepro_tls->SendSessionInfoPacket(ProcessNamePacket(context_switch.m_ProcessId, name_id));
@@ -5822,336 +5219,99 @@ namespace FramePro
 	}
 
 	//------------------------------------------------------------------------
-	Platform::Enum GetPlatformEnum()
+	void FrameProSession::InitialiseConnection(FrameProTLS* p_framepro_tls)
 	{
-#if FRAMEPRO_PLATFORM_WIN
-		return Platform::Windows;
-#elif FRAMEPRO_PLATFORM_UWP
-		return Platform::Windows_UWP;
-#elif FRAMEPRO_PLATFORM_XBOXONE
-		return Platform::XBoxOne;
-#elif FRAMEPRO_PLATFORM_XBOX360
-		return Platform::XBox360;
-// START EPIC
-#elif FRAMEPRO_PLATFORM_ANDROID || FRAMEPRO_PLATFORM_UNIX || FRAMEPRO_PLATFORM_IOS || FRAMEPRO_PLATFORM_SWITCH || FRAMEPRO_PLATFORM_MAC
-// END EPIC
-		return Platform::Unix;
-#else
-		#error
-#endif
-	}
+		// start the Send thread FIRST, but paused (because it adds another TLS that we need to call OnConnected on)
+		m_SendComplete.Reset();
+		m_SendReady.Reset();
+		CreateSendThread();
 
-	//------------------------------------------------------------------------
-#if FRAMEPRO_WIN_BASED_PLATFORM
-	struct CV_HEADER
-	{
-		int Signature;
-		int Offset;
-	};
+		// call OnConnected on all TLS threads
+		bool recording_to_file = m_RecordingFile.IsOpened();
+		{
+			CriticalSectionScope tls_lock(m_TLSListCriticalSection);
+			for(FrameProTLS* p_tls = m_FrameProTLSList.GetHead(); p_tls!=NULL; p_tls=p_tls->GetNext())
+				p_tls->OnConnected(recording_to_file);
+		}
 
-	struct CV_INFO_PDB20
-	{
-		CV_HEADER CvHeader;
-		int Signature;
-		int Age;
-		char PdbFileName[FRAMEPRO_MAX_PATH];
-	};
+		p_framepro_tls->SendConnectPacket(m_ClockFrequency, Platform::GetCurrentProcessId(), Platform::GetPlatformEnum());
 
-	struct CV_INFO_PDB70
-	{
-		int  CvSignature;
-		GUID Signature;
-		int Age;
-		char PdbFileName[FRAMEPRO_MAX_PATH];
-	};
+#if FRAMEPRO_ENABLE_CALLSTACKS
+		m_SendModules = false;
 #endif
+		// tell the send thread that there is data ready and wait for it to be sent
+		m_SendReady.Set();
+		m_CriticalSection.Leave();
+		m_SendComplete.Wait();
+		m_CriticalSection.Enter();
+		m_SendComplete.Reset();
 
-	//------------------------------------------------------------------------
-	void GetExtraModuleInfo(int64 ModuleBase, ModulePacket* p_module_packet)
-	{
-#if FRAMEPRO_PLATFORM_WIN
-		IMAGE_DOS_HEADER* p_dos_header = (IMAGE_DOS_HEADER*)ModuleBase;
-		IMAGE_NT_HEADERS* p_nt_header = (IMAGE_NT_HEADERS*)((char*)ModuleBase + p_dos_header->e_lfanew);
-		IMAGE_OPTIONAL_HEADER& optional_header = p_nt_header->OptionalHeader;
-		IMAGE_DATA_DIRECTORY& image_data_directory = optional_header.DataDirectory[IMAGE_DIRECTORY_ENTRY_DEBUG];
-		IMAGE_DEBUG_DIRECTORY* p_debug_info_array = (IMAGE_DEBUG_DIRECTORY*)(ModuleBase + image_data_directory.VirtualAddress);
-		int count = image_data_directory.Size / sizeof(IMAGE_DEBUG_DIRECTORY);
-		for(int i=0; i<count; ++i)
-		{
-			if(p_debug_info_array[i].Type == IMAGE_DEBUG_TYPE_CODEVIEW)
-			{
-				char* p_cv_data = (char*)(ModuleBase + p_debug_info_array[i].AddressOfRawData);
-				if(strncmp(p_cv_data, "RSDS", 4) == 0)
-				{
-					CV_INFO_PDB70* p_cv_info = (CV_INFO_PDB70*)p_cv_data;
+		// make sure no more TLS threads are added while we are setting stuff up
+		m_TLSListCriticalSection.Enter();
 
-					p_module_packet->m_PacketType = PacketType::ModulePacket;
-					p_module_packet->m_Age = p_cv_info->Age;
-					
-					static_assert(sizeof(p_module_packet->m_Sig) == sizeof(p_cv_info->Signature), "sig size wrong");
-					memcpy(p_module_packet->m_Sig, &p_cv_info->Signature, sizeof(p_cv_info->Signature));
+		// lock the session info buffers of all threads
+		for(FrameProTLS* p_tls = m_FrameProTLSList.GetHead(); p_tls!=NULL; p_tls=p_tls->GetNext())
+		{
+			p_tls->OnConnected(recording_to_file);		// in case any threads have been added since sending the connect packet
+			p_tls->LockSessionInfoBuffer();
+		}
+
+		// send the session info buffers for all threads
+		for(FrameProTLS* p_tls = m_FrameProTLSList.GetHead(); p_tls!=NULL; p_tls=p_tls->GetNext())
+			p_tls->SendSessionInfoBuffer();
 
-					strcpy_s(p_module_packet->m_SymbolFilename, p_cv_info->PdbFileName);
+		p_framepro_tls->SendFrameStartPacket(0);
 
-					return;									// returning here
-				}
-				else if(strncmp(p_cv_data, "NB10", 4) == 0)
-				{
-					CV_INFO_PDB20* p_cv_info = (CV_INFO_PDB20*)p_cv_data;
+		if(g_ConditionalScopeMinTime == UINT_MAX)
+			g_ConditionalScopeMinTime = (unsigned int)((((int64)FRAMEPRO_DEFAULT_COND_SCOPE_MIN_TIME) * m_ClockFrequency) / 1000000LL);
 
-					p_module_packet->m_PacketType = PacketType::ModulePacket;
-					p_module_packet->m_Age = p_cv_info->Age;
+		// Do this (almost) last. Threads will start sending data once this is set. This atomic flag also publishes all the above data.
+		std::atomic_thread_fence(std::memory_order_seq_cst);
+		g_Connected = true;
 
-					memset(p_module_packet->m_Sig, 0, sizeof(p_module_packet->m_Sig));
-					static_assert(sizeof(p_cv_info->Signature) <= sizeof(p_module_packet->m_Sig), "sig size wrong");
-					memcpy(p_module_packet->m_Sig, &p_cv_info->Signature, sizeof(p_cv_info->Signature));
+#if FRAMEPRO_SOCKETS_ENABLED
+		// create the receive thread. Must do this AFTER setting g_Connected
+		if (!m_RecordingFile.IsOpened())
+			CreateReceiveThread();
+#endif
 
-					strcpy_s(p_module_packet->m_SymbolFilename, p_cv_info->PdbFileName);
+		// unlock the session info buffers of all threads. Must do this after g_Connected is set
+		for(FrameProTLS* p_tls = m_FrameProTLSList.GetHead(); p_tls!=NULL; p_tls=p_tls->GetNext())
+			p_tls->UnlockSessionInfoBuffer();
 
-					return;									// returning here
-				}
-			}
+		m_TLSListCriticalSection.Leave();
+		
+		// start recording context switches if dumping to a file
+		if(m_StartContextSwitchRecording)
+		{
+			StartRecordingContextSitches();
+			m_StartContextSwitchRecording = false;
 		}
+
+		ClearGlobalHiResTimers();
+
+		SendScopeColours();
+		SendCustomStatGraphs();
+		SendCustomStatUnits();
+		SendCustomStatColours();
+
+		OnConnectionChanged(true, m_RecordingFile.GetFilename());
+
+#if FRAMEPRO_ENABLE_CALLSTACKS
+		bool enable_callstacks = m_StartRecordingCallstacks;
+		m_StartRecordingCallstacks = false;
+		m_CriticalSection.Leave();
+		SetCallstacksEnabled(enable_callstacks);
+		m_CriticalSection.Enter();
 #endif
 	}
 
 	//------------------------------------------------------------------------
-#if FRAMEPRO_ENUMERATE_MODULES
-    #if FRAMEPRO_WIN_BASED_PLATFORM
-		#if !defined(_IMAGEHLP_SOURCE_) && defined(_IMAGEHLP64)
-			// depending on your platform you may need to change PCSTR to PSTR for ModuleName
-			BOOL CALLBACK EnumerateLoadedModulesCallback(__in PCSTR ModuleName,__in DWORD64 ModuleBase,__in ULONG,__in_opt PVOID UserContext)
-		#else
-			BOOL CALLBACK EnumerateLoadedModulesCallback(__in PCSTR ModuleName,__in ULONG ModuleBase,__in ULONG,__in_opt PVOID UserContext)
-		#endif
-			{
-				FrameProSession* p_this = (FrameProSession*)UserContext;
-
-				p_this->EnumerateLoadedModulesCallback(ModuleBase, ModuleName, false);
+	void FrameProSession::Initialise(FrameProTLS* p_framepro_tls)
+	{
+		if(m_Initialised)
+			return;
 
-				return true;
-			}
-	#else
-		int EnumerateLoadedModulesCallback(struct dl_phdr_info* info, size_t size, void* data)
-		{
-			FrameProSession* p_this = (FrameProSession*)data;
-
-			int64 module_base = 0;
-			for (int j = 0; j < info->dlpi_phnum; j++)
-			{
-				if (info->dlpi_phdr[j].p_type == PT_LOAD)
-				{
-					module_base = info->dlpi_addr + info->dlpi_phdr[j].p_vaddr;
-					break;
-				}
-			}
-
-			static bool first = true;
-			if(first)
-			{
-				first = false;
-
-				int64 module_base = (int64)BaseAddressLookupFunction;		// use the address of the BaseAddressLookupFunction function so that we can work it out later
-
-				// get the module name
-				char arg1[20];
-				char char_filename[FRAMEPRO_MAX_PATH];
-				sprintf(arg1, "/proc/%d/exe", getpid());
-				memset(char_filename, 0, FRAMEPRO_MAX_PATH);
-				readlink(arg1, char_filename, FRAMEPRO_MAX_PATH -1);
-
-				p_this->EnumerateLoadedModulesCallback(
-					module_base,
-					char_filename,
-					true);
-			}
-			else
-			{
-				p_this->EnumerateLoadedModulesCallback(
-					module_base,
-					info->dlpi_name,
-					false);
-			}
-
-			return 0;
-		}
-	#endif
-#endif
-
-	//------------------------------------------------------------------------
-	void FrameProSession::EnumerateLoadedModulesCallback(
-		int64 module_base,
-		const char* p_module_name,
-		bool use_lookup_function_for_base_address)
-	{
-		CriticalSectionScope lock(m_CriticalSection);
-
-		ModulePacket* p_module_packet = (ModulePacket*)mp_Allocator->Alloc(sizeof(ModulePacket));
-		memset(p_module_packet, 0, sizeof(ModulePacket));
-
-		p_module_packet->m_ModuleBase = module_base;
-		strcpy_s(p_module_packet->m_ModuleName, p_module_name);
-		p_module_packet->m_UseLookupFunctionForBaseAddress = use_lookup_function_for_base_address ? 1 : 0;
-
-		GetExtraModuleInfo(module_base, p_module_packet);
-
-		m_ModulePackets.Add(p_module_packet);
-	}
-
-	//------------------------------------------------------------------------
-	void FrameProSession::EnumerateModules()
-	{
-		// if you are having problems compiling this on your platform unset FRAMEPRO_ENUMERATE_MODULES and it send info for just the main module
-		#if FRAMEPRO_ENUMERATE_MODULES
-			#if FRAMEPRO_WIN_BASED_PLATFORM
-				EnumerateLoadedModules64(GetCurrentProcess(), FramePro::EnumerateLoadedModulesCallback, this);
-			#else
-				dl_iterate_phdr(FramePro::EnumerateLoadedModulesCallback, this);
-			#endif
-		#endif
-
-		if (!m_ModulePackets.GetCount())
-		{
-			// if FRAMEPRO_ENUMERATE_MODULES is set or enumeration failed for some reason, fall back
-			// to getting the base address for the main module. This will always work for for all platforms.
-
-			ModulePacket* p_module_packet = (ModulePacket*)mp_Allocator->Alloc(sizeof(ModulePacket));
-			memset(p_module_packet, 0, sizeof(ModulePacket));
-
-			p_module_packet->m_PacketType = PacketType::ModulePacket;
-			p_module_packet->m_UseLookupFunctionForBaseAddress = 0;
-
-			#if FRAMEPRO_WIN_BASED_PLATFORM
-				static int module = 0;
-				HMODULE module_handle = 0;
-				GetModuleHandleEx(GET_MODULE_HANDLE_EX_FLAG_FROM_ADDRESS, (LPCTSTR)&module, &module_handle);
-
-				p_module_packet->m_ModuleBase = (int64)module_handle;
-
-				TCHAR tchar_filename[FRAMEPRO_MAX_PATH] = { 0 };
-				GetModuleFileName(NULL, tchar_filename, FRAMEPRO_MAX_PATH);
-
-				#ifdef UNICODE
-					size_t chars_converted = 0;
-					wcstombs_s(&chars_converted, p_module_packet->m_SymbolFilename, tchar_filename, FRAMEPRO_MAX_PATH);
-				#else
-					strcpy_s(p_module_packet->m_SymbolFilename, tchar_filename);
-				#endif
-// BEGIN EPIC
-			#elif FRAMEPRO_UE4_BASED_PLATFORM
-				p_module_packet->m_UseLookupFunctionForBaseAddress = 1;
-				p_module_packet->m_ModuleBase = (int64)BaseAddressLookupFunction;
-				strcpy_s( p_module_packet->m_SymbolFilename, FRAMEPRO_MAX_INLINE_STRING_LENGTH-1, TCHAR_TO_ANSI(FPlatformProcess::ExecutableName(false)) );
-// END EPIC
-			#else
-				p_module_packet->m_UseLookupFunctionForBaseAddress = 1;
-
-				p_module_packet->m_ModuleBase = (int64)BaseAddressLookupFunction;		// use the address of the BaseAddressLookupFunction function so that we can work it out later
-
-				// get the module name
-				char arg1[20];
-				sprintf(arg1, "/proc/%d/exe", getpid());
-				memset(p_module_packet->m_SymbolFilename, 0, FRAMEPRO_MAX_PATH);
-				readlink(arg1, p_module_packet->m_SymbolFilename, FRAMEPRO_MAX_PATH -1);
-			#endif
-
-			m_ModulePackets.Add(p_module_packet);
-		}
-
-		// send module packets
-		for (int i = 0; i < m_ModulePackets.GetCount(); ++i)
-		{
-			ModulePacket* p_module_packet = m_ModulePackets[i];
-			SendImmediate(p_module_packet, sizeof(ModulePacket), GetFrameProTLS());
-			mp_Allocator->Free(p_module_packet);
-		}
-		m_ModulePackets.Clear();
-	}
-
-	//------------------------------------------------------------------------
-	void FrameProSession::InitialiseConnection(FrameProTLS* p_framepro_tls)
-	{
-		// start the Send thread FIRST, but paused (because it adds another TLS that we need to call OnConnected on)
-		m_SendComplete.Reset();
-		m_SendReady.Reset();
-		CreateSendThread();
-
-		// call OnConnected on all TLS threads
-		bool recording_to_file = mp_RecordingFile != NULL;
-		{
-			CriticalSectionScope tls_lock(m_TLSListCriticalSection);
-			for(FrameProTLS* p_tls = m_FrameProTLSList.GetHead(); p_tls!=NULL; p_tls=p_tls->GetNext())
-				p_tls->OnConnected(recording_to_file);
-		}
-
-		p_framepro_tls->SendConnectPacket(m_ClockFrequency, GetCurrentProcessId(), GetPlatformEnum());
-
-		// tell the send thread that there is data ready and wait for it to be sent
-		m_SendReady.Set();
-		m_CriticalSection.Leave();
-		m_SendComplete.Wait();
-		m_CriticalSection.Enter();
-		m_SendComplete.Reset();
-
-		// make sure no more TLS threads are added while we are setting stuff up
-		m_TLSListCriticalSection.Enter();
-
-		// lock the session info buffers of all threads
-		for(FrameProTLS* p_tls = m_FrameProTLSList.GetHead(); p_tls!=NULL; p_tls=p_tls->GetNext())
-		{
-			p_tls->OnConnected(recording_to_file);		// in case any threads have been added since sending the connect packet
-			p_tls->LockSessionInfoBuffer();
-		}
-
-		// send the session info buffers for all threads
-		for(FrameProTLS* p_tls = m_FrameProTLSList.GetHead(); p_tls!=NULL; p_tls=p_tls->GetNext())
-			p_tls->SendSessionInfoBuffer();
-
-		p_framepro_tls->SendFrameStartPacket(0);
-
-		if (g_ConditionalScopeMinTime == UINT_MAX)
-		{
-			g_ConditionalScopeMinTime = (unsigned int)((((int64)FRAMEPRO_DEFAULT_COND_SCOPE_MIN_TIME) * m_ClockFrequency) / 1000000LL);
-		}
-
-		// Do this (almost) last. Threads will start sending data once this is set. This atomic flag also publishes all the above data.
-		std::atomic_thread_fence(std::memory_order_seq_cst);
-		g_Connected = true;
-
-#if FRAMEPRO_SOCKETS_ENABLED
-		// create the receive thread. Must do this AFTER setting g_Connected
-		if (!mp_RecordingFile)
-			CreateReceiveThread();
-#endif
-
-		// unlock the session info buffers of all threads. Must do this after g_Connected is set
-		for(FrameProTLS* p_tls = m_FrameProTLSList.GetHead(); p_tls!=NULL; p_tls=p_tls->GetNext())
-			p_tls->UnlockSessionInfoBuffer();
-
-		m_TLSListCriticalSection.Leave();
-		
-		// start recording context switches if dumping to a file
-		if(m_StartContextSwitchRecording)
-		{
-			StartRecordingContextSitches();
-			m_StartContextSwitchRecording = false;
-		}
-
-		ClearGlobalHiResTimers();
-
-		// START EPIC
-		FString ArchiveName;
-		if (mp_RecordingFile)
-		{
-			ArchiveName = mp_RecordingFile->GetArchiveName();
-		}
-		OnConnectionChanged(ArchiveName, true);
-		// END EPIC
-	}
-
-	//------------------------------------------------------------------------
-	void FrameProSession::Initialise(FrameProTLS* p_framepro_tls)
-	{
 		if(!HasSetThreadName(p_framepro_tls->GetThreadId()))
 			p_framepro_tls->SetThreadName(p_framepro_tls->GetThreadId(), "Main Thread");
 
@@ -6166,11 +5326,16 @@ namespace FramePro
 		}
 
 		m_ModulePackets.SetAllocator(GetAllocator());
+		m_ScopeColours.SetAllocator(GetAllocator());
+		m_CustomStatGraphs.SetAllocator(GetAllocator());
+		m_CustomStatUnits.SetAllocator(GetAllocator());
+		m_CustomStatColours.SetAllocator(GetAllocator());
 
 #if FRAMEPRO_SOCKETS_ENABLED
 		OpenListenSocket();
 		StartConnectThread();
 #endif
+		m_Initialised = true;
 	}
 
 	//------------------------------------------------------------------------
@@ -6184,13 +5349,13 @@ namespace FramePro
 
 		if(!bind_result)
 		{
-			DebugWrite(FRAMEPRO_STRING("FramePro ERROR: Failed to bind port. This usually means that another process is already running with FramePro enabled.\n"));
+			DebugWrite("FramePro ERROR: Failed to bind port. This usually means that another process is already running with FramePro enabled.\n");
 			return;
 		}
 
 		if(!m_ListenSocket.StartListening())
 		{
-			DebugWrite(FRAMEPRO_STRING("FramePro ERROR: Failed to start listening on socket\n"));
+			DebugWrite("FramePro ERROR: Failed to start listening on socket\n");
 		}
 	}
 #endif
@@ -6199,32 +5364,36 @@ namespace FramePro
 #if FRAMEPRO_SOCKETS_ENABLED
 	void FrameProSession::StartConnectThread()
 	{
-		m_ConnectThread.CreateThread(StaticConnectThreadMain, this);
+		m_ConnectThread.CreateThread(StaticConnectThreadMain, this, GetAllocator());
 	}
 #endif
 
 	//------------------------------------------------------------------------
+#if FRAMEPRO_SOCKETS_ENABLED
 	bool FrameProSession::SendSendBuffer(SendBuffer* p_send_buffer, Socket& socket)
 	{
-#if FRAMEPRO_DEBUG_TCP
-		static FILE* p_debug_file = NULL;
-		if (!p_debug_file)
-			fopen_s(&p_debug_file, "framepro_network_data.framepro_recording", "wb");
-		fwrite(p_send_buffer->GetBuffer(), p_send_buffer->GetSize(), 1, p_debug_file);
-#endif
+		#if FRAMEPRO_DEBUG_TCP
+			static File file;
+			if (!file.IsOpened())
+			{
+				bool opened = file.OpenForWrite("framepro_network_data.framepro_recording");
+				FRAMEPRO_ASSERT(opened);
+				FRAMEPRO_UNREFERENCED(opened);
+			}
+			file.Write(p_send_buffer->GetBuffer(), p_send_buffer->GetSize());
+		#endif
 
 		return socket.Send(p_send_buffer->GetBuffer(), p_send_buffer->GetSize());
 	}
+#endif
 
 	//------------------------------------------------------------------------
-// START EPIC
-	void FrameProSession::WriteSendBuffer(SendBuffer* p_send_buffer, FArchive* p_file, int64& file_size)
+	void FrameProSession::WriteSendBuffer(SendBuffer* p_send_buffer, File& file, int64& file_size)
 	{
 		int size = p_send_buffer->GetSize();
-		p_file->Serialize((uint8*)p_send_buffer->GetBuffer(), size);
+		file.Write(p_send_buffer->GetBuffer(), size);
 		file_size += size;
 	}
-// END EPIC
 
 	//------------------------------------------------------------------------
 	void FrameProSession::SendFrameBuffer()
@@ -6243,28 +5412,25 @@ namespace FramePro
 		// send the send buffers
 		for(SendBuffer* p_send_buffer = send_buffer_list.GetHead(); p_send_buffer; p_send_buffer = p_send_buffer->GetNext())
 		{
-// START EPIC
-            CriticalSectionScope lock2(m_CriticalSection);
-                
-            if(mp_RecordingFile)
-            {
-                WriteSendBuffer(p_send_buffer, mp_RecordingFile, m_RecordingFileSize);
-            }
-            else
-            {
+			if(m_RecordingFile.IsOpened())
+			{
+				CriticalSectionScope lock2(m_CriticalSection);
+				WriteSendBuffer(p_send_buffer, m_RecordingFile, m_RecordingFileSize);
+			}
+			else
+			{
 #if FRAMEPRO_SOCKETS_ENABLED
-                if(m_Interactive)
-                {
-                    if(!SendSendBuffer(p_send_buffer, m_ClientSocket))
-                        break;        // disconnected
-                }
-                else
-                {
-                    WriteSendBuffer(p_send_buffer, mp_NonInteractiveRecordingFile, m_NonInteractiveRecordingFileSize);
-                }
+				if(m_Interactive)
+				{
+					if(!SendSendBuffer(p_send_buffer, m_ClientSocket))
+						break;		// disconnected
+				}
+				else
+				{
+					WriteSendBuffer(p_send_buffer, m_NonInteractiveRecordingFile, m_NonInteractiveRecordingFileSize);
+				}
 #endif
-            }
-// END EPIC
+			}
 		}
 
 		// give the empty send buffers back to the TLS objects
@@ -6311,30 +5477,31 @@ namespace FramePro
 
 		g_Connected = false;
 
-// START EPIC
-		mp_NonInteractiveRecordingFile->Close();
-		delete mp_NonInteractiveRecordingFile;
-		mp_NonInteractiveRecordingFile = NULL;
+		m_NonInteractiveRecordingFile.Close();
 
-		FString FileName = FPaths::ProfilingDir() + TEXT("FramePro/") + g_NonInteractiveRecordingFilePath;
+		char folder[FRAMEPRO_MAX_PATH];
+		Platform::GetRecordingFolder(folder, sizeof(folder));
+		char path[FRAMEPRO_MAX_PATH];
+		SPrintf(path, sizeof(path), "%s%s", folder, g_NonInteractiveRecordingFilePath);
 
-		FArchive* p_read_file = IFileManager::Get().CreateFileReader(*FileName);
+		File read_file;
+		
+		bool opened = read_file.OpenForRead(path);
+		FRAMEPRO_ASSERT(opened);
+		FRAMEPRO_UNREFERENCED(opened);
 
-		size_t bytes_to_read = p_read_file->TotalSize();
+		size_t bytes_to_read = read_file.GetSize();
 
 		const int block_size = 64*1024;
-		uint8* p_read_buffer = (uint8*)mp_Allocator->Alloc(block_size);
+		char* p_read_buffer = (char*)mp_Allocator->Alloc(block_size);
 		while(bytes_to_read)
 		{
 			size_t size_to_read = block_size < bytes_to_read ? block_size : bytes_to_read;
-			p_read_file->Serialize(p_read_buffer, size_to_read);
+			read_file.Read(p_read_buffer, size_to_read);
 			m_ClientSocket.Send(p_read_buffer, size_to_read);
 			bytes_to_read -= size_to_read;
 		}
-		p_read_file->Close();
-		delete p_read_file;
-		p_read_file = NULL;
-// END EPIC
+		read_file.Close();
 		mp_Allocator->Free(p_read_buffer);
 
 		HandleDisconnect_NoLock();
@@ -6342,45 +5509,43 @@ namespace FramePro
 	}
 
 	//------------------------------------------------------------------------
-	void FrameProSession::HandleDisconnect()
+	void FrameProSession::HandleDisconnect(bool wait_for_threads_to_exit)
 	{
 		CriticalSectionScope lock(m_CriticalSection);
 		if(g_Connected)
-			HandleDisconnect_NoLock();
+			HandleDisconnect_NoLock(wait_for_threads_to_exit);
 	}
 
 	//------------------------------------------------------------------------
-	void FrameProSession::HandleDisconnect_NoLock()
+	void FrameProSession::HandleDisconnect_NoLock(bool wait_for_threads_to_exit)
 	{
-		FRAMEPRO_ASSERT(m_CriticalSection.Locked());
-
-#if FRAMEPRO_EVENT_TRACE_WIN32
-		if(mp_EventTraceWin32)
-			mp_EventTraceWin32->Stop();
-#endif
+		Platform::StopRecordingContextSitches(mp_ContextSwitchRecorder);
 
 #if FRAMEPRO_SOCKETS_ENABLED
 		m_ClientSocket.Disconnect();
 #endif
 		g_Connected = false;
 
-		// shut down the send thread
-		if(m_SendThread.IsAlive())
+		if (wait_for_threads_to_exit)
 		{
-			m_SendThreadExit = true;
-			m_SendReady.Set();
-			m_CriticalSection.Leave();
-			m_SendThreadFinished.Wait();
-			m_CriticalSection.Enter();
-			m_SendThreadExit = false;
-		}
+			// shut down the send thread
+			if (m_SendThread.IsAlive())
+			{
+				m_SendThreadExit = true;
+				m_SendReady.Set();
+				m_CriticalSection.Leave();
+				m_SendThreadFinished.Wait();
+				m_CriticalSection.Enter();
+				m_SendThreadExit = false;
+			}
 
-		// shut down the receive thread
-		if(m_ReceiveThread.IsAlive())
-		{
-			m_CriticalSection.Leave();
-			m_ReceiveThreadTerminatedEvent.Wait(10*1000);
-			m_CriticalSection.Enter();
+			// shut down the receive thread
+			if (m_ReceiveThread.IsAlive())
+			{
+				m_CriticalSection.Leave();
+				m_ReceiveThreadTerminatedEvent.Wait(10 * 1000);
+				m_CriticalSection.Enter();
+			}
 		}
 
 		{
@@ -6393,29 +5558,20 @@ namespace FramePro
 
 		m_InitialiseConnectionNextFrame = false;
 
-		// START EPIC
-		FString RecordingFileName;
-		// END EPIC
-
-        {
-            if(mp_RecordingFile)
-            {
-                // START EPIC
-				RecordingFileName = mp_RecordingFile->GetArchiveName();
-                mp_RecordingFile->Close();
-                delete mp_RecordingFile;
-                // END EPIC
-                mp_RecordingFile = NULL;
-            }
-        }
+		DynamicWString recording_filename;
+		recording_filename.SetAllocator(mp_Allocator);
+
+		if(m_RecordingFile.IsOpened())
+		{
+			recording_filename = m_RecordingFile.GetFilename().c_str();
+			m_RecordingFile.Close();
+		}
 
 #if FRAMEPRO_SOCKETS_ENABLED
 		// start listening for new connections
 		StartConnectThread();
 #endif
-		// START EPIC
-		OnConnectionChanged(RecordingFileName, false);
-		// END EPIC
+		OnConnectionChanged(false, recording_filename);
 	}
 
 	//------------------------------------------------------------------------
@@ -6459,7 +5615,7 @@ namespace FramePro
 	// when interactive mode is disabled sends over the socket directly, otherwise send as normal
 	void FrameProSession::SendImmediate(void* p_data, int size, FrameProTLS* p_framepro_tls)
 	{
-		if(mp_RecordingFile)
+		if(m_RecordingFile.IsOpened())
 		{
 			p_framepro_tls->Send(p_data, size);
 		}
@@ -6500,12 +5656,7 @@ namespace FramePro
 
 		CriticalSectionScope lock(m_CriticalSection);
 
-		// initialise FramePro ONCE
-		if(!m_Initialised)
-		{
-			Initialise(p_framepro_tls);
-			m_Initialised = true;
-		}
+		Initialise(p_framepro_tls);
 
 		char date_str[64];
 		GetDateString(date_str, sizeof(date_str));
@@ -6580,12 +5731,7 @@ namespace FramePro
 
 		CriticalSectionScope lock(m_CriticalSection);
 
-		// initialise FramePro ONCE
-		if(!m_Initialised)
-		{
-			Initialise(p_framepro_tls);
-			m_Initialised = true;
-		}
+		Initialise(p_framepro_tls);
 
 		// initialise the connection
 		if(m_InitialiseConnectionNextFrame)
@@ -6621,11 +5767,8 @@ namespace FramePro
 
 		if(g_Connected)
 		{
-			// flush context switches
-#if FRAMEPRO_EVENT_TRACE_WIN32
-			if(mp_EventTraceWin32)
-				mp_EventTraceWin32->Flush();
-#endif
+			Platform::FlushContextSwitches(mp_ContextSwitchRecorder);
+
 			int64 wait_start_time;
 			FRAMEPRO_GET_CLOCK_COUNT(wait_start_time);
 
@@ -6667,11 +5810,14 @@ namespace FramePro
 		}
 
 		// stop recording if the file has become too big
-		//if (mp_RecordingFile && m_RecordingFileSize > m_MaxRecordingFileSize)
-		if (false)
-		{
-			StopRecording();
-		}
+		#if LIMIT_RECORDING_FILE_SIZE
+			if (m_MaxRecordingFileSize &&
+				m_RecordingFile.IsOpened() &&
+				m_RecordingFileSize > m_MaxRecordingFileSize)
+			{
+				StopRecording_NoLock();
+			}
+		#endif
 	}
 
 	//------------------------------------------------------------------------
@@ -6701,7 +5847,7 @@ namespace FramePro
 		}
 
 		if(g_Connected)
-			p_framepro_tls->OnConnected(mp_RecordingFile != NULL);
+			p_framepro_tls->OnConnected(m_RecordingFile.IsOpened());
 	}
 
 	//------------------------------------------------------------------------
@@ -6757,9 +5903,7 @@ namespace FramePro
 
 		// call immediately if already connected
 		if(g_Connected)
-			// START EPIC
-			p_callback(true, mp_RecordingFile ? mp_RecordingFile->GetArchiveName() : FString(), p_context);
-			// END EPIC
+			p_callback(true, m_RecordingFile.GetFilename().c_str(), p_context);
 
 		if(GetConnectionChangedCallbackIndex(p_callback) == -1)
 		{
@@ -6784,62 +5928,101 @@ namespace FramePro
 	}
 
 	//------------------------------------------------------------------------
-	// START EPIC
-	void FrameProSession::OnConnectionChanged(const FString& ProfileName, bool connected) const
-	// END EPIC
+	void FrameProSession::OnConnectionChanged(bool connected, const DynamicWString& filename) const
 	{
-
 		CriticalSectionScope lock(m_ConnectionChangedCriticalSection);
 
 		for(int i=0; i<m_Connectionchangedcallbacks.GetCount(); ++i)
 		{
 			const ConnectionChangedcallbackInfo& data = m_Connectionchangedcallbacks[i];
-			// START EPIC
-			data.mp_Callback(connected, ProfileName, data.mp_Context);
-			// END EPIC
+			data.mp_Callback(connected, filename.c_str(), data.mp_Context);
 		}
 	}
 
 	//------------------------------------------------------------------------
-// START EPIC
-	void FrameProSession::StartRecording(const FString& p_filename, bool context_switches, int64 max_file_size)
+	void FrameProSession::StartRecording(const char* p_filename, bool context_switches, bool callstacks, int64 max_file_size)
 	{
 		CriticalSectionScope lock(m_CriticalSection);
 
-		if(mp_RecordingFile)
+		if(m_RecordingFile.IsOpened())
 			StopRecording();
 
-#if ALLOW_DEBUG_FILES
-		mp_RecordingFile = IFileManager::Get().CreateDebugFileWriter(*p_filename);
-#else
-		mp_RecordingFile = IFileManager::Get().CreateFileWriter(*p_filename);
-#endif
+		bool opened_recording_file = m_RecordingFile.OpenForWrite(p_filename);
 
-		if(mp_RecordingFile)
+		if(opened_recording_file)
 		{
 			const char* p_id = "framepro_recording";
-			mp_RecordingFile->Serialize((uint8*)p_id, strlen(p_id));
+			m_RecordingFile.Write(p_id, strlen(p_id));
 
 			#if FRAMEPRO_SOCKETS_ENABLED
 				m_ListenSocket.Disconnect();		// don't allow connections while recording
 			#endif
 			
 			m_StartContextSwitchRecording = context_switches;
+
+			#if FRAMEPRO_ENABLE_CALLSTACKS
+				m_StartRecordingCallstacks = callstacks;
+			#else
+				FRAMEPRO_ASSERT(!callstacks);	// please define FRAMEPRO_ENABLE_CALLSTACKS to enable callstack recording
+			#endif
+
+			m_InitialiseConnectionNextFrame = true;
+
+			m_RecordingFileSize = 0;
+			m_MaxRecordingFileSize = max_file_size;
+		}
+		else
+		{
+			Platform::DebugWrite("FramePro ERROR: Failed to open recording file!");
+		}
+	}
+
+	//------------------------------------------------------------------------
+	void FrameProSession::StartRecording(const wchar_t* p_filename, bool context_switches, bool callstacks, int64 max_file_size)
+	{
+		CreateDefaultAllocator();
+
+		CriticalSectionScope lock(m_CriticalSection);
+
+		if(m_RecordingFile.IsOpened())
+			StopRecording();
+
+		bool opened_recording_file = m_RecordingFile.OpenForWrite(p_filename);
+		FRAMEPRO_ASSERT(opened_recording_file);		// Failed to open recording file
+
+		if(opened_recording_file)
+		{
+			#if FRAMEPRO_SOCKETS_ENABLED
+				m_ListenSocket.Disconnect();		// don't allow connections while recording
+			#endif
 			
+			m_StartContextSwitchRecording = context_switches;
+
+			#if FRAMEPRO_ENABLE_CALLSTACKS
+				m_StartRecordingCallstacks = callstacks;
+			#else
+				FRAMEPRO_ASSERT(!callstacks);	// please define FRAMEPRO_ENABLE_CALLSTACKS to enable callstack recording
+			#endif
+
 			m_InitialiseConnectionNextFrame = true;
 
 			m_RecordingFileSize = 0;
 			m_MaxRecordingFileSize = max_file_size;
 		}
 	}
-// END EPIC
 
 	//------------------------------------------------------------------------
 	void FrameProSession::StopRecording()
 	{
 		CriticalSectionScope lock(m_CriticalSection);
 
-		if(mp_RecordingFile)
+		StopRecording_NoLock();
+	}
+
+	//------------------------------------------------------------------------
+	void FrameProSession::StopRecording_NoLock()
+	{
+		if(m_RecordingFile.IsOpened())
 		{
 			#if FRAMEPRO_SOCKETS_ENABLED
 				OpenListenSocket();					// start the listening socket again so that we can accept new connections
@@ -6892,15 +6075,9 @@ namespace FramePro
 			uint count;
 			p_timer->GetAndClear(value, count);
 
-			const char* p_unit = "cycles";
-
 			// if we are connecting to FramePro, FramePro will ask for the string value later, otherwise we need to send it now
 			if (p_framepro_tls->SendStringsImmediately())
-			{
 				p_framepro_tls->SendString(p_timer->GetName(), PacketType::StringPacket);
-				p_framepro_tls->SendString(p_timer->GetGraph(), PacketType::StringPacket);
-				p_framepro_tls->SendString(p_unit, PacketType::StringPacket);
-			}
 
 			CriticalSectionScope lock(p_framepro_tls->GetCurrentSendBufferCriticalSection());
 
@@ -6912,8 +6089,6 @@ namespace FramePro
 			p_packet->m_Count = count;
 			p_packet->m_Name = (StringId)p_timer->GetName();
 			p_packet->m_Value = value;
-			p_packet->m_Graph = (StringId)p_timer->GetGraph();
-			p_packet->m_Unit = (StringId)p_unit;
 		}
 	}
 
@@ -6939,107 +6114,268 @@ namespace FramePro
 	{
 		g_ConditionalScopeMinTime = (int)((value * m_ClockFrequency) / 1000000LL);
 	}
-}
 
-//------------------------------------------------------------------------
-#endif		// #if FRAMEPRO_ENABLED
-//------------------------------------------------------------------------
-// FrameProStackTrace.cpp
+	//------------------------------------------------------------------------
+	void FrameProSession::SetScopeColour(StringId name, uint colour)
+	{
+		CriticalSectionScope lock(m_ScopeColoursLock);
 
+		bool updated_existing = false;
+		int count = m_ScopeColours.GetCount();
+		for (int i = 0; i < count; ++i)
+		{
+			if (m_ScopeColours[i].m_Name == name)
+			{
+				m_ScopeColours[i].m_Colour = colour;
+				updated_existing = true;
+				break;
+			}
+		}
 
+		if (!updated_existing)
+		{
+			ScopeColour scope_colour;
+			scope_colour.m_Name = name;
+			scope_colour.m_Colour = colour;
+			m_ScopeColours.Add(scope_colour);
+		}
 
-//------------------------------------------------------------------------
-#if FRAMEPRO_ENABLED
+		if (g_Connected)
+			GetFrameProTLS()->SendScopeColourPacket(name, colour);
+	}
 
-//------------------------------------------------------------------------
-#if FRAMEPRO_ENABLE_CALLSTACKS
+	//------------------------------------------------------------------------
+	void FrameProSession::SendScopeColours()
+	{
+		FrameProTLS* p_framepro_tls = GetFrameProTLS();
 
-//------------------------------------------------------------------------
-using namespace FramePro;
+		CriticalSectionScope lock(m_ScopeColoursLock);
 
-//------------------------------------------------------------------------
-// if both of these options are commented out it will use CaptureStackBackTrace (or backtrace on linux)
-#define FRAMEPRO_USE_STACKWALK64 0				// much slower but possibly more reliable. FRAMEPRO_USE_STACKWALK64 only implemented for x86 builds.
-#define FRAMEPRO_USE_RTLVIRTUALUNWIND 0			// reported to be faster than StackWalk64 - only available on x64 builds
-#define FRAMEPRO_USE_RTLCAPTURESTACKBACKTRACE 0	// system version of FRAMEPRO_USE_RTLVIRTUALUNWIND - only available on x64 builds
+		FRAMEPRO_ASSERT(g_Connected);
 
-//------------------------------------------------------------------------
-namespace FramePro
-{
-	//------------------------------------------------------------------------
-	// START EPIC
-	bool GetStackTrace(void** stack, int& stack_size, unsigned int& hash)
+		int count = m_ScopeColours.GetCount();
+		for (int i = 0; i < count; ++i)
+			p_framepro_tls->SendScopeColourPacket(m_ScopeColours[i].m_Name, m_ScopeColours[i].m_Colour);
+	}
+
+	//------------------------------------------------------------------------
+	void FrameProSession::SetCustomStatGraph(StringId name, StringId graph)
 	{
-		return false;
+		// this needs to be outside the critical section lock because it might lock the critical section itself
+		FrameProTLS* p_framepro_tls = GetFrameProTLS();
+
+		CriticalSectionScope lock(m_CustomStatInfoLock);
+
+		Initialise(p_framepro_tls);
+
+		bool updated_existing = false;
+		int count = m_CustomStatGraphs.GetCount();
+		for (int i = 0; i < count; ++i)
+		{
+			if (m_CustomStatGraphs[i].m_Name == name)
+			{
+				m_CustomStatGraphs[i].m_Value = graph;
+				updated_existing = true;
+				break;
+			}
+		}
+
+		if (!updated_existing)
+		{
+			CustomStatInfo custom_stat_info;
+			custom_stat_info.m_Name = name;
+			custom_stat_info.m_Value = graph;
+			m_CustomStatGraphs.Add(custom_stat_info);
+		}
+
+		if (g_Connected)
+			GetFrameProTLS()->SendCustomStatGraphPacket(name, graph);
 	}
-	// END EPIC
-}
 
-//------------------------------------------------------------------------
-StackTrace::StackTrace(Allocator* p_allocator)
-:	m_StackCount(0),
-	m_StackHash(0),
-	m_CallstackSet(p_allocator)
-{
-	memset(m_Stack, 0, sizeof(m_Stack));
-}
+	//------------------------------------------------------------------------
+	void FrameProSession::SetCustomStatUnit(StringId name, StringId unit)
+	{
+		// this needs to be outside the critical section lock because it might lock the critical section itself
+		FrameProTLS* p_framepro_tls = GetFrameProTLS();
 
-//------------------------------------------------------------------------
-void StackTrace::Clear()
-{
-	m_CallstackSet.Clear();
-}
+		CriticalSectionScope lock(m_CustomStatInfoLock);
 
-//------------------------------------------------------------------------
-CallstackResult StackTrace::Capture()
-{
-	CallstackResult result;
+		Initialise(p_framepro_tls);
+
+		bool updated_existing = false;
+		int count = m_CustomStatUnits.GetCount();
+		for (int i = 0; i < count; ++i)
+		{
+			if (m_CustomStatUnits[i].m_Name == name)
+			{
+				m_CustomStatUnits[i].m_Value = unit;
+				updated_existing = true;
+				break;
+			}
+		}
 
-	result.m_IsNew = false;
+		if (!updated_existing)
+		{
+			CustomStatInfo custom_stat_info;
+			custom_stat_info.m_Name = name;
+			custom_stat_info.m_Value = unit;
+			m_CustomStatUnits.Add(custom_stat_info);
+		}
 
-	memset(m_Stack, 0, sizeof(m_Stack));
+		if (g_Connected)
+			GetFrameProTLS()->SendCustomStatUnitPacket(name, unit);
+	}
 
-	if (!GetStackTrace(m_Stack, m_StackCount, m_StackHash))
+	//------------------------------------------------------------------------
+	void FrameProSession::SetCustomStatColour(StringId name, uint colour)
 	{
-		result.mp_Callstack = NULL;
-		return result;
+		// this needs to be outside the critical section lock because it might lock the critical section itself
+		FrameProTLS* p_framepro_tls = GetFrameProTLS();
+
+		CriticalSectionScope lock(m_CustomStatInfoLock);
+
+		Initialise(p_framepro_tls);
+
+		bool updated_existing = false;
+		int count = m_CustomStatColours.GetCount();
+		for (int i = 0; i < count; ++i)
+		{
+			if (m_CustomStatColours[i].m_Name == name)
+			{
+				m_CustomStatColours[i].m_Colour = colour;
+				updated_existing = true;
+				break;
+			}
+		}
+
+		if (!updated_existing)
+		{
+			ScopeColour custom_stat_colour;
+			custom_stat_colour.m_Name = name;
+			custom_stat_colour.m_Colour = colour;
+			m_CustomStatColours.Add(custom_stat_colour);
+		}
+
+		if (g_Connected)
+			GetFrameProTLS()->SendCustomStatColourPacket(name, colour);
 	}
 
-	result.mp_Callstack = m_CallstackSet.Get((uint64*)m_Stack, m_StackCount, m_StackHash);
+	//------------------------------------------------------------------------
+	void FrameProSession::SendCustomStatGraphs()
+	{
+		FrameProTLS* p_framepro_tls = GetFrameProTLS();
+
+		CriticalSectionScope lock(m_CustomStatInfoLock);
+
+		FRAMEPRO_ASSERT(g_Connected);
 
-	if (!result.mp_Callstack)
+		int count = m_CustomStatGraphs.GetCount();
+		for (int i = 0; i < count; ++i)
+			p_framepro_tls->SendCustomStatGraphPacket(m_CustomStatGraphs[i].m_Name, m_CustomStatGraphs[i].m_Value);
+	}
+
+	//------------------------------------------------------------------------
+	void FrameProSession::SendCustomStatUnits()
 	{
-		result.mp_Callstack = m_CallstackSet.Add((uint64*)m_Stack, m_StackCount, m_StackHash);
-		result.m_IsNew = true;
+		FrameProTLS* p_framepro_tls = GetFrameProTLS();
+
+		CriticalSectionScope lock(m_CustomStatInfoLock);
+
+		FRAMEPRO_ASSERT(g_Connected);
+
+		int count = m_CustomStatUnits.GetCount();
+		for (int i = 0; i < count; ++i)
+			p_framepro_tls->SendCustomStatUnitPacket(m_CustomStatUnits[i].m_Name, m_CustomStatUnits[i].m_Value);
 	}
 
-	return result;
+	//------------------------------------------------------------------------
+	void FrameProSession::SendCustomStatColours()
+	{
+		FrameProTLS* p_framepro_tls = GetFrameProTLS();
+
+		CriticalSectionScope lock(m_CustomStatInfoLock);
+
+		FRAMEPRO_ASSERT(g_Connected);
+
+		int count = m_CustomStatColours.GetCount();
+		for (int i = 0; i < count; ++i)
+			p_framepro_tls->SendCustomStatColourPacket(m_CustomStatColours[i].m_Name, m_CustomStatColours[i].m_Colour);
+	}
 }
 
+
 //------------------------------------------------------------------------
-#endif		// #if FRAMEPRO_ENABLE_CALLSTACKS
+//
+// FrameProStackTrace.cpp
+//
 
 //------------------------------------------------------------------------
-#endif		// #ifdef ENABLE_MEMPRO
+#if FRAMEPRO_ENABLE_CALLSTACKS
+
 //------------------------------------------------------------------------
-// FrameProTLS.cpp
+namespace FramePro
+{
+	//------------------------------------------------------------------------
+	StackTrace::StackTrace(Allocator* p_allocator)
+		: m_StackCount(0),
+		m_StackHash(0),
+		m_CallstackSet(p_allocator)
+	{
+		memset(m_Stack, 0, sizeof(m_Stack));
+	}
+
+	//------------------------------------------------------------------------
+	void StackTrace::Clear()
+	{
+		m_CallstackSet.Clear();
+	}
+
+	//------------------------------------------------------------------------
+	CallstackResult StackTrace::Capture()
+	{
+		CallstackResult result;
+
+		result.m_IsNew = false;
+
+		memset(m_Stack, 0, sizeof(m_Stack));
+
+		if (!Platform::GetStackTrace(m_Stack, m_StackCount, m_StackHash))
+		{
+			result.mp_Callstack = NULL;
+			return result;
+		}
+
+		result.mp_Callstack = m_CallstackSet.Get((uint64*)m_Stack, m_StackCount, m_StackHash);
 
+		if (!result.mp_Callstack)
+		{
+			result.mp_Callstack = m_CallstackSet.Add((uint64*)m_Stack, m_StackCount, m_StackHash);
+			result.m_IsNew = true;
+		}
 
+		return result;
+	}
+}
 
 //------------------------------------------------------------------------
-#if FRAMEPRO_ENABLED
+#endif		// #if FRAMEPRO_ENABLE_CALLSTACKS
+
+
+//------------------------------------------------------------------------
+//
+// FrameProTLS.cpp
+//
+
+//------------------------------------------------------------------------
+#if FRAMEPRO_WIN_BASED_PLATFORM
+	#pragma warning(push)
+	#pragma warning(disable : 4127)
+#endif
 
 //------------------------------------------------------------------------
 namespace FramePro
 {
 	//------------------------------------------------------------------------
-// START EPIC
-	#if !FRAMEPRO_WIN_BASED_PLATFORM && !FRAMEPRO_UE4_BASED_PLATFORM
-// END EPIC
-		int FrameProTLS::m_NewThreadId = 1;
-	#endif
-
-	//------------------------------------------------------------------------
 	static const int g_FrameProTLSBufferMarker = 0xfbfbfbfb;
 
 	std::atomic<long> FrameProTLS::m_StringCount(0);
@@ -7051,8 +6387,7 @@ namespace FramePro
 		m_SendStringsImmediately(false),
 		mp_CurrentSendBuffer(NULL),
 		m_CurrentSendBufferSize(0),
-		m_ThreadId(0),
-		m_OSThreadId(GetCurrentThreadId()),
+		m_ThreadId(Platform::GetCurrentThreadId()),
 		m_HiResTimerScopeStartTime(0),
 		m_HiResTimerStartTime(0),
 		m_ActiveHiResTimerIndex(-1),
@@ -7076,16 +6411,10 @@ namespace FramePro
 		,m_SendCallstacks(false)
 #endif
 	{
-// START EPIC
-		#if FRAMEPRO_WIN_BASED_PLATFORM || FRAMEPRO_UE4_BASED_PLATFORM
-// END EPIC
-			m_ThreadId = (int)m_OSThreadId;
-		#else
-			m_ThreadId = m_NewThreadId++;
-		#endif
-
 		UpdateSendStringsImmediatelyFlag();
 
+		m_InitialisedCustomStats.SetAllocator(p_allocator);
+
 		memset(m_FalseSharingSpacerBuffer, g_FrameProTLSBufferMarker, sizeof(m_FalseSharingSpacerBuffer));
 
 		m_HiResTimers.SetAllocator(p_allocator);
@@ -7453,60 +6782,64 @@ namespace FramePro
 		int string_len = (int)wcslen(p_str);
 		FRAMEPRO_ASSERT(string_len <= INT_MAX);
 
-#if FRAMEPRO_WIN_BASED_PLATFORM
-		int string_size = string_len * sizeof(wchar_t);
-		int aligned_string_size = AlignUpPow2(string_size, 4);
-		int size_to_allocate = sizeof(StringPacket) + aligned_string_size;
-
 		StringPacket* p_packet = NULL;
+		int size_to_allocate = 0;
+
+		if(sizeof(wchar_t) == 2)
 		{
-			CriticalSectionScope lock(m_SessionInfoBufferLock);
+			int string_size = string_len * sizeof(wchar_t);
+			int aligned_string_size = AlignUpPow2(string_size, 4);
+			size_to_allocate = sizeof(StringPacket) + aligned_string_size;
 
-			p_packet = (StringPacket*)(m_SessionInfoBuffer.Allocate(size_to_allocate));
-			if(!p_packet)
 			{
-				ShowMemoryWarning();
-				return;
-			}
+				CriticalSectionScope lock(m_SessionInfoBufferLock);
 
-			p_packet->m_PacketType = packet_type;
-			p_packet->m_Length = (int)string_len;
-			p_packet->m_StringId = string_id;
-			memcpy(p_packet + 1, p_str, string_size);
-		}
-#else
-		static_assert(sizeof(wchar_t) == 4, "Expected wchar size to be 4 on this platform");
-		int string_size = string_len * 2;
-		int aligned_string_size = AlignUpPow2(string_size, 4);
-		int size_to_allocate = sizeof(StringPacket) + aligned_string_size;
+				p_packet = (StringPacket*)(m_SessionInfoBuffer.Allocate(size_to_allocate));
+				if(!p_packet)
+				{
+					ShowMemoryWarning();
+					return;
+				}
 
-		StringPacket* p_packet = NULL;
+				p_packet->m_PacketType = packet_type;
+				p_packet->m_Length = (int)string_len;
+				p_packet->m_StringId = string_id;
+				memcpy(p_packet + 1, p_str, string_size);
+			}
+		}
+		else
 		{
-			CriticalSectionScope lock(m_SessionInfoBufferLock);
+			FRAMEPRO_ASSERT(sizeof(wchar_t) == 4);	// FramePro only supports 2 or 4 byte wchars
+			int string_size = string_len * 2;
+			int aligned_string_size = AlignUpPow2(string_size, 4);
+			size_to_allocate = sizeof(StringPacket) + aligned_string_size;
 
-			p_packet = (StringPacket*)(m_SessionInfoBuffer.Allocate(size_to_allocate));
-			if (!p_packet)
 			{
-				ShowMemoryWarning();
-				return;
-			}
+				CriticalSectionScope lock(m_SessionInfoBufferLock);
 
-			p_packet->m_PacketType = packet_type;
-			p_packet->m_Length = (int)string_len;
-			p_packet->m_StringId = string_id;
+				p_packet = (StringPacket*)(m_SessionInfoBuffer.Allocate(size_to_allocate));
+				if (!p_packet)
+				{
+					ShowMemoryWarning();
+					return;
+				}
 
-			// convert UTF-32 to UTF-16 by truncating (take only first 2 bytes of the 4 bytes)
-			FRAMEPRO_ASSERT(sizeof(wchar_t) == 4);
-			char* p_dest = (char*)(p_packet + 1);
-			char* p_source = (char*)p_str;
-			for (int i = 0; i < string_len; ++i)
-			{
-				*p_dest++ = *p_source++;
-				*p_dest++ = *p_source++;
-				p_source += 2;
+				p_packet->m_PacketType = packet_type;
+				p_packet->m_Length = (int)string_len;
+				p_packet->m_StringId = string_id;
+
+				// convert UTF-32 to UTF-16 by truncating (take only first 2 bytes of the 4 bytes)
+				char* p_dest = (char*)(p_packet + 1);
+				char* p_source = (char*)p_str;
+				for (int i = 0; i < string_len; ++i)
+				{
+					*p_dest++ = *p_source++;
+					*p_dest++ = *p_source++;
+					p_source += 2;
+				}
 			}
 		}
-#endif
+
 		if(m_Connected)
 			Send(p_packet, size_to_allocate);
 	}
@@ -7560,6 +6893,66 @@ namespace FramePro
 		SendPacket(packet);
 	}
 
+	//------------------------------------------------------------------------
+	void FrameProTLS::SendScopeColourPacket(StringId name, uint colour)
+	{
+		FRAMEPRO_ASSERT(IsOnTLSThread());
+		FRAMEPRO_ASSERT(m_Connected);
+
+		ScopeColourPacket packet;
+
+		packet.m_PacketType = PacketType::ScopeColourPacket;
+		packet.m_Colour = colour;
+		packet.m_Name = name;
+
+		SendPacket(packet);
+	}
+
+	//------------------------------------------------------------------------
+	void FrameProTLS::SendCustomStatGraphPacket(StringId name, StringId graph)
+	{
+		FRAMEPRO_ASSERT(IsOnTLSThread());
+		FRAMEPRO_ASSERT(m_Connected);
+
+		CustomStatInfoPacket packet;
+
+		packet.m_PacketType = PacketType::CustomStatGraphPacket;
+		packet.m_Name = name;
+		packet.m_Value = graph;
+
+		SendPacket(packet);
+	}
+
+	//------------------------------------------------------------------------
+	void FrameProTLS::SendCustomStatUnitPacket(StringId name, StringId unit)
+	{
+		FRAMEPRO_ASSERT(IsOnTLSThread());
+		FRAMEPRO_ASSERT(m_Connected);
+
+		CustomStatInfoPacket packet;
+
+		packet.m_PacketType = PacketType::CustomStatUnitPacket;
+		packet.m_Name = name;
+		packet.m_Value = unit;
+
+		SendPacket(packet);
+	}
+
+	//------------------------------------------------------------------------
+	void FrameProTLS::SendCustomStatColourPacket(StringId name, uint colour)
+	{
+		FRAMEPRO_ASSERT(IsOnTLSThread());
+		FRAMEPRO_ASSERT(m_Connected);
+
+		CustomStatColourPacket packet;
+
+		packet.m_PacketType = PacketType::CustomStatColourPacket;
+		packet.m_Colour = colour;
+		packet.m_Name = name;
+
+		SendPacket(packet);
+	}
+
 	//------------------------------------------------------------------------
 	void FrameProTLS::SendString(const char* p_string, PacketType::Enum packet_type)
 	{
@@ -7579,7 +6972,6 @@ namespace FramePro
 	}
 
 	//------------------------------------------------------------------------
-#if FRAMEPRO_SOCKETS_ENABLED
 	void FrameProTLS::SendStringLiteral(StringLiteralType::Enum string_literal_type, StringId string_id)
 	{
 		switch(string_literal_type)
@@ -7600,16 +6992,19 @@ namespace FramePro
 				SendString(string_id, (const char*)string_id, PacketType::StringPacket);
 				break;
 
+			case StringLiteralType::GeneralStringW:
+				SendString(string_id, (const wchar_t*)string_id, PacketType::WStringPacket);
+				break;
+
 			case StringLiteralType::StringLiteralTimerName:
 				SendString(string_id, (const char*)string_id, PacketType::StringLiteralTimerNamePacket);
 				break;
 
 			default:
-				FRAMEPRO_BREAK();
+				FramePro::DebugBreak();
 				break;
 		}
 	}
-#endif
 
 	//------------------------------------------------------------------------
 	void FrameProTLS::ShowMemoryWarning() const
@@ -7620,7 +7015,7 @@ namespace FramePro
 
 		if(now - last_warn_time >= m_ClockFrequency)
 		{
-			OutputDebugString(FRAMEPRO_STRING("Warning: FramePro failed to allocate enough memory."));
+			Platform::DebugWrite("Warning: FramePro failed to allocate enough memory.");
 			last_warn_time = now;
 		}
 	}
@@ -7880,7 +7275,7 @@ namespace FramePro
 	}
 
 	//------------------------------------------------------------------------
-	void FrameProTLS::SetCustomTimeSpanStat(StringId name, int64 value, const char* p_unit)
+	void FrameProTLS::SetCustomTimeSpanStat(StringId name, int64 value)
 	{
 		FRAMEPRO_ASSERT(IsOnTLSThread());
 
@@ -7895,7 +7290,6 @@ namespace FramePro
 			p_packet->m_ThreadId = m_ThreadId;
 			p_packet->m_ValueType = CustomStatValueType::Int64;
 			p_packet->m_Name = name;
-			p_packet->m_Unit = (StringId&)p_unit;
 			p_packet->m_ValueInt64 = value;
 			p_packet->m_ValueDouble = 0.0;
 			p_packet->m_Time = time;
@@ -7905,7 +7299,7 @@ namespace FramePro
 	}
 
 	//------------------------------------------------------------------------
-	void FrameProTLS::SetCustomTimeSpanStat(StringId name, int64 value, const wchar_t* p_unit)
+	void FrameProTLS::SetCustomTimeSpanStat(StringId name, double value)
 	{
 		FRAMEPRO_ASSERT(IsOnTLSThread());
 
@@ -7918,11 +7312,10 @@ namespace FramePro
 			TimeSpanCustomStatPacket* p_packet = (TimeSpanCustomStatPacket*)AllocateSpaceInBuffer(sizeof(TimeSpanCustomStatPacket));
 			p_packet->m_PacketType = PacketType::TimeSpanCustomStatPacket;
 			p_packet->m_ThreadId = m_ThreadId;
-			p_packet->m_ValueType = CustomStatValueType::Int64;
+			p_packet->m_ValueType = CustomStatValueType::Double;
 			p_packet->m_Name = name;
-			p_packet->m_Unit = (StringId&)p_unit;
-			p_packet->m_ValueInt64 = value;
-			p_packet->m_ValueDouble = 0.0;
+			p_packet->m_ValueInt64 = 0;
+			p_packet->m_ValueDouble = value;
 			p_packet->m_Time = time;
 		}
 
@@ -7930,7 +7323,7 @@ namespace FramePro
 	}
 
 	//------------------------------------------------------------------------
-	void FrameProTLS::SetCustomTimeSpanStat(StringId name, double value, const char* p_unit)
+	void FrameProTLS::SetCustomTimeSpanStatW(StringId name, int64 value)
 	{
 		FRAMEPRO_ASSERT(IsOnTLSThread());
 
@@ -7941,13 +7334,12 @@ namespace FramePro
 			CriticalSectionScope lock(m_CurrentSendBufferCriticalSection);
 
 			TimeSpanCustomStatPacket* p_packet = (TimeSpanCustomStatPacket*)AllocateSpaceInBuffer(sizeof(TimeSpanCustomStatPacket));
-			p_packet->m_PacketType = PacketType::TimeSpanCustomStatPacket;
+			p_packet->m_PacketType = PacketType::TimeSpanCustomStatPacketW;
 			p_packet->m_ThreadId = m_ThreadId;
-			p_packet->m_ValueType = CustomStatValueType::Double;
+			p_packet->m_ValueType = CustomStatValueType::Int64;
 			p_packet->m_Name = name;
-			p_packet->m_Unit = (StringId&)p_unit;
-			p_packet->m_ValueInt64 = 0;
-			p_packet->m_ValueDouble = value;
+			p_packet->m_ValueInt64 = value;
+			p_packet->m_ValueDouble = 0.0;
 			p_packet->m_Time = time;
 		}
 
@@ -7955,7 +7347,7 @@ namespace FramePro
 	}
 
 	//------------------------------------------------------------------------
-	void FrameProTLS::SetCustomTimeSpanStat(StringId name, double value, const wchar_t* p_unit)
+	void FrameProTLS::SetCustomTimeSpanStatW(StringId name, double value)
 	{
 		FRAMEPRO_ASSERT(IsOnTLSThread());
 
@@ -7966,11 +7358,10 @@ namespace FramePro
 			CriticalSectionScope lock(m_CurrentSendBufferCriticalSection);
 
 			TimeSpanCustomStatPacket* p_packet = (TimeSpanCustomStatPacket*)AllocateSpaceInBuffer(sizeof(TimeSpanCustomStatPacket));
-			p_packet->m_PacketType = PacketType::TimeSpanCustomStatPacket;
+			p_packet->m_PacketType = PacketType::TimeSpanCustomStatPacketW;
 			p_packet->m_ThreadId = m_ThreadId;
 			p_packet->m_ValueType = CustomStatValueType::Double;
 			p_packet->m_Name = name;
-			p_packet->m_Unit = (StringId&)p_unit;
 			p_packet->m_ValueInt64 = 0;
 			p_packet->m_ValueDouble = value;
 			p_packet->m_Time = time;
@@ -7980,46 +7371,148 @@ namespace FramePro
 	}
 
 	//------------------------------------------------------------------------
-#if FRAMEPRO_ENABLE_CALLSTACKS
-	CallstackResult FrameProTLS::GetCallstack()
+	bool FrameProTLS::HaveSentCustomStatInfo(StringId name)
 	{
-		return m_StackTrace.Capture();
+		int name_index = (int)name;
+		return name_index < m_InitialisedCustomStats.GetCount() && m_InitialisedCustomStats[name_index];
 	}
-#endif
-}
-
-//------------------------------------------------------------------------
-#endif		// #if FRAMEPRO_ENABLED
-//------------------------------------------------------------------------
-// IncrementingBlockAllocator.cpp
-
-
-
-//------------------------------------------------------------------------
-#if FRAMEPRO_ENABLED
 
-//------------------------------------------------------------------------
-namespace FramePro
-{
 	//------------------------------------------------------------------------
-	IncrementingBlockAllocator::IncrementingBlockAllocator(Allocator* p_allocator)
-	:	mp_Allocator(p_allocator),
-		mp_BlockList(NULL),
-		m_CurrentBlockSize(m_MemoryBlockSize),
-		m_MemorySize(0)
+	void FrameProTLS::SetHaveSentCustomStatInfo(StringId name)
 	{
+		int name_index = (int)name;
+
+		int old_count = m_InitialisedCustomStats.GetCount();
+		if (old_count <= name_index)
+		{
+			int new_count = name_index + 1;
+			m_InitialisedCustomStats.Resize(new_count);
+			memset(&m_InitialisedCustomStats[old_count], 0, (new_count - old_count) * sizeof(bool));
+		}
+		m_InitialisedCustomStats[name_index] = true;
 	}
 
 	//------------------------------------------------------------------------
-	IncrementingBlockAllocator::~IncrementingBlockAllocator()
+	void FrameProTLS::SetCustomStatInfo(const char* p_name, const char* p_graph, const char* p_unit, uint colour)
 	{
-		Clear();
-	}
+		StringId name = RegisterString(p_name);
 
+		if (!HaveSentCustomStatInfo(name))
+		{
+			SendCustomStatGraphPacket(name, RegisterString(p_graph));
+			SendCustomStatUnitPacket(name, RegisterString(p_unit));
+
+			if(colour)
+				SendCustomStatColourPacket(name, colour);
+
+			SetHaveSentCustomStatInfo(name);
+		}
+	}
+		
 	//------------------------------------------------------------------------
-	void IncrementingBlockAllocator::Clear()
+	void FrameProTLS::SetCustomStatInfo(const wchar_t* p_name, const wchar_t* p_graph, const wchar_t* p_unit, uint colour)
 	{
-		Block* p_block = mp_BlockList;
+		StringId name = RegisterString(p_name);
+
+		if (HaveSentCustomStatInfo(name))
+		{
+			SendCustomStatGraphPacket(name, RegisterString(p_graph));
+			SendCustomStatUnitPacket(name, RegisterString(p_unit));
+
+			if(colour)
+				SendCustomStatColourPacket(name, colour);
+
+			SetHaveSentCustomStatInfo(name);
+		}
+	}
+
+	//------------------------------------------------------------------------
+	void FrameProTLS::SetCustomStatInfo(StringId name, const char* p_graph, const char* p_unit, uint colour)
+	{
+		if (HaveSentCustomStatInfo(name))
+		{
+			SendCustomStatGraphPacket(name, RegisterString(p_graph));
+			SendCustomStatUnitPacket(name, RegisterString(p_unit));
+
+			if(colour)
+				SendCustomStatColourPacket(name, colour);
+
+			SetHaveSentCustomStatInfo(name);
+		}
+	}
+
+	//------------------------------------------------------------------------
+	void FrameProTLS::SetCustomStatInfo(StringId name, const wchar_t* p_graph, const wchar_t* p_unit, uint colour)
+	{
+		if (HaveSentCustomStatInfo(name))
+		{
+			SendCustomStatGraphPacket(name, RegisterString(p_graph));
+			SendCustomStatUnitPacket(name, RegisterString(p_unit));
+
+			if(colour)
+				SendCustomStatColourPacket(name, colour);
+
+			SetHaveSentCustomStatInfo(name);
+		}
+	}
+
+	//------------------------------------------------------------------------
+	void FrameProTLS::SetCustomStatInfo(StringId name, StringId graph, StringId unit, uint colour)
+	{
+		if (HaveSentCustomStatInfo(name))
+		{
+			SendCustomStatGraphPacket(name, graph);
+			SendCustomStatUnitPacket(name, unit);
+
+			if(colour)
+				SendCustomStatColourPacket(name, colour);
+
+			SetHaveSentCustomStatInfo(name);
+		}
+	}
+
+	//------------------------------------------------------------------------
+#if FRAMEPRO_ENABLE_CALLSTACKS
+	CallstackResult FrameProTLS::GetCallstack()
+	{
+		return m_StackTrace.Capture();
+	}
+#endif
+}
+
+//------------------------------------------------------------------------
+#if FRAMEPRO_WIN_BASED_PLATFORM
+	#pragma warning(pop)
+#endif
+
+
+//------------------------------------------------------------------------
+//
+// IncrementingBlockAllocator.cpp
+//
+
+//------------------------------------------------------------------------
+namespace FramePro
+{
+	//------------------------------------------------------------------------
+	IncrementingBlockAllocator::IncrementingBlockAllocator(Allocator* p_allocator)
+	:	mp_Allocator(p_allocator),
+		mp_BlockList(NULL),
+		m_CurrentBlockSize(m_MemoryBlockSize),
+		m_MemorySize(0)
+	{
+	}
+
+	//------------------------------------------------------------------------
+	IncrementingBlockAllocator::~IncrementingBlockAllocator()
+	{
+		Clear();
+	}
+
+	//------------------------------------------------------------------------
+	void IncrementingBlockAllocator::Clear()
+	{
+		Block* p_block = mp_BlockList;
 		while(p_block)
 		{
 			Block* p_next = p_block->mp_Next;
@@ -8055,106 +7548,99 @@ namespace FramePro
 	}
 }
 
-//------------------------------------------------------------------------
-#endif		// #if FRAMEPRO_ENABLED
-//------------------------------------------------------------------------
-// PointerSet.cpp
-
-
 
 //------------------------------------------------------------------------
-#if FRAMEPRO_ENABLED
+//
+// PointerSet.cpp
+//
 
 //------------------------------------------------------------------------
 namespace FramePro
 {
+	//------------------------------------------------------------------------
 	const int g_InitialCapacity = 32;
-}
-
-//------------------------------------------------------------------------
-FramePro::PointerSet::PointerSet(Allocator* p_allocator)
-:	mp_Data((const void**)p_allocator->Alloc(g_InitialCapacity*sizeof(const void*))),
-	m_CapacityMask(g_InitialCapacity-1),
-	m_Count(0),
-	m_Capacity(g_InitialCapacity),
-	mp_Allocator(p_allocator)
-{
-	memset(mp_Data, 0, g_InitialCapacity*sizeof(const void*));
-}
 
-//------------------------------------------------------------------------
-FramePro::PointerSet::~PointerSet()
-{
-	mp_Allocator->Free(mp_Data);
-}
+	//------------------------------------------------------------------------
+	PointerSet::PointerSet(Allocator* p_allocator)
+	:	mp_Data((const void**)p_allocator->Alloc(g_InitialCapacity*sizeof(const void*))),
+		m_CapacityMask(g_InitialCapacity-1),
+		m_Count(0),
+		m_Capacity(g_InitialCapacity),
+		mp_Allocator(p_allocator)
+	{
+		memset(mp_Data, 0, g_InitialCapacity*sizeof(const void*));
+	}
 
-//------------------------------------------------------------------------
-void FramePro::PointerSet::Grow()
-{
-	int old_capacity = m_Capacity;
-	const void** p_old_data = mp_Data;
+	//------------------------------------------------------------------------
+	PointerSet::~PointerSet()
+	{
+		mp_Allocator->Free(mp_Data);
+	}
 
-	// allocate a new set
-	m_Capacity = m_Capacity ? 2*m_Capacity : 32;
-	FRAMEPRO_ASSERT(m_Capacity < (int)(INT_MAX/sizeof(void*)));
+	//------------------------------------------------------------------------
+	void PointerSet::Grow()
+	{
+		int old_capacity = m_Capacity;
+		const void** p_old_data = mp_Data;
 
-	m_CapacityMask = m_Capacity - 1;
-	size_t alloc_size = m_Capacity * sizeof(const void*);
-	mp_Data = (const void**)mp_Allocator->Alloc(alloc_size);
+		// allocate a new set
+		m_Capacity = m_Capacity ? 2*m_Capacity : 32;
+		FRAMEPRO_ASSERT(m_Capacity < (int)(INT_MAX/sizeof(void*)));
 
-	int size = m_Capacity * sizeof(void*);
-	memset(mp_Data, 0, size);
+		m_CapacityMask = m_Capacity - 1;
+		size_t alloc_size = m_Capacity * sizeof(const void*);
+		mp_Data = (const void**)mp_Allocator->Alloc(alloc_size);
 
-	// transfer pointers from old set
-	m_Count = 0;
-	for(int i=0; i<old_capacity; ++i)
-	{
-		const void* p = p_old_data[i];
-		if(p)
-			Add(p);
-	}
+		int size = m_Capacity * sizeof(void*);
+		memset(mp_Data, 0, size);
 
-	// release old buffer
-	mp_Allocator->Free(p_old_data);
+		// transfer pointers from old set
+		m_Count = 0;
+		for(int i=0; i<old_capacity; ++i)
+		{
+			const void* p = p_old_data[i];
+			if(p)
+				Add(p);
+		}
 
-	alloc_size -= old_capacity * sizeof(const void*);
-}
+		// release old buffer
+		mp_Allocator->Free(p_old_data);
 
-//------------------------------------------------------------------------
-// return true if added, false if already in set
-bool FramePro::PointerSet::AddInternal(const void* p, int64 hash, int index)
-{
-	if(m_Count >= m_Capacity/4)
-	{
-		Grow();
-		index = hash & m_CapacityMask;
+		alloc_size -= old_capacity * sizeof(const void*);
 	}
 
-	const void* p_existing = mp_Data[index];
-	while(p_existing)
+	//------------------------------------------------------------------------
+	// return true if added, false if already in set
+	bool PointerSet::AddInternal(const void* p, int64 hash, int index)
 	{
-		if(p_existing == p)
-			return false;
-		index = (index + 1) & m_CapacityMask;
-		p_existing = mp_Data[index];
-	}
-
-	mp_Data[index] = p;
+		if(m_Count >= m_Capacity/4)
+		{
+			Grow();
+			index = hash & m_CapacityMask;
+		}
 
-	++m_Count;
+		const void* p_existing = mp_Data[index];
+		while(p_existing)
+		{
+			if(p_existing == p)
+				return false;
+			index = (index + 1) & m_CapacityMask;
+			p_existing = mp_Data[index];
+		}
 
-	return true;
-}
+		mp_Data[index] = p;
 
-//------------------------------------------------------------------------
-#endif		// #if FRAMEPRO_ENABLED
-//------------------------------------------------------------------------
-// SendBuffer.cpp
+		++m_Count;
 
+		return true;
+	}
+}
 
 
 //------------------------------------------------------------------------
-#if FRAMEPRO_ENABLED
+//
+// SendBuffer.cpp
+//
 
 //------------------------------------------------------------------------
 namespace FramePro
@@ -8207,405 +7693,3468 @@ namespace FramePro
 	}
 }
 
-//------------------------------------------------------------------------
-#endif		// #if FRAMEPRO_ENABLED
-//------------------------------------------------------------------------
-// Socket.cpp
-
-// START EPIC
-// Remove FRAMEPRO_UNIX_BASED_PLATFORM
-// END EPIC
 
 //------------------------------------------------------------------------
-#if FRAMEPRO_ENABLED
+//
+// Socket.cpp
+//
 
 //------------------------------------------------------------------------
-#if FRAMEPRO_ENABLED && FRAMEPRO_SOCKETS_ENABLED
+#if FRAMEPRO_SOCKETS_ENABLED
 
 //------------------------------------------------------------------------
 namespace FramePro
 {
-	volatile int g_InitialiseCount = 0;
-}
+	//------------------------------------------------------------------------
+	Socket::Socket()
+	:	m_Listening(false)
+	{
+		Platform::CreateSocket(m_OSSocketMem, sizeof(m_OSSocketMem));
+	}
 
-//------------------------------------------------------------------------
-bool FramePro::Socket::InitialiseWSA()
-{
-	if(g_InitialiseCount == 0)
-	{
-#if FRAMEPRO_PLATFORM_XBOX360
-		  XNetStartupParams xnsp;
-		  memset(&xnsp, 0, sizeof(xnsp));
-		  xnsp.cfgSizeOfStruct = sizeof(XNetStartupParams);
-		  xnsp.cfgFlags = XNET_STARTUP_BYPASS_SECURITY;
-		  INT err = XNetStartup(&xnsp);
-#endif
+	//------------------------------------------------------------------------
+	Socket::~Socket()
+	{
+		Platform::UninitialiseSocketSystem();
+	}
 
-#if FRAMEPRO_WIN_BASED_PLATFORM
-		// Initialize Winsock
-		WSADATA wsaData;
-		if(WSAStartup(MAKEWORD(2,2), &wsaData) != 0)
-		{
-			HandleError();
-			return false;
-		}
-#endif
-// BEGIN EPIC
-#if FRAMEPRO_PLATFORM_SWITCH
-		if( !FramePro::PlatformDebugNetInit() )
+	//------------------------------------------------------------------------
+	void Socket::Disconnect()
+	{
+		Platform::DisconnectSocket(m_OSSocketMem, m_Listening);
+	}
+
+	//------------------------------------------------------------------------
+	bool Socket::StartListening()
+	{
+		FRAMEPRO_ASSERT(IsValid());
+
+		if(!Platform::StartSocketListening(m_OSSocketMem))
 		{
+			Platform::HandleSocketError();
 			return false;
 		}
-#endif
-// END EPIC
 
+		m_Listening = true;
+
+		return true;
 	}
 
-	++g_InitialiseCount;
+	//------------------------------------------------------------------------
+	bool Socket::Bind(const char* p_port)
+	{
+		FRAMEPRO_ASSERT(!IsValid());
 
-	return true;
-}
+		if(!Platform::InitialiseSocketSystem())
+			return false;
 
-//------------------------------------------------------------------------
-void FramePro::Socket::CleanupWSA()
-{
-	--g_InitialiseCount;
+		return Platform::BindSocket(m_OSSocketMem, p_port);
+	}
 
-	if(g_InitialiseCount == 0)
+	//------------------------------------------------------------------------
+	bool Socket::Accept(Socket& client_socket)
 	{
-#if FRAMEPRO_WIN_BASED_PLATFORM
-		if(WSACleanup() == FRAMEPRO_SOCKET_ERROR)
-			HandleError();
-#endif
-
-#if FRAMEPRO_PLATFORM_XBOX360
-		 XNetCleanup();
-#endif
+		FRAMEPRO_ASSERT(!client_socket.IsValid());
+		return Platform::AcceptSocket(m_OSSocketMem, client_socket.m_OSSocketMem);
 	}
-}
 
-//------------------------------------------------------------------------
-void FramePro::Socket::Disconnect()
-{
-	if(m_Socket != FRAMEPRO_INVALID_SOCKET)
+	//------------------------------------------------------------------------
+	bool Socket::Send(const void* p_buffer, size_t size)
 	{
-#if FRAMEPRO_WIN_BASED_PLATFORM
-		if(!m_Listening && shutdown(m_Socket, SD_BOTH) == FRAMEPRO_SOCKET_ERROR)
-			HandleError();
-#else
-		if(shutdown(m_Socket, SHUT_RDWR) == FRAMEPRO_SOCKET_ERROR)
-			HandleError();
-#endif
+		FRAMEPRO_ASSERT(size >= 0 && size <= INT_MAX);
 
-		// loop until the socket is closed to ensure all data is sent
-		unsigned int buffer = 0;
-		size_t ret = 0;
-		do { ret = recv(m_Socket, (char*)&buffer, sizeof(buffer), 0); } while(ret != 0 && ret != (size_t)FRAMEPRO_SOCKET_ERROR);
+		int bytes_to_send = (int)size;
+		while(bytes_to_send != 0)
+		{
+			int bytes_sent = 0;
+			if(!Platform::SocketSend(m_OSSocketMem, p_buffer, bytes_to_send, bytes_sent))
+			{
+				Platform::HandleSocketError();
+				Disconnect();
+				return false;
+			}
+			p_buffer = (char*)p_buffer + bytes_sent;
+			bytes_to_send -= bytes_sent;
+		}
 
-#if FRAMEPRO_WIN_BASED_PLATFORM
-	    if(closesocket(m_Socket) == FRAMEPRO_SOCKET_ERROR)
-			HandleError();
-#else
-		close(m_Socket);
-#endif
-		m_Socket = FRAMEPRO_INVALID_SOCKET;
+		return true;
 	}
-}
 
-//------------------------------------------------------------------------
-bool FramePro::Socket::StartListening()
-{
-	FRAMEPRO_ASSERT(m_Socket != FRAMEPRO_INVALID_SOCKET);
-
-	if (listen(m_Socket, SOMAXCONN) == FRAMEPRO_SOCKET_ERROR)
+	//------------------------------------------------------------------------
+	int Socket::Receive(void* p_buffer, int size)
 	{
-		HandleError();
-		return false;
-	}
+		int total_bytes_received = 0;
 
-	m_Listening = true;
+		while(size)
+		{
+			int bytes_received = 0;
+			bool result = Platform::SocketReceive(m_OSSocketMem, (char*)p_buffer, size, bytes_received);
 
-	return true;
-}
+			if (!result)
+			{
+				Platform::HandleSocketError();
+				Disconnect();
+				return total_bytes_received;
+			}
+			else if(bytes_received == 0)
+			{
+				Disconnect();
+				return bytes_received;
+			}
 
-//------------------------------------------------------------------------
-bool FramePro::Socket::Bind(const char* p_port)
-{
-	FRAMEPRO_ASSERT(m_Socket == FRAMEPRO_INVALID_SOCKET);
+			total_bytes_received += bytes_received;
 
-	if(!InitialiseWSA())
-		return false;
+			size -= bytes_received;
+			FRAMEPRO_ASSERT(size >= 0);
 
-#if FRAMEPRO_WIN_BASED_PLATFORM
-	// setup the addrinfo struct
-	addrinfo info;
-	ZeroMemory(&info, sizeof(info));
-	info.ai_family = AF_INET;
-	info.ai_socktype = SOCK_STREAM;
-	info.ai_protocol = IPPROTO_TCP;
-	info.ai_flags = AI_PASSIVE;
-
-	// Resolve the server address and port
-	addrinfo* p_result_info;
-	HRESULT result = getaddrinfo(NULL, p_port, &info, &p_result_info);
-	if (result != 0)
-	{
-		HandleError();
-		return false;
-	}
-
-	m_Socket = socket(
-		p_result_info->ai_family,
-		p_result_info->ai_socktype, 
-		p_result_info->ai_protocol);
-#else
-	m_Socket = socket(
-		AF_INET,
-		SOCK_STREAM,
-		IPPROTO_TCP);
-#endif
+			p_buffer = (char*)p_buffer + bytes_received;
+		}
 
-	if (m_Socket == FRAMEPRO_INVALID_SOCKET)
-	{
-#if FRAMEPRO_WIN_BASED_PLATFORM
-		freeaddrinfo(p_result_info);
-#endif
-		HandleError();
-		return false;
+		return total_bytes_received;
 	}
 
-	// Setup the TCP listening socket
-#if FRAMEPRO_WIN_BASED_PLATFORM
-	result = ::bind(m_Socket, p_result_info->ai_addr, (int)p_result_info->ai_addrlen);
-	freeaddrinfo(p_result_info);
-#else
-	// Bind to INADDR_ANY
-	sockaddr_in sa;
-	sa.sin_family = AF_INET;
-	sa.sin_addr.s_addr = INADDR_ANY;
-	int iport = atoi(p_port);
-	sa.sin_port = htons(iport);
-	int result = ::bind(m_Socket, (const sockaddr*)(&sa), sizeof(sockaddr_in));
-#endif
-
-	if (result == FRAMEPRO_SOCKET_ERROR)
+	//------------------------------------------------------------------------
+	bool Socket::IsValid() const
 	{
-		HandleError();
-		Disconnect();
-		return false;
+		return Platform::IsSocketValid(m_OSSocketMem);
 	}
-
-	return true;
-}
-
-//------------------------------------------------------------------------
-bool FramePro::Socket::Accept(Socket& client_socket)
-{
-	FRAMEPRO_ASSERT(client_socket.m_Socket == FRAMEPRO_INVALID_SOCKET);
-	client_socket.m_Socket = accept(m_Socket, NULL, NULL);
-	return client_socket.m_Socket != FRAMEPRO_INVALID_SOCKET;
 }
 
 //------------------------------------------------------------------------
-bool FramePro::Socket::Send(const void* p_buffer, size_t size)
-{
-	FRAMEPRO_ASSERT(size >= 0 && size <= INT_MAX);
-
-	int bytes_to_send = (int)size;
-	while(bytes_to_send != 0)
-	{
-// START EPIC
-#if FRAMEPRO_PLATFORM_UNIX
-// END EPIC
-		int flags = MSG_NOSIGNAL;
-#else
-		int flags = 0;
-#endif
+#endif		// #if FRAMEPRO_SOCKETS_ENABLED
 
-		int bytes_sent = (int)send(m_Socket, (char*)p_buffer, bytes_to_send, flags);
-		if(bytes_sent == FRAMEPRO_SOCKET_ERROR)
-		{
-			HandleError();
-			Disconnect();
-			return false;
-		}
-		p_buffer = (char*)p_buffer + bytes_sent;
-		bytes_to_send -= bytes_sent;
-	}
 
-	return true;
-}
+//------------------------------------------------------------------------
+//
+// Thread.cpp
+//
 
 //------------------------------------------------------------------------
-int FramePro::Socket::Receive(void* p_buffer, int size)
+namespace FramePro
 {
-	int total_bytes_received = 0;
+	//------------------------------------------------------------------------
+	Thread::Thread()
+	:	m_Alive(false),
+		m_ThreadTerminatedEvent(false, false)
+	{
+	}
 
-	while(size)
+	//------------------------------------------------------------------------
+	void Thread::CreateThread(FramePro::ThreadMain p_thread_main, void* p_param, Allocator* p_allocator)
 	{
-		int bytes_received = (int)recv(m_Socket, (char*)p_buffer, size, 0);
+		mp_ThreadMain = p_thread_main;
+		mp_Param = p_param;
 
-		if(bytes_received == 0)
-		{
-			Disconnect();
-			return bytes_received;
-		}
-		else if(bytes_received == FRAMEPRO_SOCKET_ERROR)
-		{
-			HandleError();
-			Disconnect();
-			return total_bytes_received;
-		}
+		Platform::CreateThread(m_OSThread, sizeof(m_OSThread), ThreadMain, this, p_allocator);
+	}
 
-		total_bytes_received += bytes_received;
+	//------------------------------------------------------------------------
+	int Thread::ThreadMain(void* p_context)
+	{
+		Thread* p_thread = (Thread*)p_context;
+		p_thread->m_Alive = true;
+		unsigned long ret = (unsigned long)p_thread->mp_ThreadMain(p_thread->mp_Param);
+		p_thread->m_Alive = false;
+		p_thread->m_ThreadTerminatedEvent.Set();
+		return ret;
+	}
 
-		size -= bytes_received;
-		FRAMEPRO_ASSERT(size >= 0);
+	//------------------------------------------------------------------------
+	void Thread::SetPriority(int priority)
+	{
+		Platform::SetThreadPriority(m_OSThread, priority);
+	}
 
-		p_buffer = (char*)p_buffer + bytes_received;
+	//------------------------------------------------------------------------
+	void Thread::SetAffinity(int affinity)
+	{
+		Platform::SetThreadAffinity(m_OSThread, affinity);
 	}
 
-	return total_bytes_received;
-}
-
-//------------------------------------------------------------------------
-void FramePro::Socket::HandleError()
-{
-#if FRAMEPRO_WIN_BASED_PLATFORM
-	if(WSAGetLastError() == WSAEADDRINUSE)
+	//------------------------------------------------------------------------
+	void Thread::WaitForThreadToTerminate(int timeout)
 	{
-		OutputDebugString(FRAMEPRO_STRING("FramePro: Network connection conflict. Please make sure that other FramePro enabled applications are shut down, or change the port in the the FramePro lib and FramePro settings.\n"));
-		return;
+		m_ThreadTerminatedEvent.Wait(timeout);
 	}
+}
 
-	int buffer_size = 1024;
-	TCHAR* p_buffer = (TCHAR*)HeapAlloc(GetProcessHeap(), 0, buffer_size*sizeof(TCHAR));
-	memset(p_buffer, 0, buffer_size*sizeof(TCHAR));
 
-	va_list args;
-	FormatMessage(
-		FORMAT_MESSAGE_FROM_SYSTEM,
-		NULL,
-		WSAGetLastError(),
-		0,
-		p_buffer,
-		buffer_size,
-		&args);
+//------------------------------------------------------------------------
+//
+// EnumModulesLinux.hpp
+//
 
-	DebugWrite(FRAMEPRO_STRING("FramePro Network Error: %s\n"), p_buffer);
+//------------------------------------------------------------------------
+#if FRAMEPRO_LINUX_BASED_PLATFORM
 
-	HeapFree(GetProcessHeap(), 0, p_buffer);
-#endif
-}
+//------------------------------------------------------------------------
+#if FRAMEPRO_ENABLE_CALLSTACKS
 
 //------------------------------------------------------------------------
-#endif		// #if FRAMEPRO_ENABLED && FRAMEPRO_SOCKETS_ENABLED
+namespace FramePro
+{
+	//------------------------------------------------------------------------
+	namespace EnumModulesLinux
+	{
+		void EnumerateModules(Array<ModulePacket*>& module_packets, Allocator* p_allocator);
+	}
+}
 
 //------------------------------------------------------------------------
-#endif		// #if FRAMEPRO_ENABLED
+#endif		// #if FRAMEPRO_ENABLE_CALLSTACKS
+
 //------------------------------------------------------------------------
-// Thread.cpp
+#endif		// #if FRAMEPRO_LINUX_BASED_PLATFORM
 
 
+//------------------------------------------------------------------------
+//
+// EnumModulesLinux.cpp
+//
 
 //------------------------------------------------------------------------
-#if FRAMEPRO_ENABLED
+#if FRAMEPRO_LINUX_BASED_PLATFORM
 
 //------------------------------------------------------------------------
-// START EPIC
-// Remove FRAMEPRO_UNIX_BASED_PLATFORM
-// END EPIC
+#if FRAMEPRO_ENABLE_CALLSTACKS
 
 //------------------------------------------------------------------------
-FramePro::Thread::Thread()
-#if FRAMEPRO_WIN_BASED_PLATFORM
-	:	m_Handle(0),
-		m_Alive(false),
-		m_ThreadTerminatedEvent(false, false)
-// BEGIN EPIC
-#elif FRAMEPRO_UE4_BASED_PLATFORM
-	:	m_Alive(false),
-		m_ThreadTerminatedEvent(false,false)
-// END EPIC
-#else
-	:	m_Alive(false),
-		m_ThreadTerminatedEvent(false, false)
+#if FRAMEPRO_ENUMERATE_ALL_MODULES
+	#include <link.h>
 #endif
-{
-}
 
 //------------------------------------------------------------------------
-void FramePro::Thread::CreateThread(ThreadMain p_thread_main, void* p_param)
+#include <sys/types.h>
+#include <unistd.h>
+
+//------------------------------------------------------------------------
+namespace FramePro
 {
-	mp_ThreadMain = p_thread_main;
-	mp_Param = p_param;
+	//------------------------------------------------------------------------
+	void BaseAddressLookupFunction()
+	{
+	}
 
-#if FRAMEPRO_WIN_BASED_PLATFORM
-	m_Handle = ::CreateThread(NULL, 0, PlatformThreadMain, this, 0, NULL);
-// BEGIN EPIC
-#elif FRAMEPRO_UE4_BASED_PLATFORM
-	m_Runnable = TUniquePtr<FRunnableThread>( FRunnableThread::Create( this, TEXT("FramePro") ) );
-// END EPIC
-#else
-	pthread_create(&m_Thread, NULL, PlatformThreadMain, this);
+	//------------------------------------------------------------------------
+#if FRAMEPRO_ENUMERATE_ALL_MODULES
+
+	//------------------------------------------------------------------------
+	struct ModuleCallbackContext
+	{
+		Array<ModulePacket*>* mp_ModulePackets;
+		Allocator* mp_Allocator;
+	};
+
+	//------------------------------------------------------------------------
+	void EnumerateLoadedModulesCallback(
+		int64 module_base,
+		const char* p_module_name,
+		bool use_lookup_function_for_base_address,
+		ModuleCallbackContext* p_context)
+	{
+		ModulePacket* p_module_packet = (ModulePacket*)p_context->mp_Allocator->Alloc(sizeof(ModulePacket));
+		memset(p_module_packet, 0, sizeof(ModulePacket));
+		p_module_packet->m_PacketType = PacketType::ModulePacket;
+
+		p_module_packet->m_ModuleBase = module_base;
+
+		size_t module_name_length = strlen(p_module_name) + 1;
+		FRAMEPRO_ASSERT(sizeof(p_module_packet->m_ModuleName) >= module_name_length);
+		memcpy(p_module_packet->m_ModuleName, p_module_name, module_name_length);
+
+		const char* p_last_slash = strrchr(p_module_name, '/');
+		p_last_slash = p_last_slash ? p_last_slash + 1 : p_module_name;
+		char filename[FRAMEPRO_MAX_PATH];
+		sprintf(filename, "%s.sym_txt", p_last_slash);
+		memcpy(p_module_packet->m_SymbolFilename, filename, strlen(filename) + 1);
+		
+		p_module_packet->m_UseLookupFunctionForBaseAddress = use_lookup_function_for_base_address ? 1 : 0;
+
+		p_context->mp_ModulePackets->Add(p_module_packet);
+	}
+
+	//------------------------------------------------------------------------
+	int EnumerateLoadedModulesCallback(struct dl_phdr_info* info, size_t size, void* data)
+	{
+		ModuleCallbackContext* p_context = (ModuleCallbackContext*)data;
+
+		int64 module_base = 0;
+		for (int j = 0; j < info->dlpi_phnum; j++)
+		{
+			if (info->dlpi_phdr[j].p_type == PT_LOAD)
+			{
+				module_base = info->dlpi_addr + info->dlpi_phdr[j].p_vaddr;
+				break;
+			}
+		}
+
+		static bool first = true;
+		if (first)
+		{
+			first = false;
+
+			module_base = (int64)BaseAddressLookupFunction;		// use the address of the BaseAddressLookupFunction function so that we can work it out later
+
+			// get the module name
+			char arg1[20];
+			char char_filename[FRAMEPRO_MAX_PATH];
+			sprintf(arg1, "/proc/%d/exe", getpid());
+			memset(char_filename, 0, FRAMEPRO_MAX_PATH);
+			readlink(arg1, char_filename, FRAMEPRO_MAX_PATH - 1);
+
+			EnumerateLoadedModulesCallback(
+				module_base,
+				char_filename,
+				true,
+				p_context);
+		}
+		else
+		{
+			EnumerateLoadedModulesCallback(
+				module_base,
+				info->dlpi_name,
+				false,
+				p_context);
+		}
+
+		return 0;
+	}
+#endif			// #if FRAMEPRO_ENUMERATE_ALL_MODULES
+
+	//------------------------------------------------------------------------
+	void EnumModulesLinux::EnumerateModules(Array<ModulePacket*>& module_packets, Allocator* p_allocator)
+	{
+		// if you are having problems compiling this on your platform unset FRAMEPRO_ENUMERATE_ALL_MODULES and it send info for just the main module
+#if FRAMEPRO_ENUMERATE_ALL_MODULES
+		ModuleCallbackContext* p_context = (ModuleCallbackContext*)p_allocator->Alloc(sizeof(ModuleCallbackContext));
+		p_context->mp_ModulePackets = &module_packets;
+		p_context->mp_Allocator = p_allocator;
+
+		dl_iterate_phdr(EnumerateLoadedModulesCallback, p_context);
+
+		Delete(p_allocator, p_context);
 #endif
+
+		if (!module_packets.GetCount())
+		{
+			// if FRAMEPRO_ENUMERATE_ALL_MODULES is set or enumeration failed for some reason, fall back
+			// to getting the base address for the main module. This will always work for for all platforms.
+
+			ModulePacket* p_module_packet = (ModulePacket*)p_allocator->Alloc(sizeof(ModulePacket));
+			memset(p_module_packet, 0, sizeof(ModulePacket));
+
+			p_module_packet->m_PacketType = PacketType::ModulePacket;
+			p_module_packet->m_UseLookupFunctionForBaseAddress = 0;
+
+			p_module_packet->m_UseLookupFunctionForBaseAddress = 1;
+
+			p_module_packet->m_ModuleBase = (int64)BaseAddressLookupFunction;		// use the address of the BaseAddressLookupFunction function so that we can work it out later
+
+			// get the module name
+			char arg1[20];
+			sprintf(arg1, "/proc/%d/exe", getpid());
+			memset(p_module_packet->m_ModuleName, 0, FRAMEPRO_MAX_PATH);
+			readlink(arg1, p_module_packet->m_ModuleName, FRAMEPRO_MAX_PATH - 1);
+
+			const char* p_last_slash = strrchr(p_module_packet->m_ModuleName, '/');
+			p_last_slash = p_last_slash ? p_last_slash + 1 : p_module_packet->m_ModuleName;
+			char filename[FRAMEPRO_MAX_PATH];
+			sprintf(filename, "%s.sym_txt", p_last_slash);
+			memcpy(p_module_packet->m_SymbolFilename, filename, strlen(filename) + 1);
+		
+			module_packets.Add(p_module_packet);
+		}
+	}
 }
 
+//------------------------------------------------------------------------
+#endif		// #if FRAMEPRO_ENABLE_CALLSTACKS
+
+//------------------------------------------------------------------------
+#endif		// #if FRAMEPRO_LINUX_BASED_PLATFORM
+
+
+//------------------------------------------------------------------------
+//
+// EnumModulesWindows.hpp
+//
+
 //------------------------------------------------------------------------
 #if FRAMEPRO_WIN_BASED_PLATFORM
-unsigned long WINAPI FramePro::Thread::PlatformThreadMain(void* p_param)
-{
-	Thread* p_thread = (Thread*)p_param;
-	p_thread->m_Alive = true;
-	unsigned long ret = (unsigned long)p_thread->mp_ThreadMain(p_thread->mp_Param);
-	p_thread->m_Alive = false;
-	p_thread->m_ThreadTerminatedEvent.Set();
-	return ret;
-}
-// BEGIN EPIC
-#elif FRAMEPRO_UE4_BASED_PLATFORM
-uint32 FramePro::Thread::Run()
-{
-	m_Alive = true;
-	mp_ThreadMain(mp_Param);
-	m_Alive = false;
-	m_ThreadTerminatedEvent.Set();
-	return 0;
-}
-// END EPIC
-#else
-void* FramePro::Thread::PlatformThreadMain(void* p_param)
+
+//------------------------------------------------------------------------
+#if FRAMEPRO_ENABLE_CALLSTACKS
+
+//------------------------------------------------------------------------
+namespace FramePro
 {
-	Thread* p_thread = (Thread*)p_param;
-	p_thread->m_Alive = true;
-	p_thread->mp_ThreadMain(p_thread->mp_Param);
-	p_thread->m_Alive = false;
-	p_thread->m_ThreadTerminatedEvent.Set();
-	return NULL;
+	//------------------------------------------------------------------------
+	namespace EnumModulesWindows
+	{
+		void EnumerateModules(Array<ModulePacket*>& module_packets, Allocator* p_allocator);
+	}
 }
-#endif
 
 //------------------------------------------------------------------------
-void FramePro::Thread::SetPriority(int priority)
-{
+#endif		// #if FRAMEPRO_ENABLE_CALLSTACKS
+
+//------------------------------------------------------------------------
+#endif		// #if FRAMEPRO_WIN_BASED_PLATFORM
+
+
+//------------------------------------------------------------------------
+//
+// EnumModulesWindows.cpp
+//
+
+//------------------------------------------------------------------------
 #if FRAMEPRO_WIN_BASED_PLATFORM
-	::SetThreadPriority(m_Handle, priority);
-#endif
-}
 
 //------------------------------------------------------------------------
-void FramePro::Thread::SetAffinity(int affinity)
-{
-#if FRAMEPRO_WIN_BASED_PLATFORM && !FRAMEPRO_PLATFORM_UWP
-	SetThreadAffinityMask(m_Handle, affinity);
+#if FRAMEPRO_ENABLE_CALLSTACKS
+
+//------------------------------------------------------------------------
+#if FRAMEPRO_PLATFORM_UE4
+	#include "Windows/AllowWindowsPlatformTypes.h"
+#endif
+#pragma warning(push)
+#pragma warning(disable:4668)
+#pragma warning(disable:4091)
+	#ifndef WIN32_LEAN_AND_MEAN
+		#define WIN32_LEAN_AND_MEAN
+	#endif
+	#include <windows.h>
+	#if FRAMEPRO_ENUMERATE_ALL_MODULES
+		#include <Dbghelp.h>
+	#endif
+#pragma warning(pop)
+#if FRAMEPRO_PLATFORM_UE4
+	#include "Windows/HideWindowsPlatformTypes.h"
+#endif
+
+#if FRAMEPRO_ENUMERATE_ALL_MODULES
+	#pragma comment(lib, "Dbghelp.lib")
 #endif
-}
 
 //------------------------------------------------------------------------
-void FramePro::Thread::WaitForThreadToTerminate(int timeout)
+namespace FramePro
 {
-	m_ThreadTerminatedEvent.Wait(timeout);
+	#if FRAMEPRO_ENUMERATE_ALL_MODULES
+		//------------------------------------------------------------------------
+		struct CV_HEADER
+		{
+			int Signature;
+			int Offset;
+		};
+
+		struct CV_INFO_PDB20
+		{
+			CV_HEADER CvHeader;
+			int Signature;
+			int Age;
+			char PdbFileName[FRAMEPRO_MAX_PATH];
+		};
+
+		struct CV_INFO_PDB70
+		{
+			int  CvSignature;
+			GUID Signature;
+			int Age;
+			char PdbFileName[FRAMEPRO_MAX_PATH];
+		};
+
+		//------------------------------------------------------------------------
+		struct ModuleCallbackContext
+		{
+			Array<ModulePacket*>* mp_ModulePackets;
+			Allocator* mp_Allocator;
+		};
+
+		//------------------------------------------------------------------------
+		void GetExtraModuleInfo(int64 ModuleBase, ModulePacket* p_module_packet)
+		{
+			IMAGE_DOS_HEADER* p_dos_header = (IMAGE_DOS_HEADER*)ModuleBase;
+			IMAGE_NT_HEADERS* p_nt_header = (IMAGE_NT_HEADERS*)((char*)ModuleBase + p_dos_header->e_lfanew);
+			IMAGE_OPTIONAL_HEADER& optional_header = p_nt_header->OptionalHeader;
+			IMAGE_DATA_DIRECTORY& image_data_directory = optional_header.DataDirectory[IMAGE_DIRECTORY_ENTRY_DEBUG];
+			IMAGE_DEBUG_DIRECTORY* p_debug_info_array = (IMAGE_DEBUG_DIRECTORY*)(ModuleBase + image_data_directory.VirtualAddress);
+			int count = image_data_directory.Size / sizeof(IMAGE_DEBUG_DIRECTORY);
+			for(int i=0; i<count; ++i)
+			{
+				if(p_debug_info_array[i].Type == IMAGE_DEBUG_TYPE_CODEVIEW)
+				{
+					char* p_cv_data = (char*)(ModuleBase + p_debug_info_array[i].AddressOfRawData);
+					if(strncmp(p_cv_data, "RSDS", 4) == 0)
+					{
+						CV_INFO_PDB70* p_cv_info = (CV_INFO_PDB70*)p_cv_data;
+
+						p_module_packet->m_PacketType = PacketType::ModulePacket;
+						p_module_packet->m_Age = p_cv_info->Age;
+					
+						static_assert(sizeof(p_module_packet->m_Sig) == sizeof(p_cv_info->Signature), "sig size wrong");
+						memcpy(p_module_packet->m_Sig, &p_cv_info->Signature, sizeof(p_cv_info->Signature));
+
+						FRAMEPRO_ASSERT(strlen(p_cv_info->PdbFileName) < sizeof(p_module_packet->m_SymbolFilename));
+						memcpy(p_module_packet->m_SymbolFilename, p_cv_info->PdbFileName, strlen(p_cv_info->PdbFileName) + 1);
+
+						return;									// returning here
+					}
+					else if(strncmp(p_cv_data, "NB10", 4) == 0)
+					{
+						CV_INFO_PDB20* p_cv_info = (CV_INFO_PDB20*)p_cv_data;
+
+						p_module_packet->m_PacketType = PacketType::ModulePacket;
+						p_module_packet->m_Age = p_cv_info->Age;
+
+						memset(p_module_packet->m_Sig, 0, sizeof(p_module_packet->m_Sig));
+						static_assert(sizeof(p_cv_info->Signature) <= sizeof(p_module_packet->m_Sig), "sig size wrong");
+						memcpy(p_module_packet->m_Sig, &p_cv_info->Signature, sizeof(p_cv_info->Signature));
+
+						FRAMEPRO_ASSERT(sizeof(p_module_packet->m_SymbolFilename) >= strlen(p_cv_info->PdbFileName) + 1);
+						memcpy(p_module_packet->m_SymbolFilename, p_cv_info->PdbFileName, strlen(p_cv_info->PdbFileName) + 1);
+
+						return;									// returning here
+					}
+				}
+			}
+		}
+
+		//------------------------------------------------------------------------
+		void EnumerateLoadedModulesCallback(
+			int64 module_base,
+			const char* p_module_name,
+			bool use_lookup_function_for_base_address,
+			ModuleCallbackContext* p_context)
+		{
+			ModulePacket* p_module_packet = (ModulePacket*)p_context->mp_Allocator->Alloc(sizeof(ModulePacket));
+			memset(p_module_packet, 0, sizeof(ModulePacket));
+			p_module_packet->m_PacketType = PacketType::ModulePacket;
+
+			p_module_packet->m_ModuleBase = module_base;
+			
+			size_t module_name_length = strlen(p_module_name) + 1;
+			FRAMEPRO_ASSERT(sizeof(p_module_packet->m_ModuleName) >= module_name_length);
+			memcpy(p_module_packet->m_ModuleName, p_module_name, module_name_length);
+
+			p_module_packet->m_UseLookupFunctionForBaseAddress = use_lookup_function_for_base_address ? 1 : 0;
+
+			GetExtraModuleInfo(module_base, p_module_packet);
+
+			p_context->mp_ModulePackets->Add(p_module_packet);
+		}
+
+		#if !defined(_IMAGEHLP_SOURCE_) && defined(_IMAGEHLP64)
+		// depending on your platform you may need to change PCSTR to PSTR for ModuleName
+		BOOL CALLBACK EnumerateLoadedModulesCallback(__in PCSTR ModuleName,__in DWORD64 ModuleBase,__in ULONG,__in_opt PVOID UserContext)
+		#else
+		BOOL CALLBACK EnumerateLoadedModulesCallback(__in PCSTR ModuleName,__in ULONG ModuleBase,__in ULONG,__in_opt PVOID UserContext)
+		#endif
+		{
+			ModuleCallbackContext* p_context = (ModuleCallbackContext*)UserContext;
+
+			EnumerateLoadedModulesCallback(ModuleBase, ModuleName, false, p_context);
+
+			return true;
+		}
+	#endif			// #if FRAMEPRO_ENUMERATE_ALL_MODULES
+
+	//------------------------------------------------------------------------
+	void EnumModulesWindows::EnumerateModules(Array<ModulePacket*>& module_packets, Allocator* p_allocator)
+	{
+		// if you are having problems compiling this on your platform unset FRAMEPRO_ENUMERATE_ALL_MODULES and it send info for just the main module
+		#if FRAMEPRO_ENUMERATE_ALL_MODULES
+			ModuleCallbackContext* p_context = (ModuleCallbackContext*)p_allocator->Alloc(sizeof(ModuleCallbackContext));
+			p_context->mp_ModulePackets = &module_packets;
+			p_context->mp_Allocator = p_allocator;
+
+			EnumerateLoadedModules64(GetCurrentProcess(), EnumerateLoadedModulesCallback, p_context);
+
+			Delete(p_allocator, p_context);
+		#endif
+
+		if (!module_packets.GetCount())
+		{
+			// if FRAMEPRO_ENUMERATE_ALL_MODULES is set or enumeration failed for some reason, fall back
+			// to getting the base address for the main module. This will always work for for all platforms.
+
+			ModulePacket* p_module_packet = (ModulePacket*)p_allocator->Alloc(sizeof(ModulePacket));
+			memset(p_module_packet, 0, sizeof(ModulePacket));
+
+			p_module_packet->m_PacketType = PacketType::ModulePacket;
+			p_module_packet->m_UseLookupFunctionForBaseAddress = 0;
+
+			static int module = 0;
+			HMODULE module_handle = 0;
+			GetModuleHandleEx(GET_MODULE_HANDLE_EX_FLAG_FROM_ADDRESS, (LPCTSTR)&module, &module_handle);
+
+			p_module_packet->m_ModuleBase = (int64)module_handle;
+
+			FRAMEPRO_TCHAR tchar_filename[FRAMEPRO_MAX_PATH] = { 0 };
+			GetModuleFileName(NULL, tchar_filename, FRAMEPRO_MAX_PATH);
+
+			#ifdef UNICODE
+				size_t chars_converted = 0;
+				wcstombs_s(&chars_converted, p_module_packet->m_ModuleName, tchar_filename, FRAMEPRO_MAX_PATH);
+			#else
+				strcpy_s(p_module_packet->m_ModuleName, tchar_filename);
+			#endif
+
+			module_packets.Add(p_module_packet);
+		}
+	}
 }
 
+//------------------------------------------------------------------------
+#endif		// #if FRAMEPRO_ENABLE_CALLSTACKS
+
+//------------------------------------------------------------------------
+#endif		// #if FRAMEPRO_WIN_BASED_PLATFORM
+
+
+//------------------------------------------------------------------------
+//
+// FrameProPlatform.cpp
+//
+//------------------------------------------------------------------------
+// ---
+// --- FRAMEPRO PLATFORM IMPLEMENTATION START ---
+// ---
+//------------------------------------------------------------------------
+
+
+//------------------------------------------------------------------------
+#include <ctime>
+
+//------------------------------------------------------------------------
+// general Win/Linux implementations that apply to all platforms
+//------------------------------------------------------------------------
+#if FRAMEPRO_WIN_BASED_PLATFORM
+
+	//------------------------------------------------------------------------
+	#if FRAMEPRO_PLATFORM_WIN
+		#pragma warning(push)
+		#pragma warning(disable:4668)
+		#if FRAMEPRO_PLATFORM_UE4
+			#include "Windows/AllowWindowsPlatformTypes.h"
+		#endif
+		#ifndef WIN32_LEAN_AND_MEAN
+			#define WIN32_LEAN_AND_MEAN
+		#endif
+		#include <windows.h>
+		#include <psapi.h>
+		#if FRAMEPRO_PLATFORM_UE4
+			#include "Windows/HideWindowsPlatformTypes.h"
+		#endif
+		#pragma warning(pop)
+	#endif
+
+	//------------------------------------------------------------------------
+	#if defined(_MSC_VER) && _MSC_VER <= 1600
+		#error FramePro only supports Visual Studio 2012 and above. This is because it needs atomics. If you really need 2010 support please contact [REDACTED_EMAIL]
+	#endif
+
+	//------------------------------------------------------------------------
+	#if FRAMEPRO_MAX_PATH != MAX_PATH
+		#error
+	#endif
+
+	//------------------------------------------------------------------------
+	#if FRAMEPRO_SOCKETS_ENABLED
+		#pragma comment(lib, "Ws2_32.lib")
+
+		#if defined(AF_IPX) && !defined(_WINSOCK2API_)
+			#error winsock already defined. Please include winsock2.h before including windows.h or use WIN32_LEAN_AND_MEAN. See the FAQ for more info.
+		#endif
+		#if FRAMEPRO_PLATFORM_UE4
+			#include "Windows/AllowWindowsPlatformTypes.h"
+		#endif
+		#include <winsock2.h>
+		#include <ws2tcpip.h>
+		#if FRAMEPRO_PLATFORM_UE4
+			#include "Windows//HideWindowsPlatformTypes.h"
+		#endif
+	#endif
+
+	//------------------------------------------------------------------------
+	namespace FramePro
+	{
+		//------------------------------------------------------------------------
+		namespace GenericPlatform
+		{
+			//------------------------------------------------------------------------
+			void DebugWrite(const char* p_string)
+			{
+				OutputDebugStringA(p_string);
+			}
+
+			//------------------------------------------------------------------------
+			SRWLOCK& GetOSLock(void* p_os_lock_mem)
+			{
+				return *(SRWLOCK*)p_os_lock_mem;
+			}
+
+			//------------------------------------------------------------------------
+			void CreateLock(void* p_os_lock_mem, int os_lock_mem_size)
+			{
+				FRAMEPRO_ASSERT(os_lock_mem_size >= sizeof(SRWLOCK));
+				FRAMEPRO_UNREFERENCED(os_lock_mem_size);
+				InitializeSRWLock(&GetOSLock(p_os_lock_mem));
+			}
+
+			//------------------------------------------------------------------------
+			void DestroyLock(void*)
+			{
+				// do nothing
+			}
+
+			//------------------------------------------------------------------------
+			void TakeLock(void* p_os_lock_mem)
+			{
+				AcquireSRWLockExclusive(&GetOSLock(p_os_lock_mem));
+			}
+
+			//------------------------------------------------------------------------
+			void ReleaseLock(void* p_os_lock_mem)
+			{
+				ReleaseSRWLockExclusive(&GetOSLock(p_os_lock_mem));
+			}
+
+			//------------------------------------------------------------------------
+			void GetLocalTime(tm* p_tm, const time_t *p_time)
+			{
+				localtime_s(p_tm, p_time);
+			}
+
+			//------------------------------------------------------------------------
+			int GetCurrentProcessId()
+			{
+				return ::GetCurrentProcessId();
+			}
+
+			//------------------------------------------------------------------------
+			void VSPrintf(char* p_buffer, size_t const buffer_size, const char* p_format, va_list arg_list)
+			{
+				vsprintf_s(p_buffer, buffer_size, p_format, arg_list);
+			}
+
+			//------------------------------------------------------------------------
+			void ToString(int value, char* p_dest, int dest_size)
+			{
+				_itoa_s(value, p_dest, dest_size, 10);
+			}
+
+			//------------------------------------------------------------------------
+			HANDLE& GetOSEventHandle(void* p_os_event_mem)
+			{
+				return *(HANDLE*)p_os_event_mem;
+			}
+
+			//------------------------------------------------------------------------
+			void CreateEventX(void* p_os_event_mem, int os_event_mem_size, bool initial_state, bool auto_reset)
+			{
+				FRAMEPRO_ASSERT(os_event_mem_size >= sizeof(HANDLE));
+				FRAMEPRO_UNREFERENCED(os_event_mem_size);
+				GetOSEventHandle(p_os_event_mem) = ::CreateEvent(NULL, !auto_reset, initial_state, NULL);
+
+			}
+
+			//------------------------------------------------------------------------
+			void DestroyEvent(void* p_os_event_mem)
+			{
+				CloseHandle(GetOSEventHandle(p_os_event_mem));
+			}
+
+			//------------------------------------------------------------------------
+			void SetEvent(void* p_os_event_mem)
+			{
+				::SetEvent(GetOSEventHandle(p_os_event_mem));
+			}
+
+			//------------------------------------------------------------------------
+			void ResetEvent(void* p_os_event_mem)
+			{
+				::ResetEvent(GetOSEventHandle(p_os_event_mem));
+			}
+
+			//------------------------------------------------------------------------
+			int WaitEvent(void* p_os_event_mem, int timeout)
+			{
+				return WaitForSingleObject(GetOSEventHandle(p_os_event_mem), timeout) == 0/*WAIT_OBJECT_0*/;
+			}
+
+			//------------------------------------------------------------------------
+			#if FRAMEPRO_SOCKETS_ENABLED
+				volatile int g_WinSockInitialiseCount = 0;
+			#endif
+
+			//------------------------------------------------------------------------
+			void HandleSocketError()
+			{
+				#if FRAMEPRO_SOCKETS_ENABLED
+					if (WSAGetLastError() == WSAEADDRINUSE)
+					{
+						Platform::DebugWrite("FramePro: Network connection conflict. Please make sure that other FramePro enabled applications are shut down, or change the port in the the FramePro lib and FramePro settings.\n");
+						return;
+					}
+
+					int buffer_size = 1024;
+					FRAMEPRO_TCHAR* p_buffer = (FRAMEPRO_TCHAR*)HeapAlloc(GetProcessHeap(), 0, buffer_size * sizeof(FRAMEPRO_TCHAR));
+					memset(p_buffer, 0, buffer_size * sizeof(FRAMEPRO_TCHAR));
+
+					va_list args;
+					FormatMessage(
+						FORMAT_MESSAGE_FROM_SYSTEM,
+						NULL,
+						WSAGetLastError(),
+						0,
+						p_buffer,
+						buffer_size,
+						&args);
+
+					FramePro::DebugWrite("FramePro Network Error: %s\n", p_buffer);
+
+					HeapFree(GetProcessHeap(), 0, p_buffer);
+				#endif
+			}
+
+			//------------------------------------------------------------------------
+			bool InitialiseSocketSystem()
+			{
+				#if FRAMEPRO_SOCKETS_ENABLED
+					if (g_WinSockInitialiseCount == 0)
+					{
+						// Initialize Winsock
+						WSADATA wsaData;
+						if (WSAStartup(MAKEWORD(2, 2), &wsaData) != 0)
+						{
+							HandleSocketError();
+							return false;
+						}
+					}
+
+					++g_WinSockInitialiseCount;
+
+					return true;
+				#else
+					return false;
+				#endif
+			}
+
+			//------------------------------------------------------------------------
+			void UninitialiseSocketSystem()
+			{
+				#if FRAMEPRO_SOCKETS_ENABLED
+					--g_WinSockInitialiseCount;
+
+					if (g_WinSockInitialiseCount == 0)
+					{
+						if (WSACleanup() == SOCKET_ERROR)
+							HandleSocketError();
+					}
+				#endif
+			}
+
+			//------------------------------------------------------------------------
+			#if FRAMEPRO_SOCKETS_ENABLED
+				SOCKET& GetOSSocket(void* p_os_socket_mem)
+				{
+					return *(SOCKET*)p_os_socket_mem;
+				}
+			#endif
+
+			//------------------------------------------------------------------------
+			#if FRAMEPRO_SOCKETS_ENABLED
+				const SOCKET& GetOSSocket(const void* p_os_socket_mem)
+				{
+					return *(const SOCKET*)p_os_socket_mem;
+				}
+			#endif
+
+			//------------------------------------------------------------------------
+			void CreateSocket(void* p_os_socket_mem, int os_socket_mem_size)
+			{
+				#if FRAMEPRO_SOCKETS_ENABLED
+					FRAMEPRO_ASSERT(os_socket_mem_size >= sizeof(SOCKET));
+					FRAMEPRO_UNREFERENCED(os_socket_mem_size);
+					new (p_os_socket_mem)SOCKET();
+					GetOSSocket(p_os_socket_mem) = INVALID_SOCKET;
+				#else
+					FRAMEPRO_UNREFERENCED(p_os_socket_mem);
+					FRAMEPRO_UNREFERENCED(os_socket_mem_size);
+				#endif
+			}
+
+			//------------------------------------------------------------------------
+			void DestroySocket(void* p_os_socket_mem)
+			{
+				#if FRAMEPRO_SOCKETS_ENABLED
+					SOCKET& socket = GetOSSocket(p_os_socket_mem);
+					socket.~SOCKET();
+					socket = INVALID_SOCKET;
+				#else
+					FRAMEPRO_UNREFERENCED(p_os_socket_mem);
+				#endif
+			}
+
+			//------------------------------------------------------------------------
+			void DisconnectSocket(void* p_os_socket_mem, bool stop_listening)
+			{
+				#if FRAMEPRO_SOCKETS_ENABLED
+					SOCKET& socket = GetOSSocket(p_os_socket_mem);
+
+					if(socket != INVALID_SOCKET)
+					{
+						if(!stop_listening && shutdown(socket, SD_BOTH) == SOCKET_ERROR)
+							HandleSocketError();
+
+						// loop until the socket is closed to ensure all data is sent
+						unsigned int buffer = 0;
+						size_t ret = 0;
+						do { ret = recv(socket, (char*)&buffer, sizeof(buffer), 0); } while(ret != 0 && ret != (size_t)SOCKET_ERROR);
+
+						if(closesocket(socket) == SOCKET_ERROR)
+							HandleSocketError();
+
+						socket = INVALID_SOCKET;
+					}
+				#else
+					FRAMEPRO_UNREFERENCED(p_os_socket_mem);
+					FRAMEPRO_UNREFERENCED(stop_listening);
+				#endif
+			}
+
+			//------------------------------------------------------------------------
+			bool StartSocketListening(void* p_os_socket_mem)
+			{
+				#if FRAMEPRO_SOCKETS_ENABLED
+					return listen(GetOSSocket(p_os_socket_mem), SOMAXCONN) != SOCKET_ERROR;
+				#else
+					FRAMEPRO_UNREFERENCED(p_os_socket_mem);
+					return false;
+				#endif
+			}
+
+			//------------------------------------------------------------------------
+			bool BindSocket(void* p_os_socket_mem, const char* p_port)
+			{
+				#if FRAMEPRO_SOCKETS_ENABLED
+					// setup the addrinfo struct
+					addrinfo info;
+					ZeroMemory(&info, sizeof(info));
+					info.ai_family = AF_INET;
+					info.ai_socktype = SOCK_STREAM;
+					info.ai_protocol = IPPROTO_TCP;
+					info.ai_flags = AI_PASSIVE;
+
+					// Resolve the server address and port
+					addrinfo* p_result_info;
+					HRESULT result = getaddrinfo(NULL, p_port, &info, &p_result_info);
+					if (result != 0)
+					{
+						HandleSocketError();
+						return false;
+					}
+
+					SOCKET& socket = GetOSSocket(p_os_socket_mem);
+
+					socket = ::socket(
+						p_result_info->ai_family,
+						p_result_info->ai_socktype, 
+						p_result_info->ai_protocol);
+
+					if (socket == INVALID_SOCKET)
+					{
+						freeaddrinfo(p_result_info);
+						HandleSocketError();
+						return false;
+					}
+
+					// Setup the TCP listening socket
+					result = ::bind(socket, p_result_info->ai_addr, (int)p_result_info->ai_addrlen);
+					freeaddrinfo(p_result_info);
+
+					if (result == SOCKET_ERROR)
+					{
+						HandleSocketError();
+						DisconnectSocket(p_os_socket_mem, true);
+						return false;
+					}
+
+					return true;
+				#else
+					FRAMEPRO_UNREFERENCED(p_os_socket_mem);
+					FRAMEPRO_UNREFERENCED(p_port);
+					return false;
+				#endif
+			}
+
+			//------------------------------------------------------------------------
+			bool AcceptSocket(void* p_source_os_socket_mem, void* p_target_os_socket_mem)
+			{
+				#if FRAMEPRO_SOCKETS_ENABLED
+					SOCKET& source_socket = GetOSSocket(p_source_os_socket_mem);
+					SOCKET& target_socket = GetOSSocket(p_target_os_socket_mem);
+
+					target_socket = accept(source_socket, NULL, NULL);
+					return target_socket != INVALID_SOCKET;
+				#else
+					FRAMEPRO_UNREFERENCED(p_source_os_socket_mem);
+					FRAMEPRO_UNREFERENCED(p_target_os_socket_mem);
+					return false;
+				#endif
+			}
+
+			//------------------------------------------------------------------------
+			bool SocketSend(void* p_os_socket_mem, const void* p_buffer, int size, int& bytes_sent)
+			{
+				#if FRAMEPRO_SOCKETS_ENABLED
+					SOCKET& socket = GetOSSocket(p_os_socket_mem);
+					bytes_sent = (int)send(socket, (char*)p_buffer, size, 0);
+					return bytes_sent != SOCKET_ERROR;
+				#else
+					FRAMEPRO_UNREFERENCED(p_os_socket_mem);
+					FRAMEPRO_UNREFERENCED(p_buffer);
+					FRAMEPRO_UNREFERENCED(size);
+					FRAMEPRO_UNREFERENCED(bytes_sent);
+					return false;
+				#endif
+			}
+
+			//------------------------------------------------------------------------
+			bool SocketReceive(void* p_os_socket_mem, const void* p_buffer, int size, int& bytes_received)
+			{
+				#if FRAMEPRO_SOCKETS_ENABLED
+					SOCKET& socket = GetOSSocket(p_os_socket_mem);
+					bytes_received = (int)recv(socket, (char*)p_buffer, size, 0);
+					return bytes_received != SOCKET_ERROR;
+				#else
+					FRAMEPRO_UNREFERENCED(p_os_socket_mem);
+					FRAMEPRO_UNREFERENCED(p_buffer);
+					FRAMEPRO_UNREFERENCED(size);
+					FRAMEPRO_UNREFERENCED(bytes_received);
+					return false;
+				#endif
+			}
+
+			//------------------------------------------------------------------------
+			bool IsSocketValid(const void* p_os_socket_mem)
+			{
+				#if FRAMEPRO_SOCKETS_ENABLED
+					const SOCKET& socket = GetOSSocket(p_os_socket_mem);
+					return socket != INVALID_SOCKET;
+				#else
+					FRAMEPRO_UNREFERENCED(p_os_socket_mem);
+					return false;
+				#endif
+			}
+
+			//------------------------------------------------------------------------
+			HANDLE& GetOSThread(void* p_os_thread_mem)
+			{
+				return *(HANDLE*)p_os_thread_mem;
+			}
+
+			//------------------------------------------------------------------------
+			struct ThreadContext
+			{
+				ThreadMain mp_ThreadMain;
+				void* mp_Context;
+				Allocator* mp_Allocator;
+			};
+
+			//------------------------------------------------------------------------
+			unsigned long WINAPI PlatformThreadMain(void* p_param)
+			{
+				ThreadContext* p_context = (ThreadContext*)p_param;
+
+				int result = p_context->mp_ThreadMain(p_context->mp_Context);
+
+				Delete(p_context->mp_Allocator, p_context);
+
+				return result;
+			}
+
+			//------------------------------------------------------------------------
+			void CreateThread(
+				void* p_os_thread_mem,
+				int os_thread_mem_size,
+				ThreadMain p_thread_main,
+				void* p_context,
+				Allocator* p_allocator)
+			{
+				FRAMEPRO_ASSERT(os_thread_mem_size >= sizeof(HANDLE));
+				FRAMEPRO_UNREFERENCED(os_thread_mem_size);
+
+				ThreadContext* p_thread_context = New<ThreadContext>(p_allocator);
+				p_thread_context->mp_ThreadMain = p_thread_main;
+				p_thread_context->mp_Context = p_context;
+				p_thread_context->mp_Allocator = p_allocator;
+
+				GetOSThread(p_os_thread_mem) = ::CreateThread(NULL, 0, PlatformThreadMain, p_thread_context, 0, NULL);
+			}
+
+			//------------------------------------------------------------------------
+			void SetThreadPriority(void* p_os_thread_mem, int priority)
+			{
+				HANDLE& handle = GetOSThread(p_os_thread_mem);
+				::SetThreadPriority(handle, priority);
+			}
+
+			//------------------------------------------------------------------------
+			void SetThreadAffinity(void* p_os_thread_mem, int priority)
+			{
+				HANDLE& handle = GetOSThread(p_os_thread_mem);
+				SetThreadAffinityMask(handle, priority);
+			}
+
+			//------------------------------------------------------------------------
+			void* CreateContextSwitchRecorder(Allocator* p_allocator)
+			{
+				#if FRAMEPRO_PLATFORM_WIN
+					return EventTraceWin32::Create(p_allocator);
+				#else
+					FRAMEPRO_UNREFERENCED(p_allocator);
+					return NULL;
+				#endif
+			}
+
+			//------------------------------------------------------------------------
+			void DestroyContextSwitchRecorder(void* p_context_switch_recorder, Allocator* p_allocator)
+			{
+				#if FRAMEPRO_PLATFORM_WIN
+					EventTraceWin32::Destroy(p_context_switch_recorder, p_allocator);
+				#else
+					FRAMEPRO_UNREFERENCED(p_context_switch_recorder);
+					FRAMEPRO_UNREFERENCED(p_allocator);
+				#endif
+			}
+
+			//------------------------------------------------------------------------
+			bool StartRecordingContextSitches(
+				void* p_context_switch_recorder,
+				Platform::ContextSwitchCallbackFunction p_callback,
+				void* p_context,
+				DynamicString& error)
+			{
+				#if FRAMEPRO_PLATFORM_WIN
+					return EventTraceWin32::Start(
+						p_context_switch_recorder,
+						p_callback,
+						p_context,
+						error);
+				#else
+					FRAMEPRO_UNREFERENCED(p_context_switch_recorder);
+					FRAMEPRO_UNREFERENCED(p_context);
+					FRAMEPRO_UNREFERENCED(error);
+					return false;
+				#endif
+			}
+
+			//------------------------------------------------------------------------
+			void StopRecordingContextSitches(void* p_context_switch_recorder)
+			{
+				#if FRAMEPRO_PLATFORM_WIN
+					EventTraceWin32::Stop(p_context_switch_recorder);
+				#else
+					FRAMEPRO_UNREFERENCED(p_context_switch_recorder);
+				#endif
+			}
+
+			//------------------------------------------------------------------------
+			void FlushContextSwitches(void* p_context_switch_recorder)
+			{
+				#if FRAMEPRO_PLATFORM_WIN
+					EventTraceWin32::Flush(p_context_switch_recorder);
+				#else
+					FRAMEPRO_UNREFERENCED(p_context_switch_recorder);
+				#endif
+			}
+		}
+	}
+
+#elif FRAMEPRO_LINUX_BASED_PLATFORM
+
+	//------------------------------------------------------------------------
+	#include <sys/signal.h>
+	#include <sys/types.h>
+	#include <unistd.h>
+	#include <sys/time.h>
+	#include <sched.h>
+	#include <errno.h>
+	#include <limits.h>
+	#include <pthread.h>
+	#include <inttypes.h>
+
+	#if FRAMEPRO_SOCKETS_ENABLED
+		#include <sys/socket.h>
+		#include <netinet/in.h>
+	#endif
+
+	//------------------------------------------------------------------------
+	namespace FramePro
+	{
+		//------------------------------------------------------------------------
+		namespace GenericPlatform
+		{
+			//------------------------------------------------------------------------
+			void DebugWrite(const char* p_string)
+			{
+				printf("%s", p_string);
+			}
+
+			//------------------------------------------------------------------------
+			pthread_mutex_t& GetOSLock(void* p_os_lock_mem)
+			{
+				return *(pthread_mutex_t*)p_os_lock_mem;
+			}
+
+			//------------------------------------------------------------------------
+			void CreateLock(void* p_os_lock_mem, int os_lock_mem_size)
+			{
+				FRAMEPRO_ASSERT((size_t)os_lock_mem_size >= sizeof(pthread_mutex_t));
+
+				pthread_mutexattr_t attr;
+				pthread_mutexattr_init(&attr);
+				pthread_mutexattr_settype(&attr, PTHREAD_MUTEX_RECURSIVE);
+				pthread_mutex_init(&GetOSLock(p_os_lock_mem), &attr);
+			}
+
+			//------------------------------------------------------------------------
+			void DestroyLock(void* p_os_lock_mem)
+			{
+				pthread_mutex_destroy(&GetOSLock(p_os_lock_mem));
+			}
+
+			//------------------------------------------------------------------------
+			void TakeLock(void* p_os_lock_mem)
+			{
+				pthread_mutex_lock(&GetOSLock(p_os_lock_mem));
+			}
+
+			//------------------------------------------------------------------------
+			void ReleaseLock(void* p_os_lock_mem)
+			{
+				pthread_mutex_unlock(&GetOSLock(p_os_lock_mem));
+			}
+
+			//------------------------------------------------------------------------
+			void GetLocalTime(tm* p_tm, const time_t *p_time)
+			{
+				tm* p_local_tm = localtime(p_time);
+				*p_tm = *p_local_tm;
+			}
+
+			//------------------------------------------------------------------------
+			int GetCurrentProcessId()
+			{
+				return getpid();
+			}
+
+			//------------------------------------------------------------------------
+			void VSPrintf(char* p_buffer, size_t const buffer_size, const char* p_format, va_list arg_list)
+			{
+				vsprintf(p_buffer, p_format, arg_list);
+			}
+
+			//------------------------------------------------------------------------
+			void ToString(int value, char* p_dest, int)
+			{
+				sprintf(p_dest, "%d", value);
+			}
+
+			//------------------------------------------------------------------------
+			bool GetProcessName(int, char*, int)
+			{
+				return false;
+			}
+
+			//------------------------------------------------------------------------
+			struct LinuxEvent
+			{
+				pthread_cond_t  m_Cond;
+				pthread_mutex_t m_Mutex;
+				volatile bool m_Signalled;
+				bool m_AutoReset;
+			};
+
+			//------------------------------------------------------------------------
+			LinuxEvent& GetOSEventHandle(void* p_os_event_mem)
+			{
+				return *(LinuxEvent*)p_os_event_mem;
+			}
+
+			//------------------------------------------------------------------------
+			void SetEvent(void* p_os_event_mem)
+			{
+				LinuxEvent& linux_event = GetOSEventHandle(p_os_event_mem);
+				pthread_mutex_lock(&linux_event.m_Mutex);
+				linux_event.m_Signalled = true;
+				pthread_mutex_unlock(&linux_event.m_Mutex);
+				pthread_cond_signal(&linux_event.m_Cond);
+			}
+
+			//------------------------------------------------------------------------
+			void CreateEventX(void* p_os_event_mem, int os_event_mem_size, bool initial_state, bool auto_reset)
+			{
+				FRAMEPRO_ASSERT((size_t)os_event_mem_size >= sizeof(LinuxEvent));
+				new (p_os_event_mem)LinuxEvent();
+
+				LinuxEvent& linux_event = GetOSEventHandle(p_os_event_mem);
+
+				pthread_cond_init(&linux_event.m_Cond, NULL);
+				pthread_mutex_init(&linux_event.m_Mutex, NULL);
+				linux_event.m_Signalled = false;
+				linux_event.m_AutoReset = auto_reset;
+
+				if (initial_state)
+					SetEvent(p_os_event_mem);
+			}
+
+			//------------------------------------------------------------------------
+			void DestroyEvent(void* p_os_event_mem)
+			{
+				LinuxEvent& linux_event = GetOSEventHandle(p_os_event_mem);
+				pthread_mutex_destroy(&linux_event.m_Mutex);
+				pthread_cond_destroy(&linux_event.m_Cond);
+				linux_event.~LinuxEvent();
+			}
+
+			//------------------------------------------------------------------------
+			void ResetEvent(void* p_os_event_mem)
+			{
+				LinuxEvent& linux_event = GetOSEventHandle(p_os_event_mem);
+				pthread_mutex_lock(&linux_event.m_Mutex);
+				linux_event.m_Signalled = false;
+				pthread_mutex_unlock(&linux_event.m_Mutex);
+			}
+
+			//------------------------------------------------------------------------
+			int WaitEvent(void* p_os_event_mem, int timeout)
+			{
+				LinuxEvent& linux_event = GetOSEventHandle(p_os_event_mem);
+
+				pthread_mutex_lock(&linux_event.m_Mutex);
+
+				if (linux_event.m_Signalled)
+				{
+					linux_event.m_Signalled = false;
+					pthread_mutex_unlock(&linux_event.m_Mutex);
+					return true;
+				}
+
+				if (timeout == -1)
+				{
+					while (!linux_event.m_Signalled)
+						pthread_cond_wait(&linux_event.m_Cond, &linux_event.m_Mutex);
+
+					if (!linux_event.m_AutoReset)
+						linux_event.m_Signalled = false;
+
+					pthread_mutex_unlock(&linux_event.m_Mutex);
+
+					return true;
+				}
+				else
+				{
+					timeval curr;
+					gettimeofday(&curr, NULL);
+
+					timespec time;
+					time.tv_sec = curr.tv_sec + timeout / 1000;
+					time.tv_nsec = (curr.tv_usec * 1000) + ((timeout % 1000) * 1000000);
+
+					time.tv_sec += time.tv_nsec / 1000000000L;
+					time.tv_nsec = time.tv_nsec % 1000000000L;
+
+					int ret = 0;
+					do
+					{
+						ret = pthread_cond_timedwait(&linux_event.m_Cond, &linux_event.m_Mutex, &time);
+
+					} while (!linux_event.m_Signalled && ret != ETIMEDOUT);
+
+					if (linux_event.m_Signalled)
+					{
+						if (!linux_event.m_AutoReset)
+							linux_event.m_Signalled = false;
+
+						pthread_mutex_unlock(&linux_event.m_Mutex);
+						return true;
+					}
+
+					pthread_mutex_unlock(&linux_event.m_Mutex);
+					return false;
+				}
+			}
+
+			//------------------------------------------------------------------------
+			#if FRAMEPRO_SOCKETS_ENABLED
+				int& GetOSSocket(void* p_os_socket_mem)
+				{
+					return *(int*)p_os_socket_mem;
+				}
+			#endif
+
+			//------------------------------------------------------------------------
+			#if FRAMEPRO_SOCKETS_ENABLED
+				const int& GetOSSocket(const void* p_os_socket_mem)
+				{
+					return *(int*)p_os_socket_mem;
+				}
+			#endif
+
+			//------------------------------------------------------------------------
+			#if FRAMEPRO_SOCKETS_ENABLED
+				static const int g_InvalidSocketId = -1;
+				static const int g_SocketErrorId = -1;
+			#endif
+
+			//------------------------------------------------------------------------
+			void CreateSocket(void* p_os_socket_mem, int os_socket_mem_size)
+			{
+				#if FRAMEPRO_SOCKETS_ENABLED
+					FRAMEPRO_ASSERT((size_t)os_socket_mem_size >= sizeof(int));
+					GetOSSocket(p_os_socket_mem) = g_InvalidSocketId;
+				#endif
+			}
+
+			//------------------------------------------------------------------------
+			void DestroySocket(void* p_os_socket_mem)
+			{
+				#if FRAMEPRO_SOCKETS_ENABLED
+					GetOSSocket(p_os_socket_mem) = g_InvalidSocketId;
+				#endif
+			}
+
+			//------------------------------------------------------------------------
+			bool InitialiseSocketSystem()
+			{
+				return true;
+			}
+
+			//------------------------------------------------------------------------
+			void UninitialiseSocketSystem()
+			{
+				// do nothing
+			}
+
+			//------------------------------------------------------------------------
+			void HandleSocketError()
+			{
+				DebugWrite("Socket Error");
+			}
+
+			//------------------------------------------------------------------------
+			void DisconnectSocket(void* p_os_socket_mem, bool stop_listening)
+			{
+				#if FRAMEPRO_SOCKETS_ENABLED
+					int& socket = GetOSSocket(p_os_socket_mem);
+
+					if (socket != g_InvalidSocketId)
+					{
+						if (shutdown(socket, SHUT_RDWR) == g_SocketErrorId)
+							HandleSocketError();
+
+						// loop until the socket is closed to ensure all data is sent
+						unsigned int buffer = 0;
+						size_t ret = 0;
+						do { ret = recv(socket, (char*)&buffer, sizeof(buffer), 0); } while (ret != 0 && ret != (size_t)g_SocketErrorId);
+
+						close(socket);
+
+						socket = g_InvalidSocketId;
+					}
+				#endif
+			}
+
+			//------------------------------------------------------------------------
+			bool StartSocketListening(void* p_os_socket_mem)
+			{
+				#if FRAMEPRO_SOCKETS_ENABLED
+					return listen(GetOSSocket(p_os_socket_mem), SOMAXCONN) != g_SocketErrorId;
+				#else
+					return false;
+				#endif
+			}
+
+			//------------------------------------------------------------------------
+			bool BindSocket(void* p_os_socket_mem, const char* p_port)
+			{
+				#if FRAMEPRO_SOCKETS_ENABLED
+					int& socket = GetOSSocket(p_os_socket_mem);
+
+					socket = ::socket(
+						AF_INET,
+						SOCK_STREAM,
+						IPPROTO_TCP);
+
+					if (socket == g_InvalidSocketId)
+					{
+						HandleSocketError();
+						return false;
+					}
+
+					// Setup the TCP listening socket
+					// Bind to INADDR_ANY
+					sockaddr_in sa;
+					sa.sin_family = AF_INET;
+					sa.sin_addr.s_addr = INADDR_ANY;
+					int iport = atoi(p_port);
+					sa.sin_port = htons(iport);
+					int result = ::bind(socket, (const sockaddr*)(&sa), sizeof(sockaddr_in));
+
+					if (result == g_SocketErrorId)
+					{
+						HandleSocketError();
+						DisconnectSocket(p_os_socket_mem, true);
+						return false;
+					}
+
+					return true;
+				#else
+					return false;
+				#endif
+			}
+
+			//------------------------------------------------------------------------
+			bool AcceptSocket(void* p_source_os_socket_mem, void* p_target_os_socket_mem)
+			{
+				#if FRAMEPRO_SOCKETS_ENABLED
+					int& source_socket = GetOSSocket(p_source_os_socket_mem);
+					int& target_socket = GetOSSocket(p_target_os_socket_mem);
+
+					target_socket = accept(source_socket, NULL, NULL);
+					return target_socket != g_InvalidSocketId;
+				#else
+					return false;
+				#endif
+			}
+
+			//------------------------------------------------------------------------
+			bool SocketSend(void* p_os_socket_mem, const void* p_buffer, int size, int& bytes_sent)
+			{
+				#if FRAMEPRO_SOCKETS_ENABLED
+					int& socket = GetOSSocket(p_os_socket_mem);
+					#if FRAMEPRO_PLATFORM_LINUX || FRAMEPRO_PLATFORM_ANDROID
+						int flags = MSG_NOSIGNAL;
+					#else
+						int flags = 0;
+					#endif
+					bytes_sent = (int)send(socket, (char*)p_buffer, size, flags);
+					return bytes_sent != g_SocketErrorId;
+				#else
+					return false;
+				#endif
+			}
+
+			//------------------------------------------------------------------------
+			bool SocketReceive(void* p_os_socket_mem, const void* p_buffer, int size, int& bytes_received)
+			{
+				#if FRAMEPRO_SOCKETS_ENABLED
+					int& socket = GetOSSocket(p_os_socket_mem);
+					bytes_received = (int)recv(socket, (char*)p_buffer, size, 0);
+					return bytes_received != g_SocketErrorId;
+				#else
+					return false;
+				#endif
+			}
+
+			//------------------------------------------------------------------------
+			bool IsSocketValid(const void* p_os_socket_mem)
+			{
+				#if FRAMEPRO_SOCKETS_ENABLED
+					const int& socket = GetOSSocket(p_os_socket_mem);
+					return socket != g_InvalidSocketId;
+				#else
+					return false;
+				#endif
+			}
+
+			//------------------------------------------------------------------------
+			struct ThreadContext
+			{
+				ThreadMain mp_ThreadMain;
+				void* mp_Context;
+				Allocator* mp_Allocator;
+			};
+
+			//------------------------------------------------------------------------
+			void* PlatformThreadMain(void* p_param)
+			{
+				ThreadContext* p_context = (ThreadContext*)p_param;
+
+				p_context->mp_ThreadMain(p_context->mp_Context);
+
+				Delete(p_context->mp_Allocator, p_context);
+
+				return NULL;
+			}
+
+			//------------------------------------------------------------------------
+			pthread_t& GetOSThread(void* p_os_thread_mem)
+			{
+				return *(pthread_t*)p_os_thread_mem;
+			}
+
+			//------------------------------------------------------------------------
+			void CreateThread(
+				void* p_os_thread_mem,
+				int os_thread_mem_size,
+				ThreadMain p_thread_main,
+				void* p_context,
+				Allocator* p_allocator)
+			{
+				FRAMEPRO_ASSERT((size_t)os_thread_mem_size >= sizeof(pthread_t));
+
+				ThreadContext* p_thread_context = New<ThreadContext>(p_allocator);
+				p_thread_context->mp_ThreadMain = p_thread_main;
+				p_thread_context->mp_Context = p_context;
+				p_thread_context->mp_Allocator = p_allocator;
+
+				pthread_create(&GetOSThread(p_os_thread_mem), NULL, PlatformThreadMain, p_thread_context);
+			}
+
+			//------------------------------------------------------------------------
+			void SetThreadPriority(void*, int)
+			{
+				// not implemented
+			}
+
+			//------------------------------------------------------------------------
+			void SetThreadAffinity(void*, int)
+			{
+				// not implemented
+			}
+
+			//------------------------------------------------------------------------
+			void* CreateContextSwitchRecorder(Allocator*)
+			{
+				// not implemented
+				return NULL;
+			}
+
+			//------------------------------------------------------------------------
+			void DestroyContextSwitchRecorder(void*, Allocator*)
+			{
+				// not implemented
+			}
+
+			//------------------------------------------------------------------------
+			bool StartRecordingContextSitches(void*, Platform::ContextSwitchCallbackFunction, void*, DynamicString&)
+			{
+				// not implemented
+				return false;
+			}
+
+			//------------------------------------------------------------------------
+			void StopRecordingContextSitches(void*)
+			{
+				// not implemented
+			}
+
+			//------------------------------------------------------------------------
+			void FlushContextSwitches(void*)
+			{
+				// not implemented
+			}
+		}
+	}
+
+#endif
+
+//------------------------------------------------------------------------
+//                         FRAMEPRO_PLATFORM_UE4
+//------------------------------------------------------------------------
+#if FRAMEPRO_PLATFORM_UE4
+
+	// implemented in FrameProPlatformUE4.cpp - contact [REDACTED_EMAIL] for this platform
+
+//------------------------------------------------------------------------
+//                         FRAMEPRO_PLATFORM_WIN
+//------------------------------------------------------------------------
+#elif FRAMEPRO_PLATFORM_WIN
+
+	//------------------------------------------------------------------------
+
+	//------------------------------------------------------------------------
+	// if both of these options are commented out it will use CaptureStackBackTrace (or backtrace on linux)
+	#define FRAMEPRO_USE_STACKWALK64 0				// much slower but possibly more reliable. FRAMEPRO_USE_STACKWALK64 only implemented for x86 builds.
+	#define FRAMEPRO_USE_RTLVIRTUALUNWIND 0			// reported to be faster than StackWalk64 - only available on x64 builds
+	#define FRAMEPRO_USE_RTLCAPTURESTACKBACKTRACE 0	// system version of FRAMEPRO_USE_RTLVIRTUALUNWIND - only available on x64 builds
+
+	//------------------------------------------------------------------------
+	// need to include windows.h in header for QueryPerformanceCounter()
+	#ifndef WIN32_LEAN_AND_MEAN
+		#define WIN32_LEAN_AND_MEAN
+	#endif
+	#pragma warning(push)
+	#pragma warning(disable:4668)
+	#include <windows.h>
+	#pragma warning(pop)
+
+	//------------------------------------------------------------------------
+	__int64 FramePro_QueryPerformanceCounter()
+	{
+		__int64 time;
+		::QueryPerformanceCounter((LARGE_INTEGER*)&time);
+		return time;
+	}
+
+	//------------------------------------------------------------------------
+	namespace FramePro
+	{
+		//------------------------------------------------------------------------
+		int64 Platform::GetTimerFrequency()
+		{
+			int64 frequency = 0;
+			QueryPerformanceFrequency((LARGE_INTEGER*)&frequency);
+			return frequency;
+		}
+
+		//------------------------------------------------------------------------
+		void Platform::DebugBreak()
+		{
+			::DebugBreak();
+		}
+
+		//------------------------------------------------------------------------
+		int Platform::GetCore()
+		{
+			return GetCurrentProcessorNumber();
+		}
+
+		//------------------------------------------------------------------------
+		Platform::Enum Platform::GetPlatformEnum()
+		{
+			return Platform::Windows;
+		}
+
+		//------------------------------------------------------------------------
+		void* Platform::CreateContextSwitchRecorder(Allocator* p_allocator)
+		{
+			return GenericPlatform::CreateContextSwitchRecorder(p_allocator);
+		}
+
+		//------------------------------------------------------------------------
+		void Platform::DestroyContextSwitchRecorder(void* p_context_switch_recorder, Allocator* p_allocator)
+		{
+			GenericPlatform::DestroyContextSwitchRecorder(p_context_switch_recorder, p_allocator);
+		}
+
+		//------------------------------------------------------------------------
+		bool Platform::StartRecordingContextSitches(
+			void* p_context_switch_recorder,
+			ContextSwitchCallbackFunction p_callback,
+			void* p_context,
+			DynamicString& error)
+		{
+			return GenericPlatform::StartRecordingContextSitches(
+				p_context_switch_recorder,
+				p_callback,
+				p_context,
+				error);
+		}
+
+		//------------------------------------------------------------------------
+		void Platform::StopRecordingContextSitches(void* p_context_switch_recorder)
+		{
+			GenericPlatform::StopRecordingContextSitches(p_context_switch_recorder);
+		}
+
+		//------------------------------------------------------------------------
+		void Platform::FlushContextSwitches(void* p_context_switch_recorder)
+		{
+			GenericPlatform::FlushContextSwitches(p_context_switch_recorder);
+		}
+
+		//------------------------------------------------------------------------
+		void Platform::EnumerateModules(Array<ModulePacket*>& module_packets, Allocator* p_allocator)
+		{
+			#if FRAMEPRO_ENABLE_CALLSTACKS
+				EnumModulesWindows::EnumerateModules(module_packets, p_allocator);
+			#else
+				FRAMEPRO_UNREFERENCED(module_packets);
+				FRAMEPRO_UNREFERENCED(p_allocator);
+			#endif
+		}
+
+		//------------------------------------------------------------------------
+		bool Platform::GetStackTrace(void** stack, int& stack_size, unsigned int& hash)
+		{
+			#if FRAMEPRO_USE_STACKWALK64
+
+				// get the context
+				CONTEXT context;
+				memset(&context, 0, sizeof(context));
+				RtlCaptureContext(&context);
+
+				// setup the stack frame
+				STACKFRAME64 stack_frame;
+				memset(&stack_frame, 0, sizeof(stack_frame));
+				stack_frame.AddrPC.Mode = AddrModeFlat;
+				stack_frame.AddrFrame.Mode = AddrModeFlat;
+				stack_frame.AddrStack.Mode = AddrModeFlat;
+				DWORD machine = IMAGE_FILE_MACHINE_IA64;
+				stack_frame.AddrPC.Offset = context.Rip;
+				stack_frame.AddrFrame.Offset = context.Rsp;
+				stack_frame.AddrStack.Offset = context.Rbp;
+				HANDLE thread = GetCurrentThread();
+
+				static HANDLE process = GetCurrentProcess();
+
+				stack_size = 0;
+				while (StackWalk64(
+					machine,
+					process,
+					thread,
+					&stack_frame,
+					&context,
+					NULL,
+					SymFunctionTableAccess64,
+					SymGetModuleBase64,
+					NULL) && stack_size < FRAMEPRO_STACK_TRACE_SIZE)
+				{
+					void* p = (void*)(stack_frame.AddrPC.Offset);
+					stack[stack_size++] = p;
+				}
+				hash = GetHash(stack, stack_size);
+			#elif FRAMEPRO_USE_RTLVIRTUALUNWIND
+				FramePro::VirtualUnwindStackWalk(stack, FRAMEPRO_STACK_TRACE_SIZE);
+				hash = GetHashAndStackSize(stack, stack_size);
+			#elif FRAMEPRO_USE_RTLCAPTURESTACKBACKTRACE
+				stack_size = ::RtlCaptureStackBackTrace(1, FRAMEPRO_STACK_TRACE_SIZE-1, stack, (PDWORD)&hash);
+			#else
+				CaptureStackBackTrace(0, FRAMEPRO_STACK_TRACE_SIZE, stack, (PDWORD)&hash);
+				for (stack_size = 0; stack_size<FRAMEPRO_STACK_TRACE_SIZE; ++stack_size)
+					if (!stack[stack_size])
+						break;
+			#endif
+			return true;
+		}
+
+		//------------------------------------------------------------------------
+		FILE*& GetOSFile(void* p_os_file_mem)
+		{
+			return *(FILE**)p_os_file_mem;
+		}
+
+		//------------------------------------------------------------------------
+		FILE* GetOSFile(const void* p_os_file_mem)
+		{
+			return *(FILE**)p_os_file_mem;
+		}
+
+		//------------------------------------------------------------------------
+		bool Platform::OpenFileForRead(void* p_os_file_mem, int os_file_mem_size, const char* p_filename)
+		{
+			FRAMEPRO_ASSERT(os_file_mem_size > sizeof(FILE*));
+			FRAMEPRO_UNREFERENCED(os_file_mem_size);
+			FILE*& p_file = GetOSFile(p_os_file_mem);
+			return fopen_s(&p_file, p_filename, "rb") == 0;
+		}
+
+		//------------------------------------------------------------------------
+		bool Platform::OpenFileForRead(void* p_os_file_mem, int os_file_mem_size, const wchar_t* p_filename)
+		{
+			FRAMEPRO_ASSERT(os_file_mem_size > sizeof(FILE*));
+			FRAMEPRO_UNREFERENCED(os_file_mem_size);
+			FILE*& p_file = GetOSFile(p_os_file_mem);
+			return _wfopen_s(&p_file, p_filename, L"rb") == 0;
+		}
+
+		//------------------------------------------------------------------------
+		bool Platform::OpenFileForWrite(void* p_os_file_mem, int os_file_mem_size, const char* p_filename)
+		{
+			FRAMEPRO_ASSERT(os_file_mem_size > sizeof(FILE*));
+			FRAMEPRO_UNREFERENCED(os_file_mem_size);
+			FILE*& p_file = GetOSFile(p_os_file_mem);
+			return fopen_s(&p_file, p_filename, "wb") == 0;
+		}
+
+		//------------------------------------------------------------------------
+		bool Platform::OpenFileForWrite(void* p_os_file_mem, int os_file_mem_size, const wchar_t* p_filename)
+		{
+			FRAMEPRO_ASSERT(os_file_mem_size > sizeof(FILE*));
+			FRAMEPRO_UNREFERENCED(os_file_mem_size);
+			FILE*& p_file = GetOSFile(p_os_file_mem);
+			return _wfopen_s(&p_file, p_filename, L"wb") == 0;
+		}
+
+		//------------------------------------------------------------------------
+		void Platform::CloseFile(void* p_os_file_mem)
+		{
+			FILE*& p_file = GetOSFile(p_os_file_mem);
+			fclose(p_file);
+			p_file = NULL;
+		}
+
+		//------------------------------------------------------------------------
+		void Platform::ReadFromFile(void* p_os_file_mem, void* p_data, size_t size)
+		{
+			FILE* p_file = GetOSFile(p_os_file_mem);
+			fread(p_data, size, 1, p_file);
+		}
+
+		//------------------------------------------------------------------------
+		void Platform::WriteToFile(void* p_os_file_mem, const void* p_data, size_t size)
+		{
+			FILE* p_file = GetOSFile(p_os_file_mem);
+			fwrite(p_data, size, 1, p_file);
+		}
+
+		//------------------------------------------------------------------------
+		int Platform::GetFileSize(const void* p_os_file_mem)
+		{
+			FILE* p_file = GetOSFile(p_os_file_mem);
+			int pos = ftell(p_file);
+			fseek(p_file, 0, SEEK_END);
+			int size = ftell(p_file);
+			fseek(p_file, pos, SEEK_SET);
+			return size;
+		}
+
+		//------------------------------------------------------------------------
+		void Platform::DebugWrite(const char* p_string)
+		{
+			GenericPlatform::DebugWrite(p_string);
+		}
+
+		//------------------------------------------------------------------------
+		void Platform::CreateLock(void* p_os_lock_mem, int os_lock_mem_size)
+		{
+			GenericPlatform::CreateLock(p_os_lock_mem, os_lock_mem_size);
+		}
+
+		//------------------------------------------------------------------------
+		void Platform::DestroyLock(void* p_os_lock_mem)
+		{
+			GenericPlatform::DestroyLock(p_os_lock_mem);
+		}
+
+		//------------------------------------------------------------------------
+		void Platform::TakeLock(void* p_os_lock_mem)
+		{
+			GenericPlatform::TakeLock(p_os_lock_mem);
+		}
+
+		//------------------------------------------------------------------------
+		void Platform::ReleaseLock(void* p_os_lock_mem)
+		{
+			GenericPlatform::ReleaseLock(p_os_lock_mem);
+		}
+
+		//------------------------------------------------------------------------
+		void Platform::GetLocalTime(tm* p_tm, const time_t *p_time)
+		{
+			GenericPlatform::GetLocalTime(p_tm, p_time);
+		}
+
+		//------------------------------------------------------------------------
+		int Platform::GetCurrentProcessId()
+		{
+			return GenericPlatform::GetCurrentProcessId();
+		}
+
+		//------------------------------------------------------------------------
+		void Platform::VSPrintf(char* p_buffer, size_t const buffer_size, const char* p_format, va_list arg_list)
+		{
+			GenericPlatform::VSPrintf(p_buffer, buffer_size, p_format, arg_list);
+		}
+
+		//------------------------------------------------------------------------
+		void Platform::ToString(int value, char* p_dest, int dest_size)
+		{
+			GenericPlatform::ToString(value, p_dest, dest_size);
+		}
+
+		//------------------------------------------------------------------------
+		int Platform::GetCurrentThreadId()
+		{
+			return ::GetCurrentThreadId();
+		}
+
+		//------------------------------------------------------------------------
+		bool Platform::GetProcessName(int process_id, char* p_name, int max_name_length)
+		{
+			HANDLE process = OpenProcess(PROCESS_ALL_ACCESS, true, process_id);
+			if(process)
+			{
+				unsigned long result = GetProcessImageFileNameA(process, p_name, max_name_length);
+				CloseHandle(process);
+
+				if(result)
+				{
+					int total_length = (int)strlen(p_name);
+					char* p_filename = strrchr(p_name, '\\');
+					if(p_filename && p_filename[1])
+					{
+						++p_filename;
+						memmove(p_name, p_filename, p_name + total_length + 1 - p_filename);
+					}
+
+					return true;
+				}
+			}
+
+			return false;
+		}
+
+		//------------------------------------------------------------------------
+		void Platform::CreateEventX(void* p_os_event_mem, int os_event_mem_size, bool initial_state, bool auto_reset)
+		{
+			return GenericPlatform::CreateEventX(p_os_event_mem, os_event_mem_size, initial_state, auto_reset);
+		}
+
+		//------------------------------------------------------------------------
+		void Platform::DestroyEvent(void* p_os_event_mem)
+		{
+			GenericPlatform::DestroyEvent(p_os_event_mem);
+		}
+
+		//------------------------------------------------------------------------
+		void Platform::SetEvent(void* p_os_event_mem)
+		{
+			GenericPlatform::SetEvent(p_os_event_mem);
+		}
+
+		//------------------------------------------------------------------------
+		void Platform::ResetEvent(void* p_os_event_mem)
+		{
+			GenericPlatform::ResetEvent(p_os_event_mem);
+		}
+
+		//------------------------------------------------------------------------
+		int Platform::WaitEvent(void* p_os_event_mem, int timeout)
+		{
+			return GenericPlatform::WaitEvent(p_os_event_mem, timeout);
+		}
+
+		//------------------------------------------------------------------------
+		bool Platform::InitialiseSocketSystem()
+		{
+			return GenericPlatform::InitialiseSocketSystem();
+		}
+
+		//------------------------------------------------------------------------
+		void Platform::UninitialiseSocketSystem()
+		{
+			GenericPlatform::UninitialiseSocketSystem();
+		}
+
+		//------------------------------------------------------------------------
+		void Platform::CreateSocket(void* p_os_socket_mem, int os_socket_mem_size)
+		{
+			GenericPlatform::CreateSocket(p_os_socket_mem, os_socket_mem_size);
+		}
+
+		//------------------------------------------------------------------------
+		void Platform::DestroySocket(void* p_os_socket_mem)
+		{
+			GenericPlatform::DestroySocket(p_os_socket_mem);
+		}
+
+		//------------------------------------------------------------------------
+		void Platform::DisconnectSocket(void* p_os_socket_mem, bool stop_listening)
+		{
+			GenericPlatform::DisconnectSocket(p_os_socket_mem, stop_listening);
+		}
+
+		//------------------------------------------------------------------------
+		bool Platform::StartSocketListening(void* p_os_socket_mem)
+		{
+			return GenericPlatform::StartSocketListening(p_os_socket_mem);
+		}
+
+		//------------------------------------------------------------------------
+		bool Platform::BindSocket(void* p_os_socket_mem, const char* p_port)
+		{
+			return GenericPlatform::BindSocket(p_os_socket_mem, p_port);
+		}
+
+		//------------------------------------------------------------------------
+		bool Platform::AcceptSocket(void* p_source_os_socket_mem, void* p_target_os_socket_mem)
+		{
+			return GenericPlatform::AcceptSocket(p_source_os_socket_mem, p_target_os_socket_mem);
+		}
+
+		//------------------------------------------------------------------------
+		bool Platform::SocketSend(void* p_os_socket_mem, const void* p_buffer, int size, int& bytes_sent)
+		{
+			return GenericPlatform::SocketSend(p_os_socket_mem, p_buffer, size, bytes_sent);
+		}
+
+		//------------------------------------------------------------------------
+		bool Platform::SocketReceive(void* p_os_socket_mem, const void* p_buffer, int size, int& bytes_received)
+		{
+			return GenericPlatform::SocketReceive(p_os_socket_mem, p_buffer, size, bytes_received);
+		}
+
+		//------------------------------------------------------------------------
+		bool Platform::IsSocketValid(const void* p_os_socket_mem)
+		{
+			return GenericPlatform::IsSocketValid(p_os_socket_mem);
+		}
+
+		//------------------------------------------------------------------------
+		void Platform::HandleSocketError()
+		{
+			GenericPlatform::HandleSocketError();
+		}
+
+		//------------------------------------------------------------------------
+		void Platform::CreateThread(
+			void* p_os_thread_mem,
+			int os_thread_mem_size,
+			ThreadMain p_thread_main,
+			void* p_context,
+			Allocator* p_allocator)
+		{
+			GenericPlatform::CreateThread(
+				p_os_thread_mem,
+				os_thread_mem_size,
+				p_thread_main,
+				p_context,
+				p_allocator);
+		}
+
+		//------------------------------------------------------------------------
+		void Platform::SetThreadPriority(void* p_os_thread_mem, int priority)
+		{
+			GenericPlatform::SetThreadPriority(p_os_thread_mem, priority);
+		}
+
+		//------------------------------------------------------------------------
+		void Platform::SetThreadAffinity(void* p_os_thread_mem, int priority)
+		{
+			GenericPlatform::SetThreadAffinity(p_os_thread_mem, priority);
+		}
+
+		//------------------------------------------------------------------------
+		#if FRAMEPRO_USE_TLS_SLOTS
+			#error this platform is not using TLS slots
+		#endif
+
+		__declspec(thread) void* gp_FrameProTLS = NULL;
+
+		//------------------------------------------------------------------------
+		uint Platform::AllocateTLSSlot()
+		{
+			return 0;
+		}
+
+		//------------------------------------------------------------------------
+		void* Platform::GetTLSValue(uint)
+		{
+			return gp_FrameProTLS;
+		}
+
+		//------------------------------------------------------------------------
+		void Platform::SetTLSValue(uint, void* p_value)
+		{
+			gp_FrameProTLS = p_value;
+		}
+
+		//------------------------------------------------------------------------
+		void Platform::GetRecordingFolder(char* p_path, int max_path_length)
+		{
+			FRAMEPRO_ASSERT(max_path_length);
+			FRAMEPRO_UNREFERENCED(max_path_length);
+			*p_path = '\0';
+		}
+	}
+
+//------------------------------------------------------------------------
+//                         FRAMEPRO_PLATFORM_UWP
+//------------------------------------------------------------------------
+#elif FRAMEPRO_PLATFORM_UWP
+
+	//------------------------------------------------------------------------
+	// need to include windows.h in header for QueryPerformanceCounter()
+	#ifndef WIN32_LEAN_AND_MEAN
+		#define WIN32_LEAN_AND_MEAN
+	#endif
+	#pragma warning(push)
+	#pragma warning(disable:4668)
+	#include <windows.h>
+	#pragma warning(pop)
+
+	//------------------------------------------------------------------------
+	__int64 FramePro_QueryPerformanceCounter()
+	{
+		__int64 time;
+		::QueryPerformanceCounter((LARGE_INTEGER*)&time);
+		return time;
+	}
+
+	//------------------------------------------------------------------------
+	namespace FramePro
+	{
+		//------------------------------------------------------------------------
+		int64 Platform::GetTimerFrequency()
+		{
+			int64 frequency = 0;
+			QueryPerformanceFrequency((LARGE_INTEGER*)&frequency);
+			return frequency;
+		}
+
+		//------------------------------------------------------------------------
+		void Platform::DebugBreak()
+		{
+			::DebugBreak();
+		}
+
+		//------------------------------------------------------------------------
+		int Platform::GetCore()
+		{
+			int cpu_info[4];
+			__cpuid(cpu_info, 1);
+			return (cpu_info[1] >> 24) & 0xff;
+		}
+
+		//------------------------------------------------------------------------
+		Platform::Enum Platform::GetPlatformEnum()
+		{
+			return Platform::Windows_UWP;
+		}
+
+		//------------------------------------------------------------------------
+		void* Platform::CreateContextSwitchRecorder(Allocator*)
+		{
+			return NULL;
+		}
+
+		//------------------------------------------------------------------------
+		void Platform::DestroyContextSwitchRecorder(void*, Allocator*)
+		{
+		}
+
+		//------------------------------------------------------------------------
+		bool Platform::StartRecordingContextSitches(
+			void*,
+			ContextSwitchCallbackFunction,
+			void*,
+			DynamicString&)
+		{
+			FramePro::DebugWrite("FramePro Warning: Failed to start recording context switches. Context switches may not be supported for this platform\n");
+			return false;
+		}
+
+		//------------------------------------------------------------------------
+		void Platform::StopRecordingContextSitches(void*)
+		{
+		}
+
+		//------------------------------------------------------------------------
+		void Platform::FlushContextSwitches(void*)
+		{
+		}
+
+		//------------------------------------------------------------------------
+		void Platform::EnumerateModules(Array<ModulePacket*>& module_packets, Allocator* p_allocator)
+		{
+			#if FRAMEPRO_ENABLE_CALLSTACKS
+				EnumModulesWindows::EnumerateModules(module_packets, p_allocator);
+			#endif
+		}
+
+		//------------------------------------------------------------------------
+		bool Platform::GetStackTrace(void**, int&, unsigned int&)
+		{
+			return false;
+		}
+
+		//------------------------------------------------------------------------
+		FILE*& GetOSFile(void* p_os_file_mem)
+		{
+			return *(FILE**)p_os_file_mem;
+		}
+
+		//------------------------------------------------------------------------
+		FILE* GetOSFile(const void* p_os_file_mem)
+		{
+			return *(FILE**)p_os_file_mem;
+		}
+
+		//------------------------------------------------------------------------
+		bool Platform::OpenFileForRead(void* p_os_file_mem, int os_file_mem_size, const char* p_filename)
+		{
+			FRAMEPRO_ASSERT(os_file_mem_size > sizeof(FILE*));
+			FILE*& p_file = GetOSFile(p_os_file_mem);
+			return fopen_s(&p_file, p_filename, "rb") == 0;
+		}
+
+		//------------------------------------------------------------------------
+		bool Platform::OpenFileForRead(void* p_os_file_mem, int os_file_mem_size, const wchar_t* p_filename)
+		{
+			FRAMEPRO_ASSERT(os_file_mem_size > sizeof(FILE*));
+			FILE*& p_file = GetOSFile(p_os_file_mem);
+			return _wfopen_s(&p_file, p_filename, L"rb") == 0;
+		}
+
+		//------------------------------------------------------------------------
+		bool Platform::OpenFileForWrite(void* p_os_file_mem, int os_file_mem_size, const char* p_filename)
+		{
+			FRAMEPRO_ASSERT(os_file_mem_size > sizeof(FILE*));
+			FILE*& p_file = GetOSFile(p_os_file_mem);
+			return fopen_s(&p_file, p_filename, "wb") == 0;
+		}
+
+		//------------------------------------------------------------------------
+		bool Platform::OpenFileForWrite(void* p_os_file_mem, int os_file_mem_size, const wchar_t* p_filename)
+		{
+			FRAMEPRO_ASSERT(os_file_mem_size > sizeof(FILE*));
+			FILE*& p_file = GetOSFile(p_os_file_mem);
+			return _wfopen_s(&p_file, p_filename, L"wb") == 0;
+		}
+
+		//------------------------------------------------------------------------
+		void Platform::CloseFile(void* p_os_file_mem)
+		{
+			FILE*& p_file = GetOSFile(p_os_file_mem);
+			fclose(p_file);
+			p_file = NULL;
+		}
+
+		//------------------------------------------------------------------------
+		void Platform::ReadFromFile(void* p_os_file_mem, void* p_data, size_t size)
+		{
+			FILE* p_file = GetOSFile(p_os_file_mem);
+			fread(p_data, size, 1, p_file);
+		}
+
+		//------------------------------------------------------------------------
+		void Platform::WriteToFile(void* p_os_file_mem, const void* p_data, size_t size)
+		{
+			FILE* p_file = GetOSFile(p_os_file_mem);
+			fwrite(p_data, size, 1, p_file);
+		}
+
+		//------------------------------------------------------------------------
+		int Platform::GetFileSize(const void* p_os_file_mem)
+		{
+			FILE* p_file = GetOSFile(p_os_file_mem);
+			int pos = ftell(p_file);
+			fseek(p_file, 0, SEEK_END);
+			int size = ftell(p_file);
+			fseek(p_file, pos, SEEK_SET);
+			return size;
+		}
+
+		//------------------------------------------------------------------------
+		void Platform::DebugWrite(const char* p_string)
+		{
+			GenericPlatform::DebugWrite(p_string);
+		}
+
+		//------------------------------------------------------------------------
+		void Platform::CreateLock(void* p_os_lock_mem, int os_lock_mem_size)
+		{
+			GenericPlatform::CreateLock(p_os_lock_mem, os_lock_mem_size);
+		}
+
+		//------------------------------------------------------------------------
+		void Platform::DestroyLock(void* p_os_lock_mem)
+		{
+			GenericPlatform::DestroyLock(p_os_lock_mem);
+		}
+
+		//------------------------------------------------------------------------
+		void Platform::TakeLock(void* p_os_lock_mem)
+		{
+			GenericPlatform::TakeLock(p_os_lock_mem);
+		}
+
+		//------------------------------------------------------------------------
+		void Platform::ReleaseLock(void* p_os_lock_mem)
+		{
+			GenericPlatform::ReleaseLock(p_os_lock_mem);
+		}
+
+		//------------------------------------------------------------------------
+		void Platform::GetLocalTime(tm* p_tm, const time_t *p_time)
+		{
+			GenericPlatform::GetLocalTime(p_tm, p_time);
+		}
+
+		//------------------------------------------------------------------------
+		int Platform::GetCurrentProcessId()
+		{
+			return GenericPlatform::GetCurrentProcessId();
+		}
+
+		//------------------------------------------------------------------------
+		void Platform::VSPrintf(char* p_buffer, size_t const buffer_size, const char* p_format, va_list arg_list)
+		{
+			GenericPlatform::VSPrintf(p_buffer, buffer_size, p_format, arg_list);
+		}
+
+		//------------------------------------------------------------------------
+		void Platform::ToString(int value, char* p_dest, int dest_size)
+		{
+			GenericPlatform::ToString(value, p_dest, dest_size);
+		}
+
+		//------------------------------------------------------------------------
+		int Platform::GetCurrentThreadId()
+		{
+			return ::GetCurrentThreadId();
+		}
+
+		//------------------------------------------------------------------------
+		bool Platform::GetProcessName(int process_id, char* p_name, int max_name_length)
+		{
+			HANDLE process = OpenProcess(PROCESS_ALL_ACCESS, true, process_id);
+			if(process)
+			{
+				unsigned long result = GetProcessImageFileNameA(process, p_name, max_name_length);
+				CloseHandle(process);
+
+				if(result)
+				{
+					int total_length = (int)strlen(p_name);
+					char* p_filename = strrchr(p_name, '\\');
+					if(p_filename && p_filename[1])
+					{
+						++p_filename;
+						memmove(p_name, p_filename, p_name + total_length + 1 - p_filename);
+					}
+
+					return true;
+				}
+			}
+
+			return false;
+		}
+
+		//------------------------------------------------------------------------
+		void Platform::CreateEventX(void* p_os_event_mem, int os_event_mem_size, bool initial_state, bool auto_reset)
+		{
+			return GenericPlatform::CreateEventX(p_os_event_mem, os_event_mem_size, initial_state, auto_reset);
+		}
+
+		//------------------------------------------------------------------------
+		void Platform::DestroyEvent(void* p_os_event_mem)
+		{
+			GenericPlatform::DestroyEvent(p_os_event_mem);
+		}
+
+		//------------------------------------------------------------------------
+		void Platform::SetEvent(void* p_os_event_mem)
+		{
+			GenericPlatform::SetEvent(p_os_event_mem);
+		}
+
+		//------------------------------------------------------------------------
+		void Platform::ResetEvent(void* p_os_event_mem)
+		{
+			GenericPlatform::ResetEvent(p_os_event_mem);
+		}
+
+		//------------------------------------------------------------------------
+		int Platform::WaitEvent(void* p_os_event_mem, int timeout)
+		{
+			return GenericPlatform::WaitEvent(p_os_event_mem, timeout);
+		}
+
+		//------------------------------------------------------------------------
+		bool Platform::InitialiseSocketSystem()
+		{
+			return GenericPlatform::InitialiseSocketSystem();
+		}
+
+		//------------------------------------------------------------------------
+		void Platform::UninitialiseSocketSystem()
+		{
+			GenericPlatform::UninitialiseSocketSystem();
+		}
+
+		//------------------------------------------------------------------------
+		void Platform::CreateSocket(void* p_os_socket_mem, int os_socket_mem_size)
+		{
+			GenericPlatform::CreateSocket(p_os_socket_mem, os_socket_mem_size);
+		}
+
+		//------------------------------------------------------------------------
+		void Platform::DestroySocket(void* p_os_socket_mem)
+		{
+			GenericPlatform::DestroySocket(p_os_socket_mem);
+		}
+
+		//------------------------------------------------------------------------
+		void Platform::DisconnectSocket(void* p_os_socket_mem, bool stop_listening)
+		{
+			GenericPlatform::DisconnectSocket(p_os_socket_mem, stop_listening);
+		}
+
+		//------------------------------------------------------------------------
+		bool Platform::StartSocketListening(void* p_os_socket_mem)
+		{
+			return GenericPlatform::StartSocketListening(p_os_socket_mem);
+		}
+
+		//------------------------------------------------------------------------
+		bool Platform::BindSocket(void* p_os_socket_mem, const char* p_port)
+		{
+			return GenericPlatform::BindSocket(p_os_socket_mem, p_port);
+		}
+
+		//------------------------------------------------------------------------
+		bool Platform::AcceptSocket(void* p_source_os_socket_mem, void* p_target_os_socket_mem)
+		{
+			return GenericPlatform::AcceptSocket(p_source_os_socket_mem, p_target_os_socket_mem);
+		}
+
+		//------------------------------------------------------------------------
+		bool Platform::SocketSend(void* p_os_socket_mem, const void* p_buffer, int size, int& bytes_sent)
+		{
+			return GenericPlatform::SocketSend(p_os_socket_mem, p_buffer, size, bytes_sent);
+		}
+
+		//------------------------------------------------------------------------
+		bool Platform::SocketReceive(void* p_os_socket_mem, const void* p_buffer, int size, int& bytes_received)
+		{
+			return GenericPlatform::SocketReceive(p_os_socket_mem, p_buffer, size, bytes_received);
+		}
+
+		//------------------------------------------------------------------------
+		bool Platform::IsSocketValid(const void* p_os_socket_mem)
+		{
+			return GenericPlatform::IsSocketValid(p_os_socket_mem);
+		}
+
+		//------------------------------------------------------------------------
+		void Platform::HandleSocketError()
+		{
+			GenericPlatform::HandleSocketError();
+		}
+
+		//------------------------------------------------------------------------
+		void Platform::CreateThread(
+			void* p_os_thread_mem,
+			int os_thread_mem_size,
+			ThreadMain p_thread_main,
+			void* p_context,
+			Allocator* p_allocator)
+		{
+			GenericPlatform::CreateThread(
+				p_os_thread_mem,
+				os_thread_mem_size,
+				p_thread_main,
+				p_context,
+				p_allocator);
+		}
+
+		//------------------------------------------------------------------------
+		void Platform::SetThreadPriority(void* p_os_thread_mem, int priority)
+		{
+			GenericPlatform::SetThreadPriority(p_os_thread_mem, priority);
+		}
+
+		//------------------------------------------------------------------------
+		void Platform::SetThreadAffinity(void* p_os_thread_mem, int priority)
+		{
+			GenericPlatform::SetThreadAffinity(p_os_thread_mem, priority);
+		}
+
+		//------------------------------------------------------------------------
+		#if FRAMEPRO_USE_TLS_SLOTS
+			#error this platform is not using TLS slots
+		#endif
+
+		__declspec(thread) void* gp_FrameProTLS = NULL;
+
+		//------------------------------------------------------------------------
+		uint Platform::AllocateTLSSlot()
+		{
+			return 0;
+		}
+
+		//------------------------------------------------------------------------
+		void* Platform::GetTLSValue(uint)
+		{
+			return gp_FrameProTLS;
+		}
+
+		//------------------------------------------------------------------------
+		void Platform::SetTLSValue(uint, void* p_value)
+		{
+			gp_FrameProTLS = p_value;
+		}
+
+		//------------------------------------------------------------------------
+		void Platform::GetRecordingFolder(char* p_path, int max_path_length)
+		{
+			FRAMEPRO_ASSERT(max_path_length);
+			*p_path = '\0';
+		}
+	}
+
+//------------------------------------------------------------------------
+//                         FRAMEPRO_PLATFORM_ANDROID
+//------------------------------------------------------------------------
+#elif FRAMEPRO_PLATFORM_ANDROID
+
+	//------------------------------------------------------------------------
+	#include <sys/signal.h>
+	#include <sys/types.h>
+	#include <unistd.h>
+	#include <unwind.h>
+	#include <dlfcn.h>
+       	
+	//------------------------------------------------------------------------
+	namespace FramePro
+	{
+		//------------------------------------------------------------------------
+		struct BacktraceState
+		{
+			void** mp_Current;
+			void** mp_End;
+		};
+
+		//------------------------------------------------------------------------
+		static _Unwind_Reason_Code UnwindCallback(struct _Unwind_Context* p_context, void* p_arg)
+		{
+			BacktraceState* p_state = static_cast<BacktraceState*>(p_arg);
+			uintptr_t p_instr_ptr = _Unwind_GetIP(p_context);
+			if (p_instr_ptr)
+			{
+				if (p_state->mp_Current == p_state->mp_End)
+					return _URC_END_OF_STACK;
+				else
+					*p_state->mp_Current++ = reinterpret_cast<void*>(p_instr_ptr);
+			}
+			return _URC_NO_REASON;
+		}
+
+		//------------------------------------------------------------------------
+		size_t backtrace(void** p_buffer, size_t buffer_size)
+		{
+			BacktraceState state = {p_buffer, p_buffer + buffer_size};
+			_Unwind_Backtrace(UnwindCallback, &state);
+
+			return state.mp_Current - p_buffer;
+		}
+
+		//------------------------------------------------------------------------
+		int64 Platform::GetTimerFrequency()
+		{
+			return 1000000000;
+		}
+
+		//------------------------------------------------------------------------
+		void Platform::DebugBreak()
+		{
+			raise(SIGTRAP);
+		}
+
+		//------------------------------------------------------------------------
+		int Platform::GetCore()
+		{
+			return sched_getcpu();
+		}
+
+		//------------------------------------------------------------------------
+		Platform::Enum Platform::GetPlatformEnum()
+		{
+			return Platform::Android;
+		}
+
+		//------------------------------------------------------------------------
+		void Platform::EnumerateModules(Array<ModulePacket*>& module_packets, Allocator* p_allocator)
+		{
+			EnumModulesLinux::EnumerateModules(module_packets, p_allocator);
+		}
+
+		//------------------------------------------------------------------------
+		bool Platform::GetStackTrace(void** stack, int& stack_size, unsigned int& hash)
+		{
+			stack_size = backtrace(stack, FRAMEPRO_STACK_TRACE_SIZE);
+			hash = GetHashAndStackSize(stack, stack_size);
+			return true;
+		}
+
+		//------------------------------------------------------------------------
+		FILE*& GetOSFile(void* p_os_file_mem)
+		{
+			return *(FILE**)p_os_file_mem;
+		}
+
+		//------------------------------------------------------------------------
+		FILE* GetOSFile(const void* p_os_file_mem)
+		{
+			return *(FILE**)p_os_file_mem;
+		}
+
+		//------------------------------------------------------------------------
+		bool Platform::OpenFileForRead(void* p_os_file_mem, int os_file_mem_size, const char* p_filename)
+		{
+			FRAMEPRO_ASSERT((size_t)os_file_mem_size > sizeof(FILE*));
+			FILE*& p_file = GetOSFile(p_os_file_mem);
+			p_file = fopen(p_filename, "rb");
+			return p_file != NULL;
+		}
+
+		//------------------------------------------------------------------------
+		bool Platform::OpenFileForRead(void* p_os_file_mem, int os_file_mem_size, const wchar_t* p_filename)
+		{
+			FRAMEPRO_ASSERT((size_t)os_file_mem_size > sizeof(FILE*));
+			FILE*& p_file = GetOSFile(p_os_file_mem);
+			char ansi_filename[FRAMEPRO_MAX_PATH];
+			wcstombs(ansi_filename, p_filename, FRAMEPRO_MAX_PATH);
+			p_file = fopen(ansi_filename, "rb");
+			return p_file != 0;
+		}
+
+		//------------------------------------------------------------------------
+		bool Platform::OpenFileForWrite(void* p_os_file_mem, int os_file_mem_size, const char* p_filename)
+		{
+			FRAMEPRO_ASSERT((size_t)os_file_mem_size > sizeof(FILE*));
+			FILE*& p_file = GetOSFile(p_os_file_mem);
+			p_file = fopen(p_filename, "wb");
+			return p_file != NULL;
+		}
+
+		//------------------------------------------------------------------------
+		bool Platform::OpenFileForWrite(void* p_os_file_mem, int os_file_mem_size, const wchar_t* p_filename)
+		{
+			FRAMEPRO_ASSERT((size_t)os_file_mem_size > sizeof(FILE*));
+			FILE*& p_file = GetOSFile(p_os_file_mem);
+			char ansi_filename[FRAMEPRO_MAX_PATH];
+			wcstombs(ansi_filename, p_filename, FRAMEPRO_MAX_PATH);
+			p_file = fopen(ansi_filename, "wb");
+			return p_file != NULL;
+		}
+
+		//------------------------------------------------------------------------
+		void Platform::CloseFile(void* p_os_file_mem)
+		{
+			FILE*& p_file = GetOSFile(p_os_file_mem);
+			fclose(p_file);
+			p_file = NULL;
+		}
+
+		//------------------------------------------------------------------------
+		void Platform::ReadFromFile(void* p_os_file_mem, void* p_data, size_t size)
+		{
+			FILE* p_file = GetOSFile(p_os_file_mem);
+			fread(p_data, size, 1, p_file);
+		}
+
+		//------------------------------------------------------------------------
+		void Platform::WriteToFile(void* p_os_file_mem, const void* p_data, size_t size)
+		{
+			FILE* p_file = GetOSFile(p_os_file_mem);
+			fwrite(p_data, size, 1, p_file);
+		}
+
+		//------------------------------------------------------------------------
+		int Platform::GetFileSize(const void* p_os_file_mem)
+		{
+			FILE* p_file = GetOSFile(p_os_file_mem);
+			int pos = ftell(p_file);
+			fseek(p_file, 0, SEEK_END);
+			int size = ftell(p_file);
+			fseek(p_file, pos, SEEK_SET);
+			return size;
+		}
+
+		//------------------------------------------------------------------------
+		void Platform::DebugWrite(const char* p_string)
+		{
+			GenericPlatform::DebugWrite(p_string);
+		}
+
+		//------------------------------------------------------------------------
+		void Platform::CreateLock(void* p_os_lock_mem, int os_lock_mem_size)
+		{
+			GenericPlatform::CreateLock(p_os_lock_mem, os_lock_mem_size);
+		}
+
+		//------------------------------------------------------------------------
+		void Platform::DestroyLock(void* p_os_lock_mem)
+		{
+			GenericPlatform::DestroyLock(p_os_lock_mem);
+		}
+
+		//------------------------------------------------------------------------
+		void Platform::TakeLock(void* p_os_lock_mem)
+		{
+			GenericPlatform::TakeLock(p_os_lock_mem);
+		}
+
+		//------------------------------------------------------------------------
+		void Platform::ReleaseLock(void* p_os_lock_mem)
+		{
+			GenericPlatform::ReleaseLock(p_os_lock_mem);
+		}
+
+		//------------------------------------------------------------------------
+		void Platform::GetLocalTime(tm* p_tm, const time_t *p_time)
+		{
+			GenericPlatform::GetLocalTime(p_tm, p_time);
+		}
+
+		//------------------------------------------------------------------------
+		int Platform::GetCurrentProcessId()
+		{
+			return GenericPlatform::GetCurrentProcessId();
+		}
+
+		//------------------------------------------------------------------------
+		void Platform::VSPrintf(char* p_buffer, size_t const buffer_size, const char* p_format, va_list arg_list)
+		{
+			GenericPlatform::VSPrintf(p_buffer, buffer_size, p_format, arg_list);
+		}
+
+		//------------------------------------------------------------------------
+		void Platform::ToString(int value, char* p_dest, int dest_size)
+		{
+			GenericPlatform::ToString(value, p_dest, dest_size);
+		}
+
+		//------------------------------------------------------------------------
+		int Platform::GetCurrentThreadId()
+		{
+			return gettid();
+		}
+
+		//------------------------------------------------------------------------
+		bool Platform::GetProcessName(int process_id, char* p_name, int max_name_length)
+		{
+			return GenericPlatform::GetProcessName(process_id, p_name, max_name_length);
+		}
+
+		//------------------------------------------------------------------------
+		void Platform::CreateEventX(void* p_os_event_mem, int os_event_mem_size, bool initial_state, bool auto_reset)
+		{
+			GenericPlatform::CreateEventX(p_os_event_mem, os_event_mem_size, initial_state, auto_reset);
+		}
+
+		//------------------------------------------------------------------------
+		void Platform::DestroyEvent(void* p_os_event_mem)
+		{
+			GenericPlatform::DestroyEvent(p_os_event_mem);
+		}
+
+		//------------------------------------------------------------------------
+		void Platform::SetEvent(void* p_os_event_mem)
+		{
+			GenericPlatform::SetEvent(p_os_event_mem);
+		}
+
+		//------------------------------------------------------------------------
+		void Platform::ResetEvent(void* p_os_event_mem)
+		{
+			GenericPlatform::ResetEvent(p_os_event_mem);
+		}
+
+		//------------------------------------------------------------------------
+		int Platform::WaitEvent(void* p_os_event_mem, int timeout)
+		{
+			return GenericPlatform::WaitEvent(p_os_event_mem, timeout);
+		}
+
+		//------------------------------------------------------------------------
+		void Platform::CreateSocket(void* p_os_socket_mem, int os_socket_mem_size)
+		{
+			GenericPlatform::CreateSocket(p_os_socket_mem, os_socket_mem_size);
+		}
+
+		//------------------------------------------------------------------------
+		void Platform::DestroySocket(void* p_os_socket_mem)
+		{
+			GenericPlatform::DestroySocket(p_os_socket_mem);
+		}
+
+		//------------------------------------------------------------------------
+		bool Platform::InitialiseSocketSystem()
+		{
+			return GenericPlatform::InitialiseSocketSystem();
+		}
+
+		//------------------------------------------------------------------------
+		void Platform::UninitialiseSocketSystem()
+		{
+			GenericPlatform::UninitialiseSocketSystem();
+		}
+
+		//------------------------------------------------------------------------
+		void Platform::DisconnectSocket(void* p_os_socket_mem, bool stop_listening)
+		{
+			GenericPlatform::DisconnectSocket(p_os_socket_mem, stop_listening);
+		}
+
+		//------------------------------------------------------------------------
+		bool Platform::StartSocketListening(void* p_os_socket_mem)
+		{
+			return GenericPlatform::StartSocketListening(p_os_socket_mem);
+		}
+
+		//------------------------------------------------------------------------
+		bool Platform::BindSocket(void* p_os_socket_mem, const char* p_port)
+		{
+			return GenericPlatform::BindSocket(p_os_socket_mem, p_port);
+		}
+
+		//------------------------------------------------------------------------
+		bool Platform::AcceptSocket(void* p_source_os_socket_mem, void* p_target_os_socket_mem)
+		{
+			return GenericPlatform::AcceptSocket(p_source_os_socket_mem, p_target_os_socket_mem);
+		}
+
+		//------------------------------------------------------------------------
+		bool Platform::SocketSend(void* p_os_socket_mem, const void* p_buffer, int size, int& bytes_sent)
+		{
+			return GenericPlatform::SocketSend(p_os_socket_mem, p_buffer, size, bytes_sent);
+		}
+
+		//------------------------------------------------------------------------
+		bool Platform::SocketReceive(void* p_os_socket_mem, const void* p_buffer, int size, int& bytes_received)
+		{
+			return GenericPlatform::SocketReceive(p_os_socket_mem, p_buffer, size, bytes_received);
+		}
+
+		//------------------------------------------------------------------------
+		bool Platform::IsSocketValid(const void* p_os_socket_mem)
+		{
+			return GenericPlatform::IsSocketValid(p_os_socket_mem);
+		}
+
+		//------------------------------------------------------------------------
+		void Platform::HandleSocketError()
+		{
+			GenericPlatform::HandleSocketError();
+		}
+
+		//------------------------------------------------------------------------
+		void Platform::CreateThread(
+			void* p_os_thread_mem,
+			int os_thread_mem_size,
+			ThreadMain p_thread_main,
+			void* p_context,
+			Allocator* p_allocator)
+		{
+			GenericPlatform::CreateThread(
+				p_os_thread_mem,
+				os_thread_mem_size,
+				p_thread_main,
+				p_context,
+				p_allocator);
+		}
+
+		//------------------------------------------------------------------------
+		void Platform::SetThreadPriority(void* p_os_thread_mem, int priority)
+		{
+			GenericPlatform::SetThreadPriority(p_os_thread_mem, priority);
+		}
+
+		//------------------------------------------------------------------------
+		void Platform::SetThreadAffinity(void* p_os_thread_mem, int affinity)
+		{
+			GenericPlatform::SetThreadAffinity(p_os_thread_mem, affinity);
+		}
+
+		//------------------------------------------------------------------------
+		#if FRAMEPRO_USE_TLS_SLOTS
+			#error this platform is not using TLS slots
+		#endif
+
+		__thread void* gp_FrameProTLS = NULL;
+
+		//------------------------------------------------------------------------
+		uint Platform::AllocateTLSSlot()
+		{
+			return 0;
+		}
+
+		//------------------------------------------------------------------------
+		void* Platform::GetTLSValue(uint)
+		{
+			return gp_FrameProTLS;
+		}
+
+		//------------------------------------------------------------------------
+		void Platform::SetTLSValue(uint, void* p_value)
+		{
+			gp_FrameProTLS = p_value;
+		}
+
+		//------------------------------------------------------------------------
+		void Platform::GetRecordingFolder(char* p_path, int max_path_length)
+		{
+			FRAMEPRO_ASSERT(max_path_length);
+			*p_path = '\0';
+		}
+		
+		//------------------------------------------------------------------------
+		void* Platform::CreateContextSwitchRecorder(Allocator* p_allocator)
+		{
+			return GenericPlatform::CreateContextSwitchRecorder(p_allocator);
+		}
+
+		//------------------------------------------------------------------------
+		void Platform::DestroyContextSwitchRecorder(void* p_context_switch_recorder, Allocator* p_allocator)
+		{
+			GenericPlatform::DestroyContextSwitchRecorder(p_context_switch_recorder, p_allocator);
+		}
+
+		//------------------------------------------------------------------------
+		bool Platform::StartRecordingContextSitches(
+			void* p_context_switch_recorder,
+			ContextSwitchCallbackFunction p_callback,
+			void* p_context,
+			DynamicString& error)
+		{
+			return GenericPlatform::StartRecordingContextSitches(
+				p_context_switch_recorder,
+				p_callback,
+				p_context,
+				error);
+		}
+
+		//------------------------------------------------------------------------
+		void Platform::StopRecordingContextSitches(void* p_context_switch_recorder)
+		{
+			GenericPlatform::StopRecordingContextSitches(p_context_switch_recorder);
+		}
+
+		//------------------------------------------------------------------------
+		void Platform::FlushContextSwitches(void* p_context_switch_recorder)
+		{
+			GenericPlatform::FlushContextSwitches(p_context_switch_recorder);
+		}
+	}
+
+//------------------------------------------------------------------------
+//                         FRAMEPRO_PLATFORM_LINUX
+//------------------------------------------------------------------------
+#elif FRAMEPRO_PLATFORM_LINUX
+
+	//------------------------------------------------------------------------
+	#include <sys/signal.h>
+	#include <sys/types.h>
+	#include <unistd.h>
+	#include <execinfo.h>
+       	
+	//------------------------------------------------------------------------
+	namespace FramePro
+	{
+		//------------------------------------------------------------------------
+		int64 Platform::GetTimerFrequency()
+		{
+			return 1000000000;
+		}
+
+		//------------------------------------------------------------------------
+		void Platform::DebugBreak()
+		{
+			raise(SIGTRAP);
+		}
+
+		//------------------------------------------------------------------------
+		int Platform::GetCore()
+		{
+			return sched_getcpu();
+		}
+
+		//------------------------------------------------------------------------
+		Platform::Enum Platform::GetPlatformEnum()
+		{
+			return Platform::Linux;
+		}
+
+		//------------------------------------------------------------------------
+		void Platform::EnumerateModules(Array<ModulePacket*>& module_packets, Allocator* p_allocator)
+		{
+			EnumModulesLinux::EnumerateModules(module_packets, p_allocator);
+		}
+
+		//------------------------------------------------------------------------
+		bool Platform::GetStackTrace(void** stack, int& stack_size, unsigned int& hash)
+		{
+			stack_size = backtrace(stack, FRAMEPRO_STACK_TRACE_SIZE);
+			hash = GetHashAndStackSize(stack, stack_size);
+			return true;
+		}
+
+		//------------------------------------------------------------------------
+		FILE*& GetOSFile(void* p_os_file_mem)
+		{
+			return *(FILE**)p_os_file_mem;
+		}
+
+		//------------------------------------------------------------------------
+		FILE* GetOSFile(const void* p_os_file_mem)
+		{
+			return *(FILE**)p_os_file_mem;
+		}
+
+		//------------------------------------------------------------------------
+		bool Platform::OpenFileForRead(void* p_os_file_mem, int os_file_mem_size, const char* p_filename)
+		{
+			FRAMEPRO_ASSERT((size_t)os_file_mem_size > sizeof(FILE*));
+			FILE*& p_file = GetOSFile(p_os_file_mem);
+			p_file = fopen(p_filename, "rb");
+			return p_file != NULL;
+		}
+
+		//------------------------------------------------------------------------
+		bool Platform::OpenFileForRead(void* p_os_file_mem, int os_file_mem_size, const wchar_t* p_filename)
+		{
+			FRAMEPRO_ASSERT((size_t)os_file_mem_size > sizeof(FILE*));
+			FILE*& p_file = GetOSFile(p_os_file_mem);
+			char ansi_filename[FRAMEPRO_MAX_PATH];
+			wcstombs(ansi_filename, p_filename, FRAMEPRO_MAX_PATH);
+			p_file = fopen(ansi_filename, "rb");
+			return p_file != 0;
+		}
+
+		//------------------------------------------------------------------------
+		bool Platform::OpenFileForWrite(void* p_os_file_mem, int os_file_mem_size, const char* p_filename)
+		{
+			FRAMEPRO_ASSERT((size_t)os_file_mem_size > sizeof(FILE*));
+			FILE*& p_file = GetOSFile(p_os_file_mem);
+			p_file = fopen(p_filename, "wb");
+			return p_file != NULL;
+		}
+
+		//------------------------------------------------------------------------
+		bool Platform::OpenFileForWrite(void* p_os_file_mem, int os_file_mem_size, const wchar_t* p_filename)
+		{
+			FRAMEPRO_ASSERT((size_t)os_file_mem_size > sizeof(FILE*));
+			FILE*& p_file = GetOSFile(p_os_file_mem);
+			char ansi_filename[FRAMEPRO_MAX_PATH];
+			wcstombs(ansi_filename, p_filename, FRAMEPRO_MAX_PATH);
+			p_file = fopen(ansi_filename, "wb");
+			return p_file != NULL;
+		}
+
+		//------------------------------------------------------------------------
+		void Platform::CloseFile(void* p_os_file_mem)
+		{
+			FILE*& p_file = GetOSFile(p_os_file_mem);
+			fclose(p_file);
+			p_file = NULL;
+		}
+
+		//------------------------------------------------------------------------
+		void Platform::ReadFromFile(void* p_os_file_mem, void* p_data, size_t size)
+		{
+			FILE* p_file = GetOSFile(p_os_file_mem);
+			fread(p_data, size, 1, p_file);
+		}
+
+		//------------------------------------------------------------------------
+		void Platform::WriteToFile(void* p_os_file_mem, const void* p_data, size_t size)
+		{
+			FILE* p_file = GetOSFile(p_os_file_mem);
+			fwrite(p_data, size, 1, p_file);
+		}
+
+		//------------------------------------------------------------------------
+		int Platform::GetFileSize(const void* p_os_file_mem)
+		{
+			FILE* p_file = GetOSFile(p_os_file_mem);
+			int pos = ftell(p_file);
+			fseek(p_file, 0, SEEK_END);
+			int size = ftell(p_file);
+			fseek(p_file, pos, SEEK_SET);
+			return size;
+		}
+
+		//------------------------------------------------------------------------
+		void Platform::DebugWrite(const char* p_string)
+		{
+			GenericPlatform::DebugWrite(p_string);
+		}
+
+		//------------------------------------------------------------------------
+		void Platform::CreateLock(void* p_os_lock_mem, int os_lock_mem_size)
+		{
+			GenericPlatform::CreateLock(p_os_lock_mem, os_lock_mem_size);
+		}
+
+		//------------------------------------------------------------------------
+		void Platform::DestroyLock(void* p_os_lock_mem)
+		{
+			GenericPlatform::DestroyLock(p_os_lock_mem);
+		}
+
+		//------------------------------------------------------------------------
+		void Platform::TakeLock(void* p_os_lock_mem)
+		{
+			GenericPlatform::TakeLock(p_os_lock_mem);
+		}
+
+		//------------------------------------------------------------------------
+		void Platform::ReleaseLock(void* p_os_lock_mem)
+		{
+			GenericPlatform::ReleaseLock(p_os_lock_mem);
+		}
+
+		//------------------------------------------------------------------------
+		void Platform::GetLocalTime(tm* p_tm, const time_t *p_time)
+		{
+			GenericPlatform::GetLocalTime(p_tm, p_time);
+		}
+
+		//------------------------------------------------------------------------
+		int Platform::GetCurrentProcessId()
+		{
+			return GenericPlatform::GetCurrentProcessId();
+		}
+
+		//------------------------------------------------------------------------
+		void Platform::VSPrintf(char* p_buffer, size_t const buffer_size, const char* p_format, va_list arg_list)
+		{
+			GenericPlatform::VSPrintf(p_buffer, buffer_size, p_format, arg_list);
+		}
+
+		//------------------------------------------------------------------------
+		void Platform::ToString(int value, char* p_dest, int dest_size)
+		{
+			GenericPlatform::ToString(value, p_dest, dest_size);
+		}
+
+		//------------------------------------------------------------------------
+		int Platform::GetCurrentThreadId()
+		{
+			return gettid();
+		}
+
+		//------------------------------------------------------------------------
+		bool Platform::GetProcessName(int process_id, char* p_name, int max_name_length)
+		{
+			return GenericPlatform::GetProcessName(process_id, p_name, max_name_length);
+		}
+
+		//------------------------------------------------------------------------
+		void Platform::CreateEventX(void* p_os_event_mem, int os_event_mem_size, bool initial_state, bool auto_reset)
+		{
+			GenericPlatform::CreateEventX(p_os_event_mem, os_event_mem_size, initial_state, auto_reset);
+		}
+
+		//------------------------------------------------------------------------
+		void Platform::DestroyEvent(void* p_os_event_mem)
+		{
+			GenericPlatform::DestroyEvent(p_os_event_mem);
+		}
+
+		//------------------------------------------------------------------------
+		void Platform::SetEvent(void* p_os_event_mem)
+		{
+			GenericPlatform::SetEvent(p_os_event_mem);
+		}
+
+		//------------------------------------------------------------------------
+		void Platform::ResetEvent(void* p_os_event_mem)
+		{
+			GenericPlatform::ResetEvent(p_os_event_mem);
+		}
+
+		//------------------------------------------------------------------------
+		int Platform::WaitEvent(void* p_os_event_mem, int timeout)
+		{
+			return GenericPlatform::WaitEvent(p_os_event_mem, timeout);
+		}
+
+		//------------------------------------------------------------------------
+		void Platform::CreateSocket(void* p_os_socket_mem, int os_socket_mem_size)
+		{
+			GenericPlatform::CreateSocket(p_os_socket_mem, os_socket_mem_size);
+		}
+
+		//------------------------------------------------------------------------
+		void Platform::DestroySocket(void* p_os_socket_mem)
+		{
+			GenericPlatform::DestroySocket(p_os_socket_mem);
+		}
+
+		//------------------------------------------------------------------------
+		bool Platform::InitialiseSocketSystem()
+		{
+			return GenericPlatform::InitialiseSocketSystem();
+		}
+
+		//------------------------------------------------------------------------
+		void Platform::UninitialiseSocketSystem()
+		{
+			GenericPlatform::UninitialiseSocketSystem();
+		}
+
+		//------------------------------------------------------------------------
+		void Platform::DisconnectSocket(void* p_os_socket_mem, bool stop_listening)
+		{
+			GenericPlatform::DisconnectSocket(p_os_socket_mem, stop_listening);
+		}
+
+		//------------------------------------------------------------------------
+		bool Platform::StartSocketListening(void* p_os_socket_mem)
+		{
+			return GenericPlatform::StartSocketListening(p_os_socket_mem);
+		}
+
+		//------------------------------------------------------------------------
+		bool Platform::BindSocket(void* p_os_socket_mem, const char* p_port)
+		{
+			return GenericPlatform::BindSocket(p_os_socket_mem, p_port);
+		}
+
+		//------------------------------------------------------------------------
+		bool Platform::AcceptSocket(void* p_source_os_socket_mem, void* p_target_os_socket_mem)
+		{
+			return GenericPlatform::AcceptSocket(p_source_os_socket_mem, p_target_os_socket_mem);
+		}
+
+		//------------------------------------------------------------------------
+		bool Platform::SocketSend(void* p_os_socket_mem, const void* p_buffer, int size, int& bytes_sent)
+		{
+			return GenericPlatform::SocketSend(p_os_socket_mem, p_buffer, size, bytes_sent);
+		}
+
+		//------------------------------------------------------------------------
+		bool Platform::SocketReceive(void* p_os_socket_mem, const void* p_buffer, int size, int& bytes_received)
+		{
+			return GenericPlatform::SocketReceive(p_os_socket_mem, p_buffer, size, bytes_received);
+		}
+
+		//------------------------------------------------------------------------
+		bool Platform::IsSocketValid(const void* p_os_socket_mem)
+		{
+			return GenericPlatform::IsSocketValid(p_os_socket_mem);
+		}
+
+		//------------------------------------------------------------------------
+		void Platform::HandleSocketError()
+		{
+			GenericPlatform::HandleSocketError();
+		}
+
+		//------------------------------------------------------------------------
+		void Platform::CreateThread(
+			void* p_os_thread_mem,
+			int os_thread_mem_size,
+			ThreadMain p_thread_main,
+			void* p_context,
+			Allocator* p_allocator)
+		{
+			GenericPlatform::CreateThread(
+				p_os_thread_mem,
+				os_thread_mem_size,
+				p_thread_main,
+				p_context,
+				p_allocator);
+		}
+
+		//------------------------------------------------------------------------
+		void Platform::SetThreadPriority(void* p_os_thread_mem, int priority)
+		{
+			GenericPlatform::SetThreadPriority(p_os_thread_mem, priority);
+		}
+
+		//------------------------------------------------------------------------
+		void Platform::SetThreadAffinity(void* p_os_thread_mem, int affinity)
+		{
+			GenericPlatform::SetThreadAffinity(p_os_thread_mem, affinity);
+		}
+
+		//------------------------------------------------------------------------
+		#if FRAMEPRO_USE_TLS_SLOTS
+			#error this platform is not using TLS slots
+		#endif
+
+		__thread void* gp_FrameProTLS = NULL;
+
+		//------------------------------------------------------------------------
+		uint Platform::AllocateTLSSlot()
+		{
+			return 0;
+		}
+
+		//------------------------------------------------------------------------
+		void* Platform::GetTLSValue(uint)
+		{
+			return gp_FrameProTLS;
+		}
+
+		//------------------------------------------------------------------------
+		void Platform::SetTLSValue(uint, void* p_value)
+		{
+			gp_FrameProTLS = p_value;
+		}
+
+		//------------------------------------------------------------------------
+		void Platform::GetRecordingFolder(char* p_path, int max_path_length)
+		{
+			FRAMEPRO_ASSERT(max_path_length);
+			*p_path = '\0';
+		}
+		
+		//------------------------------------------------------------------------
+		void* Platform::CreateContextSwitchRecorder(Allocator* p_allocator)
+		{
+			return GenericPlatform::CreateContextSwitchRecorder(p_allocator);
+		}
+
+		//------------------------------------------------------------------------
+		void Platform::DestroyContextSwitchRecorder(void* p_context_switch_recorder, Allocator* p_allocator)
+		{
+			GenericPlatform::DestroyContextSwitchRecorder(p_context_switch_recorder, p_allocator);
+		}
+
+		//------------------------------------------------------------------------
+		bool Platform::StartRecordingContextSitches(
+			void* p_context_switch_recorder,
+			ContextSwitchCallbackFunction p_callback,
+			void* p_context,
+			DynamicString& error)
+		{
+			return GenericPlatform::StartRecordingContextSitches(
+				p_context_switch_recorder,
+				p_callback,
+				p_context,
+				error);
+		}
+
+		//------------------------------------------------------------------------
+		void Platform::StopRecordingContextSitches(void* p_context_switch_recorder)
+		{
+			GenericPlatform::StopRecordingContextSitches(p_context_switch_recorder);
+		}
+
+		//------------------------------------------------------------------------
+		void Platform::FlushContextSwitches(void* p_context_switch_recorder)
+		{
+			GenericPlatform::FlushContextSwitches(p_context_switch_recorder);
+		}
+	}
+
+//------------------------------------------------------------------------
+//                         FRAMEPRO_PLATFORM_XBOXONE
+//------------------------------------------------------------------------
+#elif FRAMEPRO_PLATFORM_XBOXONE
+
+	// implemented in FrameProXboxOne.cpp - contact [REDACTED_EMAIL] for this platform
+
+//------------------------------------------------------------------------
+//                         FRAMEPRO_PLATFORM_PS4
+//------------------------------------------------------------------------
+#elif FRAMEPRO_PLATFORM_PS4
+
+	// implemented in FrameProPS4.cpp - contact [REDACTED_EMAIL] for this platform
+
+#else
+
+	#error Platform not defined
+
+#endif
+
+//------------------------------------------------------------------------
+// ---
+// --- FRAMEPRO PLATFORM IMPLEMENTATION END ---
+// ---
+//------------------------------------------------------------------------
+
+
 //------------------------------------------------------------------------
 #endif		// #if FRAMEPRO_ENABLED
