[FramePro] update to the latest version of FramePro runtime integration code (1.5.15.0)

diff --git a/Engine/Source/Runtime/Core/Private/FramePro/FrameProProfiler.cpp b/Engine/Source/Runtime/Core/Private/FramePro/FrameProProfiler.cpp
--- a/Engine/Source/Runtime/Core/Private/FramePro/FrameProProfiler.cpp
+++ b/Engine/Source/Runtime/Core/Private/FramePro/FrameProProfiler.cpp
@@ -203,28 +203,28 @@ static void SendCPUStats()
 		switch (CoreIdx)
 		{
 		case 0:
-			FRAMEPRO_CUSTOM_STAT("Core0Frequency", Freq, "CPUFreq", "GHz");
+			FRAMEPRO_CUSTOM_STAT("Core0Frequency", Freq, "CPUFreq", "GHz", FRAMEPRO_COLOUR(255,255,255) );
 			break;
 		case 1:
-			FRAMEPRO_CUSTOM_STAT("Core1Frequency", Freq, "CPUFreq", "GHz");
+			FRAMEPRO_CUSTOM_STAT("Core1Frequency", Freq, "CPUFreq", "GHz", FRAMEPRO_COLOUR(255,255,255));
 			break;
 		case 2:
-			FRAMEPRO_CUSTOM_STAT("Core2Frequency", Freq, "CPUFreq", "GHz");
+			FRAMEPRO_CUSTOM_STAT("Core2Frequency", Freq, "CPUFreq", "GHz", FRAMEPRO_COLOUR(255,255,255));
 			break;
 		case 3:
-			FRAMEPRO_CUSTOM_STAT("Core3Frequency", Freq, "CPUFreq", "GHz");
+			FRAMEPRO_CUSTOM_STAT("Core3Frequency", Freq, "CPUFreq", "GHz", FRAMEPRO_COLOUR(255,255,255));
 			break;
 		case 4:
-			FRAMEPRO_CUSTOM_STAT("Core4Frequency", Freq, "CPUFreq", "GHz");
+			FRAMEPRO_CUSTOM_STAT("Core4Frequency", Freq, "CPUFreq", "GHz", FRAMEPRO_COLOUR(255,255,255));
 			break;
 		case 5:
-			FRAMEPRO_CUSTOM_STAT("Core5Frequency", Freq, "CPUFreq", "GHz");
+			FRAMEPRO_CUSTOM_STAT("Core5Frequency", Freq, "CPUFreq", "GHz", FRAMEPRO_COLOUR(255,255,255));
 			break;
 		case 6:
-			FRAMEPRO_CUSTOM_STAT("Core6Frequency", Freq, "CPUFreq", "GHz");
+			FRAMEPRO_CUSTOM_STAT("Core6Frequency", Freq, "CPUFreq", "GHz", FRAMEPRO_COLOUR(255,255,255));
 			break;
 		case 7:
-			FRAMEPRO_CUSTOM_STAT("Core7Frequency", Freq, "CPUFreq", "GHz");
+			FRAMEPRO_CUSTOM_STAT("Core7Frequency", Freq, "CPUFreq", "GHz", FRAMEPRO_COLOUR(255,255,255));
 			break;
 		}
 	}
diff --git a/Engine/Source/Runtime/Core/Private/Unix/UnixPlatformMisc.cpp b/Engine/Source/Runtime/Core/Private/Unix/UnixPlatformMisc.cpp
--- a/Engine/Source/Runtime/Core/Private/Unix/UnixPlatformMisc.cpp
+++ b/Engine/Source/Runtime/Core/Private/Unix/UnixPlatformMisc.cpp
@@ -931,12 +931,12 @@ void FUnixPlatformMisc::EndNamedEvent()
 
 void FUnixPlatformMisc::CustomNamedStat(const TCHAR* Text, float Value, const TCHAR* Graph, const TCHAR* Unit)
 {
-	FRAMEPRO_DYNAMIC_CUSTOM_STAT(TCHAR_TO_WCHAR(Text), Value, TCHAR_TO_WCHAR(Graph), TCHAR_TO_WCHAR(Unit));
+	FRAMEPRO_DYNAMIC_CUSTOM_STAT(TCHAR_TO_WCHAR(Text), Value, TCHAR_TO_WCHAR(Graph), TCHAR_TO_WCHAR(Unit), FRAMEPRO_COLOUR(255,255,255));
 }
 
 void FUnixPlatformMisc::CustomNamedStat(const ANSICHAR* Text, float Value, const ANSICHAR* Graph, const ANSICHAR* Unit)
 {
-	FRAMEPRO_DYNAMIC_CUSTOM_STAT(Text, Value, Graph, Unit);
+	FRAMEPRO_DYNAMIC_CUSTOM_STAT(Text, Value, Graph, Unit, FRAMEPRO_COLOUR(255,255,255));
 }
 #endif
 
diff --git a/Engine/Source/Runtime/Core/Public/FramePro/FramePro.h b/Engine/Source/Runtime/Core/Public/FramePro/FramePro.h
--- a/Engine/Source/Runtime/Core/Public/FramePro/FramePro.h
+++ b/Engine/Source/Runtime/Core/Public/FramePro/FramePro.h
@@ -22,18 +22,14 @@
 	Add FramePro.cpp to your project to allow FramePro to communicate with your application.
 */
 
-// START EPIC
-#include "CoreTypes.h"
-#include "HAL/PlatformMisc.h"
-#include "HAL/PlatformTime.h"
-// END EPIC
-
 //------------------------------------------------------------------------
+//
 // FramePro.hpp
+//
 //------------------------------------------------------------------------
 /*
 	FramePro
-	Version:	1.4.7.1
+	Version:	1.5.15.0
 */
 //------------------------------------------------------------------------
 #ifndef FRAMEPRO_H_INCLUDED
@@ -41,9 +37,12 @@
 
 //------------------------------------------------------------------------
 #ifndef FRAMEPRO_ENABLED
-// START EPIC
-	#define FRAMEPRO_ENABLED	0
-// END EPIC
+	// *** enable/disable FramePro here
+	#ifdef __UNREAL__
+		#define FRAMEPRO_ENABLED 0
+	#else
+		#define FRAMEPRO_ENABLED 1
+	#endif
 #endif
 
 //------------------------------------------------------------------------
@@ -59,7 +58,7 @@
 	#define FRAMEPRO_SET_THREAD_NAME(name)													FramePro::SetThreadName(name)
 	#define FRAMEPRO_THREAD_ORDER(thread_name)												FramePro::SetThreadOrder(FramePro::RegisterString(thread_name))
 	#define FRAMEPRO_REGISTER_STRING(str)													FramePro::RegisterString(str)
-	#define FRAMEPRO_START_RECORDING(filename, context_switches, max_file_size)				FramePro::StartRecording(filename, context_switches, max_file_size)
+	#define FRAMEPRO_START_RECORDING(filename, context_switches, callstacks, max_file_size)	FramePro::StartRecording(filename, context_switches, callstacks, max_file_size)
 	#define FRAMEPRO_STOP_RECORDING()														FramePro::StopRecording()
 	#define FRAMEPRO_REGISTER_CONNECTION_CHANGED_CALLBACK(callback, context)				FramePro::RegisterConnectionChangedCallback(callback, context)
 	#define FRAMEPRO_UNREGISTER_CONNECTION_CHANGED_CALLBACK(callback)						FramePro::UnregisterConnectionChangedcallback(callback)
@@ -70,7 +69,7 @@
 	#define FRAMEPRO_CLEANUP_THREAD()														FramePro::CleanupThread()
 	#define FRAMEPRO_THREAD_SCOPE(thread_name)												FramePro::ThreadScope framepro_thread_scope(thread_name)
 	#define FRAMEPRO_LOG(message)															FramePro::Log(message)
-	#define FRAMEPRO_COLOUR(r, g, b)														((uint)((((uint)((r)&0xff))<<16) | (((uint)((g)&0xff))<<16) | ((b)&0xff)))
+	#define FRAMEPRO_COLOUR(r, g, b)														((FramePro::uint)((((FramePro::uint)((r)&0xff))<<16) | (((FramePro::uint)((g)&0xff))<<8) | ((b)&0xff)))
 	#define FRAMEPRO_SET_CONDITIONAL_SCOPE_MIN_TIME(microseconds)							FramePro::SetConditionalScopeMinTimeInMicroseconds(microseconds)
 
 	// scope macros
@@ -93,6 +92,7 @@
 	#define FRAMEPRO_CONDITIONAL_STOP_NAMED_SCOPE(name)										MULTI_STATEMENT( if(FramePro::IsConnected()) { FramePro::int64 framepro_end; FRAMEPRO_GET_CLOCK_COUNT(framepro_end); if(framepro_end - framepro_start > FramePro::GetConditionalScopeMinTime()) FramePro::AddTimeSpan(name "|" FRAMEPRO_SOURCE_STRING, framepro_start, framepro_end); } )
 	#define FRAMEPRO_CONDITIONAL_STOP_DYNAMIC_SCOPE(dynamic_string)							MULTI_STATEMENT( if(FramePro::IsConnected()) { FramePro::int64 framepro_end; FRAMEPRO_GET_CLOCK_COUNT(framepro_end); if(framepro_end - framepro_start > FramePro::GetConditionalScopeMinTime()) FramePro::AddTimeSpan(FramePro::RegisterString(dynamic_string), FRAMEPRO_SOURCE_STRING, framepro_start, framepro_end); } )
 	#define FRAMEPRO_CONDITIONAL_PARENT_SCOPE(name, callback, pre_duration, post_duration)	FramePro::ConditionalParentTimerScope FRAMEPRO_UNIQUE(timer_scope)(name, FRAMEPRO_SOURCE_STRING, callback, pre_duration, post_duration)
+	#define FRAMEPRO_SET_SCOPE_COLOUR(name, colour)											MULTI_STATEMENT( FramePro::SetScopeColour(FramePro::RegisterString(name), colour); )
 
 	// idle scope macros
 	#define FRAMEPRO_IDLE_SCOPE()															FramePro::TimerScope FRAMEPRO_UNIQUE(timer_scope)(FRAMEPRO_FUNCTION_NAME "|" FRAMEPRO_SOURCE_STRING_IDLE)
@@ -111,15 +111,18 @@
 	#define FRAMEPRO_IDLE_CONDITIONAL_STOP_DYNAMIC_SCOPE(dynamic_string)					MULTI_STATEMENT( if(FramePro::IsConnected()) { FramePro::int64 framepro_end; FRAMEPRO_GET_CLOCK_COUNT(framepro_end); if(framepro_end - framepro_start > FramePro::GetConditionalScopeMinTime()) FramePro::AddTimeSpan(FramePro::RegisterString(dynamic_string), FRAMEPRO_SOURCE_STRING_IDLE, framepro_start, framepro_end); } )
 
 	// custom stat macros
-	#define FRAMEPRO_CUSTOM_STAT(name, value, graph, unit)									MULTI_STATEMENT( if(FramePro::IsConnected()) FramePro::AddCustomStat(name, value, graph, unit); )
-	#define FRAMEPRO_DYNAMIC_CUSTOM_STAT(name, value, graph, unit)							MULTI_STATEMENT( if(FramePro::IsConnected()) FramePro::AddCustomStat(FramePro::IsConnected() ? FramePro::RegisterString(name) : -1, value, graph, unit); )
-	#define FRAMEPRO_SCOPE_CUSTOM_STAT(name, value, graph, unit)							MULTI_STATEMENT( if(FramePro::IsConnected()) FramePro::SetScopeCustomStat(name, value, graph, unit); )
+	#define FRAMEPRO_CUSTOM_STAT(name, value, graph, unit, colour)							MULTI_STATEMENT( if(FramePro::IsConnected()) FramePro::AddCustomStat(name, value, graph, unit, colour); )
+	#define FRAMEPRO_DYNAMIC_CUSTOM_STAT(name, value, graph, unit, colour)					MULTI_STATEMENT( if(FramePro::IsConnected()) FramePro::AddCustomStat(FramePro::RegisterString(name), value, FramePro::RegisterString(graph), FramePro::RegisterString(unit), colour); )
+	#define FRAMEPRO_SCOPE_CUSTOM_STAT(name, value, graph, unit, colour)					MULTI_STATEMENT( if(FramePro::IsConnected()) FramePro::SetScopeCustomStat(name, value, graph, unit, colour); )
+	#define FRAMEPRO_SET_CUSTOM_STAT_GRAPH(name, graph)										MULTI_STATEMENT( FramePro::SetCustomStatGraph(FramePro::RegisterString(name), FramePro::RegisterString(graph)); )
+	#define FRAMEPRO_SET_CUSTOM_STAT_UNIT(name, unit)										MULTI_STATEMENT( FramePro::SetCustomStatUnit(FramePro::RegisterString(name), FramePro::RegisterString(unit)); )
+	#define FRAMEPRO_SET_CUSTOM_STAT_COLOUR(name, colour)									MULTI_STATEMENT( FramePro::SetCustomStatColour(FramePro::RegisterString(name), colour); )
 
 	// high-res timers
 	#define FRAMEPRO_HIRES_SCOPE(name)														FramePro::HiResTimerScope FRAMEPRO_UNIQUE(hires_scope)(name)
 
 	// global high-res timers
-	#define FRAMEPRO_DECL_GLOBAL_HIRES_TIMER(name, graph_name)								FramePro::GlobalHiResTimer g_FrameProHiResTimer##name(#name, graph_name)
+	#define FRAMEPRO_DECL_GLOBAL_HIRES_TIMER(name)											FramePro::GlobalHiResTimer g_FrameProHiResTimer##name(#name)
 	#define FRAMEPRO_GLOBAL_HIRES_SCOPE(name)												FramePro::GlobalHiResTimerScope FRAMEPRO_UNIQUE(timer_scope)(g_FrameProHiResTimer##name)
 
 	// events
@@ -138,7 +141,7 @@
 	#define FRAMEPRO_SET_THREAD_NAME(name)													((void)0)
 	#define FRAMEPRO_THREAD_ORDER(thread_name)												((void)0)
 	#define FRAMEPRO_REGISTER_STRING(str)													0
-	#define FRAMEPRO_START_RECORDING(filename, context_switches, max_file_size)				((void)0)
+	#define FRAMEPRO_START_RECORDING(filename, context_switches, callstacks, max_file_size)	((void)0)
 	#define FRAMEPRO_STOP_RECORDING()														((void)0)
 	#define FRAMEPRO_REGISTER_CONNECTION_CHANGED_CALLBACK(callback, context)				((void)0)
 	#define FRAMEPRO_UNREGISTER_CONNECTION_CHANGED_CALLBACK(callback)						((void)0)
@@ -170,6 +173,7 @@
 	#define FRAMEPRO_CONDITIONAL_STOP_NAMED_SCOPE(name)										((void)0)
 	#define FRAMEPRO_CONDITIONAL_STOP_DYNAMIC_SCOPE(dynamic_string)							((void)0)
 	#define FRAMEPRO_CONDITIONAL_PARENT_SCOPE(name, callback, pre_duration, post_duration)	((void)0)
+	#define FRAMEPRO_SET_SCOPE_COLOUR(name, colour)											((void)0)
 
 	#define FRAMEPRO_IDLE_SCOPE()															((void)0)
 	#define FRAMEPRO_IDLE_NAMED_SCOPE(name)													((void)0)
@@ -186,13 +190,16 @@
 	#define FRAMEPRO_IDLE_CONDITIONAL_STOP_NAMED_SCOPE(name)								((void)0)
 	#define FRAMEPRO_IDLE_CONDITIONAL_STOP_DYNAMIC_SCOPE(dynamic_string)					((void)0)
 
-	#define FRAMEPRO_CUSTOM_STAT(name, value, graph, unit)									((void)0)
-	#define FRAMEPRO_DYNAMIC_CUSTOM_STAT(name, value, graph, unit)							((void)0)
-	#define FRAMEPRO_SCOPE_CUSTOM_STAT(name, value, graph, unit)							((void)0)
+	#define FRAMEPRO_CUSTOM_STAT(name, value, graph, unit, colour)							((void)0)
+	#define FRAMEPRO_DYNAMIC_CUSTOM_STAT(name, value, graph, unit, colour)					((void)0)
+	#define FRAMEPRO_SCOPE_CUSTOM_STAT(name, value, graph, unit, colour)					((void)0)
+	#define FRAMEPRO_SET_CUSTOM_STAT_GRAPH(name, graph)										((void)0)
+	#define FRAMEPRO_SET_CUSTOM_STAT_UNIT(name, unit)										((void)0)
+	#define FRAMEPRO_SET_CUSTOM_STAT_COLOUR(name, colour)									((void)0)
 
 	#define FRAMEPRO_HIRES_SCOPE(name)														((void)0)
 
-	#define FRAMEPRO_DECL_GLOBAL_HIRES_TIMER(name, graph_name)								
+	#define FRAMEPRO_DECL_GLOBAL_HIRES_TIMER(name)											
 	#define FRAMEPRO_GLOBAL_HIRES_SCOPE(name)												((void)0)
 
 	#define FRAMEPRO_EVENT(name, colour)													((void)0)
@@ -239,15 +246,18 @@ namespace FramePro
 
 //------------------------------------------------------------------------
 // Never send scopes shorter than this. Undefine to send all scopes.
-#define FRAMEPRO_SCOPE_MIN_TIME 10       // in ns
+#define FRAMEPRO_SCOPE_MIN_TIME 10       // ns
 
 //------------------------------------------------------------------------
-#define FRAMEPRO_WAIT_EVENT_MIN_TIME 10		// in ns
+// In order to not flood the channel with wait events of zero length, the event
+// is only sent if it is longer than this value
+#define FRAMEPRO_WAIT_EVENT_MIN_TIME 10		// ns
 
 //------------------------------------------------------------------------
-// START EPIC
-#define FRAMEPRO_ENABLE_CALLSTACKS 0
-// END EPIC
+#define FRAMEPRO_ENABLE_CONTEXT_SWITCH_TRACKING 1
+
+//------------------------------------------------------------------------
+#define FRAMEPRO_ENABLE_CALLSTACKS 1
 
 //------------------------------------------------------------------------
 // If this is disabled and two dynamic strings happen to hash to the same value you may see
@@ -256,6 +266,10 @@ namespace FramePro
 // which don't call RegiserString anyway.
 #define FRAMEPRO_DETECT_HASH_COLLISIONS 0			
 
+//------------------------------------------------------------------------
+#define FRAMEPRO_STACK_TRACE_SIZE 128
+
+
 //------------------------------------------------------------------------
 // DOCUMENTATION
 //------------------------------------------------------------------------
@@ -284,6 +298,19 @@ namespace FramePro
 // those to a TLS buffer along with the name id (typically the string literal pointer) - and that's it!
 // No locking, no hashmap lookups, minimal overhead.
 
+// Platforms
+// ---------
+// FramePro supports many platforms, including Windows, Linux, XboxOne, PS4, UE4.
+// If you need to add a platform please implement the functions in the Platform namespace (see FrameProPlatform.h)
+// You will also need to define the platform specific defines in FrameProPlatform.h
+// There are generic utility functions for Windows nad Linux in GenericPlatform that you can use.
+// You can use EventTraceWin32 for context switch recording on windows based platforms.
+// All platform specific code is in FrameProPlatform.cpp/h (or the end of the combined FramePro.cpp file)
+//
+// Proprietary comsole code is split out into seperate files, eg, FrameProPS4.cpp/h.
+// Please contact [REDACTED_EMAIL] for these files.
+//
+
 
 // Definitions
 // -----------
@@ -362,10 +389,11 @@ namespace FramePro
 // scope instead of the string.
 
 
-// FRAMEPRO_START_RECORDING(filename, context_switches, max_file_size)
+// FRAMEPRO_START_RECORDING(filename, context_switches, callstacks, max_file_size)
 // -------------------------------------------------------------------
 // filename: string - the filename to write the recording out to. Extension should be .framepro_recording
 // context_switches: bool - whether to record context switches
+// callstacks: bool - true to record the callstack of each scope. Warning, can negatively affect performance if you have a lot of scopes
 // max_file_size: stop recording once the file reaches this size
 // Instead of connecting to your game with FramePro you can instead write all of the network data out to a file.
 // This recording file can then be loaded in FramePro. This can be useful if you want to start/stop a recording
@@ -379,14 +407,14 @@ namespace FramePro
 // Stop a recording started with FRAMEPRO_START_RECORDING.
 
 
-// FRAMEPRO_REGISTER_CONNECTION_CHANGED_CALLBACK(callback, context)
-// ----------------------------------------------------------------
+// FRAMEPRO_REGISTER_CONNECTION_CHANGED_CALLBACK(callback, p_recording_filename, context)
+// --------------------------------------------------------------------------------------
 // callback: ConnectionChangedCallback function pointer
 // context: void*
 // Add a callback to be called when a connection is made or lost.
 // context is passed in to the callback as user_data.
 // The callback is of the form:
-//		typedef void (*ConnectionChangedCallback)(bool connected, void* user_data);
+//		typedef void (*ConnectionChangedCallback)(bool connected, const wchar_t* p_recording_filename, void* user_data);
 // If already connected the callback will be called immediately.
 // Multiple callbacks can be registered.
 
@@ -623,6 +651,14 @@ namespace FramePro
 // Note that nesting conditional parent scopes is not supported.
 
 
+// FRAMEPRO_SET_SCOPE_COLOUR(name, colour)
+// --------------------------------------
+// name: string (literal or dynamic)
+// colour: colour using FRAMEPRO_COLOUR macro
+// Set the colour for a scope. Note this will have no effect if the user has set a custom colour in the UI
+// This is a slow macro, so please only call once, not every frame.
+
+
 //--------------------------------------------------------------------------------------------------------
 // idle scope macros
 
@@ -697,28 +733,50 @@ namespace FramePro
 // the same thereafter. Graph is the graph name to display the stats. Multipole stats can be put into the same graph.
 // The Units will show in the y-axis.
 
-// FRAMEPRO_CUSTOM_STAT(name, value, graph, unit)
-// -----------------------------------------------
-// name: string literal					- the name of the custom stat
-// value: int, int64, float, double		- value to add to the custom stat
-// graph: string literal				- the name of the graph in which to display the stat
-// unit: string literal				- y-axis unit to be displayed in FramePro, eg "ms", "MB" etc.
+// FRAMEPRO_CUSTOM_STAT(name, value, graph, unit, colour)
+// ------------------------------------------------------
+// name: string literal							- the name of the custom stat
+// value: int, int64, float, double				- value to add to the custom stat
+// graph: string (literal or dynamic)			- the name of the graph in which the stat will be shown
+// unit: string (literal or dynamic)			- the y-axis unit to be displayed in FramePro, eg "ms", "MB" etc for a custom stat
+// colour: colour using FRAMEPRO_COLOUR macro	- the colour of the stat. Use zero for default colour. Note this will have no effect if the user has set a custom colour in the UI
 // Adds value to the custom stat with the specified name
 
-// FRAMEPRO_DYNAMIC_CUSTOM_STAT(name, value, graph, unit)
-// -----------------------------------------------
-// name: dynamic string					- the name of the custom stat
-// value: int, int64, float, double		- value to add to the custom stat
-// graph: string literal				- the name of the graph in which to display the stat
-// unit: string literal				- y-axis unit to be displayed in FramePro, eg "ms", "MB" etc.
+// FRAMEPRO_DYNAMIC_CUSTOM_STAT(name, value, graph, unit, colour)
+// --------------------------------------------------------------
+// name: dynamic string							- the name of the custom stat
+// value: int, int64, float, double				- value to add to the custom stat
+// graph: string (literal or dynamic)			- the name of the graph in which the stat will be shown
+// unit: string (literal or dynamic)			- the y-axis unit to be displayed in FramePro, eg "ms", "MB" etc for a custom stat
+// colour: colour using FRAMEPRO_COLOUR macro	- the colour of the stat. Use zero for default colour. Note this will have no effect if the user has set a custom colour in the UI
 // Adds value to the custom stat with the specified name
 
-// FRAMEPRO_SCOPE_CUSTOM_STAT(name, value)
-// ---------------------------------------
+// FRAMEPRO_SCOPE_CUSTOM_STAT(name, value, graph, unit, colour)
+// ------------------------------------------------------------
 // name: string literal or StringId - the name of the custom stat
 // value: int64 - the stat value
+// graph: string (literal or dynamic)			- the name of the graph in which the stat will be shown
+// unit: string (literal or dynamic)			- the y-axis unit to be displayed in FramePro, eg "ms", "MB" etc for a custom stat
+// colour: colour using FRAMEPRO_COLOUR macro	- the colour of the stat. Use zero for default colour. Note this will have no effect if the user has set a custom colour in the UI
 // Add a custom stat to the current scope. This will be displayed in the scope hover box and in the frame graph view
 
+// FRAMEPRO_SET_CUSTOM_STAT_GRAPH(name, graph)
+// -------------------------------------------
+// name: string (literal or dynamic)
+// graph: string (literal or dynamic)
+// Set the name of the graph in which to display the stat.
+
+// FRAMEPRO_SET_CUSTOM_STAT_UNIT(name, unit)
+// -----------------------------------------
+// name: string (literal or dynamic)
+// unit: string (literal or dynamic)
+// Set the y-axis unit to be displayed in FramePro, eg "ms", "MB" etc for a custom stat. 
+
+// FRAMEPRO_SET_CUSTOM_STAT_COLOUR(name, colour)
+// ---------------------------------------------
+// name: string (literal or dynamic)
+// colour: colour using FRAMEPRO_COLOUR macro	- the colour of the stat. Use zero for default colour. Note this will have no effect if the user has set a custom colour in the UI
+
 
 //--------------------------------------------------------------------------------------------------------
 // hires timers
@@ -758,8 +816,8 @@ namespace FramePro
 // at the start of each frame.
 
 // Usage:
-//		// this defines a global timer with name "my_timer" which will be shown in the "hires timers" graph in FramePro
-//		FRAMEPRO_DECL_GLOBAL_HIRES_TIMER(my_timer, "hires timers")
+//		// this defines a global timer with name "my_timer"
+//		FRAMEPRO_DECL_GLOBAL_HIRES_TIMER(my_timer)
 //
 //		void MyFunction()
 //		{
@@ -767,10 +825,9 @@ namespace FramePro
 //			...
 //		}
 
-// FRAMEPRO_DECL_GLOBAL_HIRES_TIMER(name, graph)
+// FRAMEPRO_DECL_GLOBAL_HIRES_TIMER(name)
 // ---------------------------------------------
 // name: unquoted string		the name of the global timer object
-// graph: string literal		the name of the graph in which to display the timer
 // Declares a high resolution timer with the specified name.
 
 // FRAMEPRO_GLOBAL_HIRES_SCOPE(name)
@@ -826,17 +883,11 @@ namespace FramePro
 #if FRAMEPRO_ENABLED
 
 //------------------------------------------------------------------------
-// enable this if you are linking to FRamePro.dll
-//#define FRAMEPRO_DLL
-
-//------------------------------------------------------------------------
-#if defined(_MSC_VER) && _MSC_VER <= 1600
-	#error FramePro only supports Visual Studio 2012 and above. This is because it needs atomics. If you really need 2010 support please contact [REDACTED_EMAIL]
-#endif
+#include <stdlib.h>
 
 //------------------------------------------------------------------------
-#include <atomic>
-#include <stdlib.h>
+// enable this if you are linking to FramePro.dll
+//#define FRAMEPRO_DLL
 
 //------------------------------------------------------------------------
 // enable debug here to get FramePro asserts
@@ -846,24 +897,6 @@ namespace FramePro
 	#define FRAMEPRO_DEBUG 0
 #endif
 
-//------------------------------------------------------------------------
-#ifdef __UNREAL__
-	#define FRAMEPRO_TOOLSET_UE4 1
-#else
-	#define FRAMEPRO_TOOLSET_UE4 0
-#endif
-
-//------------------------------------------------------------------------
-#if FRAMEPRO_TOOLSET_UE4
-	#define FRAMEPRO_API CORE_API
-#elif defined(FRAMEPRO_DLL_EXPORT)
-	#define FRAMEPRO_API __declspec(dllexport)
-#elif defined(FRAMEPRO_DLL)
-	#define FRAMEPRO_API __declspec(dllimport)
-#else
-	#define FRAMEPRO_API
-#endif
-
 //------------------------------------------------------------------------
 #define FRAMEPRO_STRINGIZE(x) FRAMEPRO_STRINGIZE2(x)
 #define FRAMEPRO_STRINGIZE2(x) #x
@@ -878,143 +911,643 @@ namespace FramePro
 #define FRAMEPRO_WIDESTR2(s) L##s
 
 //------------------------------------------------------------------------
-// START EPIC
-#define FRAMEPRO_PLATFORM_XBOXONE 	PLATFORM_XBOXONE
-#define FRAMEPRO_PLATFORM_XBOX360 	0
-#define FRAMEPRO_PLATFORM_UWP 		0
-#define FRAMEPRO_PLATFORM_WIN 		PLATFORM_WINDOWS
-#define FRAMEPRO_PLATFORM_UNIX 		PLATFORM_LINUX
-#define FRAMEPRO_PLATFORM_ANDROID 	PLATFORM_ANDROID
-#define FRAMEPRO_PLATFORM_IOS 		PLATFORM_IOS
-#define FRAMEPRO_PLATFORM_SWITCH	PLATFORM_SWITCH
-#define FRAMEPRO_PLATFORM_MAC 		PLATFORM_MAC
-// END EPIC
-
-// START EPIC
-#if FRAMEPRO_PLATFORM_WIN || FRAMEPRO_PLATFORM_UWP || FRAMEPRO_PLATFORM_XBOXONE
-	#define FRAMEPRO_WIN_BASED_PLATFORM 1
+#if defined(__UNREAL__)
+	#define FRAMEPRO_PLATFORM_UE4 1
+	#define FRAMEPRO_PLATFORM_XBOXONE 0
+	#define FRAMEPRO_PLATFORM_UWP 0
+	#define FRAMEPRO_PLATFORM_WIN 0
+	#define FRAMEPRO_PLATFORM_LINUX 0
+	#define FRAMEPRO_PLATFORM_PS4 0
+	#define FRAMEPRO_PLATFORM_ANDROID 0
+#elif defined(_XBOX_ONE) || defined(_DURANGO)
+	#define FRAMEPRO_PLATFORM_UE4 0
+	#define FRAMEPRO_PLATFORM_XBOXONE 1
+	#define FRAMEPRO_PLATFORM_UWP 0
+	#define FRAMEPRO_PLATFORM_WIN 0
+	#define FRAMEPRO_PLATFORM_LINUX 0
+	#define FRAMEPRO_PLATFORM_PS4 0
+	#define FRAMEPRO_PLATFORM_ANDROID 0
+#elif defined(WINAPI_FAMILY) && (WINAPI_FAMILY == WINAPI_FAMILY_APP)
+	#define FRAMEPRO_PLATFORM_UE4 0
+	#define FRAMEPRO_PLATFORM_XBOXONE 0
+	#define FRAMEPRO_PLATFORM_UWP 1
+	#define FRAMEPRO_PLATFORM_WIN 0
+	#define FRAMEPRO_PLATFORM_LINUX 0
+	#define FRAMEPRO_PLATFORM_PS4 0
+	#define FRAMEPRO_PLATFORM_ANDROID 0
+#elif defined(_WIN32) || defined(_WIN64) || defined(WIN32) || defined(WIN64) || defined(__WIN32__) || defined(__WINDOWS__)
+	#define FRAMEPRO_PLATFORM_UE4 0
+	#define FRAMEPRO_PLATFORM_XBOXONE 0
+	#define FRAMEPRO_PLATFORM_UWP 0
+	#define FRAMEPRO_PLATFORM_WIN 1
+	#define FRAMEPRO_PLATFORM_LINUX 0
+	#define FRAMEPRO_PLATFORM_PS4 0
+	#define FRAMEPRO_PLATFORM_ANDROID 0
+#elif defined(__ORBIS__)
+	#define FRAMEPRO_PLATFORM_UE4 0
+	#define FRAMEPRO_PLATFORM_XBOXONE 0
+	#define FRAMEPRO_PLATFORM_UWP 0
+	#define FRAMEPRO_PLATFORM_WIN 0
+	#define FRAMEPRO_PLATFORM_LINUX 0
+	#define FRAMEPRO_PLATFORM_PS4 1
+	#define FRAMEPRO_PLATFORM_ANDROID 0
+#elif defined(__ANDROID__)
+	#define FRAMEPRO_PLATFORM_UE4 0
+	#define FRAMEPRO_PLATFORM_XBOXONE 0
+	#define FRAMEPRO_PLATFORM_UWP 0
+	#define FRAMEPRO_PLATFORM_WIN 0
+	#define FRAMEPRO_PLATFORM_LINUX 0
+	#define FRAMEPRO_PLATFORM_PS4 0
+	#define FRAMEPRO_PLATFORM_ANDROID 1
+#elif defined(unix) || defined(__unix__) || defined(__unix)
+	#define FRAMEPRO_PLATFORM_UE4 0
+	#define FRAMEPRO_PLATFORM_XBOXONE 0
+	#define FRAMEPRO_PLATFORM_UWP 0
+	#define FRAMEPRO_PLATFORM_WIN 0
+	#define FRAMEPRO_PLATFORM_LINUX 1
+	#define FRAMEPRO_PLATFORM_PS4 0
+	#define FRAMEPRO_PLATFORM_ANDROID 0
 #else
-	#define FRAMEPRO_WIN_BASED_PLATFORM 0
+	#error FramePro platform not defined
 #endif
 
-#if FRAMEPRO_PLATFORM_SWITCH //#TODO: could add others here...
-	#define FRAMEPRO_UE4_BASED_PLATFORM 1
-#else
-	#define FRAMEPRO_UE4_BASED_PLATFORM 0
-#endif
-// END EPIC
+//------------------------------------------------------------------------
+#define FRAMEPRO_SOURCE_STRING			__FILE__ "|" FRAMEPRO_FUNCTION_NAME "|" FRAMEPRO_STRINGIZE(__LINE__) "|"
+#define FRAMEPRO_SOURCE_STRING_W		FRAMEPRO_WIDESTR(__FILE__) L"|" FRAMEPRO_WIDESTR(FRAMEPRO_FUNCTION_NAME) L"|" FRAMEPRO_WIDESTR(FRAMEPRO_STRINGIZE(__LINE__)) L"|"
+#define FRAMEPRO_SOURCE_STRING_IDLE		__FILE__ "|" FRAMEPRO_FUNCTION_NAME "|" FRAMEPRO_STRINGIZE(__LINE__) "|Idle"
+#define FRAMEPRO_SOURCE_STRING_IDLE_W	FRAMEPRO_WIDESTR(__FILE__) L"|" FRAMEPRO_WIDESTR(FRAMEPRO_FUNCTION_NAME) L"|" FRAMEPRO_WIDESTR(FRAMEPRO_STRINGIZE(__LINE__)) L"|Idle"
 
 //------------------------------------------------------------------------
-#if FRAMEPRO_WIN_BASED_PLATFORM && !defined(__clang__)
-	#define FRAMEPRO_FUNCTION_DEFINE_IS_STRING_LITERAL 1
+#if FRAMEPRO_DEBUG
+	#define FRAMEPRO_ASSERT(b) if(!(b)) FramePro::DebugBreak()
 #else
-	#define FRAMEPRO_FUNCTION_DEFINE_IS_STRING_LITERAL 0
+	#define FRAMEPRO_ASSERT(b) ((void)0)
 #endif
 
 //------------------------------------------------------------------------
-// On non-windows based platforms __FUNCTION__ can't be concatonated with
-// other string literals so we need to pass in the name separately.
-#if !FRAMEPRO_FUNCTION_DEFINE_IS_STRING_LITERAL
-	#ifdef FRAMEPRO_SCOPE
-		#undef FRAMEPRO_SCOPE
-		#define FRAMEPRO_SCOPE() 					FramePro::StringLiteralNamedTimerScope FRAMEPRO_UNIQUE(timer_scope)(__FUNCTION__, FRAMEPRO_SOURCE_STRING)
-	#endif
-	#ifdef FRAMEPRO_CONDITIONAL_SCOPE
-		#undef FRAMEPRO_CONDITIONAL_SCOPE
-		#define FRAMEPRO_CONDITIONAL_SCOPE()		FramePro::StringLiteralNamedConditionalTimerScope FRAMEPRO_UNIQUE(timer_scope)(__FUNCTION__, FRAMEPRO_SOURCE_STRING)
-	#endif
-	#ifdef FRAMEPRO_IDLE_SCOPE
-		#undef FRAMEPRO_IDLE_SCOPE
-		#define FRAMEPRO_IDLE_SCOPE()				FramePro::StringLiteralNamedTimerScope FRAMEPRO_UNIQUE(timer_scope)(__FUNCTION__, FRAMEPRO_SOURCE_STRING_IDLE)
-	#endif
-	#ifdef FRAMEPRO_IDLE_CONDITIONAL_SCOPE
-		#undef FRAMEPRO_IDLE_CONDITIONAL_SCOPE
-		#define FRAMEPRO_IDLE_CONDITIONAL_SCOPE()	FramePro::StringLiteralNamedConditionalTimerScope FRAMEPRO_UNIQUE(timer_scope)(__FUNCTION__, FRAMEPRO_SOURCE_STRING_IDLE)
-	#endif
-#endif
+#define FRAMEPRO_UNREFERENCED(s) (void)s
 
 //------------------------------------------------------------------------
-#if FRAMEPRO_WIN_BASED_PLATFORM
-	#define MULTI_STATEMENT(s) do { s } while(true,false)
-#else
-	#define MULTI_STATEMENT(s) do { s } while(false)
-#endif
+// typedefs
+namespace FramePro
+{
+	typedef long long int64;
+	typedef unsigned long long uint64;
+	typedef unsigned int uint;
+}
 
 //------------------------------------------------------------------------
-#if FRAMEPRO_FUNCTION_DEFINE_IS_STRING_LITERAL
-	#define FRAMEPRO_FUNCTION_NAME __FUNCTION__
-#else
-	#define FRAMEPRO_FUNCTION_NAME ""		// __FUNCTION__ is not a string literal on Unix platforms so we can't combine it with other string literals
-#endif
+#define FRAMEPRO_DEFAULT_COND_SCOPE_MIN_TIME 50
 
 //------------------------------------------------------------------------
-#define FRAMEPRO_SOURCE_STRING			__FILE__ "|" FRAMEPRO_FUNCTION_NAME "|" FRAMEPRO_STRINGIZE(__LINE__) "|"
-#define FRAMEPRO_SOURCE_STRING_W		FRAMEPRO_WIDESTR(__FILE__) L"|" FRAMEPRO_WIDESTR(FRAMEPRO_FUNCTION_NAME) L"|" FRAMEPRO_WIDESTR(FRAMEPRO_STRINGIZE(__LINE__)) L"|"
-#define FRAMEPRO_SOURCE_STRING_IDLE		__FILE__ "|" FRAMEPRO_FUNCTION_NAME "|" FRAMEPRO_STRINGIZE(__LINE__) "|Idle"
-#define FRAMEPRO_SOURCE_STRING_IDLE_W	FRAMEPRO_WIDESTR(__FILE__) L"|" FRAMEPRO_WIDESTR(FRAMEPRO_FUNCTION_NAME) L"|" FRAMEPRO_WIDESTR(FRAMEPRO_STRINGIZE(__LINE__)) L"|Idle"
+#define FRAMEPRO_MAX_INLINE_STRING_LENGTH 256
 
 //------------------------------------------------------------------------
-// START EPIC
-#define FRAMEPRO_BREAK() UE_DEBUG_BREAK()
-// END EPIC
+namespace FramePro
+{
+	typedef int(*ThreadMain)(void*);
+}
 
 //------------------------------------------------------------------------
-#if FRAMEPRO_DEBUG
-	#define FRAMEPRO_ASSERT(b) if(!(b)) FRAMEPRO_BREAK()
-#else
-	#define FRAMEPRO_ASSERT(b) ((void)0)
-#endif
+namespace FramePro
+{
+	namespace PacketType
+	{
+		enum Enum
+		{
+			Connect = 0xaabb,
+			FrameStart,
+			TimeSpan,
+			TimeSpanW,
+			NamedTimeSpan,
+			StringLiteralNamedTimeSpan,
+			ThreadName,
+			ThreadOrder,
+			StringPacket,
+			WStringPacket,
+			NameAndSourceInfoPacket,
+			NameAndSourceInfoPacketW,
+			SourceInfoPacket,
+			MainThreadPacket,
+			RequestStringLiteralPacket,
+			SetConditionalScopeMinTimePacket,
+			ConnectResponsePacket,
+			SessionInfoPacket,
+			RequestRecordedDataPacket,
+			SessionDetailsPacket,
+			ContextSwitchPacket,
+			ContextSwitchRecordingStartedPacket,
+			ProcessNamePacket,
+			CustomStatPacket_Depreciated,
+			StringLiteralTimerNamePacket,
+			HiResTimerScopePacket,
+			LogPacket,
+			EventPacket,
+			StartWaitEventPacket,
+			StopWaitEventPacket,
+			TriggerWaitEventPacket,
+			TimeSpanCustomStatPacket_Depreciated,
+			TimeSpanWithCallstack,
+			TimeSpanWWithCallstack,
+			NamedTimeSpanWithCallstack,
+			StringLiteralNamedTimeSpanWithCallstack,
+			ModulePacket,
+			SetCallstackRecordingEnabledPacket,
+			CustomStatPacketW,
+			TimeSpanCustomStatPacketW,
+			CustomStatPacket,
+			TimeSpanCustomStatPacket,
+			ScopeColourPacket,
+			CustomStatColourPacket,
+			CustomStatGraphPacket,
+			CustomStatUnitPacket,
+		};
+	};
+}
 
 //------------------------------------------------------------------------
-// START EPIC
-// Remove UNIX inttypes.h and sys/signal.h includes
+namespace FramePro
+{
+	struct ModulePacket
+	{
+		PacketType::Enum m_PacketType;
+		int m_UseLookupFunctionForBaseAddress;
+		int64 m_ModuleBase;
+		char m_Sig[16];
+		int m_Age;
+		int m_Padding;
+		char m_ModuleName[FRAMEPRO_MAX_INLINE_STRING_LENGTH];
+		char m_SymbolFilename[FRAMEPRO_MAX_INLINE_STRING_LENGTH];
+	};
+}
+
 //------------------------------------------------------------------------
-// typedefs
 namespace FramePro
 {
-	typedef uint32 uint;
+	namespace Platform
+	{
+		enum Enum
+		{
+			Windows = 0,
+			Windows_UWP,
+			XBoxOne,
+			Unused,
+			Linux,
+			PS4,
+			Android,
+			Mac,
+			iOS,
+			Switch,
+		};
+	}
 }
-// END EPIC
 
-#if FRAMEPRO_PLATFORM_WIN
-	#define FRAMEPRO_EVENT_TRACE_WIN32 1
-#else
-	#define FRAMEPRO_EVENT_TRACE_WIN32 0
+//------------------------------------------------------------------------
+// FRAMEPRO_API define
+#ifndef FRAMEPRO_API
+	#if defined(FRAMEPRO_DLL_EXPORT)
+		#define FRAMEPRO_API __declspec(dllexport)
+	#elif defined(FRAMEPRO_DLL)
+		#define FRAMEPRO_API __declspec(dllimport)
+	#else
+		#define FRAMEPRO_API
+	#endif
 #endif
 
 //------------------------------------------------------------------------
-// port
-#if FRAMEPRO_PLATFORM_WIN || FRAMEPRO_PLATFORM_UWP
+
+//------------------------------------------------------------------------
+//
+// FrameProPlatform.hpp
+//
+#if FRAMEPRO_ENABLED
+
+//------------------------------------------------------------------------
+// ---
+// --- FRAMEPRO PLATFORM IMPLEMENTATION START ---
+// ---
+//------------------------------------------------------------------------
+
+//------------------------------------------------------------------------
+#include <stdarg.h>
+#include <time.h>
+
+//------------------------------------------------------------------------
+//                         FRAMEPRO_PLATFORM_UE4
+//------------------------------------------------------------------------
+#if FRAMEPRO_PLATFORM_UE4
+
+	#include "FramePro/FrameProUE4.h" // contact [REDACTED_EMAIL] for this platform
+
+//------------------------------------------------------------------------
+//                         FRAMEPRO_PLATFORM_WIN
+//------------------------------------------------------------------------
+#elif FRAMEPRO_PLATFORM_WIN
+
+	FRAMEPRO_API __int64 FramePro_QueryPerformanceCounter();
+	#define FRAMEPRO_GET_CLOCK_COUNT(time) time = FramePro_QueryPerformanceCounter()
+
+	// Windows or Linux based platform
+	#define FRAMEPRO_WIN_BASED_PLATFORM 1
+	#define FRAMEPRO_LINUX_BASED_PLATFORM 0
+
+	#define FRAMEPRO_USE_TLS_SLOTS 0
+
+	// Port
 	#define FRAMEPRO_PORT "8428"
-#elif FRAMEPRO_PLATFORM_XBOXONE
-	#define FRAMEPRO_PORT "4420"
-// START EPIC
-#elif FRAMEPRO_PLATFORM_UNIX || FRAMEPRO_PLATFORM_ANDROID || FRAMEPRO_PLATFORM_IOS || FRAMEPRO_PLATFORM_SWITCH || FRAMEPRO_PLATFORM_MAC
-// END EPIC
+
+	// x64 or x32
+	#if defined(_WIN64) || defined(__LP64__) || defined(__x86_64__) || defined(__ppc64__)
+		#define FRAMEPRO_X64 1
+	#else
+		#define FRAMEPRO_X64 0
+	#endif
+
+	#define FRAMEPRO_MAX_PATH 260
+
+	#include <tchar.h>
+	#define FRAMEPRO_TCHAR _TCHAR
+	
+	#define MULTI_STATEMENT(s) do { s } while(true,false)
+	
+	#if !defined(__clang__)
+		#define FRAMEPRO_FUNCTION_DEFINE_IS_STRING_LITERAL 1
+	#else
+		#define FRAMEPRO_FUNCTION_DEFINE_IS_STRING_LITERAL 0
+	#endif
+
+	#define FRAMEPRO_NO_INLINE __declspec(noinline)
+	#define FRAMEPRO_FORCE_INLINE __forceinline
+
+	#define LIMIT_RECORDING_FILE_SIZE 1
+
+	#define FRAMEPRO_ALIGN_STRUCT(a)
+
+	#define FRAMEPRO_ENUMERATE_ALL_MODULES (FRAMEPRO_X64 && 1)
+
+//------------------------------------------------------------------------
+//                         FRAMEPRO_PLATFORM_UWP
+//------------------------------------------------------------------------
+#elif FRAMEPRO_PLATFORM_UWP
+
+	__int64 FramePro_QueryPerformanceCounter();
+	#define FRAMEPRO_GET_CLOCK_COUNT(time) time = FramePro_QueryPerformanceCounter()
+
+	// Windows or Linux based platform
+	#define FRAMEPRO_WIN_BASED_PLATFORM 1
+	#define FRAMEPRO_LINUX_BASED_PLATFORM 0
+
+	#define FRAMEPRO_USE_TLS_SLOTS 0
+
+	// Port
+	#define FRAMEPRO_PORT "8428"
+
+	// x64 or x32
+	#define FRAMEPRO_X64 1
+
+	#define FRAMEPRO_MAX_PATH 260
+
+	#include <tchar.h>
+	#define FRAMEPRO_TCHAR _TCHAR
+	
+	#define MULTI_STATEMENT(s) do { s } while(true,false)
+	
+	#if !defined(__clang__)
+		#define FRAMEPRO_FUNCTION_DEFINE_IS_STRING_LITERAL 1
+	#else
+		#define FRAMEPRO_FUNCTION_DEFINE_IS_STRING_LITERAL 0
+	#endif
+
+	#define FRAMEPRO_NO_INLINE __declspec(noinline)
+	#define FRAMEPRO_FORCE_INLINE __forceinline
+
+	#define LIMIT_RECORDING_FILE_SIZE 1
+
+	#define FRAMEPRO_ALIGN_STRUCT(a)
+
+	#define FRAMEPRO_ENUMERATE_ALL_MODULES (FRAMEPRO_X64 && 1)
+
+//------------------------------------------------------------------------
+//                         FRAMEPRO_PLATFORM_ANDROID
+//------------------------------------------------------------------------
+#elif FRAMEPRO_PLATFORM_ANDROID
+
+	// this needs to be as fast as possible and inline
+	#define FRAMEPRO_GET_CLOCK_COUNT(time)						\
+		timespec FRAMEPRO_UNIQUE(ts);							\
+		clock_gettime(CLOCK_MONOTONIC, &FRAMEPRO_UNIQUE(ts));	\
+		time = FRAMEPRO_UNIQUE(ts).tv_sec * 1000000000LL + FRAMEPRO_UNIQUE(ts).tv_nsec
+
+	// Windows or Linux based platform
+	#define FRAMEPRO_WIN_BASED_PLATFORM 0
+	#define FRAMEPRO_LINUX_BASED_PLATFORM 1
+
+	#define FRAMEPRO_USE_TLS_SLOTS 0
+
+	// Port
+	#define FRAMEPRO_PORT "8428"
+
+	// x64 or x32
+	#if defined(__LP64__) || defined(__x86_64__) || defined(__ppc64__)
+		#define FRAMEPRO_X64 1
+	#else
+		#define FRAMEPRO_X64 0
+	#endif
+
+	#define FRAMEPRO_MAX_PATH 260
+
+	#define MULTI_STATEMENT(s) do { s } while(false)
+	
+	namespace FramePro { typedef char FRAMEPRO_TCHAR; }
+
+	#define FRAMEPRO_FUNCTION_DEFINE_IS_STRING_LITERAL 0
+
+	#define FRAMEPRO_NO_INLINE
+	#define FRAMEPRO_FORCE_INLINE inline
+
+	#define LIMIT_RECORDING_FILE_SIZE 1
+
+	#define FRAMEPRO_ALIGN_STRUCT(a) __attribute__ ((aligned(a)))
+
+	#define FRAMEPRO_ENUMERATE_ALL_MODULES (FRAMEPRO_X64 && 1)
+
+//------------------------------------------------------------------------
+//                         FRAMEPRO_PLATFORM_LINUX
+//------------------------------------------------------------------------
+#elif FRAMEPRO_PLATFORM_LINUX
+
+	// this needs to be as fast as possible and inline
+	#define FRAMEPRO_GET_CLOCK_COUNT(time)						\
+		timespec FRAMEPRO_UNIQUE(ts);							\
+		clock_gettime(CLOCK_MONOTONIC, &FRAMEPRO_UNIQUE(ts));	\
+		time = FRAMEPRO_UNIQUE(ts).tv_sec * 1000000000LL + FRAMEPRO_UNIQUE(ts).tv_nsec
+
+	// Windows or Linux based platform
+	#define FRAMEPRO_WIN_BASED_PLATFORM 0
+	#define FRAMEPRO_LINUX_BASED_PLATFORM 1
+
+	#define FRAMEPRO_USE_TLS_SLOTS 0
+
+	// Port
 	#define FRAMEPRO_PORT "8428"
+
+	// x64 or x32
+	#if defined(__LP64__) || defined(__x86_64__) || defined(__ppc64__)
+		#define FRAMEPRO_X64 1
+	#else
+		#define FRAMEPRO_X64 0
+	#endif
+
+	#define FRAMEPRO_MAX_PATH 260
+
+	#define MULTI_STATEMENT(s) do { s } while(false)
+	
+	namespace FramePro { typedef char FRAMEPRO_TCHAR; }
+
+	#define FRAMEPRO_FUNCTION_DEFINE_IS_STRING_LITERAL 0
+
+	#define FRAMEPRO_NO_INLINE
+	#define FRAMEPRO_FORCE_INLINE inline
+
+	#define LIMIT_RECORDING_FILE_SIZE 1
+
+	#define FRAMEPRO_ALIGN_STRUCT(a) __attribute__ ((aligned(a)))
+
+	#define FRAMEPRO_ENUMERATE_ALL_MODULES (FRAMEPRO_X64 && 1)
+
+//------------------------------------------------------------------------
+//                         FRAMEPRO_PLATFORM_XBOXONE
+//------------------------------------------------------------------------
+#elif FRAMEPRO_PLATFORM_XBOXONE
+
+	#include "FrameProXboxOne.h" // contact [REDACTED_EMAIL] for this platform
+
+//------------------------------------------------------------------------
+//                         FRAMEPRO_PLATFORM_PS4
+//------------------------------------------------------------------------
+#elif FRAMEPRO_PLATFORM_PS4
+
+	#include "FrameProPS4.h" // contact [REDACTED_EMAIL] for this platform
+
 #else
-	#error FramePro platform not defined
+
+	#error Platform not defined
+
 #endif
 
 //------------------------------------------------------------------------
-#define FRAMEPRO_DEFAULT_COND_SCOPE_MIN_TIME 50
+// ---
+// --- FRAMEPRO PLATFORM IMPLEMENTATION END ---
+// ---
+//------------------------------------------------------------------------
+
+//------------------------------------------------------------------------
+// check all necessary defines are set
+
+#ifndef FRAMEPRO_GET_CLOCK_COUNT
+	#error Platform must defined FRAMEPRO_GET_CLOCK_COUNT
+#endif
+
+#ifndef FRAMEPRO_PORT
+	#error Platform must define FRAMEPRO_PORT
+#endif
+
+#ifndef FRAMEPRO_X64
+	#error Platform must define FRAMEPRO_X64
+#endif
+
+#ifndef MULTI_STATEMENT
+	#error Platform must define MULTI_STATEMENT
+#endif
+
+#ifndef FRAMEPRO_FUNCTION_DEFINE_IS_STRING_LITERAL
+	#error Platform must define FRAMEPRO_FUNCTION_DEFINE_IS_STRING_LITERAL
+#endif
+
+#ifndef FRAMEPRO_NO_INLINE
+	#error Platform must define FRAMEPRO_NO_INLINE
+#endif
+
+#ifndef FRAMEPRO_FORCE_INLINE
+	#error Platform must define FRAMEPRO_FORCE_INLINE
+#endif
+
+#ifndef FRAMEPRO_USE_TLS_SLOTS
+	#error Platform must define FRAMEPRO_USE_TLS_SLOTS
+#endif
+
+#ifndef LIMIT_RECORDING_FILE_SIZE
+	#error Platform must define LIMIT_RECORDING_FILE_SIZE
+#endif
 
 //------------------------------------------------------------------------
-// START EPIC
-// Remove FRAMEPRO_TIMER_CLOCK_GETTIME
-// END EPIC
-// the macro to get the current clock time
 namespace FramePro
 {
-// START EPIC
-#define FRAMEPRO_GET_CLOCK_COUNT(time) \
-	time = FPlatformTime::Cycles64()
-// END EPIC
+	//------------------------------------------------------------------------
+	class Allocator;
+	struct ContextSwitch;
+	class DynamicString;
+	struct ModulePacket;
+	template<typename T> class Array;
+
+	//------------------------------------------------------------------------
+	namespace Platform
+	{
+		typedef void (*ContextSwitchCallbackFunction)(const ContextSwitch&, void*);
+
+		enum Enum;
+
+		int64 GetTimerFrequency();
+
+		void DebugBreak();
+
+		int GetCore();
+
+		bool GetProcessName(int process_id, char* p_name, int max_name_length);
+
+		Platform::Enum GetPlatformEnum();
+
+		void* CreateContextSwitchRecorder(Allocator* p_allocator);
+
+		void DestroyContextSwitchRecorder(void* p_context_switch_recorder, Allocator* p_allocator);
+
+		bool StartRecordingContextSitches(void* p_context_switch_recorder, ContextSwitchCallbackFunction p_callback, void* p_context, DynamicString& error);
+		
+		void StopRecordingContextSitches(void* p_context_switch_recorder);
+
+		void FlushContextSwitches(void* p_context_switch_recorder);
+
+		void EnumerateModules(Array<ModulePacket*>& module_packets, Allocator* p_allocator);
+
+		bool GetStackTrace(void** stack, int& stack_size, unsigned int& hash);
+
+		int GetCurrentThreadId();
+
+		void DebugWrite(const char* p_string);
+
+		void GetLocalTime(tm* p_tm, const time_t *p_time);
+
+		int GetCurrentProcessId();
+
+		void VSPrintf(char* p_buffer, size_t const buffer_size, const char* p_format, va_list arg_list);
+
+		void ToString(int value, char* p_dest, int dest_size);
+
+		void CreateLock(void* p_os_lock_mem, int os_lock_mem_size);
+
+		void DestroyLock(void* p_os_lock_mem);
+
+		void TakeLock(void* p_os_lock_mem);
+
+		void ReleaseLock(void* p_os_lock_mem);
+
+		void CreateEventX(void* p_os_event_mem, int os_event_mem_size, bool initial_state, bool auto_reset);
+
+		void DestroyEvent(void* p_os_event_mem);
+
+		void SetEvent(void* p_os_event_mem);
+
+		void ResetEvent(void* p_os_event_mem);
+
+		int WaitEvent(void* p_os_event_mem, int timeout);
+
+		bool InitialiseSocketSystem();
+
+		void UninitialiseSocketSystem();
+
+		void CreateSocket(void* p_os_socket_mem, int os_socket_mem_size);
+
+		void DestroySocket(void* p_os_socket_mem);
+
+		void DisconnectSocket(void* p_os_socket_mem, bool stop_listening);
+
+		bool StartSocketListening(void* p_os_socket_mem);
+
+		bool BindSocket(void* p_os_socket_mem, const char* p_port);
+
+		bool AcceptSocket(void* p_source_os_socket_mem, void* p_target_os_socket_mem);
+
+		bool SocketSend(void* p_os_socket_mem, const void* p_buffer, int size, int& bytes_sent);
+
+		bool SocketReceive(void* p_os_socket_mem, const void* p_buffer, int size, int& bytes_received);
+
+		bool IsSocketValid(const void* p_os_socket_mem);
+
+		void HandleSocketError();
+
+		void CreateThread(
+			void* p_os_thread_mem,
+			int os_thread_mem_size,
+			ThreadMain p_thread_main,
+			void* p_context,
+			Allocator* p_allocator);
+
+		void SetThreadPriority(void* p_os_thread_mem, int priority);
+
+		void SetThreadAffinity(void* p_os_thread_mem, int affinity);
+
+		bool OpenFileForRead(void* p_os_file_mem, int os_file_mem_size, const char* p_filename);
+
+		bool OpenFileForRead(void* p_os_file_mem, int os_file_mem_size, const wchar_t* p_filename);
+
+		bool OpenFileForWrite(void* p_os_file_mem, int os_file_mem_size, const char* p_filename);
+
+		bool OpenFileForWrite(void* p_os_file_mem, int os_file_mem_size, const wchar_t* p_filename);
+		
+		void CloseFile(void* p_os_file_mem);
+
+		void ReadFromFile(void* p_os_file_mem, void* p_data, size_t size);
+
+		void WriteToFile(void* p_os_file_mem, const void* p_data, size_t size);
+
+		int GetFileSize(const void* p_os_file_mem);
+
+		uint AllocateTLSSlot();
+
+		void* GetTLSValue(uint slot);
+
+		void SetTLSValue(uint slot, void* p_value);
+
+		void GetRecordingFolder(char* p_path, int max_path_length);
+	}
 }
 
+//------------------------------------------------------------------------
+#endif		// #if FRAMEPRO_ENABLED
+#include <atomic>
+
+//------------------------------------------------------------------------
+#if FRAMEPRO_FUNCTION_DEFINE_IS_STRING_LITERAL
+	#define FRAMEPRO_FUNCTION_NAME __FUNCTION__
+#else
+	#define FRAMEPRO_FUNCTION_NAME ""		// __FUNCTION__ is not a string literal on Linux platforms so we can't combine it with other string literals
+
+	// On non-windows based platforms __FUNCTION__ can't be concatonated with
+	// other string literals so we need to pass in the name separately.
+	#ifdef FRAMEPRO_SCOPE
+		#undef FRAMEPRO_SCOPE
+		#define FRAMEPRO_SCOPE() 					FramePro::StringLiteralNamedTimerScope FRAMEPRO_UNIQUE(timer_scope)(__FUNCTION__, FRAMEPRO_SOURCE_STRING)
+	#endif
+	#ifdef FRAMEPRO_CONDITIONAL_SCOPE
+		#undef FRAMEPRO_CONDITIONAL_SCOPE
+		#define FRAMEPRO_CONDITIONAL_SCOPE()		FramePro::StringLiteralNamedConditionalTimerScope FRAMEPRO_UNIQUE(timer_scope)(__FUNCTION__, FRAMEPRO_SOURCE_STRING)
+	#endif
+	#ifdef FRAMEPRO_IDLE_SCOPE
+		#undef FRAMEPRO_IDLE_SCOPE
+		#define FRAMEPRO_IDLE_SCOPE()				FramePro::StringLiteralNamedTimerScope FRAMEPRO_UNIQUE(timer_scope)(__FUNCTION__, FRAMEPRO_SOURCE_STRING_IDLE)
+	#endif
+	#ifdef FRAMEPRO_IDLE_CONDITIONAL_SCOPE
+		#undef FRAMEPRO_IDLE_CONDITIONAL_SCOPE
+		#define FRAMEPRO_IDLE_CONDITIONAL_SCOPE()	FramePro::StringLiteralNamedConditionalTimerScope FRAMEPRO_UNIQUE(timer_scope)(__FUNCTION__, FRAMEPRO_SOURCE_STRING_IDLE)
+	#endif
+#endif
+
 //------------------------------------------------------------------------
 namespace FramePro
 {
+	FRAMEPRO_API void DebugBreak();
+
+	// assume aligned reads/writes < 8 bytes are atomic on all platforms
 	template<typename T>
 	class RelaxedAtomic
 	{
@@ -1033,32 +1566,29 @@ namespace FramePro
 namespace FramePro
 {
 	// callbacks
-	// START EPIC
-	typedef void(*ConnectionChangedCallback)(bool connected, const FString& ProfileName, void* user_data);
-	// END EPIC
+	typedef void(*ConnectionChangedCallback)(bool connected, const wchar_t* p_recording_filename, void* user_data);
 	typedef bool(*ConditionalParentScopeCallback)(const char* p_name, int64 start_time, int64 end_time, int64 ticks_per_second);
 
 	FRAMEPRO_API extern RelaxedAtomic<bool> g_Connected;
+
 	FRAMEPRO_API extern RelaxedAtomic<unsigned int> g_ConditionalScopeMinTime;
 
-// START EPIC
-// Remove FRAMEPRO_TIMER_QUERY_PERFORMANCE_COUNTER and FRAMEPRO_WIN_BASED_PLATFORM
-// END EPIC
+	FRAMEPRO_API inline bool IsConnected() { return g_Connected; }
+
+	FRAMEPRO_API inline unsigned int GetConditionalScopeMinTime() { return g_ConditionalScopeMinTime; }
 
 	FRAMEPRO_API void Shutdown();
 
 	FRAMEPRO_API void SetPort(int port);
 
+	FRAMEPRO_API void DebugBreak();
+
 	FRAMEPRO_API void SendSessionInfo(const char* p_name, const char* p_build_id);
 
 	FRAMEPRO_API void SendSessionInfo(const wchar_t* p_name, const wchar_t* p_build_id);
 
 	FRAMEPRO_API void SetAllocator(class Allocator* p_allocator);		// if you call this you must call it BEFORE any other calls
 
-	FRAMEPRO_API inline bool IsConnected() { return g_Connected; }
-
-	FRAMEPRO_API inline unsigned int GetConditionalScopeMinTime() { return g_ConditionalScopeMinTime; }
-
 	FRAMEPRO_API void FrameStart();		// call at the start of each of your frames
 
 	FRAMEPRO_API void AddTimeSpan(const char* p_name_file_and_line, int64 start_time, int64 end_time);
@@ -1071,39 +1601,45 @@ namespace FramePro
 
 	FRAMEPRO_API void AddTimeSpan(StringId name_id, const char* p_file_and_line, int64 start_time, int64 end_time, int thread_id, int core);
 
-	FRAMEPRO_API void AddCustomStat(const char* p_name, int value, const char* p_graph, const char* p_unit);
+	FRAMEPRO_API void AddCustomStat(const char* p_name, int value, const char* p_graph, const char* p_unit, uint colour);
+
+	FRAMEPRO_API void AddCustomStat(const char* p_name, int64 value, const char* p_graph, const char* p_unit, uint colour);
+
+	FRAMEPRO_API void AddCustomStat(const char* p_name, float value, const char* p_graph, const char* p_unit, uint colour);
+
+	FRAMEPRO_API void AddCustomStat(const char* p_name, double value, const char* p_graph, const char* p_unit, uint colour);
 
-	FRAMEPRO_API void AddCustomStat(const char* p_name, int64 value, const char* p_graph, const char* p_unit);
+	FRAMEPRO_API void AddCustomStat(const wchar_t* p_name, int value, const wchar_t* p_graph, const wchar_t* p_unit, uint colour);
 
-	FRAMEPRO_API void AddCustomStat(const char* p_name, float value, const char* p_graph, const char* p_unit);
+	FRAMEPRO_API void AddCustomStat(const wchar_t* p_name, int64 value, const wchar_t* p_graph, const wchar_t* p_unit, uint colour);
 
-	FRAMEPRO_API void AddCustomStat(const char* p_name, double value, const char* p_graph, const char* p_unit);
+	FRAMEPRO_API void AddCustomStat(const wchar_t* p_name, float value, const wchar_t* p_graph, const wchar_t* p_unit, uint colour);
 
-	FRAMEPRO_API void AddCustomStat(const wchar_t* p_name, int value, const wchar_t* p_graph, const wchar_t* p_unit);
+	FRAMEPRO_API void AddCustomStat(const wchar_t* p_name, double value, const wchar_t* p_graph, const wchar_t* p_unit, uint colour);
 
-	FRAMEPRO_API void AddCustomStat(const wchar_t* p_name, int64 value, const wchar_t* p_graph, const wchar_t* p_unit);
+	FRAMEPRO_API void AddCustomStat(StringId name, int value, const char* p_graph, const char* p_unit, uint colour);
 
-	FRAMEPRO_API void AddCustomStat(const wchar_t* p_name, float value, const wchar_t* p_graph, const wchar_t* p_unit);
+	FRAMEPRO_API void AddCustomStat(StringId name, int64 value, const char* p_graph, const char* p_unit, uint colour);
 
-	FRAMEPRO_API void AddCustomStat(const wchar_t* p_name, double value, const wchar_t* p_graph, const wchar_t* p_unit);
+	FRAMEPRO_API void AddCustomStat(StringId name, float value, const char* p_graph, const char* p_unit, uint colour);
 
-	FRAMEPRO_API void AddCustomStat(StringId name, int value, const char* p_graph, const char* p_unit);
+	FRAMEPRO_API void AddCustomStat(StringId name, double value, const char* p_graph, const char* p_unit, uint colour);
 
-	FRAMEPRO_API void AddCustomStat(StringId name, int64 value, const char* p_graph, const char* p_unit);
+	FRAMEPRO_API void AddCustomStat(StringId name, int value, const wchar_t* p_graph, const wchar_t* p_unit, uint colour);
 
-	FRAMEPRO_API void AddCustomStat(StringId name, float value, const char* p_graph, const char* p_unit);
+	FRAMEPRO_API void AddCustomStat(StringId name, int64 value, const wchar_t* p_graph, const wchar_t* p_unit, uint colour);
 
-	FRAMEPRO_API void AddCustomStat(StringId name, double value, const char* p_graph, const char* p_unit);
+	FRAMEPRO_API void AddCustomStat(StringId name, float value, const wchar_t* p_graph, const wchar_t* p_unit, uint colour);
 
-	// START EPIC 
-	FRAMEPRO_API void AddCustomStat(StringId name, int value, const wchar_t* p_graph, const wchar_t* p_unit);
+	FRAMEPRO_API void AddCustomStat(StringId name, double value, const wchar_t* p_graph, const wchar_t* p_unit, uint colour);
 
-	FRAMEPRO_API void AddCustomStat(StringId name, int64 value, const wchar_t* p_graph, const wchar_t* p_unit);
+	FRAMEPRO_API void AddCustomStat(StringId name, int value, StringId graph, StringId unit, uint colour);
 
-	FRAMEPRO_API void AddCustomStat(StringId name, float value, const wchar_t* p_graph, const wchar_t* p_unit);
+	FRAMEPRO_API void AddCustomStat(StringId name, int64 value, StringId graph, StringId unit, uint colour);
 
-	FRAMEPRO_API void AddCustomStat(StringId name, double value, const wchar_t* p_graph, const wchar_t* p_unit);
-	// END EPIC
+	FRAMEPRO_API void AddCustomStat(StringId name, float value, StringId graph, StringId unit, uint colour);
+
+	FRAMEPRO_API void AddCustomStat(StringId name, double value, StringId graph, StringId unit, uint colour);
 
 	FRAMEPRO_API void AddEvent(const char* p_name, uint colour);
 
@@ -1123,9 +1659,9 @@ namespace FramePro
 
 	FRAMEPRO_API void UnregisterConnectionChangedcallback(ConnectionChangedCallback p_callback);
 
-// START EPIC
-	FRAMEPRO_API void StartRecording(const FString& p_filename, bool context_switches, int64 max_file_size);
-// END EPIC
+	FRAMEPRO_API void StartRecording(const char* p_filename, bool context_switches, bool callstacks, int64 max_file_size);
+
+	FRAMEPRO_API void StartRecording(const wchar_t* p_filename, bool context_switches, bool callstacks, int64 max_file_size);
 
 	FRAMEPRO_API void StopRecording();
 
@@ -1155,31 +1691,39 @@ namespace FramePro
 
 	FRAMEPRO_API void Log(const char* p_message);
 
-	FRAMEPRO_API void SetScopeCustomStat(const char* p_name, int value, const char* p_graph, const char* p_unit);
+	FRAMEPRO_API void SetScopeCustomStat(const char* p_name, int value, const char* p_graph, const char* p_unit, uint colour);
 	
-	FRAMEPRO_API void SetScopeCustomStat(const wchar_t* p_name, int value, const wchar_t* p_graph, const wchar_t* p_unit);
+	FRAMEPRO_API void SetScopeCustomStat(const wchar_t* p_name, int value, const wchar_t* p_graph, const wchar_t* p_unit, uint colour);
 
-	FRAMEPRO_API void SetScopeCustomStat(StringId name, int value, const char* p_graph, const char* p_unit);
+	FRAMEPRO_API void SetScopeCustomStat(StringId name, int value, StringId graph, StringId unit, uint colour);
 
-	FRAMEPRO_API void SetScopeCustomStat(const char* p_name, int64 value, const char* p_graph, const char* p_unit);
+	FRAMEPRO_API void SetScopeCustomStat(const char* p_name, int64 value, const char* p_graph, const char* p_unit, uint colour);
 
-	FRAMEPRO_API void SetScopeCustomStat(const wchar_t* p_name, int64 value, const wchar_t* p_graph, const wchar_t* p_unit);
+	FRAMEPRO_API void SetScopeCustomStat(const wchar_t* p_name, int64 value, const wchar_t* p_graph, const wchar_t* p_unit, uint colour);
 
-	FRAMEPRO_API void SetScopeCustomStat(StringId name, int64 value, const char* p_graph, const char* p_unit);
+	FRAMEPRO_API void SetScopeCustomStat(StringId name, int64 value, StringId graph, StringId unit, uint colour);
 
-	FRAMEPRO_API void SetScopeCustomStat(const char* p_name, float value, const char* p_graph, const char* p_unit);
+	FRAMEPRO_API void SetScopeCustomStat(const char* p_name, float value, const char* p_graph, const char* p_unit, uint colour);
 
-	FRAMEPRO_API void SetScopeCustomStat(const wchar_t* p_name, float value, const wchar_t* p_graph, const wchar_t* p_unit);
+	FRAMEPRO_API void SetScopeCustomStat(const wchar_t* p_name, float value, const wchar_t* p_graph, const wchar_t* p_unit, uint colour);
 
-	FRAMEPRO_API void SetScopeCustomStat(StringId name, float value, const char* p_graph, const char* p_unit);
+	FRAMEPRO_API void SetScopeCustomStat(StringId name, float value, StringId graph, StringId unit, uint colour);
 
-	FRAMEPRO_API void SetScopeCustomStat(const char* p_name, double value, const char* p_graph, const char* p_unit);
+	FRAMEPRO_API void SetScopeCustomStat(const char* p_name, double value, const char* p_graph, const char* p_unit, uint colour);
 
-	FRAMEPRO_API void SetScopeCustomStat(const wchar_t* p_name, double value, const wchar_t* p_graph, const wchar_t* p_unit);
+	FRAMEPRO_API void SetScopeCustomStat(const wchar_t* p_name, double value, const wchar_t* p_graph, const wchar_t* p_unit, uint colour);
 
-	FRAMEPRO_API void SetScopeCustomStat(StringId name, double value, const char* p_graph, const char* p_unit);
+	FRAMEPRO_API void SetScopeCustomStat(StringId name, double value, StringId graph, StringId unit, uint colour);
 
 	FRAMEPRO_API void SetConditionalScopeMinTimeInMicroseconds(int64 value);
+
+	FRAMEPRO_API void SetScopeColour(StringId name, uint colour);
+
+	FRAMEPRO_API void SetCustomStatGraph(StringId name, StringId graph);
+
+	FRAMEPRO_API void SetCustomStatUnit(StringId name, StringId unit);
+
+	FRAMEPRO_API void SetCustomStatColour(StringId name, uint colour);
 }
 
 //------------------------------------------------------------------------
@@ -1195,6 +1739,41 @@ namespace FramePro
 		virtual ~Allocator() {}
 	};
 
+	//------------------------------------------------------------------------
+	template<class T>
+	inline T* New(Allocator* p_allocator)
+	{
+		T* p = (T*)p_allocator->Alloc(sizeof(T));
+		new (p)T();
+		return p;
+	}
+
+	//------------------------------------------------------------------------
+	template<class T, typename Targ1>
+	inline T* New(Allocator* p_allocator, Targ1 arg1)
+	{
+		T* p = (T*)p_allocator->Alloc(sizeof(T));
+		new (p)T(arg1);
+		return p;
+	}
+
+	//------------------------------------------------------------------------
+	template<class T, typename Targ1, typename Targ2, typename TArg3>
+	inline T* New(Allocator* p_allocator, Targ1 arg1, Targ2 arg2, TArg3 arg3)
+	{
+		T* p = (T*)p_allocator->Alloc(sizeof(T));
+		new (p)T(arg1, arg2, arg3);
+		return p;
+	}
+
+	//------------------------------------------------------------------------
+	template<typename T>
+	inline void Delete(Allocator* p_allocator, T* p)
+	{
+		p->~T();
+		p_allocator->Free(p);
+	}
+
 	//------------------------------------------------------------------------
 	class TimerScope
 	{
@@ -1634,7 +2213,7 @@ namespace FramePro
 	class GlobalHiResTimer
 	{
 	public:
-		GlobalHiResTimer(const char* p_name, const char* p_graph) : m_Value(0), mp_Next(NULL), mp_Name(p_name), mp_Graph(p_graph)
+		GlobalHiResTimer(const char* p_name) : m_Value(0), mp_Next(NULL), mp_Name(p_name)
 		{
 			FramePro::AddGlobalHiResTimer(this);
 		}
@@ -1673,15 +2252,12 @@ namespace FramePro
 
 		const char* GetName() const { return mp_Name; }
 
-		const char* GetGraph() const { return mp_Graph; }
-
 	private:
 		std::atomic<uint64> m_Value;
 
 		GlobalHiResTimer* mp_Next;
 
 		const char* mp_Name;
-		const char* mp_Graph;
 	};
 
 	//------------------------------------------------------------------------
@@ -1774,8 +2350,209 @@ namespace FramePro
 		int64 m_StartTime;
 	};
 
+	//------------------------------------------------------------------------
+	inline unsigned int GetHashAndStackSize(void** p_stack, int& stack_size)
+	{
+		#if FRAMEPRO_X64
+			const unsigned int prime = 0x01000193;
+			unsigned int hash = prime;
+			stack_size = 0;
+			void** p = p_stack;
+			while(*p)
+			{
+				uint64 key = (uint64)(*p++);
+				key = (~key) + (key << 18);
+				key = key ^ (key >> 31);
+				key = key * 21;
+				key = key ^ (key >> 11);
+				key = key + (key << 6);
+				key = key ^ (key >> 22);
+				hash = hash ^ (unsigned int)key;
+				++stack_size;
+			}
+
+			return hash;
+		#else
+			const unsigned int prime = 0x01000193;
+			unsigned int hash = prime;
+			stack_size = 0;
+			while(p_stack[stack_size])
+			{
+				hash = (hash * prime) ^ (unsigned int)p_stack[stack_size];
+				++stack_size;
+			}
+
+			return hash;
+		#endif
+	}
 }
 
+//------------------------------------------------------------------------
+
+//------------------------------------------------------------------------
+//
+// Array.hpp
+//
+#include <memory.h>
+
+//------------------------------------------------------------------------
+#if FRAMEPRO_ENABLED
+
+//------------------------------------------------------------------------
+namespace FramePro
+{
+	//--------------------------------------------------------------------
+	template<typename T>
+	class Array
+	{
+	public:
+		//--------------------------------------------------------------------
+		Array()
+		:	mp_Array(NULL),
+			m_Count(0),
+			m_Capacity(0),
+			mp_Allocator(NULL)
+		{
+		}
+
+		//--------------------------------------------------------------------
+		~Array()
+		{
+			FRAMEPRO_ASSERT(!mp_Array);
+		}
+
+		//--------------------------------------------------------------------
+		int GetCount() const
+		{
+			return m_Count;
+		}
+
+		//--------------------------------------------------------------------
+		void Clear()
+		{
+			if(mp_Array)
+			{
+				mp_Allocator->Free(mp_Array);
+				mp_Array = NULL;
+			}
+
+			m_Count = 0;
+			m_Capacity = 0;
+		}
+
+		//--------------------------------------------------------------------
+		void ClearNoFree()
+		{
+			m_Count = 0;
+		}
+
+		//--------------------------------------------------------------------
+		void SetAllocator(Allocator* p_allocator)
+		{
+			FRAMEPRO_ASSERT(mp_Allocator == p_allocator || !(mp_Allocator != NULL && p_allocator != NULL));
+			mp_Allocator = p_allocator;
+		}
+
+		//--------------------------------------------------------------------
+		void Add(const T& value)
+		{
+			if(m_Count == m_Capacity)
+				Grow();
+
+			mp_Array[m_Count] = value;
+			++m_Count;
+		}
+
+		//--------------------------------------------------------------------
+		const T& operator[](int index) const
+		{
+			FRAMEPRO_ASSERT(index >= 0 && index < m_Count);
+			return mp_Array[index];
+		}
+
+		//--------------------------------------------------------------------
+		T& operator[](int index)
+		{
+			FRAMEPRO_ASSERT(index >= 0 && index < m_Count);
+			return mp_Array[index];
+		}
+
+		//--------------------------------------------------------------------
+		void RemoveAt(int index)
+		{
+			FRAMEPRO_ASSERT(index >= 0 && index < m_Count);
+		
+			if(index < m_Count - 1)
+				memmove(mp_Array + index, mp_Array + index + 1, (m_Count - 1 - index) * sizeof(T));
+
+			--m_Count;
+		}
+
+		//--------------------------------------------------------------------
+		T RemoveLast()
+		{
+			FRAMEPRO_ASSERT(m_Count);
+
+			return mp_Array[--m_Count];
+		}
+
+		//--------------------------------------------------------------------
+		bool Contains(const T& value) const
+		{
+			for(int i=0; i<m_Count; ++i)
+				if(mp_Array[i] == value)
+					return true;
+			return false;
+		}
+
+		//--------------------------------------------------------------------
+		void Resize(int count)
+		{
+			int new_capacity = m_Count + count;
+
+			if (new_capacity > m_Capacity)
+			{
+				m_Capacity = m_Capacity ? 2 * m_Capacity : 32;
+				while (m_Capacity < new_capacity)
+					m_Capacity *= 2;
+
+				T* p_new_array = (T*)mp_Allocator->Alloc(sizeof(T)*m_Capacity);
+				if (mp_Array)
+					memcpy(p_new_array, mp_Array, sizeof(T)*m_Count);
+				mp_Allocator->Free(mp_Array);
+				mp_Array = p_new_array;
+			}
+
+			m_Count = count;
+		}
+
+	private:
+		//--------------------------------------------------------------------
+		void Grow()
+		{
+			m_Capacity = m_Capacity ? 2*m_Capacity : 32;
+			T* p_new_array = (T*)mp_Allocator->Alloc(sizeof(T)*m_Capacity);
+			if(mp_Array)
+				memcpy(p_new_array, mp_Array, sizeof(T)*m_Count);
+			mp_Allocator->Free(mp_Array);
+			mp_Array = p_new_array;
+		}
+
+		//------------------------------------------------------------------------
+		// data
+	private:
+		T* mp_Array;
+		
+		int m_Count;
+		int m_Capacity;
+
+		Allocator* mp_Allocator;
+	};
+}
+
+//------------------------------------------------------------------------
+#endif		// #if FRAMEPRO_ENABLED
+
 //------------------------------------------------------------------------
 #endif		// #if FRAMEPRO_ENABLED
 
