Add interface to restarting application.

#test confirmed restarting works properly with the same commandline and ViewModel::Activated behaves the same

diff --git a/Engine/Source/Runtime/Core/Private/GenericPlatform/GenericPlatformMisc.cpp b/Engine/Source/Runtime/Core/Private/GenericPlatform/GenericPlatformMisc.cpp
--- a/Engine/Source/Runtime/Core/Private/GenericPlatform/GenericPlatformMisc.cpp
+++ b/Engine/Source/Runtime/Core/Private/GenericPlatform/GenericPlatformMisc.cpp
@@ -573,6 +573,12 @@ void FGenericPlatformMisc::RequestExit( bool Force )
 	}
 }
 
+bool FGenericPlatformMisc::RestartApplication()
+{
+	UE_LOG(LogInit, Display, TEXT("Restart application is not supported or implemented in current platform"));
+	return false;
+}
+
 void FGenericPlatformMisc::RequestExitWithStatus(bool Force, uint8 ReturnCode)
 {
 	// Generic implementation will ignore the return code - this may be important, so warn.
diff --git a/Engine/Source/Runtime/Core/Public/GenericPlatform/GenericPlatformMisc.h b/Engine/Source/Runtime/Core/Public/GenericPlatform/GenericPlatformMisc.h
--- a/Engine/Source/Runtime/Core/Public/GenericPlatform/GenericPlatformMisc.h
+++ b/Engine/Source/Runtime/Core/Public/GenericPlatform/GenericPlatformMisc.h
@@ -652,6 +652,12 @@ public:
 	 */
 	static void RequestExitWithStatus( bool Force, uint8 ReturnCode );
 
+
+	/**
+	 * Requests application to restart
+	 */
+	static bool RestartApplication();
+
 	/**
 	 * Returns the last system error code in string form.  NOTE: Only one return value is valid at a time!
 	 *
