- Fixed plugin ini files not being read in properly due to changes in .ini system that improperly required a Base.ini in plugin configs

diff --git a/Engine/Source/Runtime/Core/Private/Misc/ConfigCacheIni.cpp b/Engine/Source/Runtime/Core/Private/Misc/ConfigCacheIni.cpp
--- a/Engine/Source/Runtime/Core/Private/Misc/ConfigCacheIni.cpp
+++ b/Engine/Source/Runtime/Core/Private/Misc/ConfigCacheIni.cpp
@@ -3286,7 +3286,7 @@ static void GetSourceIniHierarchyFilenames(const TCHAR* InBaseIniName, const TCH
 		LayerPath = LayerPath.Replace(TEXT("{USER}"), FPlatformProcess::UserDir(), ESearchCase::CaseSensitive);
 
 		// PROGRAMs don't require any ini files
-		const bool bIsRequired = (!IS_PROGRAM) && (Layer.Flag & Flag_Required);
+		const bool bIsRequired = (!IS_PROGRAM) && (Layer.Flag & Flag_Required) && (EngineConfigDir == FPaths::EngineConfigDir());
 
 		// expand if it it has {ED} or {EF} expansion tags
 		if (FCString::Strstr(*Layer.Path, TEXT("{ED}")) || FCString::Strstr(*Layer.Path, TEXT("{EF}")))
