Remove unused function for enumerating hot reload DLLs.

#rnx

diff --git a/Engine/Source/Runtime/Core/Private/Modules/ModuleManager.cpp b/Engine/Source/Runtime/Core/Private/Modules/ModuleManager.cpp
--- a/Engine/Source/Runtime/Core/Private/Modules/ModuleManager.cpp
+++ b/Engine/Source/Runtime/Core/Private/Modules/ModuleManager.cpp
@@ -1210,133 +1210,6 @@ bool FModuleManager::DoesLoadedModuleHaveUObjects( const FName ModuleName ) cons
 	return false;
 }
 
-FModuleManager::ModuleInfoRef FModuleManager::GetOrCreateModule(FName InModuleName)
-{
-	check(IsInGameThread());
-	ensureMsgf(InModuleName != NAME_None, TEXT("FModuleManager::AddModule() was called with an invalid module name (empty string or 'None'.)  This is not allowed."));
-
-	if (Modules.Contains(InModuleName))
-	{
-		return FindModuleChecked(InModuleName);
-	}
-
-	// Add this module to the set of modules that we know about
-	ModuleInfoRef ModuleInfo(new FModuleInfo());
-
-#if !IS_MONOLITHIC
-	FString ModuleNameString = InModuleName.ToString();
-
-	TMap<FName, FString> ModulePathMap;
-	FindModulePaths(*ModuleNameString, ModulePathMap);
-
-	if (ModulePathMap.Num() != 1)
-	{
-		return ModuleInfo;
-	}
-
-	// Add this module to the set of modules that we know about
-	ModuleInfo->OriginalFilename = TMap<FName, FString>::TConstIterator(ModulePathMap).Value();
-	ModuleInfo->Filename = ModuleInfo->OriginalFilename;
-
-	// When iterating on code during development, it's possible there are multiple rolling versions of this
-	// module's DLL file.  This can happen if the programmer is recompiling DLLs while the game is loaded.  In
-	// this case, we want to load the newest iteration of the DLL file, so that behavior is the same after
-	// restarting the application.
-
-	// NOTE: We leave this enabled in UE_BUILD_SHIPPING editor builds so module authors can iterate on custom modules
-#if !(UE_BUILD_SHIPPING || UE_BUILD_TEST) || (UE_BUILD_SHIPPING && WITH_EDITOR)
-	// In some cases, sadly, modules may be loaded before appInit() is called.  We can't cleanly support rolling files for those modules.
-
-	// First, check to see if the module we added already exists on disk
-	const FDateTime OriginalModuleFileTime = IFileManager::Get().GetTimeStamp(*ModuleInfo->OriginalFilename);
-	if (OriginalModuleFileTime == FDateTime::MinValue())
-	{
-		return ModuleInfo;
-	}
-
-	const FString ModuleName = InModuleName.ToString();
-	const int32 MatchPos = ModuleInfo->OriginalFilename.Find(ModuleName, ESearchCase::IgnoreCase, ESearchDir::FromEnd);
-	if (!ensureMsgf(MatchPos != INDEX_NONE, TEXT("Could not find module name '%s' in module filename '%s'"), *ModuleName, *ModuleInfo->OriginalFilename))
-	{
-		return ModuleInfo;
-	}
-
-	const int32 SuffixPos = MatchPos + ModuleName.Len();
-
-	const FString Prefix = ModuleInfo->OriginalFilename.Left(SuffixPos);
-	const FString Suffix = ModuleInfo->OriginalFilename.Right(ModuleInfo->OriginalFilename.Len() - SuffixPos);
-
-	const FString ModuleFileSearchString = FString::Printf(TEXT("%s-*%s"), *Prefix, *Suffix);
-	const FString ModuleFileSearchDirectory = FPaths::GetPath(ModuleFileSearchString);
-
-	// Search for module files
-	TArray<FString> FoundFiles;
-	IFileManager::Get().FindFiles(FoundFiles, *ModuleFileSearchString, true, false);
-
-	if (FoundFiles.Num() == 0)
-	{
-		return ModuleInfo;
-	}
-
-	// Figure out what the newest module file is
-	int32 NewestFoundFileIndex = INDEX_NONE;
-	FDateTime NewestFoundFileTime = OriginalModuleFileTime;
-
-	for (int32 CurFileIndex = 0; CurFileIndex < FoundFiles.Num(); ++CurFileIndex)
-	{
-		// FoundFiles contains file names with no directory information, but we need the full path up
-		// to the file, so we'll prefix it back on if we have a path.
-		const FString& FoundFile = FoundFiles[CurFileIndex];
-		const FString FoundFilePath = ModuleFileSearchDirectory.IsEmpty() ? FoundFile : (ModuleFileSearchDirectory / FoundFile);
-
-		// need to reject some files here that are not numbered...release executables, do have a suffix, so we need to make sure we don't find the debug version
-		check(FoundFilePath.Len() > Prefix.Len() + Suffix.Len());
-		FString Center = FoundFilePath.Mid(Prefix.Len(), FoundFilePath.Len() - Prefix.Len() - Suffix.Len());
-		check(Center.StartsWith(TEXT("-"), ESearchCase::CaseSensitive)); // a minus sign is still considered numeric, so we can leave it.
-		if (!Center.IsNumeric())
-		{
-			// this is a debug DLL or something, it is not a numbered hot DLL
-			continue;
-		}
-
-
-		// Check the time stamp for this file
-		const FDateTime FoundFileTime = IFileManager::Get().GetTimeStamp(*FoundFilePath);
-		if (ensure(FoundFileTime != -1.0))
-		{
-			// Was this file modified more recently than our others?
-			if (FoundFileTime > NewestFoundFileTime)
-			{
-				NewestFoundFileIndex = CurFileIndex;
-				NewestFoundFileTime = FoundFileTime;
-			}
-		}
-		else
-		{
-			// File wasn't found, should never happen as we searched for these files just now
-		}
-	}
-
-
-	// Did we find a variant of the module file that is newer than our original file?
-	if (NewestFoundFileIndex != INDEX_NONE)
-	{
-		// Update the module working file name to the most recently-modified copy of that module
-		const FString& NewestModuleFilename = FoundFiles[NewestFoundFileIndex];
-		const FString NewestModuleFilePath = ModuleFileSearchDirectory.IsEmpty() ? NewestModuleFilename : (ModuleFileSearchDirectory / NewestModuleFilename);
-		ModuleInfo->Filename = NewestModuleFilePath;
-	}
-	else
-	{
-		// No variants were found that were newer than the original module file name, so
-		// we'll continue to use that!
-	}
-#endif	// !(UE_BUILD_SHIPPING || UE_BUILD_TEST)
-#endif	// !IS_MONOLITHIC
-
-	return ModuleInfo;
-}
-
 int32 FModuleManager::GetModuleCount() const
 {
 	// Theoretically thread safe but by the time we return new modules could've been added
diff --git a/Engine/Source/Runtime/Core/Public/Modules/ModuleManager.h b/Engine/Source/Runtime/Core/Public/Modules/ModuleManager.h
--- a/Engine/Source/Runtime/Core/Public/Modules/ModuleManager.h
+++ b/Engine/Source/Runtime/Core/Public/Modules/ModuleManager.h
@@ -566,9 +566,6 @@ private:
 #endif
 
 private:
-	/** Gets module with given name from Modules or creates a new one. Doesn't modify Modules. */
-	ModuleInfoRef GetOrCreateModule(FName InModuleName);
-
 	/** Map of all modules.  Maps the case-insensitive module name to information about that module, loaded or not. */
 	FModuleMap Modules;
 
