Rely on the module manifest telling us the path to the latest version of a module, rather than guessing from timestamps of files in the same directory.

diff --git a/Engine/Source/Runtime/Core/Private/Modules/ModuleManager.cpp b/Engine/Source/Runtime/Core/Private/Modules/ModuleManager.cpp
--- a/Engine/Source/Runtime/Core/Private/Modules/ModuleManager.cpp
+++ b/Engine/Source/Runtime/Core/Private/Modules/ModuleManager.cpp
@@ -306,51 +306,7 @@ void FModuleManager::AddModule(const FName InModuleName)
 
 	// Add this module to the set of modules that we know about
 	ModuleInfo->OriginalFilename = Prefix + Suffix;
-	ModuleInfo->Filename         = ModuleInfo->OriginalFilename;
-
-	// When iterating on code during development, it's possible there are multiple rolling versions of this
-	// module's DLL file.  This can happen if the programmer is recompiling DLLs while the game is loaded.  In
-	// this case, we want to load the newest iteration of the DLL file, so that behavior is the same after
-	// restarting the application.
-
-	// NOTE: We leave this enabled in UE_BUILD_SHIPPING editor builds so module authors can iterate on custom modules
-#if !(UE_BUILD_SHIPPING || UE_BUILD_TEST) || (UE_BUILD_SHIPPING && WITH_EDITOR)
-	// In some cases, sadly, modules may be loaded before appInit() is called.  We can't cleanly support rolling files for those modules.
-
-	// First, check to see if the module we added already exists on disk
-	const FDateTime OriginalModuleFileTime = IFileManager::Get().GetTimeStamp(*ModuleInfo->OriginalFilename);
-	if (OriginalModuleFileTime == FDateTime::MinValue())
-	{
-		return;
-	}
-
-	const FString ModuleFileSearchString = FString::Printf(TEXT("%s-*%s"), *Prefix, *Suffix);
-
-	// Search for module files
-	TArray<FString> FoundFiles;
-	IFileManager::Get().FindFiles(FoundFiles, *ModuleFileSearchString, true, false);
-	if (FoundFiles.Num() == 0)
-	{
-		return;
-	}
-
-	const FString ModuleFileSearchDirectory = FPaths::GetPath(ModuleFileSearchString);
-
-	FString NewestModuleFilename;
-	bool bFoundNewestFile = FindNewestModuleFile(FoundFiles, OriginalModuleFileTime, ModuleFileSearchDirectory, Prefix, Suffix, NewestModuleFilename);
-
-	// Did we find a variant of the module file that is newer than our original file?
-	if (!bFoundNewestFile)
-	{
-		// No variants were found that were newer than the original module file name, so
-		// we'll continue to use that!
-		return;
-	}
-
-	// Update the module working file name to the most recently-modified copy of that module
-	const FString NewestModuleFilePath = ModuleFileSearchDirectory.IsEmpty() ? NewestModuleFilename : (ModuleFileSearchDirectory / NewestModuleFilename);
-	ModuleInfo->Filename = NewestModuleFilePath;
-#endif	// !(UE_BUILD_SHIPPING || UE_BUILD_TEST)
+	ModuleInfo->Filename         = ModuleFilename;
 #endif	// !IS_MONOLITHIC
 }
 
