- Fixed static analysis warninbg

diff --git a/Engine/Source/Runtime/Core/Private/Misc/ConfigCacheIni.cpp b/Engine/Source/Runtime/Core/Private/Misc/ConfigCacheIni.cpp
--- a/Engine/Source/Runtime/Core/Private/Misc/ConfigCacheIni.cpp
+++ b/Engine/Source/Runtime/Core/Private/Misc/ConfigCacheIni.cpp
@@ -3286,7 +3286,11 @@ static void GetSourceIniHierarchyFilenames(const TCHAR* InBaseIniName, const TCH
 		LayerPath = LayerPath.Replace(TEXT("{USER}"), FPlatformProcess::UserDir(), ESearchCase::CaseSensitive);
 
 		// PROGRAMs don't require any ini files
-		const bool bIsRequired = (!IS_PROGRAM) && ((Layer.Flag & Flag_Required) != 0) && (EngineConfigDir == FPaths::EngineConfigDir());
+#if IS_PROGRAM
+		const bool bIsRequired = false;
+#else
+		const bool bIsRequired = ((Layer.Flag & Flag_Required) != 0) && (EngineConfigDir == FPaths::EngineConfigDir());
+#endif
 
 		// expand if it it has {ED} or {EF} expansion tags
 		if (FCString::Strstr(*Layer.Path, TEXT("{ED}")) || FCString::Strstr(*Layer.Path, TEXT("{EF}")))
