- Send all Apple logs to Console, not just LowLevelOutputDebugString

diff --git a/Engine/Source/Runtime/Core/Private/Apple/ApplePlatformMisc.cpp b/Engine/Source/Runtime/Core/Private/Apple/ApplePlatformMisc.cpp
--- a/Engine/Source/Runtime/Core/Private/Apple/ApplePlatformMisc.cpp
+++ b/Engine/Source/Runtime/Core/Private/Apple/ApplePlatformMisc.cpp
@@ -48,7 +48,7 @@ FString FApplePlatformMisc::GetEnvironmentVariable(const TCHAR* VariableName)
 	}
 }
 
-void FApplePlatformMisc::LowLevelOutputDebugString( const TCHAR *Message )
+void FApplePlatformMisc::LocalPrint(const TCHAR* Message)
 {
 	//NsLog will out to all iOS output consoles, instead of just the Xcode console.
 	NSLog(@"[UE4] %s", TCHAR_TO_UTF8(Message));
diff --git a/Engine/Source/Runtime/Core/Public/Apple/ApplePlatformMisc.h b/Engine/Source/Runtime/Core/Public/Apple/ApplePlatformMisc.h
--- a/Engine/Source/Runtime/Core/Public/Apple/ApplePlatformMisc.h
+++ b/Engine/Source/Runtime/Core/Public/Apple/ApplePlatformMisc.h
@@ -79,7 +79,7 @@ struct CORE_API FApplePlatformMisc : public FGenericPlatformMisc
 		__sync_synchronize();
 	}
 
-	static void LowLevelOutputDebugString(const TCHAR *Message);
+	static void LocalPrint(const TCHAR* Message);
 	static const TCHAR* GetSystemErrorMessage(TCHAR* OutBuffer, int32 BufferCount, int32 Error);
 	static int32 NumberOfCores();
 	
