- Merged fix for random-ish full brightness bug

diff --git a/Engine/Source/Runtime/Core/Private/IOS/IOSPlatformMisc.cpp b/Engine/Source/Runtime/Core/Private/IOS/IOSPlatformMisc.cpp
--- a/Engine/Source/Runtime/Core/Private/IOS/IOSPlatformMisc.cpp
+++ b/Engine/Source/Runtime/Core/Private/IOS/IOSPlatformMisc.cpp
@@ -70,7 +70,7 @@ void (* GMemoryWarningHandler)(const FGenericMemoryWarningContext& Context) = NU
 
 /** global for showing the splash screen */
 bool GShowSplashScreen = true;
-float GOriginalBrightness = 1.0f;
+float GOriginalBrightness = -1.0f;
 
 static int32 GetFreeMemoryMB()
 {
@@ -262,7 +262,10 @@ void FIOSPlatformMisc::SetBrightness(float Brightness)
 
 void FIOSPlatformMisc::ResetBrightness()
 {
-    SetBrightness(GOriginalBrightness);
+	if (GOriginalBrightness >= 0.f)
+	{
+		SetBrightness(GOriginalBrightness);
+	}
 }
 
 bool FIOSPlatformMisc::IsRunningOnBattery()
