use IOS malloc for small blocks as it has an internal small block allocator which is always mapped to the VM.

diff --git a/Engine/Source/Runtime/Core/Private/HAL/MallocBinned.cpp b/Engine/Source/Runtime/Core/Private/HAL/MallocBinned.cpp
--- a/Engine/Source/Runtime/Core/Private/HAL/MallocBinned.cpp
+++ b/Engine/Source/Runtime/Core/Private/HAL/MallocBinned.cpp
@@ -136,16 +136,31 @@ struct FMallocBinned::PoolHashBucket
 	}
 };
 
+#if PLATFORM_IOS
+#define PLAT_PAGE_SIZE_LIMIT 16384
+#define PLAT_BINNED_ALLOC_POOLSIZE 16384
+#define PLAT_SMALL_BLOCK_POOL_SIZE 224
+#else
+#define PLAT_PAGE_SIZE_LIMIT 65536
+#define PLAT_BINNED_ALLOC_POOLSIZE 65536
+#define PLAT_SMALL_BLOCK_POOL_SIZE 0
+#endif
+
 struct FMallocBinned::Private
 {
 	/** Default alignment for binned allocator */
 	enum { DEFAULT_BINNED_ALLOCATOR_ALIGNMENT = sizeof(FFreeMem) };
-	enum { PAGE_SIZE_LIMIT = 65536 };
+	enum { PAGE_SIZE_LIMIT = PLAT_PAGE_SIZE_LIMIT };
 	// BINNED_ALLOC_POOL_SIZE can be increased beyond 64k to cause binned malloc to allocate
 	// the small size bins in bigger chunks. If OS Allocation is slow, increasing
 	// this number *may* help performance but YMMV.
-	enum { BINNED_ALLOC_POOL_SIZE = 65536 }; 
+	enum { BINNED_ALLOC_POOL_SIZE = PLAT_BINNED_ALLOC_POOLSIZE };
+	// On IOS can push small allocs in to a pre-allocated small block pool
+	enum { SMALL_BLOCK_POOL_SIZE = PLAT_SMALL_BLOCK_POOL_SIZE };
 
+	static uint64 SmallBlockStartPtr;
+	static uint64 SmallBlockEndPtr;
+	
 	// Implementation. 
 	static CA_NO_RETURN void OutOfMemory(uint64 Size, uint32 Alignment=0)
 	{
@@ -436,6 +451,12 @@ struct FMallocBinned::Private
 #if PLATFORM_IOS || PLATFORM_MAC
         if (Pool == NULL)
         {
+			// check the small block allocation range
+			if ((uint64)Ptr >= SmallBlockStartPtr && (uint64)Ptr <= SmallBlockEndPtr)
+			{
+				SmallOSFree(Allocator, Ptr, Private::SMALL_BLOCK_POOL_SIZE);
+				return;
+			}
             UE_LOG(LogMemory, Warning, TEXT("Attempting to free a pointer we didn't allocate!"));
             return;
         }
@@ -580,6 +601,15 @@ struct FMallocBinned::Private
 #endif
 	}
 
+	static FORCEINLINE void SmallOSFree(FMallocBinned& Allocator, void* Ptr, SIZE_T Size)
+	{
+#if PLATFORM_IOS
+		::free(Ptr);
+#else
+		FPlatformMemory::BinnedFreeToOS(Ptr, Size);
+#endif
+	}
+	
 	static FORCEINLINE void* OSAlloc(FMallocBinned& Allocator, SIZE_T NewSize, SIZE_T& OutActualSize)
 	{
 #ifdef CACHE_FREED_OS_ALLOCS
@@ -643,6 +673,24 @@ struct FMallocBinned::Private
 #endif
 	}
 
+	static FORCEINLINE void* SmallOSAlloc(FMallocBinned& Allocator, SIZE_T NewSize, SIZE_T& OutActualSize)
+	{
+#if PLATFORM_IOS
+		(void)OutActualSize;
+		LLM_PLATFORM_SCOPE(ELLMTag::FMalloc);
+		void* Ptr = ::malloc(NewSize);
+		if (Ptr == nullptr)
+		{
+			UE_LOG(LogTemp, Warning, TEXT("malloc failure allocating %d, error code: %d"), NewSize, errno);
+		}
+		return Ptr;
+#else
+		(void)OutActualSize;
+		LLM_PLATFORM_SCOPE(ELLMTag::FMalloc);
+		return FPlatformMemory::BinnedAllocFromOS(NewSize);
+#endif
+	}
+	
 #ifdef CACHE_FREED_OS_ALLOCS
 	static void FlushAllocCache(FMallocBinned& Allocator)
 	{
@@ -679,6 +727,9 @@ struct FMallocBinned::Private
 	}
 };
 
+uint64 FMallocBinned::Private::SmallBlockStartPtr = 0xffffffffffffffff;
+uint64 FMallocBinned::Private::SmallBlockEndPtr = 0x0000000000000000;
+
 void FMallocBinned::GetAllocatorStats( FGenericMemoryStats& out_Stats )
 {
 	FMalloc::GetAllocatorStats( out_Stats );
@@ -885,7 +936,37 @@ void* FMallocBinned::Malloc(SIZE_T Size, uint32 Alignment)
 	BINNED_INCREMENT_STATCOUNTER(TotalAllocs);
 	
 	FFreeMem* Free;
-	if( Size < BinnedSizeLimit )
+	bool bUsePools = true;
+	if (Size <= Private::SMALL_BLOCK_POOL_SIZE)
+	{
+		UPTRINT AlignedSize = Align(Size, Alignment);
+		SIZE_T ActualPoolSize; //TODO: use this to reduce waste?
+		Free = (FFreeMem*)Private::SmallOSAlloc(*this, AlignedSize, ActualPoolSize);
+		
+		if ((uint64)Free < 0x280000000)
+		{
+			// if this happens then we need to fall in to the bins below
+//			FPlatformMisc::LowLevelOutputDebugStringf(TEXT("Overflowed!!! Aligned Size:%d"), AlignedSize);
+			Private::SmallOSFree(*this, Free, AlignedSize);
+			Free = nullptr;
+			bUsePools = true;
+		}
+		else
+		{
+			if ((uint64)Free < Private::SmallBlockStartPtr)
+			{
+				Private::SmallBlockStartPtr = (uint64)Free;
+			}
+			if ((uint64)Free > Private::SmallBlockEndPtr)
+			{
+				Private::SmallBlockEndPtr = (uint64)Free;
+			}
+			bUsePools = false;
+		}
+	}
+	if (bUsePools)
+	{
+	if( Size < BinnedSizeLimit)
 	{
 		// Allocate from pool.
 		FPoolTable* Table = MemSizeToPoolTable[Size];
@@ -965,6 +1046,7 @@ void* FMallocBinned::Malloc(SIZE_T Size, uint32 Alignment)
 		BINNED_PEAK_STATCOUNTER(UsedPeak, BINNED_ADD_STATCOUNTER(UsedCurrent, Size));
 		BINNED_PEAK_STATCOUNTER(WastePeak, BINNED_ADD_STATCOUNTER(WasteCurrent, (int64)(AlignedSize - Size)));
 	}
+	}
 
 	MEM_TIME(MemTime += FPlatformTime::Seconds());
 	return Free;
@@ -990,6 +1072,16 @@ void* FMallocBinned::Realloc( void* Ptr, SIZE_T NewSize, uint32 Alignment )
 	void* NewPtr = Ptr;
 	if( Ptr && NewSize )
 	{
+		// special case for really small blocks
+		if ((uint64)Ptr >= Private::SmallBlockStartPtr && (uint64)Ptr <= Private::SmallBlockEndPtr)
+		{
+			// Grow or shrink.
+			NewPtr = Malloc(NewSizeUnmodified, Alignment);
+			FMemory::Memcpy(NewPtr, Ptr, FMath::Min<SIZE_T>(NewSizeUnmodified, Private::SMALL_BLOCK_POOL_SIZE)); // this is bad need to figure out the old size :(
+			Free( Ptr );
+		}
+		else
+		{
 		FPoolInfo* Pool = Private::FindPoolInfo(*this, (UPTRINT)Ptr, BasePtr);
 
 		if( Pool->TableIndex < BinnedOSTableIndex )
@@ -1033,6 +1125,7 @@ void* FMallocBinned::Realloc( void* Ptr, SIZE_T NewSize, uint32 Alignment )
 				Pool->SetAllocationSizes(NewSizeUnmodified, Pool->GetOsBytes(PageSize, BinnedOSTableIndex), BinnedOSTableIndex, BinnedOSTableIndex);
 			}
 		}
+		}
 	}
 	else if( Ptr == nullptr )
 	{
@@ -1071,6 +1164,13 @@ bool FMallocBinned::GetAllocationSize(void *Original, SIZE_T &SizeOut)
 #if PLATFORM_IOS || PLATFORM_MAC
 	if (Pool == NULL)
 	{
+#if PLATFORM_IOS
+		if ((uint64)Original >= Private::SmallBlockStartPtr && (uint64)Original <= Private::SmallBlockEndPtr)
+		{
+			SizeOut = Private::SMALL_BLOCK_POOL_SIZE;
+			return true;
+		}
+#endif
 		UE_LOG(LogMemory, Warning, TEXT("Attempting to access memory pool info for a pointer we didn't allocate!"));
 		return false;
 	}
@@ -1110,7 +1210,12 @@ SIZE_T FMallocBinned::QuantizeSize(SIZE_T Size, uint32 Alignment)
 	Size = FMath::Max<SIZE_T>(PoolTable[0].BlockSize, Size + (Alignment - SpareBytesCount));
 
 	SIZE_T Result;
-	if (Size < BinnedSizeLimit)
+	if (Size <= Private::SMALL_BLOCK_POOL_SIZE)
+	{
+		UPTRINT AlignedSize = Align(Size, Alignment);
+		Result = SIZE_T(AlignedSize);
+	}
+	else if (Size < BinnedSizeLimit)
 	{
 		// Allocate from pool.
 		FPoolTable* Table = MemSizeToPoolTable[Size];
