reduce thread sizes on IOS a little more

diff --git a/Engine/Source/Runtime/Core/Private/IOS/IOSPlatformMisc.cpp b/Engine/Source/Runtime/Core/Private/IOS/IOSPlatformMisc.cpp
--- a/Engine/Source/Runtime/Core/Private/IOS/IOSPlatformMisc.cpp
+++ b/Engine/Source/Runtime/Core/Private/IOS/IOSPlatformMisc.cpp
@@ -682,7 +682,7 @@ FIOSPlatformMisc::EIOSDevice FIOSPlatformMisc::GetIOSDeviceType()
 
 int FIOSPlatformMisc::GetDefaultStackSize()
 {
-	return 512 * 1024;
+	return 128 * 1024;
 }
 
 void FIOSPlatformMisc::SetMemoryWarningHandler(void (* InHandler)(const FGenericMemoryWarningContext& Context))
