crash was unrelated to these changes, turning back on small block allocator for IOS

diff --git a/Engine/Source/Runtime/Core/Private/HAL/MallocBinned.cpp b/Engine/Source/Runtime/Core/Private/HAL/MallocBinned.cpp
--- a/Engine/Source/Runtime/Core/Private/HAL/MallocBinned.cpp
+++ b/Engine/Source/Runtime/Core/Private/HAL/MallocBinned.cpp
@@ -139,7 +139,7 @@ struct FMallocBinned::PoolHashBucket
 #if PLATFORM_IOS
 #define PLAT_PAGE_SIZE_LIMIT 16384
 #define PLAT_BINNED_ALLOC_POOLSIZE 16384
-#define PLAT_SMALL_BLOCK_POOL_SIZE 0
+#define PLAT_SMALL_BLOCK_POOL_SIZE 224
 #else
 #define PLAT_PAGE_SIZE_LIMIT 65536
 #define PLAT_BINNED_ALLOC_POOLSIZE 65536
